fix misnamed variable CurrentPlatform should be TargetPlatform
behavior nop

diff --git a/Engine/Source/Runtime/Engine/Private/TextureDerivedData.cpp b/Engine/Source/Runtime/Engine/Private/TextureDerivedData.cpp
--- a/Engine/Source/Runtime/Engine/Private/TextureDerivedData.cpp
+++ b/Engine/Source/Runtime/Engine/Private/TextureDerivedData.cpp
@@ -664,15 +664,14 @@ static FName ConditionalRemapOodleTextureSdkVersion(FName InOodleTextureSdkVersi
 static void GetTextureBuildSettings(
 	const UTexture& Texture,
 	const UTextureLODSettings& TextureLODSettings,
-	const ITargetPlatform& CurrentPlatform,
+	const ITargetPlatform& TargetPlatform,
 	ETextureEncodeSpeed InEncodeSpeed, // must be Final or Fast.
 	FTextureBuildSettings& OutBuildSettings,
 	FTexturePlatformData::FTextureEncodeResultMetadata* OutBuildResultMetadata // can be nullptr if not needed
 	)
 {
-	// pretty sure the "CurrentPlatform" here is actually "TargetPlatform"
-	const bool bPlatformSupportsTextureStreaming = CurrentPlatform.SupportsFeature(ETargetPlatformFeatures::TextureStreaming);
-	const bool bPlatformSupportsVirtualTextureStreaming = CurrentPlatform.SupportsFeature(ETargetPlatformFeatures::VirtualTextureStreaming);
+	const bool bPlatformSupportsTextureStreaming = TargetPlatform.SupportsFeature(ETargetPlatformFeatures::TextureStreaming);
+	const bool bPlatformSupportsVirtualTextureStreaming = TargetPlatform.SupportsFeature(ETargetPlatformFeatures::VirtualTextureStreaming);
 
 	if (OutBuildResultMetadata)
 	{
@@ -764,7 +763,7 @@ static void GetTextureBuildSettings(
 	OutBuildSettings.CompositePower = Texture.CompositePower;
 	OutBuildSettings.LODBias = TextureLODSettings.CalculateLODBias(SourceSize.X, SourceSize.Y, Texture.MaxTextureSize, Texture.LODGroup, Texture.LODBias, Texture.NumCinematicMipLevels, Texture.MipGenSettings, bVirtualTextureStreaming);
 	OutBuildSettings.LODBiasWithCinematicMips = TextureLODSettings.CalculateLODBias(SourceSize.X, SourceSize.Y, Texture.MaxTextureSize, Texture.LODGroup, Texture.LODBias, 0, Texture.MipGenSettings, bVirtualTextureStreaming);
-	OutBuildSettings.bStreamable = GetTextureIsStreamableOnPlatform(Texture, CurrentPlatform);
+	OutBuildSettings.bStreamable = GetTextureIsStreamableOnPlatform(Texture, TargetPlatform);
 	OutBuildSettings.bVirtualStreamable = bVirtualTextureStreaming;
 	OutBuildSettings.PowerOfTwoMode = Texture.PowerOfTwoMode;
 	OutBuildSettings.PaddingColor = Texture.PaddingColor;
@@ -774,7 +773,7 @@ static void GetTextureBuildSettings(
 	OutBuildSettings.CompressionQuality = Texture.CompressionQuality - 1; // translate from enum's 0 .. 5 to desired compression (-1 .. 4, where -1 is default while 0 .. 4 are actual quality setting override)
 	
 	// do remap here before we send to TBW's which may not have access to config :
-	OutBuildSettings.OodleTextureSdkVersion = ConditionalRemapOodleTextureSdkVersion(Texture.OodleTextureSdkVersion,&CurrentPlatform);
+	OutBuildSettings.OodleTextureSdkVersion = ConditionalRemapOodleTextureSdkVersion(Texture.OodleTextureSdkVersion,&TargetPlatform);
 
 	// if LossyCompressionAmount is Default, inherit from LODGroup :
 	const FTextureLODGroup& LODGroup = TextureLODSettings.GetTextureLODGroup(Texture.LODGroup);
@@ -799,7 +798,7 @@ static void GetTextureBuildSettings(
 	if (MipGenSettings == TMGS_NoMipmaps && 
 		Texture.IsA(UTexture2D::StaticClass()))	// TODO: support more texture types
 	{
-		TextureLODSettings.GetDownscaleOptions(Texture, CurrentPlatform, OutBuildSettings.Downscale, (ETextureDownscaleOptions&)OutBuildSettings.DownscaleOptions);
+		TextureLODSettings.GetDownscaleOptions(Texture, TargetPlatform, OutBuildSettings.Downscale, (ETextureDownscaleOptions&)OutBuildSettings.DownscaleOptions);
 	}
 	
 	// For virtual texturing we take the address mode into consideration
@@ -851,7 +850,7 @@ static void GetTextureBuildSettings(
 	}
 
 	// By default, initialize settings for layer0
-	FinalizeBuildSettingsForLayer(Texture, 0, &CurrentPlatform, InEncodeSpeed, OutBuildSettings, OutBuildResultMetadata);
+	FinalizeBuildSettingsForLayer(Texture, 0, &TargetPlatform, InEncodeSpeed, OutBuildSettings, OutBuildResultMetadata);
 }
 
 /**
@@ -870,31 +869,31 @@ static void GetBuildSettingsForRunningPlatform(
 	ITargetPlatformManagerModule* TPM = GetTargetPlatformManager();
 	if (TPM)
 	{
-		ITargetPlatform* CurrentPlatform = NULL;
+		ITargetPlatform* TargetPlatform = NULL;
 		const TArray<ITargetPlatform*>& Platforms = TPM->GetActiveTargetPlatforms();
 
 		check(Platforms.Num());
 
-		CurrentPlatform = Platforms[0];
+		TargetPlatform = Platforms[0];
 
 		for (int32 Index = 1; Index < Platforms.Num(); Index++)
 		{
 			if (Platforms[Index]->IsRunningPlatform())
 			{
-				CurrentPlatform = Platforms[Index];
+				TargetPlatform = Platforms[Index];
 				break;
 			}
 		}
 
-		check(CurrentPlatform != NULL);
+		check(TargetPlatform != NULL);
 
-		const UTextureLODSettings* LODSettings = (UTextureLODSettings*)UDeviceProfileManager::Get().FindProfile(CurrentPlatform->PlatformName());
+		const UTextureLODSettings* LODSettings = (UTextureLODSettings*)UDeviceProfileManager::Get().FindProfile(TargetPlatform->PlatformName());
 		FTextureBuildSettings SourceBuildSettings;
 		FTexturePlatformData::FTextureEncodeResultMetadata SourceMetadata;
-		GetTextureBuildSettings(Texture, *LODSettings, *CurrentPlatform, InEncodeSpeed, SourceBuildSettings, &SourceMetadata);
+		GetTextureBuildSettings(Texture, *LODSettings, *TargetPlatform, InEncodeSpeed, SourceBuildSettings, &SourceMetadata);
 
 		TArray< TArray<FName> > PlatformFormats;
-		CurrentPlatform->GetTextureFormats(&Texture, PlatformFormats);
+		TargetPlatform->GetTextureFormats(&Texture, PlatformFormats);
 		check(PlatformFormats.Num() > 0);
 
 		const int32 NumLayers = Texture.Source.GetNumLayers();
@@ -916,7 +915,7 @@ static void GetBuildSettingsForRunningPlatform(
 				OutMetadata = &OutResultMetadataPerLayer->Add_GetRef(SourceMetadata);
 			}
 			
-			FinalizeBuildSettingsForLayer(Texture, LayerIndex, CurrentPlatform, InEncodeSpeed, OutSettings, OutMetadata);
+			FinalizeBuildSettingsForLayer(Texture, LayerIndex, TargetPlatform, InEncodeSpeed, OutSettings, OutMetadata);
 		}
 	}
 }
