Chaos - fix issue where constraint creation is skipped if the attached bodies are kinematic
- set p.AllowKinematicKinematicConstraints to 1 by default

diff --git a/Engine/Source/Runtime/Engine/Private/PhysicsEngine/ConstraintInstance.cpp b/Engine/Source/Runtime/Engine/Private/PhysicsEngine/ConstraintInstance.cpp
--- a/Engine/Source/Runtime/Engine/Private/PhysicsEngine/ConstraintInstance.cpp
+++ b/Engine/Source/Runtime/Engine/Private/PhysicsEngine/ConstraintInstance.cpp
@@ -52,7 +52,8 @@ TAutoConsoleVariable<float> CVarConstraintAngularStiffnessScale(
 bool bEnableSkeletalMeshConstraints = true;
 FAutoConsoleVariableRef CVarEnableSkeletalMeshConstraints(TEXT("p.EnableSkeletalMeshConstraints"), bEnableSkeletalMeshConstraints, TEXT("Enable skeletal mesh constraints defined within the Physics Asset Editor"));
 
-bool bAllowKinematicKinematicConstraints = false;
+// warning : setting the value to false can have negative impact on ragdoll, if they start in a kinematic state and then have their bodies turned dynamic
+bool bAllowKinematicKinematicConstraints = true;
 FAutoConsoleVariableRef CVarAllowKinematicKinematicConstraints(TEXT("p.AllowKinematicKinematicConstraints"), bAllowKinematicKinematicConstraints, TEXT("Do not create constraints between two rigid kinematics."));
 
 /** Handy macro for setting BIT of VAR based on the bool CONDITION */
