Ensure registration of engine-level parameter interfaces by MetaSoundEngine Module and split out from loading default settings call in AudioSettings (can happen on app load if cooking and AudioDeviceManager isn't loaded).
#rnx

diff --git a/Engine/Plugins/Runtime/Metasound/Source/MetasoundEngine/Private/MetasoundEngineModule.cpp b/Engine/Plugins/Runtime/Metasound/Source/MetasoundEngine/Private/MetasoundEngineModule.cpp
--- a/Engine/Plugins/Runtime/Metasound/Source/MetasoundEngine/Private/MetasoundEngineModule.cpp
+++ b/Engine/Plugins/Runtime/Metasound/Source/MetasoundEngine/Private/MetasoundEngineModule.cpp
@@ -10,6 +10,7 @@
 #include "MetasoundUObjectRegistry.h"
 #include "MetasoundWave.h"
 #include "Modules/ModuleManager.h"
+#include "Sound/AudioSettings.h"
 
 DEFINE_LOG_CATEGORY(LogMetasoundEngine);
 
@@ -25,6 +26,12 @@ class FMetasoundEngineModule : public IMetasoundEngineModule
 		FModuleManager::Get().LoadModuleChecked("MetasoundGenerator");
 		FModuleManager::Get().LoadModuleChecked("AudioCodecEngine");
 
+		// Register engine-level parameter interfaces if not done already.
+		// (Potentially not already called if plugin is loaded while cooking.)
+		UAudioSettings* AudioSettings = GetMutableDefault<UAudioSettings>();
+		check(AudioSettings);
+		AudioSettings->RegisterParameterInterfaces();
+
 		// Register interfaces
 		Metasound::Engine::RegisterInterfaces();
 
diff --git a/Engine/Source/Runtime/Engine/Classes/Sound/AudioSettings.h b/Engine/Source/Runtime/Engine/Classes/Sound/AudioSettings.h
--- a/Engine/Source/Runtime/Engine/Classes/Sound/AudioSettings.h
+++ b/Engine/Source/Runtime/Engine/Classes/Sound/AudioSettings.h
@@ -234,12 +234,19 @@ private:
 	TObjectPtr<USoundConcurrency> DefaultSoundConcurrency;
 
 public:
+	// Loads default object instances from soft object path properties
 	void LoadDefaultObjects();
 
 	USoundClass* GetDefaultSoundClass() const;
 	USoundClass* GetDefaultMediaSoundClass() const;
 	USoundConcurrency* GetDefaultSoundConcurrency() const;
 
+	// Registers Parameter Interfaces defined by the engine module.
+	// Called on engine start-up. Can be called when engine is not
+	// initialized as well by consuming plugins (ex. on cook by plugins
+	// requiring interface system to be loaded).
+	void RegisterParameterInterfaces();
+
 	// Get the quality level settings at the provided level index
 	const FAudioQualitySettings& GetQualityLevelSettings(int32 QualityLevel) const;
 	
@@ -273,6 +280,8 @@ private:
 
 	void AddDefaultSettings();
 
+	bool bParameterInterfacesRegistered;
+
 	// Whether or not the audio mixer is loaded/enabled. Used to toggle visibility of editor features.
 	bool bIsAudioMixerEnabled;
 };
diff --git a/Engine/Source/Runtime/Engine/Private/AudioDeviceManager.cpp b/Engine/Source/Runtime/Engine/Private/AudioDeviceManager.cpp
--- a/Engine/Source/Runtime/Engine/Private/AudioDeviceManager.cpp
+++ b/Engine/Source/Runtime/Engine/Private/AudioDeviceManager.cpp
@@ -542,6 +542,7 @@ bool FAudioDeviceManager::InitializeManager()
 		check(AudioSettings);
 
 		AudioSettings->LoadDefaultObjects();
+		AudioSettings->RegisterParameterInterfaces();
 
 		const bool bIsAudioMixerEnabled = AudioDeviceModule->IsAudioMixerModule();
 		AudioSettings->SetAudioMixerEnabled(bIsAudioMixerEnabled);
diff --git a/Engine/Source/Runtime/Engine/Private/AudioSettings.cpp b/Engine/Source/Runtime/Engine/Private/AudioSettings.cpp
--- a/Engine/Source/Runtime/Engine/Private/AudioSettings.cpp
+++ b/Engine/Source/Runtime/Engine/Private/AudioSettings.cpp
@@ -30,6 +30,7 @@ UAudioSettings::UAudioSettings(const FObjectInitializer& ObjectInitializer)
 	AddDefaultSettings();
 
 	bAllowPlayWhenSilent = true;
+	bParameterInterfacesRegistered = false;
 	bIsAudioMixerEnabled = false;
 
 	GlobalMinPitchScale = 0.4F;
@@ -259,10 +260,18 @@ void UAudioSettings::LoadDefaultObjects()
 	{
 		UE_LOG(LogAudio, Display, TEXT("No default SoundConcurrencyObject specified (or failed to load)."));
 	}
+}
 
-	Audio::IAudioParameterInterfaceRegistry& InterfaceRegistry = Audio::IAudioParameterInterfaceRegistry::Get();
-	InterfaceRegistry.RegisterInterface(Audio::AttenuationInterface::GetInterface());
-	InterfaceRegistry.RegisterInterface(Audio::SpatializationInterface::GetInterface());
+void UAudioSettings::RegisterParameterInterfaces()
+{
+	if (!bParameterInterfacesRegistered)
+	{
+		UE_LOG(LogAudio, Display, TEXT("Registering Engine Module Parameter Interfaces..."));
+		bParameterInterfacesRegistered = true;
+		Audio::IAudioParameterInterfaceRegistry& InterfaceRegistry = Audio::IAudioParameterInterfaceRegistry::Get();
+		InterfaceRegistry.RegisterInterface(Audio::AttenuationInterface::GetInterface());
+		InterfaceRegistry.RegisterInterface(Audio::SpatializationInterface::GetInterface());
+	}
 }
 
 USoundClass* UAudioSettings::GetDefaultMediaSoundClass() const
