PR #8286: [FEATURE] Added FindLocalLookAtRotation function to Kismet Math Library (Contributed by [REDACTED_PERSON])

diff --git a/Engine/Source/Runtime/Engine/Classes/Kismet/KismetMathLibrary.h b/Engine/Source/Runtime/Engine/Classes/Kismet/KismetMathLibrary.h
--- a/Engine/Source/Runtime/Engine/Classes/Kismet/KismetMathLibrary.h
+++ b/Engine/Source/Runtime/Engine/Classes/Kismet/KismetMathLibrary.h
@@ -2301,6 +2301,13 @@ class ENGINE_API UKismetMathLibrary : public UBlueprintFunctionLibrary
 	UFUNCTION(BlueprintPure, Category="Math|Rotator", meta=(Keywords="rotation rotate"))
 	static FRotator FindLookAtRotation(const FVector& Start, const FVector& Target);
 
+	/** 
+	 * Find a local rotation (range of [-180, 180]) for an object with StartTransform to point at TargetLocation. 
+	 * Useful for getting LookAt Azimuth or Pawn Aim Offset.
+	*/
+	UFUNCTION(BlueprintPure, Category = "Math|Rotator", meta = (Keywords = "rotation rotate local azimuth"))
+	static FRotator FindRelativeLookAtRotation(const FTransform& StartTransform, const FVector& TargetLocation);
+
 	/** Breaks apart a rotator into {Roll, Pitch, Yaw} angles in degrees */
 	UFUNCTION(BlueprintPure, Category = "Math|Rotator", meta = (Keywords = "rotation rotate rotator breakrotator", NativeBreakFunc))
 	static void BreakRotator(
diff --git a/Engine/Source/Runtime/Engine/Private/KismetMathLibrary.cpp b/Engine/Source/Runtime/Engine/Private/KismetMathLibrary.cpp
--- a/Engine/Source/Runtime/Engine/Private/KismetMathLibrary.cpp
+++ b/Engine/Source/Runtime/Engine/Private/KismetMathLibrary.cpp
@@ -968,6 +968,11 @@ FRotator UKismetMathLibrary::MakeRotationFromAxes(FVector Forward, FVector Right
 	return RotMatrix.Rotator();
 }
 
+FRotator UKismetMathLibrary::FindRelativeLookAtRotation(const FTransform& StartTransform, const FVector& TargetLocation)
+{
+	return NormalizedDeltaRotator(FindLookAtRotation(StartTransform.GetLocation(), TargetLocation), StartTransform.GetRotation().Rotator());
+}
+
 int32 UKismetMathLibrary::RandomIntegerFromStream(int32 Max, const FRandomStream& Stream)
 {
 	return Stream.RandHelper(Max);
