Refactor new HLSL analytic derivatives again
FExpressions are no longer responsible for emitting derivative code directly.  Instead, derivatives can create new expressions to compute derivative, when requested
This simplifies overall design of individual expressions.  Also allows analytic derivative code to benefit from constant folding, and preshaders

diff --git a/Engine/Source/Runtime/Engine/Private/HLSLTree/HLSLTree.cpp b/Engine/Source/Runtime/Engine/Private/HLSLTree/HLSLTree.cpp
--- a/Engine/Source/Runtime/Engine/Private/HLSLTree/HLSLTree.cpp
+++ b/Engine/Source/Runtime/Engine/Private/HLSLTree/HLSLTree.cpp
@@ -54,27 +54,6 @@ EExpressionEvaluation CombineEvaluations(EExpressionEvaluation Lhs, EExpressionE
 	return EExpressionEvaluation::Preshader;
 }
 
-EExpressionDerivative CombineDerivatives(EExpressionDerivative Lhs, EExpressionDerivative Rhs)
-{
-	if (Lhs == EExpressionDerivative::None)
-	{
-		return Rhs;
-	}
-	else if (Rhs == EExpressionDerivative::None)
-	{
-		return Lhs;
-	}
-	else if (Lhs == EExpressionDerivative::Invalid || Rhs == EExpressionDerivative::Invalid)
-	{
-		return EExpressionDerivative::Invalid;
-	}
-	else if (Lhs == EExpressionDerivative::Zero && Rhs == EExpressionDerivative::Zero)
-	{
-		return EExpressionDerivative::Zero;
-	}
-	return EExpressionDerivative::Valid;
-}
-
 FErrors::FErrors(FMemStackBase& InAllocator)
 	: Allocator(&InAllocator)
 {
@@ -317,7 +296,7 @@ FEmitShaderCode* FEmitContext::EmitPreshaderOrConstant(const FRequestedType& Req
 		const Shader::EValueType FieldType = Type.GetFlatFieldType(FieldIndex);
 		const Shader::FValueTypeDescription TypeDesc = Shader::GetValueTypeDescription(FieldType);
 		const int32 NumFieldComponents = TypeDesc.NumComponents;
-		const EExpressionEvaluation FieldEvaluation = Expression->GetPreparedType().GetFieldData(ComponentIndex, NumFieldComponents).Evaluation;
+		const EExpressionEvaluation FieldEvaluation = Expression->GetPreparedType().GetFieldEvaluation(ComponentIndex, NumFieldComponents);
 
 		if (FieldEvaluation == EExpressionEvaluation::Preshader)
 		{
@@ -552,6 +531,30 @@ FScope* FScope::FindSharedParent(FScope* Lhs, FScope* Rhs)
 	return Scope0;
 }
 
+FExpressionLocalPHI::FExpressionLocalPHI(const FExpressionLocalPHI* Source, EDerivativeCoordinate Coord)
+	: DerivativeChain(Source->DerivativeChain)
+	, LocalName(Source->LocalName)
+	, NumValues(Source->NumValues)
+{
+	DerivativeChain.Add(Coord);
+	for (int32 i = 0; i < NumValues; ++i)
+	{
+		Scopes[i] = Source->Scopes[i];
+	}
+}
+
+void FExpressionLocalPHI::ComputeAnalyticDerivatives(FTree& Tree, FExpressionDerivatives& OutResult) const
+{
+	// We don't have values assigned at the time analytic derivatives are computed
+	// It's possible the derivatives will be end up being invalid, but that case will need to be detected later, during PrepareValue
+	FExpressionLocalPHI* ExpressionDdx = Tree.NewExpression<FExpressionLocalPHI>(this, EDerivativeCoordinate::Ddx);
+	FExpressionLocalPHI* ExpressionDdy = Tree.NewExpression<FExpressionLocalPHI>(this, EDerivativeCoordinate::Ddy);
+	Tree.PHIExpressions.Add(ExpressionDdx);
+	Tree.PHIExpressions.Add(ExpressionDdy);
+	OutResult.ExpressionDdx = ExpressionDdx;
+	OutResult.ExpressionDdy = ExpressionDdy;
+}
+
 void FExpressionLocalPHI::PrepareValue(FEmitContext& Context, const FRequestedType& RequestedType, FPrepareValueResult& OutResult) const
 {
 	check(NumValues <= MaxNumPreviousScopes);
@@ -647,7 +650,7 @@ void FExpressionLocalPHI::PrepareValue(FEmitContext& Context, const FRequestedTy
 	}
 }
 
-void FExpressionLocalPHI::EmitValueShader(FEmitContext& Context, const FRequestedType& RequestedType, FEmitShaderValues& OutResult) const
+void FExpressionLocalPHI::EmitValueShader(FEmitContext& Context, const FRequestedType& RequestedType, FEmitValueShaderResult& OutResult) const
 {
 	int32 LocalPHIIndex = Context.LocalPHIs.Find(this);
 	if (LocalPHIIndex == INDEX_NONE)
@@ -667,9 +670,7 @@ void FExpressionLocalPHI::EmitValueShader(FEmitContext& Context, const FRequeste
 			}
 		}
 
-		const EExpressionDerivative Derivative = GetPreparedType().GetData().Derivative;
 		const FRequestedType LocalType = GetRequestedType();
-		const Shader::FType LocalDerivativeType = LocalType.GetType().GetDerivativeType();
 
 		bool bNeedToAddDeclaration = true;
 		for (int32 i = 0; i < NumValues; ++i)
@@ -678,7 +679,7 @@ void FExpressionLocalPHI::EmitValueShader(FEmitContext& Context, const FRequeste
 			check(IsScopeLive(ValueScope));
 
 			Context.ScopeStack.Add(ValueScope);
-			const FEmitShaderValues ShaderValue = Values[i]->GetValueShader(Context, LocalType);
+			FEmitShaderCode* ShaderValue = Values[i]->GetValueShader(Context, LocalType);
 			Context.ScopeStack.Pop();
 
 			if (ValueScope == DeclarationScope)
@@ -686,35 +687,14 @@ void FExpressionLocalPHI::EmitValueShader(FEmitContext& Context, const FRequeste
 				ValueScope->EmitDeclarationf(Context, TEXT("%s LocalPHI%d = %s;"),
 					LocalType.GetName(),
 					LocalPHIIndex,
-					ShaderValue.Code->Reference);
-				if (Derivative == EExpressionDerivative::Valid)
-				{
-					ValueScope->EmitDeclarationf(Context, TEXT("%s LocalPHI%dDdx = %s;"),
-						LocalDerivativeType.GetName(),
-						LocalPHIIndex,
-						ShaderValue.CodeDdx->Reference);
-					ValueScope->EmitDeclarationf(Context, TEXT("%s LocalPHI%dDdy = %s;"),
-						LocalDerivativeType.GetName(),
-						LocalPHIIndex,
-						ShaderValue.CodeDdy->Reference);
-				}
-
+					ShaderValue->Reference);
 				bNeedToAddDeclaration = false;
 			}
 			else
 			{
 				ValueScope->EmitStatementf(Context, TEXT("LocalPHI%d = %s;"),
 					LocalPHIIndex,
-					ShaderValue.Code->Reference);
-				if (Derivative == EExpressionDerivative::Valid)
-				{
-					ValueScope->EmitStatementf(Context, TEXT("LocalPHI%dDdx = %s;"),
-						LocalPHIIndex,
-						ShaderValue.CodeDdx->Reference);
-					ValueScope->EmitStatementf(Context, TEXT("LocalPHI%dDdy = %s;"),
-						LocalPHIIndex,
-						ShaderValue.CodeDdy->Reference);
-				}
+					ShaderValue->Reference);
 			}
 		}
 
@@ -724,26 +704,11 @@ void FExpressionLocalPHI::EmitValueShader(FEmitContext& Context, const FRequeste
 			DeclarationScope->EmitDeclarationf(Context, TEXT("%s LocalPHI%d;"),
 				LocalType.GetName(),
 				LocalPHIIndex);
-			if (Derivative == EExpressionDerivative::Valid)
-			{
-				DeclarationScope->EmitDeclarationf(Context, TEXT("%s LocalPHI%dDdx;"),
-					LocalDerivativeType.GetName(),
-					LocalPHIIndex);
-				DeclarationScope->EmitDeclarationf(Context, TEXT("%s LocalPHI%dDdy;"),
-					LocalDerivativeType.GetName(),
-					LocalPHIIndex);
-			}
 		}
 	}
 
 	const Shader::FType LocalType = GetType();
 	OutResult.Code = Context.EmitInlineCode(LocalType, TEXT("LocalPHI%"), LocalPHIIndex);
-	if (GetDerivative(RequestedType) == EExpressionDerivative::Valid)
-	{
-		const Shader::FType LocalDerivativeType = LocalType.GetDerivativeType();
-		OutResult.CodeDdx = Context.EmitInlineCode(LocalDerivativeType, TEXT("LocalPHI%Ddx"), LocalPHIIndex);
-		OutResult.CodeDdy = Context.EmitInlineCode(LocalDerivativeType, TEXT("LocalPHI%Ddy"), LocalPHIIndex);
-	}
 }
 
 void FStatement::Reset()
@@ -790,11 +755,11 @@ const FPreparedType& PrepareExpressionValue(FEmitContext& Context, FExpression*
 		InExpression->CurrentRequestedType.RequestedComponents.PadToNum(NumComponents, false);
 		for (int32 Index = 0; Index < NumComponents; ++Index)
 		{
-			const EComponentRequest PrevRequest = InExpression->CurrentRequestedType.GetComponentRequest(Index);
-			const EComponentRequest Request = RequestedType.GetComponentRequest(Index);
-			if ((uint32)Request > (uint32)PrevRequest)
+			const bool bPrevRequest = InExpression->CurrentRequestedType.IsComponentRequested(Index);
+			const bool bRequest = RequestedType.IsComponentRequested(Index);
+			if (!bPrevRequest && bRequest)
 			{
-				InExpression->CurrentRequestedType.SetComponentRequest(Index, Request);
+				InExpression->CurrentRequestedType.SetComponentRequest(Index);
 				bNeedToUpdateType = true;
 			}
 		}
@@ -819,13 +784,12 @@ const FPreparedType& PrepareExpressionValue(FEmitContext& Context, FExpression*
 	return InExpression->PrepareValueResult.GetPreparedType();
 }
 
-FRequestedType::FRequestedType(int32 NumComponents, EComponentRequest DefaultRequest)
+FRequestedType::FRequestedType(int32 NumComponents, bool bDefaultRequest)
 {
-	RequestedComponents.Init(IsRequested(DefaultRequest), NumComponents);
-	RequestedComponentDerivatives.Init(IsDerivativeRequested(DefaultRequest), NumComponents);
+	RequestedComponents.Init(bDefaultRequest, NumComponents);
 }
 
-FRequestedType::FRequestedType(const Shader::FType& InType, EComponentRequest DefaultRequest)
+FRequestedType::FRequestedType(const Shader::FType& InType, bool bDefaultRequest)
 {
 	int32 NumComponents = 0;
 	if (InType.IsStruct())
@@ -840,16 +804,14 @@ FRequestedType::FRequestedType(const Shader::FType& InType, EComponentRequest De
 		NumComponents = TypeDesc.NumComponents;
 	}
 
-	RequestedComponents.Init(IsRequested(DefaultRequest), NumComponents);
-	RequestedComponentDerivatives.Init(IsDerivativeRequested(DefaultRequest), NumComponents);
+	RequestedComponents.Init(bDefaultRequest, NumComponents);
 }
 
-FRequestedType::FRequestedType(const Shader::EValueType& InType, EComponentRequest DefaultRequest)
+FRequestedType::FRequestedType(const Shader::EValueType& InType, bool bDefaultRequest)
 {
 	const Shader::FValueTypeDescription TypeDesc = Shader::GetValueTypeDescription(InType);
 	ValueComponentType = TypeDesc.ComponentType;
-	RequestedComponents.Init(IsRequested(DefaultRequest), TypeDesc.NumComponents);
-	RequestedComponentDerivatives.Init(IsDerivativeRequested(DefaultRequest), TypeDesc.NumComponents);
+	RequestedComponents.Init(bDefaultRequest, TypeDesc.NumComponents);
 }
 
 const Shader::FType FRequestedType::GetType() const
@@ -878,51 +840,24 @@ int32 FRequestedType::GetNumComponents() const
 	return 0;
 }
 
-EComponentRequest FRequestedType::GetComponentRequest(int32 Index) const
+void FRequestedType::SetComponentRequest(int32 Index, bool bRequested)
 {
-	if (RequestedComponents.IsValidIndex(Index))
-	{
-		if (RequestedComponentDerivatives.IsValidIndex(Index) && RequestedComponentDerivatives[Index])
-		{
-			return EComponentRequest::RequestedWithDerivative;
-		}
-		if (RequestedComponents[Index])
-		{
-			return EComponentRequest::Requested;
-		}
-	}
-	return EComponentRequest::None;
-}
-
-void FRequestedType::SetComponentRequest(int32 Index, EComponentRequest Request)
-{
-	const bool bRequested = IsRequested(Request);
-	const bool bRequestedDerivative = IsDerivativeRequested(Request);
 	if (bRequested)
 	{
 		RequestedComponents.PadToNum(Index + 1, false);
 	}
-	if (bRequestedDerivative)
-	{
-		RequestedComponentDerivatives.PadToNum(Index + 1, false);
-	}
-
 	if (RequestedComponents.IsValidIndex(Index))
 	{
 		RequestedComponents[Index] = bRequested;
 	}
-	if (RequestedComponentDerivatives.IsValidIndex(Index))
-	{
-		RequestedComponentDerivatives[Index] = bRequestedDerivative;
-	}
 }
 
-void FRequestedType::SetFieldRequested(const Shader::FStructField* Field, EComponentRequest Request)
+void FRequestedType::SetFieldRequested(const Shader::FStructField* Field, bool bRequested)
 {
 	const int32 NumComponents = Field->GetNumComponents();
 	for (int32 Index = 0; Index < NumComponents; ++Index)
 	{
-		SetComponentRequest(Field->ComponentIndex + Index, Request);
+		SetComponentRequest(Field->ComponentIndex + Index, bRequested);
 	}
 }
 
@@ -931,7 +866,7 @@ void FRequestedType::SetField(const Shader::FStructField* Field, const FRequeste
 	const int32 NumComponents = Field->GetNumComponents();
 	for (int32 Index = 0; Index < NumComponents; ++Index)
 	{
-		SetComponentRequest(Field->ComponentIndex + Index, InRequest.GetComponentRequest(Index));
+		SetComponentRequest(Field->ComponentIndex + Index, InRequest.IsComponentRequested(Index));
 	}
 }
 
@@ -941,7 +876,7 @@ FRequestedType FRequestedType::GetField(const Shader::FStructField* Field) const
 	const int32 NumComponents = Field->GetNumComponents();
 	for (int32 Index = 0; Index < NumComponents; ++Index)
 	{
-		Result.SetComponentRequest(Index, GetComponentRequest(Field->ComponentIndex + Index));
+		Result.SetComponentRequest(Index, IsComponentRequested(Field->ComponentIndex + Index));
 	}
 	return Result;
 }
@@ -966,7 +901,7 @@ int32 FPreparedType::GetNumComponents() const
 	}
 	else if (ValueComponentType != Shader::EValueComponentType::Void)
 	{
-		const int32 MaxComponentIndex = PreparedComponents.FindLastByPredicate([](const FPreparedData& InData) { return InData.IsValid(); });
+		const int32 MaxComponentIndex = PreparedComponents.FindLastByPredicate([](EExpressionEvaluation InEvaluation) { return InEvaluation != EExpressionEvaluation::None; });
 		if (MaxComponentIndex != INDEX_NONE)
 		{
 			return MaxComponentIndex + 1;
@@ -981,7 +916,6 @@ FRequestedType MakeRequestedType(Shader::EValueComponentType ComponentType, cons
 	FRequestedType Result;
 	Result.ValueComponentType = ComponentType;
 	Result.RequestedComponents = RequestedComponents.RequestedComponents;
-	Result.RequestedComponentDerivatives = RequestedComponents.RequestedComponentDerivatives;
 	return Result;
 }
 
@@ -1015,119 +949,85 @@ FRequestedType FPreparedType::GetRequestedType() const
 		}
 		for (int32 Index = 0; Index < NumComponents; ++Index)
 		{
-			Result.SetComponentRequest(Index, GetComponentData(Index).GetRequest());
+			const EExpressionEvaluation ComponentEvaluation = GetComponentEvaluation(Index);
+			if (ComponentEvaluation != EExpressionEvaluation::None)
+			{
+				Result.SetComponentRequest(Index);
+			}
 		}
 	}
 	return Result;
 }
 
-FPreparedData FPreparedType::GetData() const
+EExpressionEvaluation FPreparedType::GetEvaluation() const
 {
-	FPreparedData Result;
+	EExpressionEvaluation Result = EExpressionEvaluation::None;
 	for (int32 Index = 0; Index < PreparedComponents.Num(); ++Index)
 	{
-		const FPreparedData& Component = PreparedComponents[Index];
-		if (Component.IsValid())
-		{
-			Result = CombinePreparedData(Result, Component);
-		}
+		Result = CombineEvaluations(Result, PreparedComponents[Index]);
 	}
 	return Result;
 }
 
-FPreparedData FPreparedType::GetData(const FRequestedType& RequestedType) const
+EExpressionEvaluation FPreparedType::GetEvaluation(const FRequestedType& RequestedType) const
 {
-	FPreparedData Result;
+	EExpressionEvaluation Result = EExpressionEvaluation::None;
 	for (int32 Index = 0; Index < PreparedComponents.Num(); ++Index)
 	{
-		const FPreparedData& Component = PreparedComponents[Index];
-		if (Component.IsValid())
+		if (RequestedType.IsComponentRequested(Index))
 		{
-			const EComponentRequest ComponentRequest = RequestedType.GetComponentRequest(Index);
-			if (IsRequested(ComponentRequest))
-			{
-				Result.Evaluation = CombineEvaluations(Result.Evaluation, Component.Evaluation);
-				if (IsDerivativeRequested(ComponentRequest))
-				{
-					Result.Derivative = CombineDerivatives(Result.Derivative, Component.Derivative);
-				}
-			}
+			Result = CombineEvaluations(Result, PreparedComponents[Index]);
 		}
 	}
 	return Result;
 }
 
-FPreparedData FPreparedType::GetFieldData(int32 ComponentIndex, int32 NumComponents) const
+EExpressionEvaluation FPreparedType::GetFieldEvaluation(int32 ComponentIndex, int32 NumComponents) const
 {
-	FPreparedData Result;
+	EExpressionEvaluation Result = EExpressionEvaluation::None;
 	for (int32 Index = 0; Index < NumComponents; ++Index)
 	{
-		Result = CombinePreparedData(Result, GetComponentData(ComponentIndex + Index));
+		Result = CombineEvaluations(Result, GetComponentEvaluation(ComponentIndex + Index));
 	}
 	return Result;
 }
 
-FPreparedData FPreparedType::GetComponentData(int32 Index) const
+EExpressionEvaluation FPreparedType::GetComponentEvaluation(int32 Index) const
 {
-	FPreparedData Result;
-	if (PreparedComponents.IsValidIndex(Index))
-	{
-		Result = PreparedComponents[Index];
-	}
-	return Result;
+	return PreparedComponents.IsValidIndex(Index) ? PreparedComponents[Index] : EExpressionEvaluation::None;
 }
 
-void FPreparedType::SetComponentData(int32 Index, const FPreparedData& Data)
+void FPreparedType::SetComponentEvaluation(int32 Index, EExpressionEvaluation Evaluation)
 {
-	if (Data.IsValid() && Index >= PreparedComponents.Num())
+	if (Evaluation != EExpressionEvaluation::None && Index >= PreparedComponents.Num())
 	{
-		PreparedComponents.AddDefaulted(Index + 1 - PreparedComponents.Num());
+		PreparedComponents.AddZeroed(Index + 1 - PreparedComponents.Num());
 	}
 	if (PreparedComponents.IsValidIndex(Index))
 	{
-		PreparedComponents[Index] = Data;
+		PreparedComponents[Index] = Evaluation;
 	}
 }
 
-void FPreparedType::MergeComponentData(int32 Index, EComponentRequest Request, const FPreparedData& Data)
+void FPreparedType::MergeComponentEvaluation(int32 Index, EExpressionEvaluation Evaluation)
 {
-	if (Request != EComponentRequest::None)
+	if (Evaluation != EExpressionEvaluation::None && Index >= PreparedComponents.Num())
 	{
-		if (Data.IsValid() && Index >= PreparedComponents.Num())
-		{
-			PreparedComponents.AddDefaulted(Index + 1 - PreparedComponents.Num());
-		}
-		if (PreparedComponents.IsValidIndex(Index))
-		{
-			FPreparedData& Component = PreparedComponents[Index];
-			Component.Evaluation = CombineEvaluations(Component.Evaluation, Data.Evaluation);
-			if (Request == EComponentRequest::RequestedWithDerivative)
-			{
-				Component.Derivative = CombineDerivatives(Component.Derivative, Data.Derivative);
-			}
-		}
+		PreparedComponents.AddZeroed(Index + 1 - PreparedComponents.Num());
 	}
-}
-
-void FPreparedType::SetEvaluation(EExpressionEvaluation Evaluation)
-{
-	for (int32 Index = 0; Index < PreparedComponents.Num(); ++Index)
+	if (PreparedComponents.IsValidIndex(Index))
 	{
-		if (PreparedComponents[Index].IsValid())
-		{
-			PreparedComponents[Index].Evaluation = Evaluation;
-		}
+		PreparedComponents[Index] = CombineEvaluations(PreparedComponents[Index], Evaluation);
 	}
 }
 
-void FPreparedType::SetDerivative(EExpressionDerivative Derivative)
+void FPreparedType::SetEvaluation(EExpressionEvaluation Evaluation)
 {
 	for (int32 Index = 0; Index < PreparedComponents.Num(); ++Index)
 	{
-		FPreparedData& ComponentData = PreparedComponents[Index];
-		if (ComponentData.IsValid())
+		if (PreparedComponents[Index] != EExpressionEvaluation::None)
 		{
-			ComponentData.Derivative = Derivative;
+			PreparedComponents[Index] = Evaluation;
 		}
 	}
 }
@@ -1136,7 +1036,7 @@ void FPreparedType::SetField(const Shader::FStructField* Field, const FPreparedT
 {
 	for (int32 Index = 0; Index < Field->GetNumComponents(); ++Index)
 	{
-		SetComponentData(Field->ComponentIndex + Index, FieldType.GetComponentData(Index));
+		SetComponentEvaluation(Field->ComponentIndex + Index, FieldType.GetComponentEvaluation(Index));
 	}
 }
 
@@ -1145,7 +1045,7 @@ FPreparedType FPreparedType::GetFieldType(const Shader::FStructField* Field) con
 	FPreparedType Result(Field->Type);
 	for (int32 Index = 0; Index < Field->GetNumComponents(); ++Index)
 	{
-		Result.SetComponentData(Index, GetComponentData(Field->ComponentIndex + Index));
+		Result.SetComponentEvaluation(Index, GetComponentEvaluation(Field->ComponentIndex + Index));
 	}
 	return Result;
 }
@@ -1182,9 +1082,9 @@ FPreparedType MergePreparedTypes(const FPreparedType& Lhs, const FPreparedType&
 
 	for (int32 Index = 0; Index < NumComponents; ++Index)
 	{
-		const FPreparedData LhsData = Lhs.GetComponentData(Index);
-		const FPreparedData RhsData = Rhs.GetComponentData(Index);
-		Result.SetComponentData(Index, CombinePreparedData(LhsData, RhsData));
+		const EExpressionEvaluation LhsEvaluation = Lhs.GetComponentEvaluation(Index);
+		const EExpressionEvaluation RhsEvaluation = Rhs.GetComponentEvaluation(Index);
+		Result.SetComponentEvaluation(Index, CombineEvaluations(LhsEvaluation, RhsEvaluation));
 	}
 
 	return Result;
@@ -1220,17 +1120,19 @@ bool FPrepareValueResult::TryMergePreparedType(FEmitContext& Context, const Shad
 	return true;
 }
 
-void FPrepareValueResult::SetType(FEmitContext& Context, const FRequestedType& RequestedType, const FPreparedData& Data, const Shader::FType& Type)
+void FPrepareValueResult::SetType(FEmitContext& Context, const FRequestedType& RequestedType, EExpressionEvaluation Evaluation, const Shader::FType& Type)
 {
 	if (TryMergePreparedType(Context, Type.StructType, Shader::GetValueTypeDescription(Type.ValueType).ComponentType))
 	{
-		if (Data.IsValid())
+		if (Evaluation != EExpressionEvaluation::None)
 		{
 			const int32 NumComponents = Type.GetNumComponents();
 			for (int32 Index = 0; Index < NumComponents; ++Index)
 			{
-				const EComponentRequest ComponentRequest = RequestedType.GetComponentRequest(Index);
-				PreparedType.MergeComponentData(Index, ComponentRequest, Data);
+				if (RequestedType.IsComponentRequested(Index))
+				{
+					PreparedType.MergeComponentEvaluation(Index, Evaluation);
+				}
 			}
 		}
 	}
@@ -1243,8 +1145,10 @@ void FPrepareValueResult::SetType(FEmitContext& Context, const FRequestedType& R
 		const int32 NumComponents = RequestedType.GetNumComponents();
 		for (int32 Index = 0; Index < NumComponents; ++Index)
 		{
-			const EComponentRequest ComponentRequest = RequestedType.GetComponentRequest(Index);
-			PreparedType.MergeComponentData(Index, ComponentRequest, Type.GetComponentData(Index));
+			if (RequestedType.IsComponentRequested(Index))
+			{
+				PreparedType.MergeComponentEvaluation(Index, Type.GetComponentEvaluation(Index));
+			}
 		}
 	}
 }
@@ -1259,7 +1163,12 @@ void FPrepareValueResult::SetForwardValue(FEmitContext& Context, const FRequeste
 	}
 }
 
-void FExpression::EmitValueShader(FEmitContext& Context, const FRequestedType& RequestedType, FEmitShaderValues& OutResult) const
+void FExpression::ComputeAnalyticDerivatives(FTree& Tree, FExpressionDerivatives& OutResult) const
+{
+	// nop
+}
+
+void FExpression::EmitValueShader(FEmitContext& Context, const FRequestedType& RequestedType, FEmitValueShaderResult& OutResult) const
 {
 	check(false);
 }
@@ -1269,56 +1178,39 @@ void FExpression::EmitValuePreshader(FEmitContext& Context, const FRequestedType
 	check(false);
 }
 
-FEmitShaderValues FExpression::GetValueShader(FEmitContext& Context, const FRequestedType& RequestedType, const Shader::FType& ResultType)
+FEmitShaderCode* FExpression::GetValueShader(FEmitContext& Context, const FRequestedType& RequestedType, const Shader::FType& ResultType)
 {
 	if (PrepareValueResult.ForwardValue)
 	{
 		return PrepareValueResult.ForwardValue->GetValueShader(Context, RequestedType, ResultType);
 	}
 
-	const FPreparedData Data = PrepareValueResult.PreparedType.GetData(RequestedType);
-	check(Data.IsValid());
+	const EExpressionEvaluation Evaluation = PrepareValueResult.PreparedType.GetEvaluation(RequestedType);
+	check(Evaluation != EExpressionEvaluation::None);
 
-	FEmitShaderValues Result;
-	if (Data.Evaluation == EExpressionEvaluation::Constant || Data.Evaluation == EExpressionEvaluation::Preshader)
+	FEmitShaderCode* Value = nullptr;
+	if (Evaluation == EExpressionEvaluation::Constant || Evaluation == EExpressionEvaluation::Preshader)
 	{
-		Result.Code = Context.EmitPreshaderOrConstant(RequestedType, this);
-		check(Data.Derivative != EExpressionDerivative::Valid); // If a constant has a valid derivative, it should be 'Zero'
-		check(!Result.HasDerivatives());
+		Value = Context.EmitPreshaderOrConstant(RequestedType, this);
 	}
 	else
 	{
-		check(Data.Evaluation == EExpressionEvaluation::Shader);
+		check(Evaluation == EExpressionEvaluation::Shader);
+		FEmitValueShaderResult Result;
 		EmitValueShader(Context, RequestedType, Result);
-		if (Result.HasDerivatives())
-		{
-			checkf(Data.Derivative == EExpressionDerivative::Valid, TEXT("Expression emitted derivatives, but didn't request them during PrepareValue"));
-		}
-		else
-		{
-			checkf(Data.Derivative != EExpressionDerivative::Valid, TEXT("Expression requested derivatives during PrepareValue, but didn't emit them"));
-		}
+		Value = Result.Code;
 	}
 
-	Result = Context.EmitCast(Result, ResultType);
-	if (Data.Derivative == EExpressionDerivative::Zero)
-	{
-		const Shader::FType DerivativeResultType = ResultType.GetDerivativeType();
-		check(!DerivativeResultType.IsVoid());
-		check(!Result.HasDerivatives());
-		Result.CodeDdx = Context.EmitConstantZero(DerivativeResultType);
-		Result.CodeDdy = Result.CodeDdx;
-	}
-
-	return Result;
+	Value = Context.EmitCast(Value, ResultType);
+	return Value;
 }
 
-FEmitShaderValues FExpression::GetValueShader(FEmitContext& Context, const FRequestedType& RequestedType)
+FEmitShaderCode* FExpression::GetValueShader(FEmitContext& Context, const FRequestedType& RequestedType)
 {
 	return GetValueShader(Context, RequestedType, RequestedType.GetType());
 }
 
-FEmitShaderValues FExpression::GetValueShader(FEmitContext& Context)
+FEmitShaderCode* FExpression::GetValueShader(FEmitContext& Context)
 {
 	return GetValueShader(Context, GetRequestedType());
 }
@@ -1331,7 +1223,7 @@ void FExpression::GetValuePreshader(FEmitContext& Context, const FRequestedType&
 	}
 
 	check(!bReentryFlag);
-	const EExpressionEvaluation Evaluation = PrepareValueResult.PreparedType.GetData(RequestedType).Evaluation;
+	const EExpressionEvaluation Evaluation = PrepareValueResult.PreparedType.GetEvaluation(RequestedType);
 	if (Evaluation == EExpressionEvaluation::Preshader)
 	{
 		bReentryFlag = true;
@@ -1354,7 +1246,7 @@ Shader::FValue FExpression::GetValueConstant(FEmitContext& Context, const FReque
 	}
 
 	check(!bReentryFlag);
-	check(PrepareValueResult.PreparedType.GetData(RequestedType).Evaluation == EExpressionEvaluation::Constant);
+	check(PrepareValueResult.PreparedType.GetEvaluation(RequestedType) == EExpressionEvaluation::Constant);
 	
 	Shader::FPreshaderData ConstantPreshader;
 	bReentryFlag = true;
@@ -1573,6 +1465,35 @@ void FTree::ResetNodes()
 	}
 }
 
+bool FTree::Finalize()
+{
+	// Resolve values for any PHI nodes that were generated
+	// Resolving a PHI may produce additional PHIs
+	while (PHIExpressions.Num() > 0)
+	{
+		FExpressionLocalPHI* Expression = PHIExpressions.Pop(false);
+		for (int32 i = 0; i < Expression->NumValues; ++i)
+		{
+			FExpression* LocalValue = AcquireLocal(*Expression->Scopes[i], Expression->LocalName);
+			if (!LocalValue)
+			{
+				//Errorf(TEXT("Local %s is not assigned on all control paths"), *Expression->LocalName.ToString());
+				return false;
+			}
+
+			for(EDerivativeCoordinate DerivativeCoord : Expression->DerivativeChain)
+			{
+				const FExpressionDerivatives Derivatives = GetAnalyticDerivatives(LocalValue);
+				LocalValue = Derivatives.Get(DerivativeCoord);
+			}
+			// May be nullptr if derivatives are not valid
+			Expression->Values[i] = LocalValue;
+		}
+	}
+
+	return true;
+}
+
 bool FTree::EmitShader(FEmitContext& Context, FStringBuilderBase& OutCode) const
 {
 	if (RootScope->ContainedStatement)
@@ -1605,6 +1526,59 @@ void FTree::RegisterStatement(FScope& Scope, FStatement* Statement)
 	Scope.ContainedStatement = Statement;
 }
 
+void FTree::AssignLocal(FScope& Scope, const FName& LocalName, FExpression* Value)
+{
+	Scope.LocalMap.Add(LocalName, Value);
+}
+
+FExpression* FTree::AcquireLocal(FScope& Scope, const FName& LocalName)
+{
+	FExpression** FoundExpression = Scope.LocalMap.Find(LocalName);
+	if (FoundExpression)
+	{
+		return *FoundExpression;
+	}
+
+	const TArrayView<FScope*> PreviousScopes = Scope.GetPreviousScopes();
+	if (PreviousScopes.Num() > 1)
+	{
+		FExpressionLocalPHI* Expression = NewExpression<FExpressionLocalPHI>();
+		Expression->LocalName = LocalName;
+		Expression->NumValues = PreviousScopes.Num();
+		for (int32 i = 0; i < PreviousScopes.Num(); ++i)
+		{
+			Expression->Scopes[i] = PreviousScopes[i];
+		}
+		PHIExpressions.Add(Expression);
+		Scope.LocalMap.Add(LocalName, Expression);
+		return Expression;
+	}
+
+	if (PreviousScopes.Num() == 1)
+	{
+		return AcquireLocal(*PreviousScopes[0], LocalName);
+	}
+
+	return nullptr;
+}
+
+const FExpressionDerivatives& FTree::GetAnalyticDerivatives(FExpression* InExpression)
+{
+	if (!InExpression)
+	{
+		static const FExpressionDerivatives EmptyDerivatives;
+		return EmptyDerivatives;
+	}
+
+	if (!InExpression->bComputedDerivatives)
+	{
+		FExpressionReentryScope ReentryScope(InExpression);
+		InExpression->ComputeAnalyticDerivatives(*this, InExpression->Derivatives);
+		InExpression->bComputedDerivatives = true;
+	}
+	return InExpression->Derivatives;
+}
+
 FScope* FTree::NewScope(FScope& Scope)
 {
 	FScope* NewScope = NewNode<FScope>();
@@ -1631,4 +1605,4 @@ FTextureParameterDeclaration* FTree::NewTextureParameterDeclaration(const FName&
 }
 
 } // namespace HLSLTree
-} // namespace UE
\ No newline at end of file
+} // namespace UE
diff --git a/Engine/Source/Runtime/Engine/Private/HLSLTree/HLSLTreeCommon.cpp b/Engine/Source/Runtime/Engine/Private/HLSLTree/HLSLTreeCommon.cpp
--- a/Engine/Source/Runtime/Engine/Private/HLSLTree/HLSLTreeCommon.cpp
+++ b/Engine/Source/Runtime/Engine/Private/HLSLTree/HLSLTreeCommon.cpp
@@ -9,6 +9,50 @@ namespace UE
 namespace HLSLTree
 {
 
+FExternalInputDescription GetExternalInputDescription(EExternalInput Input)
+{
+	switch (Input)
+	{
+	case EExternalInput::None: return FExternalInputDescription(TEXT("None"), Shader::EValueType::Void);
+
+	case EExternalInput::TexCoord0: return FExternalInputDescription(TEXT("TexCoord0"), Shader::EValueType::Float2, EExternalInput::TexCoord0_Ddx, EExternalInput::TexCoord0_Ddy);
+	case EExternalInput::TexCoord1: return FExternalInputDescription(TEXT("TexCoord1"), Shader::EValueType::Float2, EExternalInput::TexCoord1_Ddx, EExternalInput::TexCoord1_Ddy);
+	case EExternalInput::TexCoord2: return FExternalInputDescription(TEXT("TexCoord2"), Shader::EValueType::Float2, EExternalInput::TexCoord2_Ddx, EExternalInput::TexCoord2_Ddy);
+	case EExternalInput::TexCoord3: return FExternalInputDescription(TEXT("TexCoord3"), Shader::EValueType::Float2, EExternalInput::TexCoord3_Ddx, EExternalInput::TexCoord3_Ddy);
+	case EExternalInput::TexCoord4: return FExternalInputDescription(TEXT("TexCoord4"), Shader::EValueType::Float2, EExternalInput::TexCoord4_Ddx, EExternalInput::TexCoord4_Ddy);
+	case EExternalInput::TexCoord5: return FExternalInputDescription(TEXT("TexCoord5"), Shader::EValueType::Float2, EExternalInput::TexCoord5_Ddx, EExternalInput::TexCoord5_Ddy);
+	case EExternalInput::TexCoord6: return FExternalInputDescription(TEXT("TexCoord6"), Shader::EValueType::Float2, EExternalInput::TexCoord6_Ddx, EExternalInput::TexCoord6_Ddy);
+	case EExternalInput::TexCoord7: return FExternalInputDescription(TEXT("TexCoord7"), Shader::EValueType::Float2, EExternalInput::TexCoord7_Ddx, EExternalInput::TexCoord7_Ddy);
+
+	case EExternalInput::TexCoord0_Ddx: return FExternalInputDescription(TEXT("TexCoord0_Ddx"), Shader::EValueType::Float2);
+	case EExternalInput::TexCoord1_Ddx: return FExternalInputDescription(TEXT("TexCoord1_Ddx"), Shader::EValueType::Float2);
+	case EExternalInput::TexCoord2_Ddx: return FExternalInputDescription(TEXT("TexCoord2_Ddx"), Shader::EValueType::Float2);
+	case EExternalInput::TexCoord3_Ddx: return FExternalInputDescription(TEXT("TexCoord3_Ddx"), Shader::EValueType::Float2);
+	case EExternalInput::TexCoord4_Ddx: return FExternalInputDescription(TEXT("TexCoord4_Ddx"), Shader::EValueType::Float2);
+	case EExternalInput::TexCoord5_Ddx: return FExternalInputDescription(TEXT("TexCoord5_Ddx"), Shader::EValueType::Float2);
+	case EExternalInput::TexCoord6_Ddx: return FExternalInputDescription(TEXT("TexCoord6_Ddx"), Shader::EValueType::Float2);
+	case EExternalInput::TexCoord7_Ddx: return FExternalInputDescription(TEXT("TexCoord7_Ddx"), Shader::EValueType::Float2);
+
+	case EExternalInput::TexCoord0_Ddy: return FExternalInputDescription(TEXT("TexCoord0_Ddy"), Shader::EValueType::Float2);
+	case EExternalInput::TexCoord1_Ddy: return FExternalInputDescription(TEXT("TexCoord1_Ddy"), Shader::EValueType::Float2);
+	case EExternalInput::TexCoord2_Ddy: return FExternalInputDescription(TEXT("TexCoord2_Ddy"), Shader::EValueType::Float2);
+	case EExternalInput::TexCoord3_Ddy: return FExternalInputDescription(TEXT("TexCoord3_Ddy"), Shader::EValueType::Float2);
+	case EExternalInput::TexCoord4_Ddy: return FExternalInputDescription(TEXT("TexCoord4_Ddy"), Shader::EValueType::Float2);
+	case EExternalInput::TexCoord5_Ddy: return FExternalInputDescription(TEXT("TexCoord5_Ddy"), Shader::EValueType::Float2);
+	case EExternalInput::TexCoord6_Ddy: return FExternalInputDescription(TEXT("TexCoord6_Ddy"), Shader::EValueType::Float2);
+	case EExternalInput::TexCoord7_Ddy: return FExternalInputDescription(TEXT("TexCoord7_Ddy"), Shader::EValueType::Float2);
+
+	case EExternalInput::WorldPosition: return FExternalInputDescription(TEXT("WorldPosition"), Shader::EValueType::Double3, EExternalInput::WorldPosition_Ddx, EExternalInput::WorldPosition_Ddy);
+	case EExternalInput::WorldPosition_NoOffsets: return FExternalInputDescription(TEXT("WorldPosition_NoOffsets"), Shader::EValueType::Double3, EExternalInput::WorldPosition_Ddx, EExternalInput::WorldPosition_Ddy);
+	case EExternalInput::TranslatedWorldPosition: return FExternalInputDescription(TEXT("TranslatedWorldPosition"), Shader::EValueType::Float3, EExternalInput::WorldPosition_Ddx, EExternalInput::WorldPosition_Ddy);
+	case EExternalInput::TranslatedWorldPosition_NoOffsets: return FExternalInputDescription(TEXT("TranslatedWorldPosition_NoOffsets"), Shader::EValueType::Float3, EExternalInput::WorldPosition_Ddx, EExternalInput::WorldPosition_Ddy);
+
+	case EExternalInput::WorldPosition_Ddx: return FExternalInputDescription(TEXT("WorldPosition_Ddx"), Shader::EValueType::Float3);
+	case EExternalInput::WorldPosition_Ddy: return FExternalInputDescription(TEXT("WorldPosition_Ddx"), Shader::EValueType::Float3);
+	default: checkNoEntry(); return FExternalInputDescription(TEXT("Invalid"), Shader::EValueType::Void);
+	}
+}
+
 FSwizzleParameters::FSwizzleParameters(int8 InR, int8 InG, int8 InB, int8 InA) : NumComponents(0)
 {
 	ComponentIndex[0] = InR;
@@ -47,11 +91,10 @@ FRequestedType FSwizzleParameters::GetRequestedInputType(const FRequestedType& R
 
 	for (int32 Index = 0; Index < NumComponents; ++Index)
 	{
-		const EComponentRequest ComponentRequest = RequestedType.GetComponentRequest(Index);
-		if (ComponentRequest != EComponentRequest::None)
+		if (RequestedType.IsComponentRequested(Index))
 		{
 			const int32 SwizzledComponentIndex = ComponentIndex[Index];
-			RequestedInputType.SetComponentRequest(SwizzledComponentIndex, ComponentRequest);
+			RequestedInputType.SetComponentRequest(SwizzledComponentIndex);
 		}
 	}
 	
@@ -81,6 +124,22 @@ FSwizzleParameters MakeSwizzleMask(bool bInR, bool bInG, bool bInB, bool bInA)
 	return FSwizzleParameters(ComponentIndex[0], ComponentIndex[1], ComponentIndex[2], ComponentIndex[3]);
 }
 
+FExpression* FTree::NewConstant(const Shader::FValue& Value)
+{
+	return NewExpression<FExpressionConstant>(Value);
+}
+
+void FExpressionConstant::ComputeAnalyticDerivatives(FTree& Tree, FExpressionDerivatives& OutResult) const
+{
+	const Shader::FType DerivativeType = Value.Type.GetDerivativeType();
+	if (!DerivativeType.IsVoid())
+	{
+		const Shader::FValue ZeroValue(DerivativeType);
+		OutResult.ExpressionDdx = Tree.NewConstant(ZeroValue);
+		OutResult.ExpressionDdy = OutResult.ExpressionDdx;
+	}
+}
+
 void FExpressionConstant::PrepareValue(FEmitContext& Context, const FRequestedType& RequestedType, FPrepareValueResult& OutResult) const
 {
 	OutResult.SetType(Context, RequestedType, EExpressionEvaluation::Constant, Value.Type);
@@ -91,6 +150,18 @@ void FExpressionConstant::EmitValuePreshader(FEmitContext& Context, const FReque
 	OutPreshader.WriteOpcode(Shader::EPreshaderOpcode::Constant).Write(Value);
 }
 
+void FExpressionMaterialParameter::ComputeAnalyticDerivatives(FTree& Tree, FExpressionDerivatives& OutResult) const
+{
+	const Shader::FType Type = GetShaderValueType(ParameterType);
+	const Shader::FType DerivativeType = Type.GetDerivativeType();
+	if (!DerivativeType.IsVoid())
+	{
+		const Shader::FValue ZeroValue(DerivativeType);
+		OutResult.ExpressionDdx = Tree.NewConstant(ZeroValue);
+		OutResult.ExpressionDdy = OutResult.ExpressionDdx;
+	}
+}
+
 void FExpressionMaterialParameter::PrepareValue(FEmitContext& Context, const FRequestedType& RequestedType, FPrepareValueResult& OutResult) const
 {
 	const EExpressionEvaluation Evaluation = IsStaticMaterialParameter(ParameterType) ? EExpressionEvaluation::Constant : EExpressionEvaluation::Preshader;
@@ -132,67 +203,84 @@ void FExpressionMaterialParameter::EmitValuePreshader(FEmitContext& Context, con
 	}
 }
 
+void FExpressionExternalInput::ComputeAnalyticDerivatives(FTree& Tree, FExpressionDerivatives& OutResult) const
+{
+	const FExternalInputDescription InputDesc = GetExternalInputDescription(InputType);
+	if (InputDesc.Ddx != EExternalInput::None)
+	{
+		check(InputDesc.Ddy != EExternalInput::None);
+		OutResult.ExpressionDdx = Tree.NewExpression<FExpressionExternalInput>(InputDesc.Ddx);
+		OutResult.ExpressionDdy = Tree.NewExpression<FExpressionExternalInput>(InputDesc.Ddy);
+	}
+}
+
 void FExpressionExternalInput::PrepareValue(FEmitContext& Context, const FRequestedType& RequestedType, FPrepareValueResult& OutResult) const
 {
-	// TODO - derivative type
-	OutResult.SetTypeWithDerivative(Context, RequestedType, EExpressionEvaluation::Shader, GetInputExpressionType(InputType));
+	const FExternalInputDescription InputDesc = GetExternalInputDescription(InputType);
+	OutResult.SetType(Context, RequestedType, EExpressionEvaluation::Shader, InputDesc.Type);
 }
 
-void FExpressionExternalInput::EmitValueShader(FEmitContext& Context, const FRequestedType& RequestedType, FEmitShaderValues& OutResult) const
+void FExpressionExternalInput::EmitValueShader(FEmitContext& Context, const FRequestedType& RequestedType, FEmitValueShaderResult& OutResult) const
 {
+	const int32 TypeIndex = (int32)InputType;
 	if (IsTexCoord(InputType))
 	{
-		const int32 TypeIndex = (int32)InputType;
-		const int32 TexCoordIndex = TypeIndex - (int32)EExternalInputType::TexCoord0;
-
+		const int32 TexCoordIndex = TypeIndex - (int32)EExternalInput::TexCoord0;
 		Context.NumTexCoords = FMath::Max(Context.NumTexCoords, TexCoordIndex + 1);
 		OutResult.Code = Context.EmitInlineCode(Shader::EValueType::Float2, TEXT("Parameters.TexCoords[%].xy"), TexCoordIndex);
-		if (GetDerivative(RequestedType) == EExpressionDerivative::Valid)
-		{
-			OutResult.CodeDdx = Context.EmitInlineCode(Shader::EValueType::Float2, TEXT("Parameters.TexCoords_DDX[%].xy"), TexCoordIndex);
-			OutResult.CodeDdy = Context.EmitInlineCode(Shader::EValueType::Float2, TEXT("Parameters.TexCoords_DDY[%].xy"), TexCoordIndex);
-		}
+	}
+	else if (IsTexCoord_Ddx(InputType))
+	{
+		const int32 TexCoordIndex = TypeIndex - (int32)EExternalInput::TexCoord0_Ddx;
+		OutResult.Code = Context.EmitInlineCode(Shader::EValueType::Float2, TEXT("Parameters.TexCoords_DDX[%].xy"), TexCoordIndex);
+	}
+	else if (IsTexCoord_Ddy(InputType))
+	{
+		const int32 TexCoordIndex = TypeIndex - (int32)EExternalInput::TexCoord0_Ddy;
+		OutResult.Code = Context.EmitInlineCode(Shader::EValueType::Float2, TEXT("Parameters.TexCoords_DDY[%].xy"), TexCoordIndex);
 	}
 	else
 	{
-		const Shader::EValueType Type = GetInputExpressionType(InputType);
+		const FExternalInputDescription InputDesc = GetExternalInputDescription(InputType);
 
 		// TODO - handle PrevFrame position
 		switch (InputType)
 		{
-		case EExternalInputType::WorldPosition:
-			OutResult.Code = Context.EmitInlineCode(Type, TEXT("GetWorldPosition(Parameters)"));
+		case EExternalInput::WorldPosition:
+			OutResult.Code = Context.EmitInlineCode(InputDesc.Type, TEXT("GetWorldPosition(Parameters)"));
+			break;
+		case EExternalInput::WorldPosition_NoOffsets:
+			OutResult.Code = Context.EmitInlineCode(InputDesc.Type, TEXT("GetWorldPosition_NoMaterialOffsets(Parameters)"));
+			break;
+		case EExternalInput::TranslatedWorldPosition:
+			OutResult.Code = Context.EmitInlineCode(InputDesc.Type, TEXT("GetTranslatedWorldPosition(Parameters)"));
 			break;
-		case EExternalInputType::WorldPosition_NoOffsets:
-			OutResult.Code = Context.EmitInlineCode(Type, TEXT("GetWorldPosition_NoMaterialOffsets(Parameters)"));
+		case EExternalInput::TranslatedWorldPosition_NoOffsets:
+			OutResult.Code = Context.EmitInlineCode(InputDesc.Type, TEXT("GetTranslatedWorldPosition_NoMaterialOffsets(Parameters)"));
 			break;
-		case EExternalInputType::TranslatedWorldPosition:
-			OutResult.Code = Context.EmitInlineCode(Type, TEXT("GetTranslatedWorldPosition(Parameters)"));
+		case EExternalInput::WorldPosition_Ddx:
+			OutResult.Code = Context.EmitInlineCode(InputDesc.Type, TEXT("Parameters.WorldPosition_DDX"));
 			break;
-		case EExternalInputType::TranslatedWorldPosition_NoOffsets:
-			OutResult.Code = Context.EmitInlineCode(Type, TEXT("GetTranslatedWorldPosition_NoMaterialOffsets(Parameters)"));
+		case EExternalInput::WorldPosition_Ddy:
+			OutResult.Code = Context.EmitInlineCode(InputDesc.Type, TEXT("Parameters.WorldPosition_DDY"));
 			break;
 		default:
 			checkNoEntry();
 			break;
 		}
-
-		if (GetDerivative(RequestedType) == EExpressionDerivative::Valid)
-		{
-			OutResult.CodeDdx = Context.EmitInlineCode(Shader::EValueType::Float3, TEXT("Parameters.WorldPosition_DDX"));
-			OutResult.CodeDdy = Context.EmitInlineCode(Shader::EValueType::Float3, TEXT("Parameters.WorldPosition_DDY"));
-		}
 	}
 }
 
 void FExpressionTextureSample::PrepareValue(FEmitContext& Context, const FRequestedType& RequestedType, FPrepareValueResult& OutResult) const
 {
-	const FRequestedType TexCoordRequestedType(Shader::EValueType::Float2, EComponentRequest::RequestedWithDerivative);
-	PrepareExpressionValue(Context, TexCoordExpression, TexCoordRequestedType);
+	PrepareExpressionValue(Context, TexCoordExpression, Shader::EValueType::Float2);
+	PrepareExpressionValue(Context, TexCoordDerivatives.ExpressionDdx, Shader::EValueType::Float2);
+	PrepareExpressionValue(Context, TexCoordDerivatives.ExpressionDdy, Shader::EValueType::Float2);
+
 	OutResult.SetType(Context, RequestedType, EExpressionEvaluation::Shader, Shader::EValueType::Float4);
 }
 
-void FExpressionTextureSample::EmitValueShader(FEmitContext& Context, const FRequestedType& RequestedType, FEmitShaderValues& OutResult) const
+void FExpressionTextureSample::EmitValueShader(FEmitContext& Context, const FRequestedType& RequestedType, FEmitValueShaderResult& OutResult) const
 {
 	const FTextureDescription& Desc = Declaration->Description;
 	const EMaterialValueType MaterialValueType = Desc.Texture->GetMaterialType();
@@ -318,21 +406,35 @@ void FExpressionTextureSample::EmitValueShader(FEmitContext& Context, const FReq
 		}
 	}
 
-	const FRequestedType TexCoordRequestedType(Shader::EValueType::Float2, EComponentRequest::RequestedWithDerivative);
-	const FEmitShaderValues TexCoordValue = TexCoordExpression->GetValueShader(Context, TexCoordRequestedType);
+	FEmitShaderCode* TexCoordValue = TexCoordExpression->GetValueShader(Context, Shader::EValueType::Float2);
+	FEmitShaderCode* TexCoordValueDdx = TexCoordDerivatives.ExpressionDdx->GetValueShader(Context, Shader::EValueType::Float2);
+	FEmitShaderCode* TexCoordValueDdy = TexCoordDerivatives.ExpressionDdy->GetValueShader(Context, Shader::EValueType::Float2);
 	OutResult.Code = Context.EmitCode(Shader::EValueType::Float4, TEXT("%(%Grad(%, %, %, %, %))"),
 		SamplerTypeFunction,
 		SampleFunctionName,
 		*TextureName,
 		*SamplerStateCode,
-		TexCoordValue.Code,
-		TexCoordValue.CodeDdx,
-		TexCoordValue.CodeDdy);
+		TexCoordValue,
+		TexCoordValueDdx,
+		TexCoordValueDdy);
+}
+
+void FExpressionGetStructField::ComputeAnalyticDerivatives(FTree& Tree, FExpressionDerivatives& OutResult) const
+{
+	const FExpressionDerivatives StructDerivatives = Tree.GetAnalyticDerivatives(StructExpression);
+	if (StructDerivatives.IsValid())
+	{
+		const Shader::FStructType* DerivativeStructType = StructType->DerivativeType;
+		check(DerivativeStructType);
+
+		OutResult.ExpressionDdx = Tree.NewExpression<FExpressionGetStructField>(DerivativeStructType, Field->Name, StructDerivatives.ExpressionDdx);
+		OutResult.ExpressionDdy = Tree.NewExpression<FExpressionGetStructField>(DerivativeStructType, Field->Name, StructDerivatives.ExpressionDdy);
+	}
 }
 
 void FExpressionGetStructField::PrepareValue(FEmitContext& Context, const FRequestedType& RequestedType, FPrepareValueResult& OutResult) const
 {
-	FRequestedType RequestedStructType(StructType, EComponentRequest::None);
+	FRequestedType RequestedStructType(StructType, false);
 	RequestedStructType.SetField(Field, RequestedType);
 
 	const FPreparedType& StructPreparedType = PrepareExpressionValue(Context, StructExpression, RequestedStructType);
@@ -344,38 +446,42 @@ void FExpressionGetStructField::PrepareValue(FEmitContext& Context, const FReque
 	OutResult.SetType(Context, RequestedType, StructPreparedType.GetFieldType(Field));
 }
 
-void FExpressionGetStructField::EmitValueShader(FEmitContext& Context, const FRequestedType& RequestedType, FEmitShaderValues& OutResult) const
+void FExpressionGetStructField::EmitValueShader(FEmitContext& Context, const FRequestedType& RequestedType, FEmitValueShaderResult& OutResult) const
 {
-	FRequestedType RequestedStructType(StructType, EComponentRequest::None);
+	FRequestedType RequestedStructType(StructType, false);
 	RequestedStructType.SetField(Field, RequestedType);
 
-	const FEmitShaderValues StructValue = StructExpression->GetValueShader(Context, RequestedStructType);
+	FEmitShaderCode* StructValue = StructExpression->GetValueShader(Context, RequestedStructType);
 
 	OutResult.Code = Context.EmitInlineCode(Field->Type, TEXT("%.%"),
-		StructValue.Code,
+		StructValue,
 		Field->Name);
-
-	if (StructValue.HasDerivatives() && GetDerivative(RequestedType) == EExpressionDerivative::Valid)
-	{
-		const Shader::FType DerivativeType = Field->Type.GetDerivativeType();
-		OutResult.CodeDdx = Context.EmitInlineCode(DerivativeType, TEXT("%.%"),
-			StructValue.CodeDdx,
-			Field->Name);
-		OutResult.CodeDdy = Context.EmitInlineCode(DerivativeType, TEXT("%.%"),
-			StructValue.CodeDdy,
-			Field->Name);
-	}
 }
 
 void FExpressionGetStructField::EmitValuePreshader(FEmitContext& Context, const FRequestedType& RequestedType, Shader::FPreshaderData& OutPreshader) const
 {
-	FRequestedType RequestedStructType(StructType, EComponentRequest::None);
+	FRequestedType RequestedStructType(StructType, false);
 	RequestedStructType.SetField(Field, RequestedType);
 
 	StructExpression->GetValuePreshader(Context, RequestedStructType, OutPreshader);
 	OutPreshader.WriteOpcode(Shader::EPreshaderOpcode::GetField).Write(Field->Type).Write(Field->ComponentIndex);
 }
 
+void FExpressionSetStructField::ComputeAnalyticDerivatives(FTree& Tree, FExpressionDerivatives& OutResult) const
+{
+	const FExpressionDerivatives StructDerivatives = Tree.GetAnalyticDerivatives(StructExpression);
+	const FExpressionDerivatives FieldDerivatives = Tree.GetAnalyticDerivatives(FieldExpression);
+
+	if (StructDerivatives.IsValid() && FieldDerivatives.IsValid())
+	{
+		const Shader::FStructType* DerivativeStructType = StructType->DerivativeType;
+		check(DerivativeStructType);
+
+		OutResult.ExpressionDdx = Tree.NewExpression<FExpressionSetStructField>(DerivativeStructType, Field->Name, StructDerivatives.ExpressionDdx, FieldDerivatives.ExpressionDdx);
+		OutResult.ExpressionDdy = Tree.NewExpression<FExpressionSetStructField>(DerivativeStructType, Field->Name, StructDerivatives.ExpressionDdy, FieldDerivatives.ExpressionDdy);
+	}
+}
+
 void FExpressionSetStructField::PrepareValue(FEmitContext& Context, const FRequestedType& RequestedType, FPrepareValueResult& OutResult) const
 {
 	FRequestedType RequestedStructType(RequestedType);
@@ -404,7 +510,7 @@ void FExpressionSetStructField::PrepareValue(FEmitContext& Context, const FReque
 	}
 }
 
-void FExpressionSetStructField::EmitValueShader(FEmitContext& Context, const FRequestedType& RequestedType, FEmitShaderValues& OutResult) const
+void FExpressionSetStructField::EmitValueShader(FEmitContext& Context, const FRequestedType& RequestedType, FEmitValueShaderResult& OutResult) const
 {
 	FRequestedType RequestedStructType(RequestedType);
 	RequestedStructType.ClearFieldRequested(Field);
@@ -412,11 +518,7 @@ void FExpressionSetStructField::EmitValueShader(FEmitContext& Context, const FRe
 
 	const FRequestedType RequestedFieldType = RequestedType.GetField(Field);
 	const EExpressionEvaluation FieldEvaluation = FieldExpression->GetEvaluation(RequestedFieldType);
-
-	const FEmitShaderValues FieldValue = FieldExpression->GetValueShader(Context, RequestedFieldType, Field->Type);
-
-	const Shader::FStructType* DerivativeStructType = StructType->DerivativeType;
-	const EExpressionDerivative Derivative = GetDerivative(RequestedType);
+	FEmitShaderCode* FieldValue = FieldExpression->GetValueShader(Context, RequestedFieldType, Field->Type);
 
 	if (StructEvaluation == EExpressionEvaluation::None)
 	{
@@ -428,42 +530,16 @@ void FExpressionSetStructField::EmitValueShader(FEmitContext& Context, const FRe
 			StructType->Name,
 			Field->Name,
 			StructType->Name,
-			FieldValue.Code);
-		if (Derivative == EExpressionDerivative::Valid && DerivativeStructType && FieldValue.HasDerivatives())
-		{
-			OutResult.CodeDdx = Context.EmitCode(DerivativeStructType, TEXT("%_Set%((%)0, %)"),
-				DerivativeStructType->Name,
-				Field->Name,
-				DerivativeStructType->Name,
-				FieldValue.CodeDdx);
-			OutResult.CodeDdy = Context.EmitCode(DerivativeStructType, TEXT("%_Set%((%)0, %)"),
-				DerivativeStructType->Name,
-				Field->Name,
-				DerivativeStructType->Name,
-				FieldValue.CodeDdy);
-		}
+			FieldValue);
 	}
 	else
 	{
-		const FEmitShaderValues StructValue = StructExpression->GetValueShader(Context, RequestedStructType);
+		FEmitShaderCode* StructValue = StructExpression->GetValueShader(Context, RequestedStructType);
 		OutResult.Code = Context.EmitCode(StructType, TEXT("%_Set%(%, %)"),
 			StructType->Name,
 			Field->Name,
-			StructValue.Code,
-			FieldValue.Code);
-		if (Derivative == EExpressionDerivative::Valid && DerivativeStructType && StructValue.HasDerivatives() && FieldValue.HasDerivatives())
-		{
-			OutResult.CodeDdx = Context.EmitCode(DerivativeStructType, TEXT("%_Set%(%, %)"),
-				DerivativeStructType->Name,
-				Field->Name,
-				StructValue.CodeDdx,
-				FieldValue.CodeDdx);
-			OutResult.CodeDdy = Context.EmitCode(DerivativeStructType, TEXT("%_Set%(%, %)"),
-				DerivativeStructType->Name,
-				Field->Name,
-				StructValue.CodeDdy,
-				FieldValue.CodeDdy);
-		}
+			StructValue,
+			FieldValue);
 	}
 }
 
@@ -510,15 +586,15 @@ void FExpressionSelect::PrepareValue(FEmitContext& Context, const FRequestedType
 	OutResult.SetType(Context, RequestedType, MergePreparedTypes(LhsType, RhsType));
 }
 
-void FExpressionSelect::EmitValueShader(FEmitContext& Context, const FRequestedType& RequestedType, FEmitShaderValues& OutResult) const
+void FExpressionSelect::EmitValueShader(FEmitContext& Context, const FRequestedType& RequestedType, FEmitValueShaderResult& OutResult) const
 {
-	const FEmitShaderValues TrueValue = TrueExpression->GetValueShader(Context, RequestedType);
-	const FEmitShaderValues FalseValue = FalseExpression->GetValueShader(Context, RequestedType);
+	FEmitShaderCode* TrueValue = TrueExpression->GetValueShader(Context, RequestedType);
+	FEmitShaderCode* FalseValue = FalseExpression->GetValueShader(Context, RequestedType);
 
 	OutResult.Code = Context.EmitCode(RequestedType.GetType(), TEXT("(% ? % : %)"),
-		ConditionExpression->GetValueShader(Context, Shader::EValueType::Bool1).Code,
-		TrueValue.Code,
-		FalseValue.Code);
+		ConditionExpression->GetValueShader(Context, Shader::EValueType::Bool1),
+		TrueValue,
+		FalseValue);
 }
 
 void FExpressionSelect::EmitValuePreshader(FEmitContext& Context, const FRequestedType& RequestedType, Shader::FPreshaderData& OutPreshader) const
@@ -526,6 +602,37 @@ void FExpressionSelect::EmitValuePreshader(FEmitContext& Context, const FRequest
 	check(false); // TODO
 }
 
+FExpression* FTree::NewUnaryOp(EUnaryOp Op, FExpression* Input)
+{
+	return NewExpression<FExpressionUnaryOp>(Op, Input);
+}
+
+void FExpressionUnaryOp::ComputeAnalyticDerivatives(FTree& Tree, FExpressionDerivatives& OutResult) const
+{
+	const FExpressionDerivatives InputDerivatives = Tree.GetAnalyticDerivatives(Input);
+	if (InputDerivatives.IsValid())
+	{
+		switch (Op)
+		{
+		case EUnaryOp::Neg:
+			OutResult.ExpressionDdx = Tree.NewNeg(InputDerivatives.ExpressionDdx);
+			OutResult.ExpressionDdy = Tree.NewNeg(InputDerivatives.ExpressionDdy);
+			break;
+		case EUnaryOp::Rcp:
+		{
+			FExpression* Result = Tree.NewRcp(Input);
+			FExpression* dFdA = Tree.NewNeg(Tree.NewMul(Result, Result));
+			OutResult.ExpressionDdx = Tree.NewMul(dFdA, InputDerivatives.ExpressionDdx);
+			OutResult.ExpressionDdy = Tree.NewMul(dFdA, InputDerivatives.ExpressionDdy);
+			break;
+		}
+		default:
+			checkNoEntry();
+			break;
+		}
+	}
+}
+
 void FExpressionUnaryOp::PrepareValue(FEmitContext& Context, const FRequestedType& RequestedType, FPrepareValueResult& OutResult) const
 {
 	const FPreparedType& InputType = PrepareExpressionValue(Context, Input, RequestedType);
@@ -550,11 +657,39 @@ void FExpressionUnaryOp::PrepareValue(FEmitContext& Context, const FRequestedTyp
 	OutResult.SetType(Context, RequestedType, ResultType);
 }
 
-void FExpressionUnaryOp::EmitValueShader(FEmitContext& Context, const FRequestedType& RequestedType, FEmitShaderValues& OutResult) const
+void FExpressionUnaryOp::EmitValueShader(FEmitContext& Context, const FRequestedType& RequestedType, FEmitValueShaderResult& OutResult) const
 {
-	OutResult = Context.EmitUnaryOp(Op,
-		Input->GetValueShader(Context, RequestedType),
-		GetDerivative(RequestedType));
+	FEmitShaderCode* InputValue = Input->GetValueShader(Context, RequestedType);
+	Shader::FType ResultType = InputValue->Type;
+	const Shader::FValueTypeDescription InputTypeDesc = Shader::GetValueTypeDescription(ResultType);
+
+	switch (Op)
+	{
+	case EUnaryOp::Neg:
+		if (InputTypeDesc.ComponentType == Shader::EValueComponentType::Double)
+		{
+			OutResult.Code = Context.EmitCode(ResultType, TEXT("LWCNegate(%)"), InputValue);
+		}
+		else
+		{
+			OutResult.Code = Context.EmitInlineCode(ResultType, TEXT("(-%)"), InputValue);
+		}
+		break;
+	case EUnaryOp::Rcp:
+		if (InputTypeDesc.ComponentType == Shader::EValueComponentType::Double)
+		{
+			ResultType = Shader::MakeValueType(Shader::EValueComponentType::Float, ResultType.GetNumComponents());
+			OutResult.Code = Context.EmitCode(ResultType, TEXT("LWCRcp(%)"), InputValue);
+		}
+		else
+		{
+			OutResult.Code = Context.EmitCode(ResultType, TEXT("rcp(%)"), InputValue);
+		}
+		break;
+	default:
+		checkNoEntry();
+		break;
+	}
 }
 
 void FExpressionUnaryOp::EmitValuePreshader(FEmitContext& Context, const FRequestedType& RequestedType, Shader::FPreshaderData& OutPreshader) const
@@ -566,6 +701,63 @@ void FExpressionUnaryOp::EmitValuePreshader(FEmitContext& Context, const FReques
 	OutPreshader.WriteOpcode(OpDesc.PreshaderOpcode);
 }
 
+FExpression* FTree::NewBinaryOp(EBinaryOp Op, FExpression* Lhs, FExpression* Rhs)
+{
+	return NewExpression<FExpressionBinaryOp>(Op, Lhs, Rhs);
+}
+
+void FExpressionBinaryOp::ComputeAnalyticDerivatives(FTree& Tree, FExpressionDerivatives& OutResult) const
+{
+	// Operations with constant derivatives
+	switch (Op)
+	{
+	case EBinaryOp::Less:
+		OutResult.ExpressionDdx = Tree.NewConstant(0.0f);
+		OutResult.ExpressionDdx = OutResult.ExpressionDdy;
+		break;
+	default:
+		break;
+	}
+
+	if (OutResult.IsValid())
+	{
+		return;
+	}
+
+	const FExpressionDerivatives LhsDerivatives = Tree.GetAnalyticDerivatives(Lhs);
+	const FExpressionDerivatives RhsDerivatives = Tree.GetAnalyticDerivatives(Rhs);
+	if (LhsDerivatives.IsValid() && RhsDerivatives.IsValid())
+	{
+		switch (Op)
+		{
+		case EBinaryOp::Add:
+			OutResult.ExpressionDdx = Tree.NewAdd(LhsDerivatives.ExpressionDdx, RhsDerivatives.ExpressionDdx);
+			OutResult.ExpressionDdy = Tree.NewAdd(LhsDerivatives.ExpressionDdy, RhsDerivatives.ExpressionDdy);
+			break;
+		case EBinaryOp::Sub:
+			OutResult.ExpressionDdx = Tree.NewSub(LhsDerivatives.ExpressionDdx, RhsDerivatives.ExpressionDdx);
+			OutResult.ExpressionDdy = Tree.NewSub(LhsDerivatives.ExpressionDdy, RhsDerivatives.ExpressionDdy);
+			break;
+		case EBinaryOp::Mul:
+			OutResult.ExpressionDdx = Tree.NewAdd(Tree.NewMul(LhsDerivatives.ExpressionDdx, Rhs), Tree.NewMul(RhsDerivatives.ExpressionDdx, Lhs));
+			OutResult.ExpressionDdy = Tree.NewAdd(Tree.NewMul(LhsDerivatives.ExpressionDdy, Rhs), Tree.NewMul(RhsDerivatives.ExpressionDdy, Lhs));
+			break;
+		case EBinaryOp::Div:
+		{
+			FExpression* Denom = Tree.NewRcp(Tree.NewMul(Rhs, Rhs));
+			FExpression* dFdA = Tree.NewMul(Rhs, Denom);
+			FExpression* dFdB = Tree.NewNeg(Tree.NewMul(Lhs, Denom));
+			OutResult.ExpressionDdx = Tree.NewAdd(Tree.NewMul(dFdA, LhsDerivatives.ExpressionDdx), Tree.NewMul(dFdB, RhsDerivatives.ExpressionDdx));
+			OutResult.ExpressionDdy = Tree.NewAdd(Tree.NewMul(dFdA, LhsDerivatives.ExpressionDdy), Tree.NewMul(dFdB, RhsDerivatives.ExpressionDdy));
+			break;
+		}
+		default:
+			checkNoEntry();
+			break;
+		}
+	}
+}
+
 void FExpressionBinaryOp::PrepareValue(FEmitContext& Context, const FRequestedType& RequestedType, FPrepareValueResult& OutResult) const
 {
 	const FPreparedType& LhsType = PrepareExpressionValue(Context, Lhs, RequestedType);
@@ -584,7 +776,6 @@ void FExpressionBinaryOp::PrepareValue(FEmitContext& Context, const FRequestedTy
 	if (Op == EBinaryOp::Less)
 	{
 		ResultType.ValueComponentType = Shader::EValueComponentType::Bool;
-		ResultType.SetDerivative(EExpressionDerivative::Zero); // Comparisons result in 'zero' derivative...is this correct?
 	}
 
 	const FBinaryOpDescription OpDesc = GetBinaryOpDesription(Op);
@@ -607,13 +798,71 @@ FRequestedType GetBinaryOpInputType(const FRequestedType& RequestedType, EBinary
 }
 } // namespace Private
 
-void FExpressionBinaryOp::EmitValueShader(FEmitContext& Context, const FRequestedType& RequestedType, FEmitShaderValues& OutResult) const
+void FExpressionBinaryOp::EmitValueShader(FEmitContext& Context, const FRequestedType& RequestedType, FEmitValueShaderResult& OutResult) const
 {
 	const FRequestedType InputType = Private::GetBinaryOpInputType(RequestedType, Op, Lhs->GetRequestedType(), Rhs->GetRequestedType());
-	OutResult = Context.EmitBinaryOp(Op,
-		Lhs->GetValueShader(Context, InputType),
-		Rhs->GetValueShader(Context, InputType),
-		GetDerivative(RequestedType));
+	const Shader::EValueComponentType InputComponentType = InputType.ValueComponentType;
+	Shader::FType ResultType = InputType.GetType();
+	FEmitShaderCode* LhsValue = Lhs->GetValueShader(Context, InputType);
+	FEmitShaderCode* RhsValue = Rhs->GetValueShader(Context, InputType);
+
+	switch (Op)
+	{
+	case EBinaryOp::Add:
+		if (InputComponentType == Shader::EValueComponentType::Double)
+		{
+			OutResult.Code = Context.EmitCode(ResultType, TEXT("LWCAdd(%, %)"), LhsValue, RhsValue);
+		}
+		else
+		{
+			OutResult.Code = Context.EmitCode(ResultType, TEXT("(% + %)"), LhsValue, RhsValue);
+		}
+		break;
+	case EBinaryOp::Sub:
+		if (InputComponentType == Shader::EValueComponentType::Double)
+		{
+			OutResult.Code = Context.EmitCode(ResultType, TEXT("LWCSubtract(%, %)"), LhsValue, RhsValue);
+		}
+		else
+		{
+			OutResult.Code = Context.EmitCode(ResultType, TEXT("(% - %)"), LhsValue, RhsValue);
+		}
+		break;
+	case EBinaryOp::Mul:
+		if (InputComponentType == Shader::EValueComponentType::Double)
+		{
+			OutResult.Code = Context.EmitCode(ResultType, TEXT("LWCMultiply(%, %)"), LhsValue, RhsValue);
+		}
+		else
+		{
+			OutResult.Code = Context.EmitCode(ResultType, TEXT("(% * %)"), LhsValue, RhsValue);
+		}
+		break;
+	case EBinaryOp::Div:
+		if (InputComponentType == Shader::EValueComponentType::Double)
+		{
+			OutResult.Code = Context.EmitCode(ResultType, TEXT("LWCDivide(%, %)"), LhsValue, RhsValue);
+		}
+		else
+		{
+			OutResult.Code = Context.EmitCode(ResultType, TEXT("(% / %)"), LhsValue, RhsValue);
+		}
+		break;
+	case EBinaryOp::Less:
+		ResultType = Shader::MakeValueType(Shader::EValueComponentType::Bool, ResultType.GetNumComponents());
+		if (InputComponentType == Shader::EValueComponentType::Double)
+		{
+			OutResult.Code = Context.EmitCode(ResultType, TEXT("LWCLess(%, %)"), LhsValue, RhsValue);
+		}
+		else
+		{
+			OutResult.Code = Context.EmitCode(ResultType, TEXT("(% < %)"), LhsValue, RhsValue);
+		}
+		break;
+	default:
+		checkNoEntry();
+		break;
+	}
 }
 
 void FExpressionBinaryOp::EmitValuePreshader(FEmitContext& Context, const FRequestedType& RequestedType, Shader::FPreshaderData& OutPreshader) const
@@ -627,6 +876,16 @@ void FExpressionBinaryOp::EmitValuePreshader(FEmitContext& Context, const FReque
 	OutPreshader.WriteOpcode(OpDesc.PreshaderOpcode);
 }
 
+void FExpressionSwizzle::ComputeAnalyticDerivatives(FTree& Tree, FExpressionDerivatives& OutResult) const
+{
+	const FExpressionDerivatives InputDerivatives = Tree.GetAnalyticDerivatives(Input);
+	if (InputDerivatives.IsValid())
+	{
+		OutResult.ExpressionDdx = Tree.NewExpression<FExpressionSwizzle>(Parameters, InputDerivatives.ExpressionDdx);
+		OutResult.ExpressionDdy = Tree.NewExpression<FExpressionSwizzle>(Parameters, InputDerivatives.ExpressionDdy);
+	}
+}
+
 void FExpressionSwizzle::PrepareValue(FEmitContext& Context, const FRequestedType& RequestedType, FPrepareValueResult& OutResult) const
 {
 	const FRequestedType RequestedInputType = Parameters.GetRequestedInputType(RequestedType);
@@ -636,13 +895,13 @@ void FExpressionSwizzle::PrepareValue(FEmitContext& Context, const FRequestedTyp
 	for (int32 ComponentIndex = 0; ComponentIndex < Parameters.NumComponents; ++ComponentIndex)
 	{
 		const int32 SwizzledComponentIndex = Parameters.ComponentIndex[ComponentIndex];
-		ResultType.SetComponentData(ComponentIndex, InputType.GetComponentData(SwizzledComponentIndex));
+		ResultType.SetComponentEvaluation(ComponentIndex, InputType.GetComponentEvaluation(SwizzledComponentIndex));
 	}
 
 	OutResult.SetType(Context, RequestedType, ResultType);
 }
 
-void FExpressionSwizzle::EmitValueShader(FEmitContext& Context, const FRequestedType& RequestedType, FEmitShaderValues& OutResult) const
+void FExpressionSwizzle::EmitValueShader(FEmitContext& Context, const FRequestedType& RequestedType, FEmitValueShaderResult& OutResult) const
 {
 	const FRequestedType RequestedInputType = Parameters.GetRequestedInputType(RequestedType);
 
@@ -660,16 +919,16 @@ void FExpressionSwizzle::EmitValueShader(FEmitContext& Context, const FRequested
 		}
 	}
 
-	const FEmitShaderValues InputValue = Input->GetValueShader(Context, RequestedInputType);
+	FEmitShaderCode* InputValue = Input->GetValueShader(Context, RequestedInputType);
 	if (bHasSwizzleReorder)
 	{
 		OutResult.Code = Context.EmitInlineCode(RequestedType.GetType(), TEXT("%.%"),
-			InputValue.Code,
+			InputValue,
 			Swizzle);
 	}
 	else
 	{
-		OutResult.Code = InputValue.Code;
+		OutResult.Code = InputValue;
 	}
 }
 
@@ -686,6 +945,17 @@ void FExpressionSwizzle::EmitValuePreshader(FEmitContext& Context, const FReques
 		.Write((uint8)Parameters.ComponentIndex[3]);
 }
 
+void FExpressionAppend::ComputeAnalyticDerivatives(FTree& Tree, FExpressionDerivatives& OutResult) const
+{
+	const FExpressionDerivatives LhsDerivatives = Tree.GetAnalyticDerivatives(Lhs);
+	const FExpressionDerivatives RhsDerivatives = Tree.GetAnalyticDerivatives(Rhs);
+	if (LhsDerivatives.IsValid() && RhsDerivatives.IsValid())
+	{
+		OutResult.ExpressionDdx = Tree.NewExpression<FExpressionAppend>(LhsDerivatives.ExpressionDdx, RhsDerivatives.ExpressionDdx);
+		OutResult.ExpressionDdy = Tree.NewExpression<FExpressionAppend>(LhsDerivatives.ExpressionDdy, RhsDerivatives.ExpressionDdy);
+	}
+}
+
 void FExpressionAppend::PrepareValue(FEmitContext& Context, const FRequestedType& RequestedType, FPrepareValueResult& OutResult) const
 {
 	const FPreparedType& LhsType = PrepareExpressionValue(Context, Lhs, RequestedType);
@@ -695,7 +965,7 @@ void FExpressionAppend::PrepareValue(FEmitContext& Context, const FRequestedType
 	FRequestedType RhsRequestedType;
 	for (int32 Index = NumLhsComponents; Index < NumRequestedComponents; ++Index)
 	{
-		RhsRequestedType.SetComponentRequest(Index - NumLhsComponents, RequestedType.GetComponentRequest(Index));
+		RhsRequestedType.SetComponentRequest(Index - NumLhsComponents, RequestedType.IsComponentRequested(Index));
 	}
 
 	const FPreparedType& RhsType = PrepareExpressionValue(Context, Rhs, RhsRequestedType);
@@ -709,11 +979,11 @@ void FExpressionAppend::PrepareValue(FEmitContext& Context, const FRequestedType
 	FPreparedType ResultType(LhsType.ValueComponentType);
 	for (int32 Index = 0; Index < NumLhsComponents; ++Index)
 	{
-		ResultType.SetComponentData(Index, LhsType.GetComponentData(Index));
+		ResultType.SetComponentEvaluation(Index, LhsType.GetComponentEvaluation(Index));
 	}
 	for (int32 Index = 0; Index < NumRhsComponents; ++Index)
 	{
-		ResultType.SetComponentData(NumLhsComponents + Index, LhsType.GetComponentData(Index));
+		ResultType.SetComponentEvaluation(NumLhsComponents + Index, LhsType.GetComponentEvaluation(Index));
 	}
 
 	OutResult.SetType(Context, RequestedType, ResultType);
@@ -739,32 +1009,32 @@ FAppendTypes GetAppendTypes(const FRequestedType& RequestedType, const FRequeste
 	Types.RhsType.ValueComponentType = RequestedType.ValueComponentType;
 	for (int32 Index = 0; Index < NumLhsComponents; ++Index)
 	{
-		Types.LhsType.SetComponentRequest(Index, RequestedType.GetComponentRequest(Index));
+		Types.LhsType.SetComponentRequest(Index, RequestedType.IsComponentRequested(Index));
 	}
 	for (int32 Index = 0; Index < NumRhsComponents; ++Index)
 	{
-		Types.RhsType.SetComponentRequest(Index, RequestedType.GetComponentRequest(NumLhsComponents + Index));
+		Types.RhsType.SetComponentRequest(Index, RequestedType.IsComponentRequested(NumLhsComponents + Index));
 	}
 	return Types;
 }
 } // namespace Private
 
-void FExpressionAppend::EmitValueShader(FEmitContext& Context, const FRequestedType& RequestedType, FEmitShaderValues& OutResult) const
+void FExpressionAppend::EmitValueShader(FEmitContext& Context, const FRequestedType& RequestedType, FEmitValueShaderResult& OutResult) const
 {
 	const Private::FAppendTypes Types = Private::GetAppendTypes(RequestedType, Lhs->GetRequestedType(), Rhs->GetRequestedType());
-	const FEmitShaderValues LhsValue = Lhs->GetValueShader(Context, Types.LhsType);
+	FEmitShaderCode* LhsValue = Lhs->GetValueShader(Context, Types.LhsType);
 
 	if (Types.RhsType.IsVoid())
 	{
-		OutResult.Code = LhsValue.Code;
+		OutResult.Code = LhsValue;
 	}
 	else
 	{
-		const FEmitShaderValues RhsValue = Rhs->GetValueShader(Context, Types.RhsType);
+		FEmitShaderCode* RhsValue = Rhs->GetValueShader(Context, Types.RhsType);
 		OutResult.Code = Context.EmitInlineCode(Types.ResultType, TEXT("%(%, %)"),
 			RequestedType.GetName(),
-			LhsValue.Code,
-			RhsValue.Code);
+			LhsValue,
+			RhsValue);
 	}
 }
 
@@ -785,7 +1055,7 @@ void FExpressionReflectionVector::PrepareValue(FEmitContext& Context, const FReq
 	OutResult.SetType(Context, RequestedType, EExpressionEvaluation::Shader, Shader::EValueType::Float3);
 }
 
-void FExpressionReflectionVector::EmitValueShader(FEmitContext& Context, const FRequestedType& RequestedType, FEmitShaderValues& OutResult) const
+void FExpressionReflectionVector::EmitValueShader(FEmitContext& Context, const FRequestedType& RequestedType, FEmitValueShaderResult& OutResult) const
 {
 	OutResult.Code = Context.EmitInlineCode(Shader::EValueType::Float3, TEXT("Parameters.ReflectionVector"));
 }
@@ -805,7 +1075,7 @@ void FStatementReturn::Prepare(FEmitContext& Context) const
 
 void FStatementReturn::EmitShader(FEmitContext& Context) const
 {
-	ParentScope->EmitStatementf(Context, TEXT("return %s;"), Expression->GetValueShader(Context).Code->Reference);
+	ParentScope->EmitStatementf(Context, TEXT("return %s;"), Expression->GetValueShader(Context)->Reference);
 }
 
 void FStatementIf::Prepare(FEmitContext& Context) const
@@ -850,8 +1120,8 @@ void FStatementIf::EmitShader(FEmitContext& Context) const
 	{
 		if (IsScopeLive(ThenScope) || IsScopeLive(ElseScope))
 		{
-			const FEmitShaderValues ConditionValue = ConditionExpression->GetValueShader(Context, Shader::EValueType::Bool1);
-			ParentScope->EmitNestedScopef(Context, ThenScope, TEXT("if (%s)"), ConditionValue.Code->Reference);
+			FEmitShaderCode* ConditionValue = ConditionExpression->GetValueShader(Context, Shader::EValueType::Bool1);
+			ParentScope->EmitNestedScopef(Context, ThenScope, TEXT("if (%s)"), ConditionValue->Reference);
 			if (IsScopeLive(ElseScope))
 			{
 				ParentScope->EmitNestedScopef(Context, ElseScope, TEXT("else"));
diff --git a/Engine/Source/Runtime/Engine/Private/HLSLTree/HLSLTreeIntrinsics.cpp b/Engine/Source/Runtime/Engine/Private/HLSLTree/HLSLTreeIntrinsics.cpp
--- a/Engine/Source/Runtime/Engine/Private/HLSLTree/HLSLTreeIntrinsics.cpp
+++ b/Engine/Source/Runtime/Engine/Private/HLSLTree/HLSLTreeIntrinsics.cpp
@@ -182,198 +182,5 @@ FEmitShaderCode* FEmitContext::EmitCast(FEmitShaderCode* ShaderValue, const Shad
 	return ShaderValue;
 }
 
-FEmitShaderValues FEmitContext::EmitCast(FEmitShaderValues ShaderValue, const Shader::FType& DestType)
-{
-	FEmitShaderValues Result;
-	Result.Code = EmitCast(ShaderValue.Code, DestType);
-	if (ShaderValue.HasDerivatives())
-	{
-		const Shader::FType DerivativeType = DestType.GetDerivativeType();
-		Result.CodeDdx = EmitCast(ShaderValue.CodeDdx, DerivativeType);
-		Result.CodeDdy = EmitCast(ShaderValue.CodeDdy, DerivativeType);
-	}
-	return Result;
-}
-
-FEmitShaderCode* FEmitContext::EmitUnaryOp(EUnaryOp Op, FEmitShaderCode* Input)
-{
-	const Shader::FValueTypeDescription InputTypeDesc = Shader::GetValueTypeDescription(Input->Type);
-	const Shader::EValueComponentType InputComponentType = InputTypeDesc.ComponentType;
-	const int8 NumComponents = InputTypeDesc.NumComponents;
-	Shader::EValueType ResultType = Input->Type;
-
-	FEmitShaderCode* Result = nullptr;
-	switch (Op)
-	{
-	case EUnaryOp::Neg:
-		if (InputComponentType == Shader::EValueComponentType::Double)
-		{
-			Result = EmitCode(ResultType, TEXT("LWCNegate(%)"), Input);
-		}
-		else
-		{
-			Result = EmitInlineCode(ResultType, TEXT("(-%)"), Input);
-		}
-		break;
-	case EUnaryOp::Rcp:
-		if (InputComponentType == Shader::EValueComponentType::Double)
-		{
-			ResultType = Shader::MakeValueType(Shader::EValueComponentType::Float, NumComponents);
-			Result = EmitCode(ResultType, TEXT("LWCRcp(%)"), Input);
-		}
-		else
-		{
-			Result = EmitCode(ResultType, TEXT("rcp(%)"), Input);
-		}
-		break;
-	default:
-		checkNoEntry();
-		break;
-	}
-	return Result;
-}
-
-FEmitShaderValues FEmitContext::EmitUnaryOp(EUnaryOp Op, FEmitShaderValues Input, EExpressionDerivative Derivative)
-{
-	FEmitShaderValues Result;
-	Result.Code = EmitUnaryOp(Op, Input.Code);
-	if (Derivative == EExpressionDerivative::Valid && Input.HasDerivatives())
-	{
-		switch (Op)
-		{
-		case EUnaryOp::Neg:
-			Result.CodeDdx = EmitNeg(Input.CodeDdx);
-			Result.CodeDdy = EmitNeg(Input.CodeDdy);
-			break;
-		case EUnaryOp::Rcp:
-		{
-			FEmitShaderCode* dFdA = EmitNeg(EmitMul(Result.Code, Result.Code));
-			Result.CodeDdx = EmitMul(dFdA, Input.CodeDdx);
-			Result.CodeDdy = EmitMul(dFdA, Input.CodeDdy);
-			break;
-		}
-		default:
-			checkNoEntry();
-			break;
-		}
-	}
-	return Result;
-}
-
-FEmitShaderCode* FEmitContext::EmitBinaryOp(EBinaryOp Op, FEmitShaderCode* Lhs, FEmitShaderCode* Rhs)
-{
-	const Shader::FValueTypeDescription LhsTypeDesc = Shader::GetValueTypeDescription(Lhs->Type);
-	const Shader::FValueTypeDescription RhsTypeDesc = Shader::GetValueTypeDescription(Rhs->Type);
-	const int8 NumComponents = (LhsTypeDesc.NumComponents == 1 || RhsTypeDesc.NumComponents == 1) ? FMath::Max(LhsTypeDesc.NumComponents, RhsTypeDesc.NumComponents) : FMath::Min(LhsTypeDesc.NumComponents, RhsTypeDesc.NumComponents);
-	const Shader::EValueComponentType InputComponentType = Shader::CombineComponentTypes(LhsTypeDesc.ComponentType, RhsTypeDesc.ComponentType);
-	const Shader::EValueType InputType = Shader::MakeValueType(InputComponentType, NumComponents);
-	Shader::EValueType ResultType = InputType;
-
-	FEmitShaderCode* LhsCast = EmitCast(Lhs, InputType);
-	FEmitShaderCode* RhsCast = EmitCast(Rhs, InputType);
-
-	FEmitShaderCode* Result = nullptr;
-	switch (Op)
-	{
-	case EBinaryOp::Add:
-		if (InputComponentType == Shader::EValueComponentType::Double)
-		{
-			Result = EmitCode(ResultType, TEXT("LWCAdd(%, %)"), LhsCast, RhsCast);
-		}
-		else
-		{
-			Result = EmitCode(ResultType, TEXT("(% + %)"), LhsCast, RhsCast);
-		}
-		break;
-	case EBinaryOp::Sub:
-		if (InputComponentType == Shader::EValueComponentType::Double)
-		{
-			Result = EmitCode(ResultType, TEXT("LWCSubtract(%, %)"), LhsCast, RhsCast);
-		}
-		else
-		{
-			Result = EmitCode(ResultType, TEXT("(% - %)"), LhsCast, RhsCast);
-		}
-		break;
-	case EBinaryOp::Mul:
-		if (InputComponentType == Shader::EValueComponentType::Double)
-		{
-			Result = EmitCode(ResultType, TEXT("LWCMultiply(%, %)"), LhsCast, RhsCast);
-		}
-		else
-		{
-			Result = EmitCode(ResultType, TEXT("(% * %)"), LhsCast, RhsCast);
-		}
-		break;
-	case EBinaryOp::Div:
-		if (InputComponentType == Shader::EValueComponentType::Double)
-		{
-			Result = EmitCode(ResultType, TEXT("LWCDivide(%, %)"), LhsCast, RhsCast);
-		}
-		else
-		{
-			Result = EmitCode(ResultType, TEXT("(% / %)"), LhsCast, RhsCast);
-		}
-		break;
-	case EBinaryOp::Less:
-		ResultType = Shader::MakeValueType(Shader::EValueComponentType::Bool, NumComponents);
-		if (InputComponentType == Shader::EValueComponentType::Double)
-		{
-			Result = EmitCode(ResultType, TEXT("LWCLess(%, %)"), LhsCast, RhsCast);
-		}
-		else
-		{
-			Result = EmitCode(ResultType, TEXT("(% < %)"), LhsCast, RhsCast);
-		}
-		break;
-	default:
-		checkNoEntry();
-		break;
-	}
-
-	return Result;
-}
-
-FEmitShaderValues FEmitContext::EmitBinaryOp(EBinaryOp Op, FEmitShaderValues Lhs, FEmitShaderValues Rhs, EExpressionDerivative Derivative)
-{
-	FEmitShaderValues Result;
-	Result.Code = EmitBinaryOp(Op, Lhs.Code, Rhs.Code);
-	if (Derivative == EExpressionDerivative::Valid && Lhs.HasDerivatives() && Rhs.HasDerivatives())
-	{
-		const Shader::FValueTypeDescription ResultTypeDesc = Shader::GetValueTypeDescription(Result.Code->Type);
-		const Shader::EValueType DerivativeType = Shader::MakeValueType(Shader::EValueComponentType::Float, ResultTypeDesc.NumComponents);
-
-		switch (Op)
-		{
-		case EBinaryOp::Add:
-			Result.CodeDdx = EmitAdd(Lhs.CodeDdx, Rhs.CodeDdx);
-			Result.CodeDdy = EmitAdd(Lhs.CodeDdy, Rhs.CodeDdy);
-			break;
-		case EBinaryOp::Sub:
-			Result.CodeDdx = EmitSub(Lhs.CodeDdx, Rhs.CodeDdx);
-			Result.CodeDdy = EmitSub(Lhs.CodeDdy, Rhs.CodeDdy);
-			break;
-		case EBinaryOp::Mul:
-			Result.CodeDdx = EmitAdd(EmitMul(Lhs.CodeDdx, Rhs.Code), EmitMul(Rhs.CodeDdx, Lhs.Code));
-			Result.CodeDdy = EmitAdd(EmitMul(Lhs.CodeDdy, Rhs.Code), EmitMul(Rhs.CodeDdy, Lhs.Code));
-			break;
-		case EBinaryOp::Div:
-		{
-			FEmitShaderCode* Denom = EmitRcp(EmitMul(Rhs.Code, Rhs.Code));
-			FEmitShaderCode* dFdA = EmitMul(Rhs.Code, Denom);
-			FEmitShaderCode* dFdB = EmitNeg(EmitMul(Lhs.Code, Denom));
-			Result.CodeDdx = EmitAdd(EmitMul(dFdA, Lhs.CodeDdx), EmitMul(dFdB, Rhs.CodeDdx));
-			Result.CodeDdy = EmitAdd(EmitMul(dFdA, Lhs.CodeDdy), EmitMul(dFdB, Rhs.CodeDdy));
-			break;
-		}
-		default:
-			checkNoEntry();
-			break;
-		}
-	}
-
-	return Result;
-}
-
 } // namespace HLSLTree
 } // namespace UE
diff --git a/Engine/Source/Runtime/Engine/Private/Materials/MaterialExpressionHLSL.cpp b/Engine/Source/Runtime/Engine/Private/Materials/MaterialExpressionHLSL.cpp
--- a/Engine/Source/Runtime/Engine/Private/Materials/MaterialExpressionHLSL.cpp
+++ b/Engine/Source/Runtime/Engine/Private/Materials/MaterialExpressionHLSL.cpp
@@ -107,7 +107,7 @@ EMaterialGenerateHLSLStatus UMaterialExpressionStaticSwitch::GenerateHLSLExpress
 
 EMaterialGenerateHLSLStatus UMaterialExpressionGetLocal::GenerateHLSLExpression(FMaterialHLSLGenerator& Generator, UE::HLSLTree::FScope& Scope, int32 OutputIndex, UE::HLSLTree::FExpression*& OutExpression)
 {
-	OutExpression = Generator.AcquireLocalValue(Scope, LocalName);
+	OutExpression = Generator.GetTree().AcquireLocal(Scope, LocalName);
 	if (!OutExpression)
 	{
 		return Generator.Error(TEXT("Local accessed before assigned"));
@@ -140,14 +140,14 @@ EMaterialGenerateHLSLStatus UMaterialExpressionStaticBoolParameter::GenerateHLSL
 }
 EMaterialGenerateHLSLStatus UMaterialExpressionWorldPosition::GenerateHLSLExpression(FMaterialHLSLGenerator& Generator, UE::HLSLTree::FScope& Scope, int32 OutputIndex, UE::HLSLTree::FExpression*& OutExpression)
 {
-	UE::HLSLTree::EExternalInputType InputType = UE::HLSLTree::EExternalInputType::None;
+	UE::HLSLTree::EExternalInput InputType = UE::HLSLTree::EExternalInput::None;
 
 	switch (WorldPositionShaderOffset)
 	{
-	case WPT_Default: InputType = UE::HLSLTree::EExternalInputType::WorldPosition; break;
-	case WPT_ExcludeAllShaderOffsets: InputType = UE::HLSLTree::EExternalInputType::WorldPosition_NoOffsets; break;
-	case WPT_CameraRelative: InputType = UE::HLSLTree::EExternalInputType::TranslatedWorldPosition; break;
-	case WPT_CameraRelativeNoOffsets: InputType = UE::HLSLTree::EExternalInputType::TranslatedWorldPosition_NoOffsets; break;
+	case WPT_Default: InputType = UE::HLSLTree::EExternalInput::WorldPosition; break;
+	case WPT_ExcludeAllShaderOffsets: InputType = UE::HLSLTree::EExternalInput::WorldPosition_NoOffsets; break;
+	case WPT_CameraRelative: InputType = UE::HLSLTree::EExternalInput::TranslatedWorldPosition; break;
+	case WPT_CameraRelativeNoOffsets: InputType = UE::HLSLTree::EExternalInput::TranslatedWorldPosition_NoOffsets; break;
 	default: checkNoEntry(); break;
 	}
 
@@ -178,15 +178,18 @@ EMaterialGenerateHLSLStatus UMaterialExpressionTextureObjectParameter::GenerateH
 
 EMaterialGenerateHLSLStatus UMaterialExpressionTextureSample::GenerateHLSLExpressionBase(FMaterialHLSLGenerator& Generator, UE::HLSLTree::FScope& Scope, UE::HLSLTree::FTextureParameterDeclaration* TextureDeclaration, UE::HLSLTree::FExpression*& OutExpression)
 {
+	using namespace UE::HLSLTree;
 	if (!TextureDeclaration)
 	{
 		return Generator.Error(TEXT("Missing input texture"));
 	}
 
-	UE::HLSLTree::FExpression* TexCoordExpression = Coordinates.GetTracedInput().Expression ? Coordinates.AcquireHLSLExpression(Generator, Scope) : Generator.NewTexCoord(ConstCoordinate);
-	UE::HLSLTree::FExpressionTextureSample* ExpressionTextureSample = Generator.GetTree().NewExpression<UE::HLSLTree::FExpressionTextureSample>(TextureDeclaration, TexCoordExpression);
+	FExpression* TexCoordExpression = Coordinates.GetTracedInput().Expression ? Coordinates.AcquireHLSLExpression(Generator, Scope) : Generator.NewTexCoord(ConstCoordinate);
+	FExpressionTextureSample* ExpressionTextureSample = Generator.GetTree().NewExpression<FExpressionTextureSample>(TextureDeclaration, TexCoordExpression);
 	ExpressionTextureSample->SamplerSource = SamplerSource;
 	ExpressionTextureSample->MipValueMode = MipValueMode;
+	ExpressionTextureSample->TexCoordDerivatives = Generator.GetTree().GetAnalyticDerivatives(TexCoordExpression);
+
 	OutExpression = ExpressionTextureSample;
 	return EMaterialGenerateHLSLStatus::Success;
 }
@@ -389,7 +392,7 @@ EMaterialGenerateHLSLStatus UMaterialExpressionSetLocal::GenerateHLSLStatements(
 		return Generator.Error(TEXT("Missing value connection"));
 	}
 
-	Generator.GenerateAssignLocal(Scope, LocalName, ValueExpression);
+	Generator.GetTree().AssignLocal(Scope, LocalName, ValueExpression);
 	Exec.GenerateHLSLStatements(Generator, Scope);
 	return EMaterialGenerateHLSLStatus::Success;
 }
@@ -463,7 +466,7 @@ EMaterialGenerateHLSLStatus UMaterialExpressionForLoop::GenerateHLSLExpression(F
 		return Generator.Error(TEXT("For loop index accessed outside loop scope"));
 	}
 
-	OutExpression = Generator.AcquireLocalValue(Scope, ExpressionData->LocalName);
+	OutExpression = Generator.GetTree().AcquireLocal(Scope, ExpressionData->LocalName);
 	return EMaterialGenerateHLSLStatus::Success;
 }
 
@@ -492,7 +495,7 @@ EMaterialGenerateHLSLStatus UMaterialExpressionForLoop::GenerateHLSLStatements(F
 
 	UE::HLSLTree::FExpression* StepExpression = IndexStep.GetTracedInput().Expression ? IndexStep.AcquireHLSLExpression(Generator, Scope) : Generator.NewConstant(1);
 
-	Generator.GenerateAssignLocal(Scope, ExpressionData->LocalName, StartExpression);
+	Generator.GetTree().AssignLocal(Scope, ExpressionData->LocalName, StartExpression);
 
 	UE::HLSLTree::FStatementLoop* LoopStatement = Generator.GetTree().NewStatement<UE::HLSLTree::FStatementLoop>(Scope);
 	LoopStatement->LoopScope = Generator.NewOwnedScope(*LoopStatement);
@@ -507,13 +510,13 @@ EMaterialGenerateHLSLStatus UMaterialExpressionForLoop::GenerateHLSLStatements(F
 
 	Generator.GetTree().NewStatement<UE::HLSLTree::FStatementBreak>(*IfStatement->ElseScope);
 
-	UE::HLSLTree::FExpression* LocalExpression = Generator.AcquireLocalValue(*LoopStatement->LoopScope, ExpressionData->LocalName);
+	UE::HLSLTree::FExpression* LocalExpression = Generator.GetTree().AcquireLocal(*LoopStatement->LoopScope, ExpressionData->LocalName);
 
 	IfStatement->ConditionExpression = Generator.GetTree().NewExpression<UE::HLSLTree::FExpressionBinaryOp>(UE::HLSLTree::EBinaryOp::Less, LocalExpression, EndExpression);
 	LoopBody.GenerateHLSLStatements(Generator, *IfStatement->ThenScope);
 
-	UE::HLSLTree::FExpression* NewLocalExpression = Generator.GetTree().NewExpression<UE::HLSLTree::FExpressionBinaryOp>(UE::HLSLTree::EBinaryOp::Add, Generator.AcquireLocalValue(*IfStatement->ThenScope, ExpressionData->LocalName), StepExpression);
-	Generator.GenerateAssignLocal(*IfStatement->ThenScope, ExpressionData->LocalName, NewLocalExpression);
+	UE::HLSLTree::FExpression* NewLocalExpression = Generator.GetTree().NewExpression<UE::HLSLTree::FExpressionBinaryOp>(UE::HLSLTree::EBinaryOp::Add, Generator.GetTree().AcquireLocal(*IfStatement->ThenScope, ExpressionData->LocalName), StepExpression);
+	Generator.GetTree().AssignLocal(*IfStatement->ThenScope, ExpressionData->LocalName, NewLocalExpression);
 
 	Completed.GenerateHLSLStatements(Generator, *LoopStatement->NextScope);
 
diff --git a/Engine/Source/Runtime/Engine/Private/Materials/MaterialHLSLEmitter.cpp b/Engine/Source/Runtime/Engine/Private/Materials/MaterialHLSLEmitter.cpp
--- a/Engine/Source/Runtime/Engine/Private/Materials/MaterialHLSLEmitter.cpp
+++ b/Engine/Source/Runtime/Engine/Private/Materials/MaterialHLSLEmitter.cpp
@@ -786,7 +786,7 @@ bool MaterialEmitHLSL(const FMaterialCompileTargetParameters& InCompilerTarget,
 		}
 
 		// Prepare code for just the normal
-		FRequestedType RequestedMaterialNormal(Generator.GetMaterialAttributesType(), EComponentRequest::None);
+		FRequestedType RequestedMaterialNormal(Generator.GetMaterialAttributesType(), false);
 		RequestedMaterialNormal.SetFieldRequested(NormalField);
 		PrepareExpressionValue(EmitContext, Generator.GetResultExpression(), RequestedMaterialNormal);
 
diff --git a/Engine/Source/Runtime/Engine/Private/Materials/MaterialHLSLGenerator.cpp b/Engine/Source/Runtime/Engine/Private/Materials/MaterialHLSLGenerator.cpp
--- a/Engine/Source/Runtime/Engine/Private/Materials/MaterialHLSLGenerator.cpp
+++ b/Engine/Source/Runtime/Engine/Private/Materials/MaterialHLSLGenerator.cpp
@@ -125,23 +125,7 @@ bool FMaterialHLSLGenerator::Finalize()
 		return false;
 	}
 
-	// Resolve values for any PHI nodes that were generated
-	// Resolving a PHI may produce additional PHIs
-	while(PHIExpressions.Num() > 0)
-	{
-		UE::HLSLTree::FExpressionLocalPHI* Expression = PHIExpressions.Pop(false);
-		for (int32 i = 0; i < Expression->NumValues; ++i)
-		{
-			Expression->Values[i] = InternalAcquireLocalValue(*Expression->Scopes[i], Expression->LocalName);
-			if (!Expression->Values[i])
-			{
-				Errorf(TEXT("Local %s is not assigned on all control paths"), *Expression->LocalName.ToString());
-				return false;
-			}
-		}
-	}
-
-	return true;
+	return HLSLTree->Finalize();
 }
 
 EMaterialGenerateHLSLStatus FMaterialHLSLGenerator::Error(const FString& Message)
@@ -369,50 +353,6 @@ UE::HLSLTree::FTextureParameterDeclaration* FMaterialHLSLGenerator::AcquireTextu
 	return Declaration;
 }
 
-bool FMaterialHLSLGenerator::GenerateAssignLocal(UE::HLSLTree::FScope& Scope, const FName& LocalName, UE::HLSLTree::FExpression* Value)
-{
-	const FLocalKey Key(&Scope, LocalName);
-	LocalMap.Add(Key, Value);
-	return true;
-}
-
-UE::HLSLTree::FExpression* FMaterialHLSLGenerator::InternalAcquireLocalValue(UE::HLSLTree::FScope& Scope, const FName& LocalName)
-{
-	const FLocalKey Key(&Scope, LocalName);
-	UE::HLSLTree::FExpression** FoundExpression = LocalMap.Find(Key);
-	if (FoundExpression)
-	{
-		return *FoundExpression;
-	}
-
-	const TArrayView<UE::HLSLTree::FScope*> PreviousScopes = Scope.GetPreviousScopes();
-	if (PreviousScopes.Num() > 1)
-	{
-		UE::HLSLTree::FExpressionLocalPHI* Expression = HLSLTree->NewExpression<UE::HLSLTree::FExpressionLocalPHI>();
-		Expression->LocalName = LocalName;
-		Expression->NumValues = PreviousScopes.Num();
-		for (int32 i = 0; i < PreviousScopes.Num(); ++i)
-		{
-			Expression->Scopes[i] = PreviousScopes[i];
-		}
-		PHIExpressions.Add(Expression);
-		LocalMap.Add(Key, Expression);
-		return Expression;
-	}
-
-	if (PreviousScopes.Num() == 1)
-	{
-		return InternalAcquireLocalValue(*PreviousScopes[0], LocalName);
-	}
-
-	return nullptr;
-}
-
-UE::HLSLTree::FExpression* FMaterialHLSLGenerator::AcquireLocalValue(UE::HLSLTree::FScope& Scope, const FName& LocalName)
-{
-	return InternalAcquireLocalValue(Scope, LocalName);
-}
-
 UE::HLSLTree::FExpression* FMaterialHLSLGenerator::AcquireExpression(UE::HLSLTree::FScope& Scope, UMaterialExpression* MaterialExpression, int32 OutputIndex)
 {
 	FFunctionCallEntry* FunctionEntry = FunctionCallStack.Last();
diff --git a/Engine/Source/Runtime/Engine/Public/HLSLTree/HLSLTree.h b/Engine/Source/Runtime/Engine/Public/HLSLTree/HLSLTree.h
--- a/Engine/Source/Runtime/Engine/Public/HLSLTree/HLSLTree.h
+++ b/Engine/Source/Runtime/Engine/Public/HLSLTree/HLSLTree.h
@@ -36,7 +36,6 @@ class FScope;
 class FExpressionLocalPHI;
 class FRequestedType;
 struct FEmitShaderCode;
-struct FEmitShaderValues;
 
 static constexpr int32 MaxNumPreviousScopes = 2;
 
@@ -96,15 +95,6 @@ struct FEmitShaderCode
 	FSHAHash Hash;
 };
 
-struct FEmitShaderValues
-{
-	FEmitShaderCode* Code = nullptr;
-	FEmitShaderCode* CodeDdx = nullptr;
-	FEmitShaderCode* CodeDdy = nullptr;
-
-	bool HasDerivatives() const { return (bool)CodeDdx && (bool)CodeDdy; }
-};
-
 using FEmitShaderValueDependencies = TArray<FEmitShaderCode*, TInlineAllocator<8>>;
 
 enum class EFormatArgType : uint8
@@ -176,24 +166,8 @@ public:
 	
 	FEmitShaderCode* EmitPreshaderOrConstant(const FRequestedType& RequestedType, FExpression* Expression);
 	FEmitShaderCode* EmitConstantZero(const Shader::FType& Type);
-
 	FEmitShaderCode* EmitCast(FEmitShaderCode* ShaderValue, const Shader::FType& DestType);
-	FEmitShaderValues EmitCast(FEmitShaderValues ShaderValue, const Shader::FType& DestType);
-
-	FEmitShaderCode* EmitUnaryOp(EUnaryOp Op, FEmitShaderCode* Input);
-	FEmitShaderValues EmitUnaryOp(EUnaryOp Op, FEmitShaderValues Input, EExpressionDerivative Derivative);
-
-	FEmitShaderCode* EmitBinaryOp(EBinaryOp Op, FEmitShaderCode* Lhs, FEmitShaderCode* Rhs);
-	FEmitShaderValues EmitBinaryOp(EBinaryOp Op, FEmitShaderValues Lhs, FEmitShaderValues Rhs, EExpressionDerivative Derivative);
-
-	template<typename T> inline T EmitNeg(T Input) { return EmitUnaryOp(EUnaryOp::Neg, Input); }
-	template<typename T> inline T EmitRcp(T Input) { return EmitUnaryOp(EUnaryOp::Rcp, Input); }
-
-	template<typename T> inline T EmitAdd(T Lhs, T Rhs) { return EmitBinaryOp(EBinaryOp::Add, Lhs, Rhs); }
-	template<typename T> inline T EmitSub(T Lhs, T Rhs) { return EmitBinaryOp(EBinaryOp::Sub, Lhs, Rhs); }
-	template<typename T> inline T EmitMul(T Lhs, T Rhs) { return EmitBinaryOp(EBinaryOp::Mul, Lhs, Rhs); }
-	template<typename T> inline T EmitDiv(T Lhs, T Rhs) { return EmitBinaryOp(EBinaryOp::Div, Lhs, Rhs); }
-
+	
 	FMemStackBase* Allocator = nullptr;
 	const Shader::FStructTypeRegistry* TypeRegistry = nullptr;
 	TArray<FScope*, TInlineAllocator<16>> ScopeStack;
@@ -239,15 +213,6 @@ public:
 	bool bEmitShader = false;
 };
 
-enum class EComponentRequest : uint8
-{
-	None,
-	Requested,
-	RequestedWithDerivative,
-};
-inline bool IsRequested(EComponentRequest InRequest) { return InRequest != EComponentRequest::None; }
-inline bool IsDerivativeRequested(EComponentRequest InRequest) { return InRequest == EComponentRequest::RequestedWithDerivative; }
-
 /**
  * Like Shader::FType, but tracks which individual components are needed
  */
@@ -255,9 +220,9 @@ class FRequestedType
 {
 public:
 	FRequestedType() = default;
-	FRequestedType(int32 NumComponents, EComponentRequest DefaultRequest = EComponentRequest::Requested);
-	FRequestedType(const Shader::FType& InType, EComponentRequest DefaultRequest = EComponentRequest::Requested);
-	FRequestedType(const Shader::EValueType& InType, EComponentRequest DefaultRequest = EComponentRequest::Requested);
+	FRequestedType(int32 NumComponents, bool bDefaultRequest = true);
+	FRequestedType(const Shader::FType& InType, bool bDefaultRequest = true);
+	FRequestedType(const Shader::EValueType& InType, bool bDefaultRequest = true);
 	
 	const TCHAR* GetName() const { return GetType().GetName(); }
 	bool IsStruct() const { return StructType != nullptr; }
@@ -266,25 +231,22 @@ public:
 
 	int32 GetNumComponents() const;
 	bool IsComponentRequested(int32 Index) const { return RequestedComponents.IsValidIndex(Index) ? (bool)RequestedComponents[Index] : false; }
-	bool IsComponentDerivativeRequested(int32 Index) const { return RequestedComponentDerivatives.IsValidIndex(Index) ? (bool)RequestedComponentDerivatives[Index] : false; }
 	bool IsVoid() const { return RequestedComponents.Find(true) == INDEX_NONE; }
 
-	EComponentRequest GetComponentRequest(int32 Index) const;
-	void SetComponentRequest(int32 Index, EComponentRequest Request = EComponentRequest::Requested);
+	void SetComponentRequest(int32 Index, bool bRequest = true);
 
 	void Reset()
 	{
 		StructType = nullptr;
 		ValueComponentType = Shader::EValueComponentType::Void;
 		RequestedComponents.Reset();
-		RequestedComponentDerivatives.Reset();
 	}
 
 	/** Marks the given field as requested (or not) */
-	void SetFieldRequested(const Shader::FStructField* Field, EComponentRequest Request = EComponentRequest::Requested);
+	void SetFieldRequested(const Shader::FStructField* Field, bool bRequest = true);
 	void ClearFieldRequested(const Shader::FStructField* Field)
 	{
-		SetFieldRequested(Field, EComponentRequest::None);
+		SetFieldRequested(Field, false);
 	}
 
 	/** Marks the given field as requested, based on the input request type (which should match the field type) */
@@ -302,57 +264,10 @@ public:
 
 	/** 1 bit per component, a value of 'true' means the specified component is requsted */
 	TBitArray<> RequestedComponents;
-	TBitArray<> RequestedComponentDerivatives;
 };
 
 FRequestedType MakeRequestedType(Shader::EValueComponentType ComponentType, const FRequestedType& RequestedComponents);
 
-struct FPreparedData
-{
-	FPreparedData() = default;
-
-	FPreparedData(EExpressionEvaluation InEvaluation, EExpressionDerivative InDerivative = EExpressionDerivative::Invalid)
-		: Evaluation(InEvaluation)
-	{
-		switch (InEvaluation)
-		{
-		case EExpressionEvaluation::None:
-			Derivative = EExpressionDerivative::None;
-			break;
-		case EExpressionEvaluation::Constant:
-		case EExpressionEvaluation::Preshader:
-			Derivative = EExpressionDerivative::Zero;
-			break;
-		default:
-			Derivative = InDerivative;
-			break;
-		}
-	}
-
-	bool IsValid() const { return Evaluation != EExpressionEvaluation::None; }
-
-	EComponentRequest GetRequest() const
-	{
-		if (Evaluation != EExpressionEvaluation::None)
-		{
-			if (Derivative == EExpressionDerivative::Valid || Derivative == EExpressionDerivative::Zero)
-			{
-				return  EComponentRequest::RequestedWithDerivative;
-			}
-			return  EComponentRequest::Requested;
-		}
-		return EComponentRequest::None;
-	}
-
-	EExpressionEvaluation Evaluation = EExpressionEvaluation::None;
-	EExpressionDerivative Derivative = EExpressionDerivative::None;
-};
-
-inline FPreparedData CombinePreparedData(const FPreparedData& Lhs, const FPreparedData& Rhs)
-{
-	return FPreparedData(CombineEvaluations(Lhs.Evaluation, Rhs.Evaluation), CombineDerivatives(Lhs.Derivative, Rhs.Derivative));
-}
-
 /**
  * Like FRequestedType, but tracks an EExpressionEvaluation per component, rather than a simple requested flag
  */
@@ -365,7 +280,6 @@ public:
 	FPreparedType(const Shader::FType& InType);
 
 	void SetEvaluation(EExpressionEvaluation Evaluation);
-	void SetDerivative(EExpressionDerivative Derivative);
 
 	void SetField(const Shader::FStructField* Field, const FPreparedType& FieldType);
 	FPreparedType GetFieldType(const Shader::FStructField* Field) const;
@@ -378,23 +292,20 @@ public:
 	bool IsInitialized() const { return StructType != nullptr || ValueComponentType != Shader::EValueComponentType::Void; }
 	bool IsVoid() const;
 
-	FPreparedData GetData() const;
-	FPreparedData GetData(const FRequestedType& RequestedType) const;
-	FPreparedData GetFieldData(int32 ComponentIndex, int32 NumComponents) const;
-	FPreparedData GetComponentData(int32 Index) const;
-
-	inline EExpressionEvaluation GetEvaluation(const FRequestedType& RequestedType) const { return GetData(RequestedType).Evaluation; }
-	inline EExpressionDerivative GetDerivative(const FRequestedType& RequestedType) const { return GetData(RequestedType).Derivative; }
+	EExpressionEvaluation GetEvaluation() const;
+	EExpressionEvaluation GetEvaluation(const FRequestedType& RequestedType) const;
+	EExpressionEvaluation GetFieldEvaluation(int32 ComponentIndex, int32 NumComponents) const;
+	EExpressionEvaluation GetComponentEvaluation(int32 Index) const;
 
-	void SetComponentData(int32 Index, const FPreparedData& Data);
-	void MergeComponentData(int32 Index, EComponentRequest Request, const FPreparedData& Data);
+	void SetComponentEvaluation(int32 Index, EExpressionEvaluation Evaluation);
+	void MergeComponentEvaluation(int32 Index, EExpressionEvaluation Evaluation);
 
 	/** Unlike FRequestedType, one of these should be set */
 	const Shader::FStructType* StructType = nullptr;
 	Shader::EValueComponentType ValueComponentType = Shader::EValueComponentType::Void;
 
 	/** Evaluation type for each component, may be 'None' for components that are unused */
-	TArray<FPreparedData, TInlineAllocator<16>> PreparedComponents;
+	TArray<EExpressionEvaluation, TInlineAllocator<16>> PreparedComponents;
 };
 
 FPreparedType MergePreparedTypes(const FPreparedType& Lhs, const FPreparedType& Rhs);
@@ -404,14 +315,9 @@ class FPrepareValueResult
 public:
 	const FPreparedType& GetPreparedType() const { return PreparedType; }
 
-	void SetType(FEmitContext& Context, const FRequestedType& RequestedType, const FPreparedData& Data, const Shader::FType& Type);
+	void SetType(FEmitContext& Context, const FRequestedType& RequestedType, EExpressionEvaluation Evaluation, const Shader::FType& Type);
 	void SetType(FEmitContext& Context, const FRequestedType& RequestedType, const FPreparedType& Type);
 
-	void SetTypeWithDerivative(FEmitContext& Context, const FRequestedType& RequestedType, EExpressionEvaluation Evaluation, const Shader::FType& Type)
-	{
-		SetType(Context, RequestedType, FPreparedData(Evaluation, EExpressionDerivative::Valid), Type);
-	}
-
 	void SetForwardValue(FEmitContext& Context, const FRequestedType& RequestedType, FExpression* InValue);
 
 private:
@@ -423,6 +329,27 @@ private:
 	friend class FExpression;
 };
 
+struct FEmitValueShaderResult
+{
+	FEmitShaderCode* Code = nullptr;
+};
+
+enum class EDerivativeCoordinate : uint8
+{
+	Ddx,
+	Ddy,
+};
+
+struct FExpressionDerivatives
+{
+	FExpression* ExpressionDdx = nullptr;
+	FExpression* ExpressionDdy = nullptr;
+
+	FExpression* Get(EDerivativeCoordinate Coord) const { return (Coord == EDerivativeCoordinate::Ddx) ? ExpressionDdx : ExpressionDdy; }
+
+	bool IsValid() const { return (bool)ExpressionDdx && (bool)ExpressionDdy; }
+};
+
 /**
  * Represents an HLSL expression.  This is a piece of code that evaluates to a value, but has no side effects.
  * Unlike statements, expressions are not expected to execute in any particular order.  They may be cached (or not) in generated code, without the underlying implementation needing to care.
@@ -436,30 +363,58 @@ public:
 	FRequestedType GetRequestedType() const { return PrepareValueResult.PreparedType.GetRequestedType(); }
 	Shader::FType GetType() const { return PrepareValueResult.PreparedType.GetType(); }
 	EExpressionEvaluation GetEvaluation(const FRequestedType& RequestedType) const { return PrepareValueResult.PreparedType.GetEvaluation(RequestedType); }
-	EExpressionDerivative GetDerivative(const FRequestedType& RequestedType) const { return PrepareValueResult.PreparedType.GetDerivative(RequestedType); }
 
 	virtual void Reset() override;
 
 	friend const FPreparedType& PrepareExpressionValue(FEmitContext& Context, FExpression* InExpression, const FRequestedType& RequestedType);
 
-	FEmitShaderValues GetValueShader(FEmitContext& Context, const FRequestedType& RequestedType, const Shader::FType& ResultType);
+	FEmitShaderCode* GetValueShader(FEmitContext& Context, const FRequestedType& RequestedType, const Shader::FType& ResultType);
 	void GetValuePreshader(FEmitContext& Context, const FRequestedType& RequestedType, Shader::FPreshaderData& OutPreshader);
 	Shader::FValue GetValueConstant(FEmitContext& Context, const FRequestedType& RequestedType);
 
-	FEmitShaderValues GetValueShader(FEmitContext& Context, const FRequestedType& RequestedType);
-	FEmitShaderValues GetValueShader(FEmitContext& Context);
+	FEmitShaderCode* GetValueShader(FEmitContext& Context, const FRequestedType& RequestedType);
+	FEmitShaderCode* GetValueShader(FEmitContext& Context);
 
 protected:
+	virtual void ComputeAnalyticDerivatives(FTree& Tree, FExpressionDerivatives& OutResult) const;
 	virtual void PrepareValue(FEmitContext& Context, const FRequestedType& RequestedType, FPrepareValueResult& OutResult) const = 0;
-	virtual void EmitValueShader(FEmitContext& Context, const FRequestedType& RequestedType, FEmitShaderValues& OutResult) const;
+	virtual void EmitValueShader(FEmitContext& Context, const FRequestedType& RequestedType, FEmitValueShaderResult& OutResult) const;
 	virtual void EmitValuePreshader(FEmitContext& Context, const FRequestedType& RequestedType, Shader::FPreshaderData& OutPreshader) const;
 
 private:
+	FExpressionDerivatives Derivatives;
 	FRequestedType CurrentRequestedType;
 	FPrepareValueResult PrepareValueResult;
 	bool bReentryFlag = false;
+	bool bComputedDerivatives = false;
 
+	friend class FTree;
 	friend class FEmitContext;
+	friend class FExpressionReentryScope;
+};
+
+class FExpressionReentryScope
+{
+public:
+	FExpressionReentryScope(FExpression* InExpression) : Expression(InExpression)
+	{
+		if (Expression)
+		{
+			check(!Expression->bReentryFlag);
+			Expression->bReentryFlag = true;
+		}
+	}
+
+	~FExpressionReentryScope()
+	{
+		if (Expression)
+		{
+			check(Expression->bReentryFlag);
+			Expression->bReentryFlag = false;
+		}
+	}
+
+	FExpression* Expression;
 };
 
 
@@ -471,13 +426,18 @@ private:
 class FExpressionLocalPHI final : public FExpression
 {
 public:
+	FExpressionLocalPHI() = default;
+	FExpressionLocalPHI(const FExpressionLocalPHI* Source, EDerivativeCoordinate Coord);
+
+	virtual void ComputeAnalyticDerivatives(FTree& Tree, FExpressionDerivatives& OutResult) const override;
 	virtual void PrepareValue(FEmitContext& Context, const FRequestedType& RequestedType, FPrepareValueResult& OutResult) const override;
-	virtual void EmitValueShader(FEmitContext& Context, const FRequestedType& RequestedType, FEmitShaderValues& OutResult) const override;
+	virtual void EmitValueShader(FEmitContext& Context, const FRequestedType& RequestedType, FEmitValueShaderResult& OutResult) const override;
 
+	TArray<EDerivativeCoordinate, TInlineAllocator<8>> DerivativeChain;
 	FName LocalName;
 	FScope* Scopes[MaxNumPreviousScopes];
 	FExpression* Values[MaxNumPreviousScopes];
-	int32 NumValues;
+	int32 NumValues = 0;
 };
 
 /**
@@ -595,6 +555,7 @@ private:
 	FScope* ParentScope = nullptr;
 	FStatement* ContainedStatement = nullptr;
 	FScope* PreviousScope[MaxNumPreviousScopes];
+	TMap<FName, FExpression*> LocalMap;
 	FCodeList Declarations;
 	FCodeList Statements;
 	int32 NumPreviousScopes = 0;
@@ -636,6 +597,8 @@ public:
 
 	void ResetNodes();
 
+	bool Finalize();
+
 	bool EmitShader(FEmitContext& Context, FStringBuilderBase& OutCode) const;
 
 	FScope& GetRootScope() const { return *RootScope; }
@@ -656,9 +619,26 @@ public:
 		return Statement;
 	}
 
+	void AssignLocal(FScope& Scope, const FName& LocalName, FExpression* Value);
+	FExpression* AcquireLocal(FScope& Scope, const FName& LocalName);
+
+	const FExpressionDerivatives& GetAnalyticDerivatives(FExpression* InExpression);
+
 	FScope* NewScope(FScope& Scope);
 	FScope* NewOwnedScope(FStatement& Owner);
 
+	FExpression* NewConstant(const Shader::FValue& Value);
+	FExpression* NewUnaryOp(EUnaryOp Op, FExpression* Input);
+	FExpression* NewBinaryOp(EBinaryOp Op, FExpression* Lhs, FExpression* Rhs);
+
+	FExpression* NewNeg(FExpression* Input) { return NewUnaryOp(EUnaryOp::Neg, Input); }
+	FExpression* NewRcp(FExpression* Input) { return NewUnaryOp(EUnaryOp::Rcp, Input); }
+
+	FExpression* NewAdd(FExpression* Lhs, FExpression* Rhs) { return NewBinaryOp(EBinaryOp::Add, Lhs, Rhs); }
+	FExpression* NewSub(FExpression* Lhs, FExpression* Rhs) { return NewBinaryOp(EBinaryOp::Sub, Lhs, Rhs); }
+	FExpression* NewMul(FExpression* Lhs, FExpression* Rhs) { return NewBinaryOp(EBinaryOp::Mul, Lhs, Rhs); }
+	FExpression* NewDiv(FExpression* Lhs, FExpression* Rhs) { return NewBinaryOp(EBinaryOp::Div, Lhs, Rhs); }
+
 	FTextureParameterDeclaration* NewTextureParameterDeclaration(const FName& Name, const FTextureDescription& DefaultValue);
 
 private:
@@ -677,6 +657,9 @@ private:
 	FMemStackBase* Allocator = nullptr;
 	FNode* Nodes = nullptr;
 	FScope* RootScope = nullptr;
+	TArray<FExpressionLocalPHI*> PHIExpressions;
+
+	friend class FExpressionLocalPHI;
 };
 
 //Shader::EValueType RequestExpressionType(FEmitContext& Context, FExpression* InExpression, int8 InRequestedNumComponents); // friend of FExpression
diff --git a/Engine/Source/Runtime/Engine/Public/HLSLTree/HLSLTreeCommon.h b/Engine/Source/Runtime/Engine/Public/HLSLTree/HLSLTreeCommon.h
--- a/Engine/Source/Runtime/Engine/Public/HLSLTree/HLSLTreeCommon.h
+++ b/Engine/Source/Runtime/Engine/Public/HLSLTree/HLSLTreeCommon.h
@@ -20,6 +20,7 @@ public:
 
 	Shader::FValue Value;
 
+	virtual void ComputeAnalyticDerivatives(FTree& Tree, FExpressionDerivatives& OutResult) const override;
 	virtual void PrepareValue(FEmitContext& Context, const FRequestedType& RequestedType, FPrepareValueResult& OutResult) const override;
 	virtual void EmitValuePreshader(FEmitContext& Context, const FRequestedType& RequestedType, Shader::FPreshaderData& OutPreshader) const override;
 };
@@ -35,13 +36,15 @@ public:
 	Shader::FValue DefaultValue;
 	EMaterialParameterType ParameterType;
 
+	virtual void ComputeAnalyticDerivatives(FTree& Tree, FExpressionDerivatives& OutResult) const override;
 	virtual void PrepareValue(FEmitContext& Context, const FRequestedType& RequestedType, FPrepareValueResult& OutResult) const override;
 	virtual void EmitValuePreshader(FEmitContext& Context, const FRequestedType& RequestedType, Shader::FPreshaderData& OutPreshader) const override;
 };
 
-enum class EExternalInputType : uint8
+enum class EExternalInput : uint8
 {
 	None,
+
 	TexCoord0,
 	TexCoord1,
 	TexCoord2,
@@ -50,55 +53,77 @@ enum class EExternalInputType : uint8
 	TexCoord5,
 	TexCoord6,
 	TexCoord7,
+
+	TexCoord0_Ddx,
+	TexCoord1_Ddx,
+	TexCoord2_Ddx,
+	TexCoord3_Ddx,
+	TexCoord4_Ddx,
+	TexCoord5_Ddx,
+	TexCoord6_Ddx,
+	TexCoord7_Ddx,
+
+	TexCoord0_Ddy,
+	TexCoord1_Ddy,
+	TexCoord2_Ddy,
+	TexCoord3_Ddy,
+	TexCoord4_Ddy,
+	TexCoord5_Ddy,
+	TexCoord6_Ddy,
+	TexCoord7_Ddy,
+
 	WorldPosition,
 	WorldPosition_NoOffsets,
 	TranslatedWorldPosition,
 	TranslatedWorldPosition_NoOffsets,
+
+	WorldPosition_Ddx,
+	WorldPosition_Ddy,
 };
 static constexpr int32 NumTexCoords = 8;
 
-inline bool IsTexCoord(EExternalInputType Type)
+struct FExternalInputDescription
 {
-	return FMath::IsWithin((int32)Type, (int32)EExternalInputType::TexCoord0, (int32)EExternalInputType::TexCoord0 + NumTexCoords);
-}
+	FExternalInputDescription(const TCHAR* InName, Shader::EValueType InType, EExternalInput InDdx = EExternalInput::None, EExternalInput InDdy = EExternalInput::None)
+		: Name(InName), Type(InType), Ddx(InDdx), Ddy(InDdy)
+	{}
 
-inline Shader::EValueType GetInputExpressionType(EExternalInputType Type)
-{
-	if (IsTexCoord(Type))
-	{
-		return Shader::EValueType::Float2;
-	}
+	const TCHAR* Name;
+	Shader::EValueType Type;
+	EExternalInput Ddx;
+	EExternalInput Ddy;
+};
 
-	switch (Type)
-	{
-	case EExternalInputType::None:
-		return Shader::EValueType::Void;
-	case EExternalInputType::WorldPosition:
-	case EExternalInputType::WorldPosition_NoOffsets:
-		return Shader::EValueType::Double3;
-	case EExternalInputType::TranslatedWorldPosition:
-	case EExternalInputType::TranslatedWorldPosition_NoOffsets:
-		return Shader::EValueType::Float3;
-	default:
-		checkNoEntry();
-		return Shader::EValueType::Void;
-	}
+FExternalInputDescription GetExternalInputDescription(EExternalInput Input);
+
+inline bool IsTexCoord(EExternalInput Type)
+{
+	return FMath::IsWithin((int32)Type, (int32)EExternalInput::TexCoord0, (int32)EExternalInput::TexCoord0 + NumTexCoords);
+}
+inline bool IsTexCoord_Ddx(EExternalInput Type)
+{
+	return FMath::IsWithin((int32)Type, (int32)EExternalInput::TexCoord0_Ddx, (int32)EExternalInput::TexCoord0_Ddx + NumTexCoords);
+}
+inline bool IsTexCoord_Ddy(EExternalInput Type)
+{
+	return FMath::IsWithin((int32)Type, (int32)EExternalInput::TexCoord0_Ddy, (int32)EExternalInput::TexCoord0_Ddy + NumTexCoords);
 }
-inline EExternalInputType MakeInputTexCoord(int32 Index)
+inline EExternalInput MakeInputTexCoord(int32 Index)
 {
 	check(Index >= 0 && Index < NumTexCoords);
-	return (EExternalInputType)((int32)EExternalInputType::TexCoord0 + Index);
+	return (EExternalInput)((int32)EExternalInput::TexCoord0 + Index);
 }
 
 class FExpressionExternalInput : public FExpression
 {
 public:
-	FExpressionExternalInput(EExternalInputType InInputType) : InputType(InInputType) {}
+	FExpressionExternalInput(EExternalInput InInputType) : InputType(InInputType) {}
 
-	EExternalInputType InputType;
+	EExternalInput InputType;
 
+	virtual void ComputeAnalyticDerivatives(FTree& Tree, FExpressionDerivatives& OutResult) const override;
 	virtual void PrepareValue(FEmitContext& Context, const FRequestedType& RequestedType, FPrepareValueResult& OutResult) const override;
-	virtual void EmitValueShader(FEmitContext& Context, const FRequestedType& RequestedType, FEmitShaderValues& OutResult) const override;
+	virtual void EmitValueShader(FEmitContext& Context, const FRequestedType& RequestedType, FEmitValueShaderResult& OutResult) const override;
 };
 
 class FExpressionTextureSample : public FExpression
@@ -113,11 +138,12 @@ public:
 
 	FTextureParameterDeclaration* Declaration;
 	FExpression* TexCoordExpression;
+	FExpressionDerivatives TexCoordDerivatives;
 	ESamplerSourceMode SamplerSource;
 	ETextureMipValueMode MipValueMode;
 
 	virtual void PrepareValue(FEmitContext& Context, const FRequestedType& RequestedType, FPrepareValueResult& OutResult) const override;
-	virtual void EmitValueShader(FEmitContext& Context, const FRequestedType& RequestedType, FEmitShaderValues& OutResult) const override;
+	virtual void EmitValueShader(FEmitContext& Context, const FRequestedType& RequestedType, FEmitValueShaderResult& OutResult) const override;
 };
 
 class FExpressionGetStructField : public FExpression
@@ -135,8 +161,9 @@ public:
 	const Shader::FStructField* Field;
 	FExpression* StructExpression;
 
+	virtual void ComputeAnalyticDerivatives(FTree& Tree, FExpressionDerivatives& OutResult) const override;
 	virtual void PrepareValue(FEmitContext& Context, const FRequestedType& RequestedType, FPrepareValueResult& OutResult) const override;
-	virtual void EmitValueShader(FEmitContext& Context, const FRequestedType& RequestedType, FEmitShaderValues& OutResult) const override;
+	virtual void EmitValueShader(FEmitContext& Context, const FRequestedType& RequestedType, FEmitValueShaderResult& OutResult) const override;
 	virtual void EmitValuePreshader(FEmitContext& Context, const FRequestedType& RequestedType, Shader::FPreshaderData& OutPreshader) const override;
 };
 
@@ -157,8 +184,9 @@ public:
 	FExpression* StructExpression;
 	FExpression* FieldExpression;
 
+	virtual void ComputeAnalyticDerivatives(FTree& Tree, FExpressionDerivatives& OutResult) const override;
 	virtual void PrepareValue(FEmitContext& Context, const FRequestedType& RequestedType, FPrepareValueResult& OutResult) const override;
-	virtual void EmitValueShader(FEmitContext& Context, const FRequestedType& RequestedType, FEmitShaderValues& OutResult) const override;
+	virtual void EmitValueShader(FEmitContext& Context, const FRequestedType& RequestedType, FEmitValueShaderResult& OutResult) const override;
 	virtual void EmitValuePreshader(FEmitContext& Context, const FRequestedType& RequestedType, Shader::FPreshaderData& OutPreshader) const override;
 };
 
@@ -176,7 +204,7 @@ public:
 	FExpression* FalseExpression;
 
 	virtual void PrepareValue(FEmitContext& Context, const FRequestedType& RequestedType, FPrepareValueResult& OutResult) const override;
-	virtual void EmitValueShader(FEmitContext& Context, const FRequestedType& RequestedType, FEmitShaderValues& OutResult) const override;
+	virtual void EmitValueShader(FEmitContext& Context, const FRequestedType& RequestedType, FEmitValueShaderResult& OutResult) const override;
 	virtual void EmitValuePreshader(FEmitContext& Context, const FRequestedType& RequestedType, Shader::FPreshaderData& OutPreshader) const override;
 };
 
@@ -191,8 +219,9 @@ public:
 	EUnaryOp Op;
 	FExpression* Input;
 
+	virtual void ComputeAnalyticDerivatives(FTree& Tree, FExpressionDerivatives& OutResult) const override;
 	virtual void PrepareValue(FEmitContext& Context, const FRequestedType& RequestedType, FPrepareValueResult& OutResult) const override;
-	virtual void EmitValueShader(FEmitContext& Context, const FRequestedType& RequestedType, FEmitShaderValues& OutResult) const override;
+	virtual void EmitValueShader(FEmitContext& Context, const FRequestedType& RequestedType, FEmitValueShaderResult& OutResult) const override;
 	virtual void EmitValuePreshader(FEmitContext& Context, const FRequestedType& RequestedType, Shader::FPreshaderData& OutPreshader) const override;
 };
 
@@ -209,8 +238,9 @@ public:
 	FExpression* Lhs;
 	FExpression* Rhs;
 
+	virtual void ComputeAnalyticDerivatives(FTree& Tree, FExpressionDerivatives& OutResult) const override;
 	virtual void PrepareValue(FEmitContext& Context, const FRequestedType& RequestedType, FPrepareValueResult& OutResult) const override;
-	virtual void EmitValueShader(FEmitContext& Context, const FRequestedType& RequestedType, FEmitShaderValues& OutResult) const override;
+	virtual void EmitValueShader(FEmitContext& Context, const FRequestedType& RequestedType, FEmitValueShaderResult& OutResult) const override;
 	virtual void EmitValuePreshader(FEmitContext& Context, const FRequestedType& RequestedType, Shader::FPreshaderData& OutPreshader) const override;
 };
 
@@ -238,8 +268,9 @@ public:
 	FSwizzleParameters Parameters;
 	FExpression* Input;
 
+	virtual void ComputeAnalyticDerivatives(FTree& Tree, FExpressionDerivatives& OutResult) const override;
 	virtual void PrepareValue(FEmitContext& Context, const FRequestedType& RequestedType, FPrepareValueResult& OutResult) const override;
-	virtual void EmitValueShader(FEmitContext& Context, const FRequestedType& RequestedType, FEmitShaderValues& OutResult) const override;
+	virtual void EmitValueShader(FEmitContext& Context, const FRequestedType& RequestedType, FEmitValueShaderResult& OutResult) const override;
 	virtual void EmitValuePreshader(FEmitContext& Context, const FRequestedType& RequestedType, Shader::FPreshaderData& OutPreshader) const override;
 };
 
@@ -254,8 +285,9 @@ public:
 	FExpression* Lhs;
 	FExpression* Rhs;
 
+	virtual void ComputeAnalyticDerivatives(FTree& Tree, FExpressionDerivatives& OutResult) const override;
 	virtual void PrepareValue(FEmitContext& Context, const FRequestedType& RequestedType, FPrepareValueResult& OutResult) const override;
-	virtual void EmitValueShader(FEmitContext& Context, const FRequestedType& RequestedType, FEmitShaderValues& OutResult) const override;
+	virtual void EmitValueShader(FEmitContext& Context, const FRequestedType& RequestedType, FEmitValueShaderResult& OutResult) const override;
 	virtual void EmitValuePreshader(FEmitContext& Context, const FRequestedType& RequestedType, Shader::FPreshaderData& OutPreshader) const override;
 };
 
@@ -263,7 +295,7 @@ class FExpressionReflectionVector : public FExpression
 {
 public:
 	virtual void PrepareValue(FEmitContext& Context, const FRequestedType& RequestedType, FPrepareValueResult& OutResult) const override;
-	virtual void EmitValueShader(FEmitContext& Context, const FRequestedType& RequestedType, FEmitShaderValues& OutResult) const override;
+	virtual void EmitValueShader(FEmitContext& Context, const FRequestedType& RequestedType, FEmitValueShaderResult& OutResult) const override;
 };
 
 class FStatementReturn : public FStatement
diff --git a/Engine/Source/Runtime/Engine/Public/HLSLTree/HLSLTreeTypes.h b/Engine/Source/Runtime/Engine/Public/HLSLTree/HLSLTreeTypes.h
--- a/Engine/Source/Runtime/Engine/Public/HLSLTree/HLSLTreeTypes.h
+++ b/Engine/Source/Runtime/Engine/Public/HLSLTree/HLSLTreeTypes.h
@@ -45,23 +45,6 @@ enum class EExpressionEvaluation : uint8
 
 EExpressionEvaluation CombineEvaluations(EExpressionEvaluation Lhs, EExpressionEvaluation Rhs);
 
-enum class EExpressionDerivative : uint8
-{
-	/** Uninitialized/unknown */
-	None,
-
-	/** Derivative is not valid */
-	Invalid,
-
-	/** Derivative is known to be 0 */
-	Zero,
-
-	/** Derivative is valid */
-	Valid,
-};
-
-EExpressionDerivative CombineDerivatives(EExpressionDerivative Lhs, EExpressionDerivative Rhs);
-
 enum class EUnaryOp : uint8
 {
 	None,
diff --git a/Engine/Source/Runtime/Engine/Public/MaterialHLSLGenerator.h b/Engine/Source/Runtime/Engine/Public/MaterialHLSLGenerator.h
--- a/Engine/Source/Runtime/Engine/Public/MaterialHLSLGenerator.h
+++ b/Engine/Source/Runtime/Engine/Public/MaterialHLSLGenerator.h
@@ -102,9 +102,6 @@ public:
 	/** Returns a declaration to access the given texture parameter */
 	UE::HLSLTree::FTextureParameterDeclaration* AcquireTextureParameterDeclaration(const FName& Name, const UE::HLSLTree::FTextureDescription& DefaultValue);
 
-	bool GenerateAssignLocal(UE::HLSLTree::FScope& Scope, const FName& LocalName, UE::HLSLTree::FExpression* Value);
-	UE::HLSLTree::FExpression* AcquireLocalValue(UE::HLSLTree::FScope& Scope, const FName& LocalName);
-
 	/**
 	 * Returns the appropriate HLSLNode representing the given UMaterialExpression.
 	 * The node will be created if it doesn't exist. Otherwise, the tree will be updated to ensure the node is visible in the given scope
@@ -165,24 +162,6 @@ private:
 		}
 	};
 
-	struct FLocalKey
-	{
-		FLocalKey(UE::HLSLTree::FScope* InScope, const FName& InName) : Scope(InScope), Name(InName) {}
-
-		UE::HLSLTree::FScope* Scope;
-		FName Name;
-
-		friend inline uint32 GetTypeHash(const FLocalKey& Value)
-		{
-			return HashCombine(GetTypeHash(Value.Scope), GetTypeHash(Value.Name));
-		}
-
-		friend inline bool operator==(const FLocalKey& Lhs, const FLocalKey& Rhs)
-		{
-			return Lhs.Scope == Rhs.Scope && Lhs.Name == Rhs.Name;
-		}
-	};
-
 	struct FExpressionDataKey
 	{
 		FExpressionDataKey(const FName& InTypeName, const UMaterialExpression* InMaterialExpression) : MaterialExpression(InMaterialExpression), TypeName(InTypeName) {}
@@ -213,8 +192,6 @@ private:
 		int32 NumInputs = 0;
 	};
 
-	UE::HLSLTree::FExpression* InternalAcquireLocalValue(UE::HLSLTree::FScope& Scope, const FName& LocalName);
-
 	void InternalRegisterExpressionData(const FName& Type, const UMaterialExpression* MaterialExpression, void* Data);
 	void* InternalFindExpressionData(const FName& Type, const UMaterialExpression* MaterialExpression);
 
@@ -232,13 +209,11 @@ private:
 	TArray<FExpressionKey> ExpressionStack;
 	TArray<FFunctionCallEntry*> FunctionCallStack;
 	TArray<UE::HLSLTree::FScope*> JoinedScopeStack;
-	TArray<UE::HLSLTree::FExpressionLocalPHI*> PHIExpressions;
 	TArray<FString> CompileErrors;
 	TArray<UMaterialExpression*> ErrorExpressions;
 	TMap<UE::HLSLTree::FTextureDescription, UE::HLSLTree::FTextureParameterDeclaration*> TextureDeclarationMap;
 	TMap<FName, UE::HLSLTree::FTextureParameterDeclaration*> TextureParameterDeclarationMap;
 	TMap<FSHAHash, FFunctionCallEntry*> FunctionCallMap;
-	TMap<FLocalKey, UE::HLSLTree::FExpression*> LocalMap;
 	TMap<UMaterialExpression*, FStatementEntry> StatementMap;
 	TMap<FExpressionDataKey, void*> ExpressionDataMap;
 	bool bGeneratedResult;
diff --git a/Engine/Source/Runtime/Engine/Public/Shader/Preshader.h b/Engine/Source/Runtime/Engine/Public/Shader/Preshader.h
--- a/Engine/Source/Runtime/Engine/Public/Shader/Preshader.h
+++ b/Engine/Source/Runtime/Engine/Public/Shader/Preshader.h
@@ -8,6 +8,7 @@
 
 class FUniformExpressionSet;
 struct FMaterialRenderContext;
+class FMaterial;
 
 namespace UE
 {
