Remove caching of UMaterialExpression -> HLSLTree::FExpression in new material generator
Generated expression may depend on local variable state, which is difficult to cache (and not previously captured correctly)
Duplicated FExpressions will still de-duplicate results as needed, so simplest solution is to just remove this cache
It can be re-added with proper local variable support later, if needed for optimization

diff --git a/Engine/Source/Runtime/Engine/Private/Materials/MaterialExpressionHLSL.cpp b/Engine/Source/Runtime/Engine/Private/Materials/MaterialExpressionHLSL.cpp
--- a/Engine/Source/Runtime/Engine/Private/Materials/MaterialExpressionHLSL.cpp
+++ b/Engine/Source/Runtime/Engine/Private/Materials/MaterialExpressionHLSL.cpp
@@ -359,7 +359,12 @@ EMaterialGenerateHLSLStatus UMaterialExpressionReflectionVectorWS::GenerateHLSLE
 
 EMaterialGenerateHLSLStatus UMaterialExpressionFunctionInput::GenerateHLSLExpression(FMaterialHLSLGenerator& Generator, UE::HLSLTree::FScope& Scope, int32 OutputIndex, UE::HLSLTree::FExpression*& OutExpression)
 {
-	return Generator.Error(TEXT("Invalid"));
+	OutExpression = Generator.AcquireFunctionInputExpression(Scope, this);
+	if (OutExpression)
+	{
+		return EMaterialGenerateHLSLStatus::Success;
+	}
+	return Generator.Error("Invalid material function input");
 }
 
 EMaterialGenerateHLSLStatus UMaterialExpressionMaterialFunctionCall::GenerateHLSLExpression(FMaterialHLSLGenerator& Generator, UE::HLSLTree::FScope& Scope, int32 OutputIndex, UE::HLSLTree::FExpression*& OutExpression)
diff --git a/Engine/Source/Runtime/Engine/Private/Materials/MaterialHLSLGenerator.cpp b/Engine/Source/Runtime/Engine/Private/Materials/MaterialHLSLGenerator.cpp
--- a/Engine/Source/Runtime/Engine/Private/Materials/MaterialHLSLGenerator.cpp
+++ b/Engine/Source/Runtime/Engine/Private/Materials/MaterialHLSLGenerator.cpp
@@ -135,7 +135,7 @@ EMaterialGenerateHLSLStatus FMaterialHLSLGenerator::Error(const FString& Message
 
 	if (ExpressionStack.Num() > 0)
 	{
-		UMaterialExpression* ErrorExpression = ExpressionStack.Last().Expression;
+		UMaterialExpression* ErrorExpression = ExpressionStack.Last();
 		check(ErrorExpression);
 
 		if (ErrorExpression->GetClass() != UMaterialExpressionMaterialFunctionCall::StaticClass()
@@ -355,25 +355,34 @@ UE::HLSLTree::FTextureParameterDeclaration* FMaterialHLSLGenerator::AcquireTextu
 
 UE::HLSLTree::FExpression* FMaterialHLSLGenerator::AcquireExpression(UE::HLSLTree::FScope& Scope, UMaterialExpression* MaterialExpression, int32 OutputIndex)
 {
-	FFunctionCallEntry* FunctionEntry = FunctionCallStack.Last();
-
-	const FExpressionKey Key(MaterialExpression, OutputIndex);
-	UE::HLSLTree::FExpression** PrevExpression = FunctionEntry->ExpressionMap.Find(Key);
 	UE::HLSLTree::FExpression* Expression = nullptr;
-	if (!PrevExpression)
+
+	ExpressionStack.Add(MaterialExpression);
+	const EMaterialGenerateHLSLStatus Status = MaterialExpression->GenerateHLSLExpression(*this, Scope, OutputIndex, Expression);
+	verify(ExpressionStack.Pop() == MaterialExpression);
+
+	return Expression;
+}
+
+UE::HLSLTree::FExpression* FMaterialHLSLGenerator::AcquireFunctionInputExpression(UE::HLSLTree::FScope& Scope, const UMaterialExpressionFunctionInput* MaterialExpression)
+{
+	const FFunctionCallEntry* FunctionEntry = FunctionCallStack.Last();
+	const FFunctionInput* ConnectedInput = nullptr;
+	for (const FFunctionInput& Input : FunctionEntry->Inputs)
 	{
-		// TODO - need to rethink this caching, won't work to cache expressions that depend on current value of local variables
-		// May just remove caching at this level....need to rework function inputs in this case
-		ExpressionStack.Add(Key);
-		const EMaterialGenerateHLSLStatus Status = MaterialExpression->GenerateHLSLExpression(*this, Scope, OutputIndex, Expression);
-		verify(ExpressionStack.Pop() == Key);
-		FunctionEntry->ExpressionMap.Add(Key, Expression);
+		if (Input.FunctionInputExpression == MaterialExpression)
+		{
+			ConnectedInput = &Input;
+			break;
+		}
 	}
-	else
+
+	if (!ConnectedInput)
 	{
-		Expression = *PrevExpression;
+		return nullptr;
 	}
-	return Expression;
+
+	return AcquireExpression(Scope, ConnectedInput->ConnectedExpression, ConnectedInput->ConnectedOutputIndex);
 }
 
 UE::HLSLTree::FTextureParameterDeclaration* FMaterialHLSLGenerator::AcquireTextureDeclaration(UE::HLSLTree::FScope& Scope, UMaterialExpression* MaterialExpression, int32 OutputIndex)
@@ -416,29 +425,18 @@ bool FMaterialHLSLGenerator::GenerateStatements(UE::HLSLTree::FScope& Scope, UMa
 			}
 		}
 
-		const FExpressionKey Key(MaterialExpression);
-		ExpressionStack.Add(Key);
+		ExpressionStack.Add(MaterialExpression);
 		const EMaterialGenerateHLSLStatus Status = MaterialExpression->GenerateHLSLStatements(*this, *ScopeToUse);
-		verify(ExpressionStack.Pop() == Key);
+		verify(ExpressionStack.Pop() == MaterialExpression);
 	}
 
 	return bResult;
 }
 
-template<typename T>
-TArrayView<T> CopyArrayView(FMemStackBase& Allocator, TArrayView<T> InView)
+UE::HLSLTree::FExpression* FMaterialHLSLGenerator::GenerateFunctionCall(UE::HLSLTree::FScope& Scope, UMaterialFunctionInterface* Function, TArrayView<const FFunctionExpressionInput> ConnectedInputs, int32 OutputIndex)
 {
-	const int32 Num = InView.Num();
-	T* Result = new(Allocator) T[Num];
-	for (int32 i = 0; i < Num; ++i)
-	{
-		Result[i] = InView[i];
-	}
-	return MakeArrayView(Result, Num);
-}
+	using namespace UE::HLSLTree;
 
-UE::HLSLTree::FExpression* FMaterialHLSLGenerator::GenerateFunctionCall(UE::HLSLTree::FScope& Scope, UMaterialFunctionInterface* Function, TArrayView<const FFunctionExpressionInput> Inputs, int32 OutputIndex)
-{
 	if (!Function)
 	{
 		Error(TEXT("Missing material function"));
@@ -449,7 +447,7 @@ UE::HLSLTree::FExpression* FMaterialHLSLGenerator::GenerateFunctionCall(UE::HLSL
 	TArray<FFunctionExpressionOutput> FunctionOutputs;
 	Function->GetInputsAndOutputs(FunctionInputs, FunctionOutputs);
 
-	if (FunctionInputs.Num() != Inputs.Num())
+	if (FunctionInputs.Num() != ConnectedInputs.Num())
 	{
 		Error(TEXT("Mismatched function inputs"));
 		return nullptr;
@@ -465,22 +463,25 @@ UE::HLSLTree::FExpression* FMaterialHLSLGenerator::GenerateFunctionCall(UE::HLSL
 	FSHA1 Hasher;
 	Hasher.Update((uint8*)&Function, sizeof(UMaterialFunctionInterface*));
 
-	TArray<UE::HLSLTree::FExpression*> InputExpressions;
-	InputExpressions.Reserve(Inputs.Num());
-	for (int32 InputIndex = 0; InputIndex < Inputs.Num(); ++InputIndex)
+	FFunctionInputArray Inputs;
+	for (int32 InputIndex = 0; InputIndex < ConnectedInputs.Num(); ++InputIndex)
 	{
-		const FFunctionExpressionInput& Input = Inputs[InputIndex];
-		UE::HLSLTree::FExpression* InputExpression = Input.Input.AcquireHLSLExpression(*this, Scope);
-
-		check(InputExpression);
-		InputExpressions.Add(InputExpression);
-		Hasher.Update((uint8*)&InputExpression, sizeof(UE::HLSLTree::FExpression*));
+		// FunctionInputs are the inputs from the UMaterialFunction object
+		// ConnectedInputs are the inputs from the UMaterialFunctionCall object
+		// We want to connect the UMaterialExpressionFunctionInput from the UMaterialFunction to whatever UMaterialExpression is passed to the UMaterialFunctionCall
+		const FFunctionExpressionInput& FunctionInput = FunctionInputs[InputIndex];
+		const FFunctionExpressionInput& ConnectedInput = ConnectedInputs[InputIndex];
+
+		FFunctionInput& Input = Inputs.AddDefaulted_GetRef();
+		Input.FunctionInputExpression = FunctionInput.ExpressionInput;
+		Input.ConnectedExpression = ConnectedInput.Input.Expression;
+		Input.ConnectedOutputIndex = ConnectedInput.Input.OutputIndex;
+
+		Hasher.Update((uint8*)&Input.ConnectedExpression, sizeof(Input.ConnectedExpression));
+		Hasher.Update((uint8*)&Input.ConnectedOutputIndex, sizeof(Input.ConnectedOutputIndex));
 	}
 
-	Hasher.Final();
-
-	FSHAHash Hash;
-	Hasher.GetHash(Hash.Hash);
+	const FSHAHash Hash = Hasher.Finalize();
 
 	FFunctionCallEntry** ExistingFunctionCall = FunctionCallMap.Find(Hash);
 	FFunctionCallEntry* FunctionCall = ExistingFunctionCall ? *ExistingFunctionCall : nullptr;
@@ -488,15 +489,7 @@ UE::HLSLTree::FExpression* FMaterialHLSLGenerator::GenerateFunctionCall(UE::HLSL
 	{
 		FunctionCall = new(HLSLTree->GetAllocator()) FFunctionCallEntry();
 		FunctionCall->Function = Function;
-		// Inject the function inputs into the function scope
-		for (int32 InputIndex = 0; InputIndex < Inputs.Num(); ++InputIndex)
-		{
-			const FFunctionExpressionInput& Input = FunctionInputs[InputIndex];
-			UE::HLSLTree::FExpression* InputExpression = InputExpressions[InputIndex];
-			const FExpressionKey ExpressionKey(Input.ExpressionInput, 0);
-			FunctionCall->ExpressionMap.Add(ExpressionKey, InputExpression);
-		}
-
+		FunctionCall->Inputs = MoveTemp(Inputs);
 		FunctionCallMap.Add(Hash, FunctionCall);
 	}
 
diff --git a/Engine/Source/Runtime/Engine/Public/MaterialHLSLGenerator.h b/Engine/Source/Runtime/Engine/Public/MaterialHLSLGenerator.h
--- a/Engine/Source/Runtime/Engine/Public/MaterialHLSLGenerator.h
+++ b/Engine/Source/Runtime/Engine/Public/MaterialHLSLGenerator.h
@@ -108,10 +108,12 @@ public:
 	 * Note that a given UMaterialExpression may only support 1 of these node types, attempting to access an invalid node type will generate an error
 	 */
 	UE::HLSLTree::FExpression* AcquireExpression(UE::HLSLTree::FScope& Scope, UMaterialExpression* MaterialExpression, int32 OutputIndex);
+	UE::HLSLTree::FExpression* AcquireFunctionInputExpression(UE::HLSLTree::FScope& Scope, const UMaterialExpressionFunctionInput* MaterialExpression);
+
 	UE::HLSLTree::FTextureParameterDeclaration* AcquireTextureDeclaration(UE::HLSLTree::FScope& Scope, UMaterialExpression* MaterialExpression, int32 OutputIndex);
 	bool GenerateStatements(UE::HLSLTree::FScope& Scope, UMaterialExpression* MaterialExpression);
 
-	UE::HLSLTree::FExpression* GenerateFunctionCall(UE::HLSLTree::FScope& Scope, UMaterialFunctionInterface* Function, TArrayView<const FFunctionExpressionInput> Inputs, int32 OutputIndex);
+	UE::HLSLTree::FExpression* GenerateFunctionCall(UE::HLSLTree::FScope& Scope, UMaterialFunctionInterface* Function, TArrayView<const FFunctionExpressionInput> ConnectedInputs, int32 OutputIndex);
 
 	template<typename T, typename... ArgTypes>
 	T* NewExpressionData(const UMaterialExpression* MaterialExpression, ArgTypes... Args)
@@ -144,24 +146,7 @@ public:
 
 private:
 	static constexpr int32 MaxNumPreviousScopes = UE::HLSLTree::MaxNumPreviousScopes;
-	struct FExpressionKey
-	{
-		explicit FExpressionKey(UMaterialExpression* InExpression, int32 InOutputIndex = INDEX_NONE) : Expression(InExpression), OutputIndex(InOutputIndex) {}
-
-		UMaterialExpression* Expression;
-		int32 OutputIndex;
-
-		friend inline uint32 GetTypeHash(const FExpressionKey& Value)
-		{
-			return HashCombine(GetTypeHash(Value.Expression), GetTypeHash(Value.OutputIndex));
-		}
-
-		friend inline bool operator==(const FExpressionKey& Lhs, const FExpressionKey& Rhs)
-		{
-			return Lhs.Expression == Rhs.Expression && Lhs.OutputIndex == Rhs.OutputIndex;
-		}
-	};
-
+	
 	struct FExpressionDataKey
 	{
 		FExpressionDataKey(const FName& InTypeName, const UMaterialExpression* InMaterialExpression) : MaterialExpression(InMaterialExpression), TypeName(InTypeName) {}
@@ -180,10 +165,18 @@ private:
 		}
 	};
 
+	struct FFunctionInput
+	{
+		const UMaterialExpressionFunctionInput* FunctionInputExpression = nullptr;
+		UMaterialExpression* ConnectedExpression = nullptr;
+		int32 ConnectedOutputIndex = INDEX_NONE;
+	};
+	using FFunctionInputArray = TArray<FFunctionInput, TInlineAllocator<4>>;
+
 	struct FFunctionCallEntry
 	{
 		UMaterialFunctionInterface* Function = nullptr;
-		TMap<FExpressionKey, UE::HLSLTree::FExpression*> ExpressionMap;
+		FFunctionInputArray Inputs;
 	};
 
 	struct FStatementEntry
@@ -206,8 +199,8 @@ private:
 	UE::HLSLTree::FExpression* ResultExpression = nullptr;
 	UE::HLSLTree::FStatement* ResultStatement = nullptr;
 
-	TArray<FExpressionKey> ExpressionStack;
-	TArray<FFunctionCallEntry*> FunctionCallStack;
+	TArray<UMaterialExpression*> ExpressionStack;
+	TArray<FFunctionCallEntry*, TInlineAllocator<8>> FunctionCallStack;
 	TArray<UE::HLSLTree::FScope*> JoinedScopeStack;
 	TArray<FString> CompileErrors;
 	TArray<UMaterialExpression*> ErrorExpressions;
