Engine changes to expose velocity functionality to location based impulses

diff --git a/Engine/Source/Editor/PhysicsAssetEditor/Private/PhysicsAssetEditorSkeletalMeshComponent.cpp b/Engine/Source/Editor/PhysicsAssetEditor/Private/PhysicsAssetEditorSkeletalMeshComponent.cpp
--- a/Engine/Source/Editor/PhysicsAssetEditor/Private/PhysicsAssetEditorSkeletalMeshComponent.cpp
+++ b/Engine/Source/Editor/PhysicsAssetEditor/Private/PhysicsAssetEditorSkeletalMeshComponent.cpp
@@ -573,13 +573,18 @@ void UPhysicsAssetEditorSkeletalMeshComponent::RefreshBoneTransforms(FActorCompo
 	}
 }
 
-void UPhysicsAssetEditorSkeletalMeshComponent::AddImpulseAtLocation(FVector Impulse, FVector Location, FName BoneName)
+void UPhysicsAssetEditorSkeletalMeshComponent::AddImpulseAtLocation(FVector Impulse, FVector Location, FName BoneName, bool bVelChange)
 {
 #if !WITH_CHAOS
-	Super::AddImpulseAtLocation(Impulse, Location, BoneName);
+	Super::AddImpulseAtLocation(Impulse, Location, BoneName, bVelChange);
 #else
 	if (PreviewInstance != nullptr)
 	{
+		if (bVelChange)
+		{
+			ensureMsgf(false, TEXT("Velocity Change does not function in SkeletalMeshComponent."));
+		}
+
 		PreviewInstance->AddImpulseAtLocation(Impulse, Location, BoneName);
 	}
 #endif
diff --git a/Engine/Source/Editor/PhysicsAssetEditor/Private/PhysicsAssetEditorSkeletalMeshComponent.h b/Engine/Source/Editor/PhysicsAssetEditor/Private/PhysicsAssetEditorSkeletalMeshComponent.h
--- a/Engine/Source/Editor/PhysicsAssetEditor/Private/PhysicsAssetEditorSkeletalMeshComponent.h
+++ b/Engine/Source/Editor/PhysicsAssetEditor/Private/PhysicsAssetEditorSkeletalMeshComponent.h
@@ -50,7 +50,7 @@ class UPhysicsAssetEditorSkeletalMeshComponent : public UDebugSkelMeshComponent
 
 	/** UPrimitiveComponent interface */
 	virtual FPrimitiveSceneProxy* CreateSceneProxy() override;
-	virtual void AddImpulseAtLocation(FVector Impulse, FVector Location, FName BoneName = NAME_None) override;
+	virtual void AddImpulseAtLocation(FVector Impulse, FVector Location, FName BoneName = NAME_None, bool bVelChange = false) override;
 	virtual bool ShouldCreatePhysicsState() const override;
 
 	/** USkinnedMeshComponent interface */
diff --git a/Engine/Source/Runtime/Engine/Classes/Components/PrimitiveComponent.h b/Engine/Source/Runtime/Engine/Classes/Components/PrimitiveComponent.h
--- a/Engine/Source/Runtime/Engine/Classes/Components/PrimitiveComponent.h
+++ b/Engine/Source/Runtime/Engine/Classes/Components/PrimitiveComponent.h
@@ -1380,9 +1380,10 @@ public:
 	 *	@param	Impulse		Magnitude and direction of impulse to apply.
 	 *	@param	Location	Point in world space to apply impulse at.
 	 *	@param	BoneName	If a SkeletalMeshComponent, name of bone to apply impulse to. 'None' indicates root body.
+	 *  @param	bVelChange	If true, the Strength is taken as a change in angular velocity instead of an impulse (ie. mass will have no effect).
 	 */
 	UFUNCTION(BlueprintCallable, Category="Physics", meta=(UnsafeDuringActorConstruction="true"))
-	virtual void AddImpulseAtLocation(FVector Impulse, FVector Location, FName BoneName = NAME_None);
+	virtual void AddImpulseAtLocation(FVector Impulse, FVector Location, FName BoneName = NAME_None, bool bVelChange = false);
 
 	/**
 	 * Add an impulse to all rigid bodies in this component, radiating out from the specified position.
diff --git a/Engine/Source/Runtime/Engine/Classes/PhysicsEngine/BodyInstance.h b/Engine/Source/Runtime/Engine/Classes/PhysicsEngine/BodyInstance.h
--- a/Engine/Source/Runtime/Engine/Classes/PhysicsEngine/BodyInstance.h
+++ b/Engine/Source/Runtime/Engine/Classes/PhysicsEngine/BodyInstance.h
@@ -756,7 +756,7 @@ public:
 	/** Add an impulse to this body */
 	void AddImpulse(const FVector& Impulse, bool bVelChange);
 	/** Add an impulse to this body and a particular world position */
-	void AddImpulseAtPosition(const FVector& Impulse, const FVector& Position);
+	void AddImpulseAtPosition(const FVector& Impulse, const FVector& Position, bool bVelChange = false);
 	/** Set the linear velocity of this body */
 	void SetLinearVelocity(const FVector& NewVel, bool bAddToCurrent, bool bAutoWake = true);
 
diff --git a/Engine/Source/Runtime/Engine/Private/PhysicsEngine/BodyInstance.cpp b/Engine/Source/Runtime/Engine/Private/PhysicsEngine/BodyInstance.cpp
--- a/Engine/Source/Runtime/Engine/Private/PhysicsEngine/BodyInstance.cpp
+++ b/Engine/Source/Runtime/Engine/Private/PhysicsEngine/BodyInstance.cpp
@@ -3586,15 +3586,22 @@ void FBodyInstance::AddImpulse(const FVector& Impulse, bool bVelChange)
 	});
 }
 
-void FBodyInstance::AddImpulseAtPosition(const FVector& Impulse, const FVector& Position)
+void FBodyInstance::AddImpulseAtPosition(const FVector& Impulse, const FVector& Position, bool bVelChange)
 {
 	FPhysicsCommand::ExecuteWrite(ActorHandle, [&](const FPhysicsActorHandle& Actor)
-	{
-		if(FPhysicsInterface::IsRigidBody(Actor) && FPhysicsInterface::IsInScene(Actor) && !IsRigidBodyKinematic_AssumesLocked(Actor))
 		{
-			FPhysicsInterface::AddImpulseAtLocation_AssumesLocked(Actor, Impulse, Position);
-		}
-	});
+			if (FPhysicsInterface::IsRigidBody(Actor) && FPhysicsInterface::IsInScene(Actor) && !IsRigidBodyKinematic_AssumesLocked(Actor))
+			{
+				if (bVelChange)
+				{
+					FPhysicsInterface::AddVelocityAtLocation_AssumesLocked(Actor, Impulse, Position);
+				}
+				else
+				{
+					FPhysicsInterface::AddImpulseAtLocation_AssumesLocked(Actor, Impulse, Position);
+				}
+			}
+		});
 }
 
 void FBodyInstance::SetInstanceNotifyRBCollision(bool bNewNotifyCollision)
diff --git a/Engine/Source/Runtime/Engine/Private/PrimitiveComponentPhysics.cpp b/Engine/Source/Runtime/Engine/Private/PrimitiveComponentPhysics.cpp
--- a/Engine/Source/Runtime/Engine/Private/PrimitiveComponentPhysics.cpp
+++ b/Engine/Source/Runtime/Engine/Private/PrimitiveComponentPhysics.cpp
@@ -134,12 +134,12 @@ void UPrimitiveComponent::AddAngularImpulseInRadians(FVector Impulse, FName Bone
 	}
 }
 
-void UPrimitiveComponent::AddImpulseAtLocation(FVector Impulse, FVector Location, FName BoneName)
+void UPrimitiveComponent::AddImpulseAtLocation(FVector Impulse, FVector Location, FName BoneName, bool bVelChange)
 {
 	if (FBodyInstance* BI = GetBodyInstance(BoneName))
 	{
 		WarnInvalidPhysicsOperations(LOCTEXT("AddImpulseAtLocation", "AddImpulseAtLocation"), BI, BoneName);
-		BI->AddImpulseAtPosition(Impulse, Location);
+		BI->AddImpulseAtPosition(Impulse, Location, bVelChange);
 	}
 }
 
diff --git a/Engine/Source/Runtime/PhysicsCore/Private/ChaosEngineInterface.cpp b/Engine/Source/Runtime/PhysicsCore/Private/ChaosEngineInterface.cpp
--- a/Engine/Source/Runtime/PhysicsCore/Private/ChaosEngineInterface.cpp
+++ b/Engine/Source/Runtime/PhysicsCore/Private/ChaosEngineInterface.cpp
@@ -774,6 +774,14 @@ void FChaosEngineInterface::AddImpulseAtLocation_AssumesLocked(const FPhysicsAct
 	}
 }
 
+void FChaosEngineInterface::AddVelocityAtLocation_AssumesLocked(const FPhysicsActorHandle& InActorReference, const FVector& InVelocityDelta, const FVector& InLocation)
+{
+	if (ensure(FChaosEngineInterface::IsValid(InActorReference)))
+	{
+		AddImpulseAtLocation_AssumesLocked(InActorReference, InActorReference->GetGameThreadAPI().M() * InVelocityDelta, InLocation);
+	}
+}
+
 void FChaosEngineInterface::AddRadialImpulse_AssumesLocked(const FPhysicsActorHandle& InActorReference,const FVector& InOrigin,float InRadius,float InStrength,ERadialImpulseFalloff InFalloff,bool bInVelChange)
 {
 	using namespace Chaos;
diff --git a/Engine/Source/Runtime/PhysicsCore/Public/Chaos/ChaosEngineInterface.h b/Engine/Source/Runtime/PhysicsCore/Public/Chaos/ChaosEngineInterface.h
--- a/Engine/Source/Runtime/PhysicsCore/Public/Chaos/ChaosEngineInterface.h
+++ b/Engine/Source/Runtime/PhysicsCore/Public/Chaos/ChaosEngineInterface.h
@@ -450,6 +450,7 @@ public:
 	static void AddVelocity_AssumesLocked(const FPhysicsActorHandle& InActorReference,const FVector& InForce);
 	static void AddAngularVelocityInRadians_AssumesLocked(const FPhysicsActorHandle& InActorReference,const FVector& InTorque);
 	static void AddImpulseAtLocation_AssumesLocked(const FPhysicsActorHandle& InActorReference,const FVector& InImpulse,const FVector& InLocation);
+	static void AddVelocityAtLocation_AssumesLocked(const FPhysicsActorHandle& InActorReference, const FVector& InVelocityDelta, const FVector& InLocation);
 	static void AddRadialImpulse_AssumesLocked(const FPhysicsActorHandle& InActorReference,const FVector& InOrigin,float InRadius,float InStrength,ERadialImpulseFalloff InFalloff,bool bInVelChange);
 
 	static bool IsGravityEnabled_AssumesLocked(const FPhysicsActorHandle& InActorReference);
