Fixed non unity compile error

diff --git a/Engine/Source/Runtime/Engine/Public/RayTracingSkinnedGeometry.h b/Engine/Source/Runtime/Engine/Public/RayTracingSkinnedGeometry.h
--- a/Engine/Source/Runtime/Engine/Public/RayTracingSkinnedGeometry.h
+++ b/Engine/Source/Runtime/Engine/Public/RayTracingSkinnedGeometry.h
@@ -3,6 +3,7 @@
 #pragma once
 
 #include "CoreMinimal.h"
+#include "RHIDefinitions.h"
 
 #if RHI_RAYTRACING
 
