Try to make static analysis happy

diff --git a/Engine/Source/Runtime/Engine/Private/Shader/ShaderValue.cpp b/Engine/Source/Runtime/Engine/Private/Shader/ShaderValue.cpp
--- a/Engine/Source/Runtime/Engine/Private/Shader/ShaderValue.cpp
+++ b/Engine/Source/Runtime/Engine/Private/Shader/ShaderValue.cpp
@@ -68,8 +68,11 @@ EValueComponentType FType::GetComponentType(int32 Index) const
 		check(ValueType == EValueType::Struct);
 		return StructType->ComponentTypes.IsValidIndex(Index) ? StructType->ComponentTypes[Index] : EValueComponentType::Void;
 	}
-	const FValueTypeDescription TypeDesc = GetValueTypeDescription(ValueType);
-	return (Index >= 0 && Index < TypeDesc.NumComponents) ? TypeDesc.ComponentType : EValueComponentType::Void;
+	else
+	{
+		const FValueTypeDescription TypeDesc = GetValueTypeDescription(ValueType);
+		return (Index >= 0 && Index < TypeDesc.NumComponents) ? TypeDesc.ComponentType : EValueComponentType::Void;
+	}
 }
 
 EValueType FType::GetFlatFieldType(int32 Index) const
@@ -79,7 +82,10 @@ EValueType FType::GetFlatFieldType(int32 Index) const
 		check(ValueType == EValueType::Struct);
 		return StructType->FlatFieldTypes.IsValidIndex(Index) ? StructType->FlatFieldTypes[Index] : EValueType::Void;
 	}
-	return (Index == 0) ? ValueType : EValueType::Void;
+	else
+	{
+		return (Index == 0) ? ValueType : EValueType::Void;
+	}
 }
 
 bool FType::Merge(const FType& OtherType)
