Avoid caching UMaterialExpression->HLSLTree::FExpression, as material expressions may depend on current local variable values
Additional misc cleanup

diff --git a/Engine/Source/Runtime/Engine/Classes/Materials/MaterialExpressionComponentMask.h b/Engine/Source/Runtime/Engine/Classes/Materials/MaterialExpressionComponentMask.h
--- a/Engine/Source/Runtime/Engine/Classes/Materials/MaterialExpressionComponentMask.h
+++ b/Engine/Source/Runtime/Engine/Classes/Materials/MaterialExpressionComponentMask.h
@@ -35,6 +35,7 @@ class UMaterialExpressionComponentMask : public UMaterialExpression
 	virtual int32 Compile(class FMaterialCompiler* Compiler, int32 OutputIndex) override;
 	virtual void GetCaption(TArray<FString>& OutCaptions) const override;
 	virtual FText GetKeywords() const override { return FText::FromString(TEXT("component mask")); }
+	virtual EMaterialGenerateHLSLStatus GenerateHLSLExpression(FMaterialHLSLGenerator& Generator, UE::HLSLTree::FScope& Scope, int32 OutputIndex, UE::HLSLTree::FExpression*& OutExpression) override;
 #endif
 	//~ End UMaterialExpression Interface
 };
diff --git a/Engine/Source/Runtime/Engine/Private/Materials/MaterialExpressionHLSL.cpp b/Engine/Source/Runtime/Engine/Private/Materials/MaterialExpressionHLSL.cpp
--- a/Engine/Source/Runtime/Engine/Private/Materials/MaterialExpressionHLSL.cpp
+++ b/Engine/Source/Runtime/Engine/Private/Materials/MaterialExpressionHLSL.cpp
@@ -35,6 +35,7 @@
 #include "Materials/MaterialExpressionDivide.h"
 #include "Materials/MaterialExpressionBinaryOp.h"
 #include "Materials/MaterialExpressionAppendVector.h"
+#include "Materials/MaterialExpressionComponentMask.h"
 #include "Materials/MaterialExpressionGetMaterialAttributes.h"
 #include "Materials/MaterialExpressionSetMaterialAttributes.h"
 #include "Materials/MaterialExpressionReflectionVectorWS.h"
@@ -235,7 +236,7 @@ EMaterialGenerateHLSLStatus UMaterialExpressionBinaryOp::GenerateHLSLExpression(
 		return EMaterialGenerateHLSLStatus::Error;
 	}
 
-	OutExpression = Generator.GetTree().NewExpression<UE::HLSLTree::FExpressionBinaryOp>(GetBinaryOp(), Lhs, Rhs);
+	OutExpression = Generator.GetTree().NewBinaryOp(GetBinaryOp(), Lhs, Rhs);
 	return EMaterialGenerateHLSLStatus::Success;
 }
 
@@ -247,7 +248,7 @@ EMaterialGenerateHLSLStatus UMaterialExpressionAdd::GenerateHLSLExpression(FMate
 	{
 		return EMaterialGenerateHLSLStatus::Error;
 	}
-	OutExpression = Generator.GetTree().NewExpression<UE::HLSLTree::FExpressionBinaryOp>(UE::HLSLTree::EBinaryOp::Add, Lhs, Rhs);
+	OutExpression = Generator.GetTree().NewBinaryOp(UE::HLSLTree::EBinaryOp::Add, Lhs, Rhs);
 	return EMaterialGenerateHLSLStatus::Success;
 }
 
@@ -259,7 +260,7 @@ EMaterialGenerateHLSLStatus UMaterialExpressionMultiply::GenerateHLSLExpression(
 	{
 		return EMaterialGenerateHLSLStatus::Error;
 	}
-	OutExpression = Generator.GetTree().NewExpression<UE::HLSLTree::FExpressionBinaryOp>(UE::HLSLTree::EBinaryOp::Mul, Lhs, Rhs);
+	OutExpression = Generator.GetTree().NewBinaryOp(UE::HLSLTree::EBinaryOp::Mul, Lhs, Rhs);
 	return EMaterialGenerateHLSLStatus::Success;
 }
 
@@ -271,7 +272,7 @@ EMaterialGenerateHLSLStatus UMaterialExpressionDivide::GenerateHLSLExpression(FM
 	{
 		return EMaterialGenerateHLSLStatus::Error;
 	}
-	OutExpression = Generator.GetTree().NewExpression<UE::HLSLTree::FExpressionBinaryOp>(UE::HLSLTree::EBinaryOp::Div, Lhs, Rhs);
+	OutExpression = Generator.GetTree().NewBinaryOp(UE::HLSLTree::EBinaryOp::Div, Lhs, Rhs);
 	return EMaterialGenerateHLSLStatus::Success;
 }
 
@@ -287,6 +288,19 @@ EMaterialGenerateHLSLStatus UMaterialExpressionAppendVector::GenerateHLSLExpress
 	return EMaterialGenerateHLSLStatus::Success;
 }
 
+EMaterialGenerateHLSLStatus UMaterialExpressionComponentMask::GenerateHLSLExpression(FMaterialHLSLGenerator& Generator, UE::HLSLTree::FScope& Scope, int32 OutputIndex, UE::HLSLTree::FExpression*& OutExpression)
+{
+	UE::HLSLTree::FExpression* InputExpression = Input.AcquireHLSLExpression(Generator, Scope);
+	if (!InputExpression)
+	{
+		return EMaterialGenerateHLSLStatus::Error;
+	}
+
+	OutExpression = Generator.GetTree().NewExpression<UE::HLSLTree::FExpressionSwizzle>(UE::HLSLTree::MakeSwizzleMask(!!R, !!G, !!B, !!A), InputExpression);
+	return EMaterialGenerateHLSLStatus::Success;
+}
+
+
 EMaterialGenerateHLSLStatus UMaterialExpressionGetMaterialAttributes::GenerateHLSLExpression(FMaterialHLSLGenerator& Generator, UE::HLSLTree::FScope& Scope, int32 OutputIndex, UE::HLSLTree::FExpression*& OutExpression)
 {
 	UE::HLSLTree::FExpression* AttributesExpression = MaterialAttributes.AcquireHLSLExpression(Generator, Scope);
@@ -421,6 +435,7 @@ EMaterialGenerateHLSLStatus UMaterialExpressionIfThenElse::GenerateHLSLStatement
 
 EMaterialGenerateHLSLStatus UMaterialExpressionWhileLoop::GenerateHLSLStatements(FMaterialHLSLGenerator& Generator, UE::HLSLTree::FScope& Scope)
 {
+	using namespace UE::HLSLTree;
 	if (!Condition.IsConnected())
 	{
 		return Generator.Error(TEXT("Missing condition connection"));
@@ -431,17 +446,17 @@ EMaterialGenerateHLSLStatus UMaterialExpressionWhileLoop::GenerateHLSLStatements
 		return Generator.Error(TEXT("Missing LoopBody connection"));
 	}
 
-	UE::HLSLTree::FStatementLoop* LoopStatement = Generator.GetTree().NewStatement<UE::HLSLTree::FStatementLoop>(Scope);
+	FStatementLoop* LoopStatement = Generator.GetTree().NewStatement<FStatementLoop>(Scope);
 	LoopStatement->LoopScope = Generator.NewOwnedScope(*LoopStatement);
 
-	UE::HLSLTree::FStatementIf* IfStatement = Generator.GetTree().NewStatement<UE::HLSLTree::FStatementIf>(*LoopStatement->LoopScope);
+	FStatementIf* IfStatement = Generator.GetTree().NewStatement<FStatementIf>(*LoopStatement->LoopScope);
 	IfStatement->ThenScope = Generator.NewOwnedScope(*IfStatement);
 	IfStatement->ElseScope = Generator.NewOwnedScope(*IfStatement);
 	LoopStatement->NextScope = Generator.NewScope(Scope, EMaterialNewScopeFlag::NoPreviousScope);
 	LoopStatement->LoopScope->AddPreviousScope(*IfStatement->ThenScope);
 	LoopStatement->NextScope->AddPreviousScope(*IfStatement->ElseScope);
 
-	Generator.GetTree().NewStatement<UE::HLSLTree::FStatementBreak>(*IfStatement->ElseScope);
+	Generator.GetTree().NewStatement<FStatementBreak>(*IfStatement->ElseScope);
 
 	IfStatement->ConditionExpression = Condition.AcquireHLSLExpression(Generator, *LoopStatement->LoopScope);
 	LoopBody.GenerateHLSLStatements(Generator, *IfStatement->ThenScope);
@@ -477,18 +492,19 @@ EMaterialGenerateHLSLStatus UMaterialExpressionForLoop::GenerateHLSLExpression(F
 
 EMaterialGenerateHLSLStatus UMaterialExpressionForLoop::GenerateHLSLStatements(FMaterialHLSLGenerator& Generator, UE::HLSLTree::FScope& Scope)
 {
+	using namespace UE::HLSLTree;
 	if (!LoopBody.GetExpression())
 	{
 		return Generator.Error(TEXT("Missing LoopBody connection"));
 	}
 
-	UE::HLSLTree::FExpression* StartExpression = StartIndex.AcquireHLSLExpression(Generator, Scope);
+	FExpression* StartExpression = StartIndex.AcquireHLSLExpression(Generator, Scope);
 	if (!StartExpression)
 	{
 		return Generator.Error(TEXT("Missing StartIndex connection"));
 	}
 
-	UE::HLSLTree::FExpression* EndExpression = EndIndex.AcquireHLSLExpression(Generator, Scope);
+	FExpression* EndExpression = EndIndex.AcquireHLSLExpression(Generator, Scope);
 	if (!EndExpression)
 	{
 		return Generator.Error(TEXT("Missing EndIndex connection"));
@@ -498,29 +514,29 @@ EMaterialGenerateHLSLStatus UMaterialExpressionForLoop::GenerateHLSLStatements(F
 	FExpressionDataForLoop* ExpressionData = Generator.NewExpressionData<FExpressionDataForLoop>(this);
 	ExpressionData->LocalName = *FString::Printf(TEXT("ForLoopControl%d"), GlobalData->NumLoops++);
 
-	UE::HLSLTree::FExpression* StepExpression = IndexStep.GetTracedInput().Expression ? IndexStep.AcquireHLSLExpression(Generator, Scope) : Generator.NewConstant(1);
+	FExpression* StepExpression = IndexStep.GetTracedInput().Expression ? IndexStep.AcquireHLSLExpression(Generator, Scope) : Generator.NewConstant(1);
 
 	Generator.GetTree().AssignLocal(Scope, ExpressionData->LocalName, StartExpression);
 
-	UE::HLSLTree::FStatementLoop* LoopStatement = Generator.GetTree().NewStatement<UE::HLSLTree::FStatementLoop>(Scope);
+	FStatementLoop* LoopStatement = Generator.GetTree().NewStatement<FStatementLoop>(Scope);
 	LoopStatement->LoopScope = Generator.NewOwnedScope(*LoopStatement);
 	ExpressionData->LoopScope = LoopStatement->LoopScope;
 
-	UE::HLSLTree::FStatementIf* IfStatement = Generator.GetTree().NewStatement<UE::HLSLTree::FStatementIf>(*LoopStatement->LoopScope);
+	FStatementIf* IfStatement = Generator.GetTree().NewStatement<FStatementIf>(*LoopStatement->LoopScope);
 	IfStatement->ThenScope = Generator.NewOwnedScope(*IfStatement);
 	IfStatement->ElseScope = Generator.NewOwnedScope(*IfStatement);
 	LoopStatement->NextScope = Generator.NewScope(Scope, EMaterialNewScopeFlag::NoPreviousScope);
 	LoopStatement->LoopScope->AddPreviousScope(*IfStatement->ThenScope);
 	LoopStatement->NextScope->AddPreviousScope(*IfStatement->ElseScope);
 
-	Generator.GetTree().NewStatement<UE::HLSLTree::FStatementBreak>(*IfStatement->ElseScope);
+	Generator.GetTree().NewStatement<FStatementBreak>(*IfStatement->ElseScope);
 
-	UE::HLSLTree::FExpression* LocalExpression = Generator.GetTree().AcquireLocal(*LoopStatement->LoopScope, ExpressionData->LocalName);
+	FExpression* LocalExpression = Generator.GetTree().AcquireLocal(*LoopStatement->LoopScope, ExpressionData->LocalName);
 
-	IfStatement->ConditionExpression = Generator.GetTree().NewExpression<UE::HLSLTree::FExpressionBinaryOp>(UE::HLSLTree::EBinaryOp::Less, LocalExpression, EndExpression);
+	IfStatement->ConditionExpression = Generator.GetTree().NewLess(LocalExpression, EndExpression);
 	LoopBody.GenerateHLSLStatements(Generator, *IfStatement->ThenScope);
 
-	UE::HLSLTree::FExpression* NewLocalExpression = Generator.GetTree().NewExpression<UE::HLSLTree::FExpressionBinaryOp>(UE::HLSLTree::EBinaryOp::Add, Generator.GetTree().AcquireLocal(*IfStatement->ThenScope, ExpressionData->LocalName), StepExpression);
+	FExpression* NewLocalExpression = Generator.GetTree().NewAdd(Generator.GetTree().AcquireLocal(*IfStatement->ThenScope, ExpressionData->LocalName), StepExpression);
 	Generator.GetTree().AssignLocal(*IfStatement->ThenScope, ExpressionData->LocalName, NewLocalExpression);
 
 	Completed.GenerateHLSLStatements(Generator, *LoopStatement->NextScope);
diff --git a/Engine/Source/Runtime/Engine/Private/Materials/MaterialHLSLGenerator.cpp b/Engine/Source/Runtime/Engine/Private/Materials/MaterialHLSLGenerator.cpp
--- a/Engine/Source/Runtime/Engine/Private/Materials/MaterialHLSLGenerator.cpp
+++ b/Engine/Source/Runtime/Engine/Private/Materials/MaterialHLSLGenerator.cpp
@@ -173,17 +173,20 @@ static UE::HLSLTree::FExpression* CompileMaterialInput(FMaterialHLSLGenerator& G
 	UMaterial* Material,
 	UE::HLSLTree::FExpression* AttributesExpression)
 {
-#if 0
-	UE::HLSLTree::FExpression* Expression = nullptr;
-	if (Material->IsPropertyActive(InputProperty))
+	using namespace UE::HLSLTree;
+
+	// We're only interesting in attributes that map to valid fields
+	const UE::Shader::FStructField* AttributeField = Generator.GetMaterialAttributesType()->FindFieldByName(*FMaterialAttributeDefinitionMap::GetAttributeName(InputProperty));
+
+	FExpression* Expression = nullptr;
+	if (AttributeField && Material->IsPropertyActive(InputProperty))
 	{
 		FMaterialInputDescription InputDescription;
 		if (Material->GetExpressionInputDescription(InputProperty, InputDescription))
 		{
 			if (InputDescription.bUseConstant)
 			{
-				UE::Shader::FValue DefaultValue = FMaterialAttributeDefinitionMap::GetDefaultValue(InputProperty);
-				DefaultValue.NumComponents = UE::Shader::GetValueTypeDescription(InputDescription.Type).NumComponents;
+				UE::Shader::FValue DefaultValue = UE::Shader::Cast(FMaterialAttributeDefinitionMap::GetDefaultValue(InputProperty), InputDescription.Type);
 				if (InputDescription.ConstantValue != DefaultValue)
 				{
 					Expression = Generator.NewConstant(InputDescription.ConstantValue);
@@ -201,7 +204,7 @@ static UE::HLSLTree::FExpression* CompileMaterialInput(FMaterialHLSLGenerator& G
 					}
 					if (!Expression)
 					{
-						Expression = Generator.NewTexCoord(Scope, TexCoordIndex);
+						Expression = Generator.NewTexCoord(TexCoordIndex);
 					}
 				}
 				else
@@ -214,18 +217,21 @@ static UE::HLSLTree::FExpression* CompileMaterialInput(FMaterialHLSLGenerator& G
 
 	if (Expression)
 	{
-		UE::HLSLTree::FExpressionSetMaterialAttribute* SetAttributeExpression = Generator.GetTree().NewExpression<UE::HLSLTree::FExpressionSetMaterialAttribute>(Scope);
-		SetAttributeExpression->AttributeID = FMaterialAttributeDefinitionMap::GetID(InputProperty);
-		SetAttributeExpression->AttributesExpression = AttributesExpression;
-		SetAttributeExpression->ValueExpression = Expression;
+		FExpressionSetStructField* SetAttributeExpression = Generator.GetTree().NewExpression<FExpressionSetStructField>(
+			Generator.GetMaterialAttributesType(),
+			AttributeField,
+			AttributesExpression,
+			Expression);
 		return SetAttributeExpression;
 	}
-#endif
+
 	return AttributesExpression;
 }
 
 bool FMaterialHLSLGenerator::GenerateResult(UE::HLSLTree::FScope& Scope)
 {
+	using namespace UE::HLSLTree;
+
 	bool bResult = false;
 	if (bGeneratedResult)
 	{
@@ -238,7 +244,7 @@ bool FMaterialHLSLGenerator::GenerateResult(UE::HLSLTree::FScope& Scope)
 
 		if (TargetMaterial)
 		{
-			UE::HLSLTree::FExpression* AttributesExpression = nullptr;
+			FExpression* AttributesExpression = nullptr;
 			if (TargetMaterial->bUseMaterialAttributes)
 			{
 				FMaterialInputDescription InputDescription;
@@ -250,7 +256,7 @@ bool FMaterialHLSLGenerator::GenerateResult(UE::HLSLTree::FScope& Scope)
 			}
 			else
 			{
-				AttributesExpression = HLSLTree->NewExpression<UE::HLSLTree::FExpressionConstant>(MaterialAttributesDefaultValue);
+				AttributesExpression = HLSLTree->NewExpression<FExpressionConstant>(MaterialAttributesDefaultValue);
 				for (int32 PropertyIndex = 0; PropertyIndex < MP_MAX; ++PropertyIndex)
 				{
 					const EMaterialProperty Property = (EMaterialProperty)PropertyIndex;
@@ -260,7 +266,7 @@ bool FMaterialHLSLGenerator::GenerateResult(UE::HLSLTree::FScope& Scope)
 
 			if (AttributesExpression)
 			{
-				UE::HLSLTree::FStatementReturn* ReturnStatement = HLSLTree->NewStatement<UE::HLSLTree::FStatementReturn>(Scope);
+				FStatementReturn* ReturnStatement = HLSLTree->NewStatement<FStatementReturn>(Scope);
 				ReturnStatement->Expression = AttributesExpression;
 				ResultExpression = AttributesExpression;
 				ResultStatement = ReturnStatement;
@@ -301,22 +307,19 @@ UE::HLSLTree::FScope* FMaterialHLSLGenerator::NewJoinedScope(UE::HLSLTree::FScop
 	return NewScope;
 }
 
-UE::HLSLTree::FExpressionConstant* FMaterialHLSLGenerator::NewConstant(const UE::Shader::FValue& Value)
+UE::HLSLTree::FExpression* FMaterialHLSLGenerator::NewConstant(const UE::Shader::FValue& Value)
 {
-	UE::HLSLTree::FExpressionConstant* Expression = HLSLTree->NewExpression<UE::HLSLTree::FExpressionConstant>(Value);
-	return Expression;
+	return HLSLTree->NewConstant(Value);
 }
 
-UE::HLSLTree::FExpressionExternalInput* FMaterialHLSLGenerator::NewTexCoord(int32 Index)
+UE::HLSLTree::FExpression* FMaterialHLSLGenerator::NewTexCoord(int32 Index)
 {
-	UE::HLSLTree::FExpressionExternalInput* Expression = HLSLTree->NewExpression<UE::HLSLTree::FExpressionExternalInput>(UE::HLSLTree::MakeInputTexCoord(Index));
-	return Expression;
+	return HLSLTree->NewExpression<UE::HLSLTree::FExpressionExternalInput>(UE::HLSLTree::MakeInputTexCoord(Index));
 }
 
-UE::HLSLTree::FExpressionSwizzle* FMaterialHLSLGenerator::NewSwizzle(const UE::HLSLTree::FSwizzleParameters& Params, UE::HLSLTree::FExpression* Input)
+UE::HLSLTree::FExpression* FMaterialHLSLGenerator::NewSwizzle(const UE::HLSLTree::FSwizzleParameters& Params, UE::HLSLTree::FExpression* Input)
 {
-	UE::HLSLTree::FExpressionSwizzle* Expression = HLSLTree->NewExpression<UE::HLSLTree::FExpressionSwizzle>(Params, Input);
-	return Expression;
+	return HLSLTree->NewExpression<UE::HLSLTree::FExpressionSwizzle>(Params, Input);
 }
 
 UE::HLSLTree::FTextureParameterDeclaration* FMaterialHLSLGenerator::AcquireTextureDeclaration(const UE::HLSLTree::FTextureDescription& Value)
diff --git a/Engine/Source/Runtime/Engine/Public/HLSLTree/HLSLTree.h b/Engine/Source/Runtime/Engine/Public/HLSLTree/HLSLTree.h
--- a/Engine/Source/Runtime/Engine/Public/HLSLTree/HLSLTree.h
+++ b/Engine/Source/Runtime/Engine/Public/HLSLTree/HLSLTree.h
@@ -627,6 +627,7 @@ public:
 	FScope* NewScope(FScope& Scope);
 	FScope* NewOwnedScope(FStatement& Owner);
 
+	/** Shortcuts to create various common expression types */
 	FExpression* NewConstant(const Shader::FValue& Value);
 	FExpression* NewUnaryOp(EUnaryOp Op, FExpression* Input);
 	FExpression* NewBinaryOp(EBinaryOp Op, FExpression* Lhs, FExpression* Rhs);
@@ -638,6 +639,7 @@ public:
 	FExpression* NewSub(FExpression* Lhs, FExpression* Rhs) { return NewBinaryOp(EBinaryOp::Sub, Lhs, Rhs); }
 	FExpression* NewMul(FExpression* Lhs, FExpression* Rhs) { return NewBinaryOp(EBinaryOp::Mul, Lhs, Rhs); }
 	FExpression* NewDiv(FExpression* Lhs, FExpression* Rhs) { return NewBinaryOp(EBinaryOp::Div, Lhs, Rhs); }
+	FExpression* NewLess(FExpression* Lhs, FExpression* Rhs) { return NewBinaryOp(EBinaryOp::Less, Lhs, Rhs); }
 
 	FTextureParameterDeclaration* NewTextureParameterDeclaration(const FName& Name, const FTextureDescription& DefaultValue);
 
diff --git a/Engine/Source/Runtime/Engine/Public/HLSLTree/HLSLTreeCommon.h b/Engine/Source/Runtime/Engine/Public/HLSLTree/HLSLTreeCommon.h
--- a/Engine/Source/Runtime/Engine/Public/HLSLTree/HLSLTreeCommon.h
+++ b/Engine/Source/Runtime/Engine/Public/HLSLTree/HLSLTreeCommon.h
@@ -170,6 +170,16 @@ public:
 class FExpressionSetStructField : public FExpression
 {
 public:
+	FExpressionSetStructField(const Shader::FStructType* InStructType, const Shader::FStructField* InField, FExpression* InStructExpression, FExpression* InFieldExpression)
+		: StructType(InStructType)
+		, Field(InField)
+		, StructExpression(InStructExpression)
+		, FieldExpression(InFieldExpression)
+	{
+		check(InStructType);
+		check(InField);
+	}
+
 	FExpressionSetStructField(const Shader::FStructType* InStructType, const TCHAR* InFieldName, FExpression* InStructExpression, FExpression* InFieldExpression)
 		: StructType(InStructType)
 		, Field(InStructType->FindFieldByName(InFieldName))
diff --git a/Engine/Source/Runtime/Engine/Public/MaterialHLSLGenerator.h b/Engine/Source/Runtime/Engine/Public/MaterialHLSLGenerator.h
--- a/Engine/Source/Runtime/Engine/Public/MaterialHLSLGenerator.h
+++ b/Engine/Source/Runtime/Engine/Public/MaterialHLSLGenerator.h
@@ -30,10 +30,6 @@ namespace UE
 namespace HLSLTree
 {
 struct FSwizzleParameters;
-class FExpressionConstant;
-class FExpressionExternalInput;
-class FExpressionSwizzle;
-class FExpressionLocalPHI;
 }
 namespace Shader
 {
@@ -92,9 +88,9 @@ public:
 
 	UE::HLSLTree::FScope* NewJoinedScope(UE::HLSLTree::FScope& Scope);
 
-	UE::HLSLTree::FExpressionConstant* NewConstant(const UE::Shader::FValue& Value);
-	UE::HLSLTree::FExpressionExternalInput* NewTexCoord(int32 Index);
-	UE::HLSLTree::FExpressionSwizzle* NewSwizzle(const UE::HLSLTree::FSwizzleParameters& Params, UE::HLSLTree::FExpression* Input);
+	UE::HLSLTree::FExpression* NewConstant(const UE::Shader::FValue& Value);
+	UE::HLSLTree::FExpression* NewTexCoord(int32 Index);
+	UE::HLSLTree::FExpression* NewSwizzle(const UE::HLSLTree::FSwizzleParameters& Params, UE::HLSLTree::FExpression* Input);
 
 	/** Returns a declaration to access the given texture, with no parameter */
 	UE::HLSLTree::FTextureParameterDeclaration* AcquireTextureDeclaration(const UE::HLSLTree::FTextureDescription& Value);
