Add out of order packet analytics.

diff --git a/Engine/Source/Runtime/Engine/Classes/Engine/NetConnection.h b/Engine/Source/Runtime/Engine/Classes/Engine/NetConnection.h
--- a/Engine/Source/Runtime/Engine/Classes/Engine/NetConnection.h
+++ b/Engine/Source/Runtime/Engine/Classes/Engine/NetConnection.h
@@ -1266,6 +1266,46 @@ public:
 	 */
 	void FlushPacketOrderCache(bool bFlushWholeCache=false);
 
+	/**
+	 * Get the total number of out of order packets on this connection.
+	 *
+	 * @return The total number of out of order packets.
+	 */
+	int32 GetTotalOutOfOrderPackets() const
+	{
+		return TotalOutOfOrderPacketsLost + TotalOutOfOrderPacketsRecovered + TotalOutOfOrderPacketsDuplicate;
+	}
+
+	/**
+	 * Get the total number of out of order packets lost on this connection.
+	 *
+	 * @return The total number of out of order packets lost.
+	 */
+	int32 GetTotalOutOfOrderPacketsLost() const
+	{
+		return TotalOutOfOrderPacketsLost;
+	}
+
+	/**
+	 * Get the total number of out of order packets recovered on this connection.
+	 *
+	 * @return The total number of out of order packets recovered.
+	 */
+	int32 GetTotalOutOfOrderPacketsRecovered() const
+	{
+		return TotalOutOfOrderPacketsRecovered;
+	}
+
+	/**
+	 * Get the total number of out of order packets that were duplicates on this connection.
+	 *
+	 * @return The total number of out of order packets that were duplicates.
+	 */
+	int32 GetTotalOutOfOrderPacketsDuplicate() const
+	{
+		return TotalOutOfOrderPacketsDuplicate;
+	}
+
 	/**
 	 * Sets the OS/NIC level timestamp, for the last packet that was received
 	 *
@@ -1515,7 +1555,17 @@ private:
 	/** Out of order packet tracking/correction */
 
 	/** Stat tracking for the total number of out of order packets, for this connection */
-	int32 TotalOutOfOrderPackets;
+	UE_DEPRECATED(5.1, "Use GetTotalOutOfOrderPackets instead.")
+	int32 TotalOutOfOrderPackets = 0;
+
+	/** Stat tracking for the total number of out of order packets lost */
+	int32 TotalOutOfOrderPacketsLost = 0;
+
+	/** Stat tracking for the total number of out of order packets recovered */
+	int32 TotalOutOfOrderPacketsRecovered = 0;
+
+	/** Stat tracking for the total number of out of order packets that were duplicates */
+	int32 TotalOutOfOrderPacketsDuplicate = 0;
 
 	/** Buffer of partially read (post-PacketHandler) sequenced packets, which are waiting for a missing packet/sequence */
 	TOptional<TCircularBuffer<TUniquePtr<FBitReader>>> PacketOrderCache;
diff --git a/Engine/Source/Runtime/Engine/Classes/Engine/NetDriver.h b/Engine/Source/Runtime/Engine/Classes/Engine/NetDriver.h
--- a/Engine/Source/Runtime/Engine/Classes/Engine/NetDriver.h
+++ b/Engine/Source/Runtime/Engine/Classes/Engine/NetDriver.h
@@ -957,10 +957,10 @@ public:
 	uint32						OutPacketsLost;
 	/** Total packets lost that have been sent by the server since the net driver's creation  */
 	uint32						OutTotalPacketsLost;
-	/** todo document */
-	uint32						InOutOfOrderPackets;
-	/** todo document */
-	uint32						OutOutOfOrderPackets;
+	UE_DEPRECATED(5.1, "Use GetTotalOutOfOrderPackets instead.")
+	uint32						InOutOfOrderPackets = 0;
+	UE_DEPRECATED(5.1, "OutOutOfOrderPackets is not updated and is now deprecated.")
+	uint32						OutOutOfOrderPackets = 0;
 	/** Tracks the total number of voice packets sent */
 	uint32						VoicePacketsSent;
 	/** Tracks the total number of voice bytes sent */
@@ -1709,6 +1709,80 @@ public:
 	 */
 	inline void IncreaseOutTotalNotifiedPackets() { ++OutTotalNotifiedPackets; }
 
+	/**
+	 * Get the total number of out of order packets for all connections.
+	 *
+	 * @return The total number of out of order packets.
+	 */
+	int32 GetTotalOutOfOrderPackets() const
+	{
+		return TotalOutOfOrderPacketsLost + TotalOutOfOrderPacketsRecovered + TotalOutOfOrderPacketsDuplicate;
+	}
+
+	/**
+	 * Get the total number of out of order packets lost for all connections.
+	 *
+	 * @return The total number of out of order packets lost.
+	 */
+	int32 GetTotalOutOfOrderPacketsLost() const
+	{
+		return TotalOutOfOrderPacketsLost;
+	}
+
+	/**
+	 * Increase the value of TotalOutOfOrderPacketsLost.
+	 *
+	 * @param Count		The amount to add to TotalOutOfOrderPacketsLost
+	 */
+	void IncreaseTotalOutOfOrderPacketsLost(int32 Count=1)
+	{
+		TotalOutOfOrderPacketsLost += Count;
+
+		PRAGMA_DISABLE_DEPRECATION_WARNINGS
+		InOutOfOrderPackets += Count;
+		PRAGMA_ENABLE_DEPRECATION_WARNINGS
+	}
+
+	/**
+	 * Get the total number of out of order packets recovered for all connections.
+	 *
+	 * @return The total number of out of order packets recovered.
+	 */
+	int32 GetTotalOutOfOrderPacketsRecovered() const
+	{
+		return TotalOutOfOrderPacketsRecovered;
+	}
+
+	/**
+	 * Increase the value of TotalOutOfOrderPacketsRecovered.
+	 *
+	 * @param Count		The amount to add to TotalOutOfOrderPacketsRecovered
+	 */
+	void IncreaseTotalOutOfOrderPacketsRecovered(int32 Count=1)
+	{
+		TotalOutOfOrderPacketsRecovered += Count;
+	}
+
+	/**
+	 * Get the total number of out of order packets that were duplicates for all connections.
+	 *
+	 * @return The total number of out of order packets that were duplicates.
+	 */
+	int32 GetTotalOutOfOrderPacketsDuplicate() const
+	{
+		return TotalOutOfOrderPacketsDuplicate;
+	}
+
+	/**
+	 * Increase the value of TotalOutOfOrderPacketsDuplicate.
+	 *
+	 * @param Count		The amount to add to TotalOutOfOrderPacketsDuplicate
+	 */
+	void IncreaseTotalOutOfOrderPacketsDuplicate(int32 Count=1)
+	{
+		TotalOutOfOrderPacketsDuplicate += Count;
+	}
+
 	bool DidHitchLastFrame() const;
 
 	static bool IsDormInitialStartupActor(AActor* Actor);
@@ -1787,4 +1861,13 @@ private:
 
 	/** cache whether or not we have a replay connection, updated when a connection is added or removed */
 	bool bHasReplayConnection;
+
+	/** Stat tracking for the total number of out of order packets lost */
+	int32 TotalOutOfOrderPacketsLost = 0;
+
+	/** Stat tracking for the total number of out of order packets recovered */
+	int32 TotalOutOfOrderPacketsRecovered = 0;
+
+	/** Stat tracking for the total number of out of order packets that were duplicates */
+	int32 TotalOutOfOrderPacketsDuplicate = 0;
 };
diff --git a/Engine/Source/Runtime/Engine/Private/Analytics/EngineNetAnalytics.cpp b/Engine/Source/Runtime/Engine/Private/Analytics/EngineNetAnalytics.cpp
--- a/Engine/Source/Runtime/Engine/Private/Analytics/EngineNetAnalytics.cpp
+++ b/Engine/Source/Runtime/Engine/Private/Analytics/EngineNetAnalytics.cpp
@@ -27,6 +27,9 @@ bool FNetConnAnalyticsVars::operator == (const FNetConnAnalyticsVars& A) const
 {
 	return OutAckOnlyCount == A.OutAckOnlyCount &&
 			OutKeepAliveCount == A.OutKeepAliveCount &&
+			OutOfOrderPacketsLostCount == A.OutOfOrderPacketsLostCount &&
+			OutOfOrderPacketsRecoveredCount == A.OutOfOrderPacketsRecoveredCount &&
+			OutOfOrderPacketsDuplicateCount == A.OutOfOrderPacketsDuplicateCount &&
 			/** Close results can't be shared - if either are set, equality comparison fails */
 			!CloseReason.IsValid() && !A.CloseReason.IsValid() &&
 			ClientCloseReasons == A.ClientCloseReasons &&
@@ -37,6 +40,9 @@ void FNetConnAnalyticsVars::CommitAnalytics(FNetConnAnalyticsVars& AggregatedDat
 {
 	AggregatedData.OutAckOnlyCount += OutAckOnlyCount;
 	AggregatedData.OutKeepAliveCount += OutKeepAliveCount;
+	AggregatedData.OutOfOrderPacketsLostCount += OutOfOrderPacketsLostCount;
+	AggregatedData.OutOfOrderPacketsRecoveredCount += OutOfOrderPacketsRecoveredCount;
+	AggregatedData.OutOfOrderPacketsDuplicateCount += OutOfOrderPacketsDuplicateCount;
 
 	for (TMap<FString, int32>::TConstIterator It(RecoveredFaults); It; ++It)
 	{
@@ -184,6 +190,9 @@ void FNetConnAnalyticsData::SendAnalytics()
 
 		UE_LOG(LogNet, Log, TEXT(" - OutAckOnlyCount: %llu"), OutAckOnlyCount);
 		UE_LOG(LogNet, Log, TEXT(" - OutKeepAliveCount: %llu"), OutKeepAliveCount);
+		UE_LOG(LogNet, Log, TEXT(" - OutOfOrderPacketsLostCount: %llu"), OutOfOrderPacketsLostCount);
+		UE_LOG(LogNet, Log, TEXT(" - OutOfOrderPacketsRecoveredCount: %llu"), OutOfOrderPacketsRecoveredCount);
+		UE_LOG(LogNet, Log, TEXT(" - OutOfOrderPacketsDuplicateCount: %llu"), OutOfOrderPacketsDuplicateCount);
 
 
 		UE_LOG(LogNet, Log, TEXT(" - CloseReasons:"));
@@ -233,6 +242,9 @@ void FNetConnAnalyticsData::SendAnalytics()
 		static const FString EZEventName = TEXT("Core.ServerNetConn");
 		static const FString EZAttrib_OutAckOnlyCount = TEXT("OutAckOnlyCount");
 		static const FString EZAttrib_OutKeepAliveCount = TEXT("OutKeepAliveCount");
+		static const FString EZAttrib_OutOfOrderPacketsLostCount = TEXT("OutOfOrderPacketsLostCount");
+		static const FString EZAttrib_OutOfOrderPacketsRecoveredCount = TEXT("OutOfOrderPacketsRecoveredCount");
+		static const FString EZAttrib_OutOfOrderPacketsDuplicateCount = TEXT("OutOfOrderPacketsDuplicateCount");
 		static const FString EZAttrib_CloseReasons = TEXT("CloseReasons");
 		static const FString EZAttrib_FullCloseReasons = TEXT("FullCloseReasons");
 		static const FString EZAttrib_ClientCloseReasons = TEXT("ClientCloseReasons");
@@ -349,6 +361,9 @@ void FNetConnAnalyticsData::SendAnalytics()
 		AnalyticsProvider->RecordEvent(EZEventName, MakeAnalyticsEventAttributeArray(
 			EZAttrib_OutAckOnlyCount, OutAckOnlyCount,
 			EZAttrib_OutKeepAliveCount, OutKeepAliveCount,
+			EZAttrib_OutOfOrderPacketsLostCount, OutOfOrderPacketsLostCount,
+			EZAttrib_OutOfOrderPacketsRecoveredCount, OutOfOrderPacketsRecoveredCount,
+			EZAttrib_OutOfOrderPacketsDuplicateCount, OutOfOrderPacketsDuplicateCount,
 			EZAttrib_CloseReasons, FJsonFragment(MoveTemp(CloseReasonsJsonStr)),
 			EZAttrib_FullCloseReasons, FJsonFragment(MoveTemp(FullCloseReasonsJsonStr)),
 			EZAttrib_ClientCloseReasons, FJsonFragment(MoveTemp(ClientCloseReasonsJsonStr)),
diff --git a/Engine/Source/Runtime/Engine/Private/NetConnection.cpp b/Engine/Source/Runtime/Engine/Private/NetConnection.cpp
--- a/Engine/Source/Runtime/Engine/Private/NetConnection.cpp
+++ b/Engine/Source/Runtime/Engine/Private/NetConnection.cpp
@@ -288,7 +288,6 @@ UNetConnection::UNetConnection(const FObjectInitializer& ObjectInitializer)
 ,	HasDirtyAcks(0u)
 ,	bHasWarnedAboutChannelLimit(false)
 ,	bConnectionPendingCloseDueToSocketSendFailure(false)
-,	TotalOutOfOrderPackets(0)
 ,	PacketOrderCache()
 ,	PacketOrderCacheStartIdx(0)
 ,	PacketOrderCacheCount(0)
@@ -2522,14 +2521,23 @@ void UNetConnection::ReceivedPacket( FBitReader& Reader, bool bIsReinjectedPacke
 					UE_LOG(LogNet, VeryVerbose, TEXT("'Out of Order' Packet Cache, caching sequence order '%i' (capacity: %i)"), LinearCacheIdx, CacheCapacity);
 
 					CurCachePacket = MakeUnique<FBitReader>(Reader);
+
 					PacketOrderCacheCount++;
+					TotalOutOfOrderPacketsRecovered++;
+					AnalyticsVars.IncreaseOutOfOrderPacketsRecoveredCount();
+					Driver->IncreaseTotalOutOfOrderPacketsRecovered();
 
 					ResetReaderMark.Pop(*CurCachePacket);
 				}
 				else
 				{
+					PRAGMA_DISABLE_DEPRECATION_WARNINGS
 					TotalOutOfOrderPackets++;
-					Driver->InOutOfOrderPackets++;
+					PRAGMA_ENABLE_DEPRECATION_WARNINGS
+
+					TotalOutOfOrderPacketsDuplicate++;
+					AnalyticsVars.IncreaseOutOfOrderPacketsDuplicateCount();
+					Driver->IncreaseTotalOutOfOrderPacketsDuplicate();
 				}
 
 				return;
@@ -2550,14 +2558,19 @@ void UNetConnection::ReceivedPacket( FBitReader& Reader, bool bIsReinjectedPacke
 		}
 		else
 		{
+			PRAGMA_DISABLE_DEPRECATION_WARNINGS
 			TotalOutOfOrderPackets++;
-			Driver->InOutOfOrderPackets++;
+			PRAGMA_ENABLE_DEPRECATION_WARNINGS
+
+			TotalOutOfOrderPacketsLost++;
+			AnalyticsVars.IncreaseOutOfOrderPacketsLostCount();
+			Driver->IncreaseTotalOutOfOrderPacketsLost();
 
 			if (!PacketOrderCache.IsSet() && CVarNetDoPacketOrderCorrection.GetValueOnAnyThread() != 0)
 			{
 				int32 EnableThreshold = CVarNetPacketOrderCorrectionEnableThreshold.GetValueOnAnyThread();
 
-				if (TotalOutOfOrderPackets >= EnableThreshold)
+				if (TotalOutOfOrderPacketsLost >= EnableThreshold)
 				{
 					UE_LOG(LogNet, Verbose, TEXT("Hit threshold of %i 'out of order' packet sequences. Enabling out of order packet correction."), EnableThreshold);
 
diff --git a/Engine/Source/Runtime/Engine/Private/NetDriver.cpp b/Engine/Source/Runtime/Engine/Private/NetDriver.cpp
--- a/Engine/Source/Runtime/Engine/Private/NetDriver.cpp
+++ b/Engine/Source/Runtime/Engine/Private/NetDriver.cpp
@@ -389,8 +389,6 @@ UNetDriver::UNetDriver(const FObjectInitializer& ObjectInitializer)
 ,	InTotalPacketsLost(0)
 ,	OutPacketsLost(0)
 ,	OutTotalPacketsLost(0)
-,	InOutOfOrderPackets(0)
-,	OutOutOfOrderPackets(0)
 ,	StatUpdateTime(0.0)
 ,	StatPeriod(1.f)
 ,	TotalRPCsCalled(0)
diff --git a/Engine/Source/Runtime/Engine/Public/Analytics/EngineNetAnalytics.h b/Engine/Source/Runtime/Engine/Public/Analytics/EngineNetAnalytics.h
--- a/Engine/Source/Runtime/Engine/Public/Analytics/EngineNetAnalytics.h
+++ b/Engine/Source/Runtime/Engine/Public/Analytics/EngineNetAnalytics.h
@@ -12,6 +12,8 @@
  */
 struct FNetConnAnalyticsVars
 {
+	friend struct FNetConnAnalyticsData;
+
 	using FNetResult = UE::Net::FNetResult;
 
 public:
@@ -23,6 +25,24 @@ public:
 	void CommitAnalytics(FNetConnAnalyticsVars& AggregatedData);
 
 
+	/** Accessors */
+
+	void IncreaseOutOfOrderPacketsLostCount(int32 Count=1)
+	{
+		OutOfOrderPacketsLostCount += Count;
+	}
+
+	void IncreaseOutOfOrderPacketsRecoveredCount(int32 Count=1)
+	{
+		OutOfOrderPacketsRecoveredCount += Count;
+	}
+
+	void IncreaseOutOfOrderPacketsDuplicateCount(int32 Count=1)
+	{
+		OutOfOrderPacketsDuplicateCount += Count;
+	}
+
+
 public:
 	/** The number of packets that were exclusively ack packets */
 	uint64 OutAckOnlyCount;
@@ -30,6 +50,17 @@ public:
 	/** The number of packets that were just keep-alive packets */
 	uint64 OutKeepAliveCount;
 
+private:
+	/** The number of out of order packets lost */
+	uint64 OutOfOrderPacketsLostCount = 0;
+
+	/** The number of out of order packets recovered */
+	uint64 OutOfOrderPacketsRecoveredCount = 0;
+
+	/** The number of out of order packets that were duplicates */
+	uint64 OutOfOrderPacketsDuplicateCount = 0;
+
+public:
 	/** The result/reason for triggering NetConnection Close (local) */
 	TUniquePtr<FNetResult> CloseReason;
 
