Releasing references to other objects when PlayerCameraManager gets destroyed to avoid UObject leaks

diff --git a/Engine/Source/Runtime/Engine/Private/PlayerCameraManager.cpp b/Engine/Source/Runtime/Engine/Private/PlayerCameraManager.cpp
--- a/Engine/Source/Runtime/Engine/Private/PlayerCameraManager.cpp
+++ b/Engine/Source/Runtime/Engine/Private/PlayerCameraManager.cpp
@@ -826,8 +826,11 @@ void APlayerCameraManager::EndPlay(const EEndPlayReason::Type EndPlayReason)
 		{
 			AnimCameraActor->Destroy();
 		}
-		AnimCameraActor = NULL;
+		AnimCameraActor->SetOwner(nullptr);
+		AnimCameraActor = nullptr;
 	}
+	ActiveAnims.Empty();
+	ModifierList.Empty();
 	Super::EndPlay(EndPlayReason);
 }
 
