World Partition: minor code/includes cleanups.

diff --git a/Engine/Source/Runtime/Engine/Private/WorldPartition/WorldPartitionEditorCell.cpp b/Engine/Source/Runtime/Engine/Private/WorldPartition/WorldPartitionEditorCell.cpp
--- a/Engine/Source/Runtime/Engine/Private/WorldPartition/WorldPartitionEditorCell.cpp
+++ b/Engine/Source/Runtime/Engine/Private/WorldPartition/WorldPartitionEditorCell.cpp
@@ -5,6 +5,15 @@
 #include "WorldPartition/WorldPartition.h"
 #include "WorldPartition/WorldPartitionActorDesc.h"
 
+UWorldPartitionEditorCell::UWorldPartitionEditorCell(const FObjectInitializer& ObjectInitializer)
+	: Super(ObjectInitializer)
+#if WITH_EDITOR
+	, Bounds(ForceInitToZero)
+	, bLoaded(false)
+	, bLoadedChangedByUserOperation(false)
+#endif
+{}
+
 #if WITH_EDITOR
 void UWorldPartitionEditorCell::AddReferencedObjects(UObject* InThis, FReferenceCollector& Collector)
 {
@@ -34,7 +43,7 @@ void UWorldPartitionEditorCell::AddReferencedObjects(UObject* InThis, FReference
 
 void UWorldPartitionEditorCell::BeginDestroy()
 {
-	// Release WorldPartition Actor Handles/References
+	// Release handles/references
 	Actors.Empty();
 	LoadedActors.Empty();
 	Super::BeginDestroy();
@@ -109,13 +118,4 @@ void UWorldPartitionEditorCell::RemoveActor(const FGuid& Source, const FWorldPar
 		}
 	}
 }
-#endif
-
-UWorldPartitionEditorCell::UWorldPartitionEditorCell(const FObjectInitializer& ObjectInitializer)
-	: Super(ObjectInitializer)
-#if WITH_EDITOR
-	, Bounds(ForceInitToZero)
-	, bLoaded(false)
-	, bLoadedChangedByUserOperation(false)
-#endif
-{}
+#endif
\ No newline at end of file
diff --git a/Engine/Source/Runtime/Engine/Public/WorldPartition/WorldPartitionSubsystem.h b/Engine/Source/Runtime/Engine/Public/WorldPartition/WorldPartitionSubsystem.h
--- a/Engine/Source/Runtime/Engine/Public/WorldPartition/WorldPartitionSubsystem.h
+++ b/Engine/Source/Runtime/Engine/Public/WorldPartition/WorldPartitionSubsystem.h
@@ -3,14 +3,9 @@
 #pragma once
 
 #include "Subsystems/WorldSubsystem.h"
-#include "Templates/SubclassOf.h"
-#include "Misc/Optional.h"
-#include "WorldPartition/WorldPartitionStreamingSource.h"
 #include "WorldPartitionSubsystem.generated.h"
 
 class UWorldPartition;
-class UWorldPartitionEditorCell;
-class FWorldPartitionActorDesc;
 
 enum class EWorldPartitionRuntimeCellState : uint8;
 
