Renamed misnamed "LOD" Actor category to "HLOD"
* More fitting documentation/name for the bEnableAutoLODGeneration property (now using the HLOD category)
* Side effect of this change is that the HLOD Layer property is now hidden on a lot of classes were it wasn't wanted
* Improved property display name/comments for AActor::bEnableAutoLODGeneration, now "Include Actor in HLOD" (UPrimitiveComponent has a "Include Component in HLOD" property)

diff --git a/Engine/Plugins/Enterprise/VariantManagerContent/Source/VariantManagerContent/Public/LevelVariantSetsActor.h b/Engine/Plugins/Enterprise/VariantManagerContent/Source/VariantManagerContent/Public/LevelVariantSetsActor.h
--- a/Engine/Plugins/Enterprise/VariantManagerContent/Source/VariantManagerContent/Public/LevelVariantSetsActor.h
+++ b/Engine/Plugins/Enterprise/VariantManagerContent/Source/VariantManagerContent/Public/LevelVariantSetsActor.h
@@ -15,7 +15,7 @@ class ULevelVariantSets;
 class ULevelVariantSetsFunctionDirector;
 class UVariantSet;
 
-UCLASS(hideCategories=(Rendering, Physics, LOD, Activation, Input, Actor, Cooking))
+UCLASS(hideCategories=(Rendering, Physics, HLOD, Activation, Input, Actor, Cooking))
 class VARIANTMANAGERCONTENT_API ALevelVariantSetsActor : public AActor
 {
 public:
diff --git a/Engine/Plugins/Enterprise/VariantManagerContent/Source/VariantManagerContentEditor/Public/VariantManagerTestActor.h b/Engine/Plugins/Enterprise/VariantManagerContent/Source/VariantManagerContentEditor/Public/VariantManagerTestActor.h
--- a/Engine/Plugins/Enterprise/VariantManagerContent/Source/VariantManagerContentEditor/Public/VariantManagerTestActor.h
+++ b/Engine/Plugins/Enterprise/VariantManagerContent/Source/VariantManagerContentEditor/Public/VariantManagerTestActor.h
@@ -17,7 +17,7 @@ enum class EVariantManagerTestEnum : uint8
 	ThirdOption = 45,
 };
 
-UCLASS(hideCategories=(Rendering, Physics, LOD, Activation, Input, Actor, Cooking))
+UCLASS(hideCategories=(Rendering, Physics, HLOD, Activation, Input, Actor, Cooking))
 class VARIANTMANAGERCONTENTEDITOR_API AVariantManagerTestActor : public AActor
 {
 public:
diff --git a/Engine/Plugins/Experimental/VirtualProductionUtilities/Source/VPUtilities/Public/VPRootActor.h b/Engine/Plugins/Experimental/VirtualProductionUtilities/Source/VPUtilities/Public/VPRootActor.h
--- a/Engine/Plugins/Experimental/VirtualProductionUtilities/Source/VPUtilities/Public/VPRootActor.h
+++ b/Engine/Plugins/Experimental/VirtualProductionUtilities/Source/VPUtilities/Public/VPRootActor.h
@@ -19,7 +19,7 @@ class UMeshComponent;
 /**
  * AVPRootActor
  */
-UCLASS(meta = (DisplayName = "VP Root Actor"), hidecategories=(Actor, Collision, Input, LOD, Replication))
+UCLASS(meta = (DisplayName = "VP Root Actor"), hidecategories=(Actor, Collision, Input, HLOD, Replication))
 class VPUTILITIES_API AVPRootActor : public AActor
 {
 	GENERATED_BODY()
diff --git a/Engine/Plugins/Runtime/GeoReferencing/Source/GeoReferencing/Public/GeoReferencingSystem.h b/Engine/Plugins/Runtime/GeoReferencing/Source/GeoReferencing/Public/GeoReferencingSystem.h
--- a/Engine/Plugins/Runtime/GeoReferencing/Source/GeoReferencing/Public/GeoReferencingSystem.h
+++ b/Engine/Plugins/Runtime/GeoReferencing/Source/GeoReferencing/Public/GeoReferencingSystem.h
@@ -35,7 +35,7 @@ enum class EPlanetShape : uint8 {
 /**
  *
  */
-UCLASS(hidecategories=(Transform, Replication, Actor, LOD, Cooking))
+UCLASS(hidecategories=(Transform, Replication, Actor, Cooking))
 class GEOREFERENCING_API AGeoReferencingSystem : public AInfo
 {
 	GENERATED_BODY()
diff --git a/Engine/Plugins/Runtime/nDisplay/Source/DisplayCluster/Public/DisplayClusterRootActor.h b/Engine/Plugins/Runtime/nDisplay/Source/DisplayCluster/Public/DisplayClusterRootActor.h
--- a/Engine/Plugins/Runtime/nDisplay/Source/DisplayCluster/Public/DisplayClusterRootActor.h
+++ b/Engine/Plugins/Runtime/nDisplay/Source/DisplayCluster/Public/DisplayClusterRootActor.h
@@ -40,7 +40,7 @@ class UProceduralMeshComponent;
 /**
  * VR root. This contains nDisplay VR hierarchy in the game.
  */
-UCLASS(HideCategories=(Replication, Collision, Input, Actor, LOD, Cooking, Physics, Activation, AssetUserData, ActorTick, Advanced), meta=(DisplayName = "nDisplay Root Actor"))
+UCLASS(HideCategories=(Replication, Collision, Input, Actor, HLOD, Cooking, Physics, Activation, AssetUserData, ActorTick, Advanced), meta=(DisplayName = "nDisplay Root Actor"))
 class DISPLAYCLUSTER_API ADisplayClusterRootActor
 	: public AActor
 {
diff --git a/Engine/Source/Runtime/Engine/Classes/Components/PrimitiveComponent.h b/Engine/Source/Runtime/Engine/Classes/Components/PrimitiveComponent.h
--- a/Engine/Source/Runtime/Engine/Classes/Components/PrimitiveComponent.h
+++ b/Engine/Source/Runtime/Engine/Classes/Components/PrimitiveComponent.h
@@ -272,15 +272,15 @@ public:
 
 #if WITH_EDITORONLY_DATA
 	/** Whether to include this component in HLODs or not. */
-	UPROPERTY(EditAnywhere, AdvancedDisplay, BlueprintReadWrite, Category = HLOD, meta = (DisplayName = "Include Component in HLOD"))
+	UPROPERTY(EditAnywhere, BlueprintReadWrite, Category=HLOD, meta=(DisplayName="Include Component in HLOD"))
 	uint8 bEnableAutoLODGeneration : 1;
 
 	/** Which specific HLOD levels this component should be excluded from */
-	UPROPERTY(EditAnywhere, AdvancedDisplay, BlueprintReadWrite, Category = HLOD, meta = (DisplayName = "Exclude from HLOD Levels", EditConditionHides, EditCondition = "bEnableAutoLODGeneration"))
+	UPROPERTY(EditAnywhere, AdvancedDisplay, BlueprintReadWrite, Category=HLOD, meta=(DisplayName="Exclude from HLOD Levels", EditConditionHides, EditCondition="bEnableAutoLODGeneration"))
 	TArray<int32> ExcludeForSpecificHLODLevels;
 
 	/** Determines how the geometry of a component will be incorporated in proxy (simplified) HLODs. */
-	UPROPERTY(EditAnywhere, AdvancedDisplay, BlueprintReadWrite, Category = HLOD, meta = (DisplayName = "HLOD Batching Policy", EditConditionHides, EditCondition = "bEnableAutoLODGeneration"))
+	UPROPERTY(EditAnywhere, AdvancedDisplay, BlueprintReadWrite, Category=HLOD, meta=(DisplayName="HLOD Batching Policy", EditConditionHides, EditCondition="bEnableAutoLODGeneration"))
 	EHLODBatchingPolicy HLODBatchingPolicy;
 
 
diff --git a/Engine/Source/Runtime/Engine/Classes/GameFramework/Actor.h b/Engine/Source/Runtime/Engine/Classes/GameFramework/Actor.h
--- a/Engine/Source/Runtime/Engine/Classes/GameFramework/Actor.h
+++ b/Engine/Source/Runtime/Engine/Classes/GameFramework/Actor.h
@@ -355,8 +355,8 @@ public:
 	UPROPERTY(EditAnywhere, AdvancedDisplay, Category=Actor)
 	uint8 bIgnoresOriginShifting:1;
 
-	/** If true, and if World setting has bEnableHierarchicalLOD equal to true, then it will generate LODActor from groups of clustered Actor */
-	UPROPERTY(EditAnywhere, AdvancedDisplay, BlueprintReadWrite, Category=LOD, meta=(DisplayName="Include Actor for HLOD Mesh generation"))
+	/** Whether this actor should be considered or not during HLOD generation. */
+	UPROPERTY(EditAnywhere, BlueprintReadWrite, Category=HLOD, meta=(DisplayName="Include Actor in HLOD"))
 	uint8 bEnableAutoLODGeneration:1;
 
 	/** Whether this actor is editor-only. Use with care, as if this actor is referenced by anything else that reference will be NULL in cooked builds */
diff --git a/Engine/Source/Runtime/Engine/Classes/GameFramework/Info.h b/Engine/Source/Runtime/Engine/Classes/GameFramework/Info.h
--- a/Engine/Source/Runtime/Engine/Classes/GameFramework/Info.h
+++ b/Engine/Source/Runtime/Engine/Classes/GameFramework/Info.h
@@ -16,7 +16,7 @@
  * Info is the base class of an Actor that isn't meant to have a physical representation in the world, used primarily
  * for "manager" type classes that hold settings data about the world, but might need to be an Actor for replication purposes.
  */
-UCLASS(abstract, hidecategories=(Input, Movement, Collision, Rendering, LOD, WorldPartition, DataLayers, Transformation), showcategories=("Input|MouseInput", "Input|TouchInput"), MinimalAPI, NotBlueprintable)
+UCLASS(abstract, hidecategories=(Input, Movement, Collision, Rendering, HLOD, WorldPartition, DataLayers, Transformation), showcategories=("Input|MouseInput", "Input|TouchInput"), MinimalAPI, NotBlueprintable)
 class AInfo : public AActor
 {
 	GENERATED_UCLASS_BODY()
diff --git a/Engine/Source/Runtime/Engine/Public/LevelInstance/LevelInstanceEditorPivotActor.h b/Engine/Source/Runtime/Engine/Public/LevelInstance/LevelInstanceEditorPivotActor.h
--- a/Engine/Source/Runtime/Engine/Public/LevelInstance/LevelInstanceEditorPivotActor.h
+++ b/Engine/Source/Runtime/Engine/Public/LevelInstance/LevelInstanceEditorPivotActor.h
@@ -14,7 +14,7 @@ class ULevelStreaming;
 /**
  * 
  */
-UCLASS(transient, notplaceable, hidecategories=(Rendering,Replication,Collision,Partition,Input,HLOD,Actor,LOD,Cooking))
+UCLASS(transient, notplaceable, hidecategories=(Rendering,Replication,Collision,Partition,Input,HLOD,Actor,Cooking))
 class ENGINE_API ALevelInstancePivot : public AActor
 {
 	GENERATED_UCLASS_BODY()
diff --git a/Engine/Source/Runtime/Landscape/Classes/LandscapeProxy.h b/Engine/Source/Runtime/Landscape/Classes/LandscapeProxy.h
--- a/Engine/Source/Runtime/Landscape/Classes/LandscapeProxy.h
+++ b/Engine/Source/Runtime/Landscape/Classes/LandscapeProxy.h
@@ -357,7 +357,7 @@ struct FLandscapeProxyMaterialOverride
 #endif
 };
 
-UCLASS(Abstract, MinimalAPI, NotBlueprintable, NotPlaceable, hidecategories=(Display, Attachment, Physics, Debug, Lighting, LOD), showcategories=(Lighting, Rendering, Transformation), hidecategories=(Mobility))
+UCLASS(Abstract, MinimalAPI, NotBlueprintable, NotPlaceable, hidecategories=(Display, Attachment, Physics, Debug, Lighting, HLOD), showcategories=(Lighting, Rendering, Transformation), hidecategories=(Mobility))
 class ALandscapeProxy : public APartitionActor, public ILandscapeSplineInterface
 {
 	GENERATED_BODY()
diff --git a/Engine/Source/Runtime/LevelSequence/Public/LevelSequenceActor.h b/Engine/Source/Runtime/LevelSequence/Public/LevelSequenceActor.h
--- a/Engine/Source/Runtime/LevelSequence/Public/LevelSequenceActor.h
+++ b/Engine/Source/Runtime/LevelSequence/Public/LevelSequenceActor.h
@@ -58,7 +58,7 @@ protected:
 /**
  * Actor responsible for controlling a specific level sequence in the world.
  */
-UCLASS(hideCategories=(Rendering, Physics, LOD, Activation, Input))
+UCLASS(hideCategories=(Rendering, Physics, HLOD, Activation, Input))
 class LEVELSEQUENCE_API ALevelSequenceActor
 	: public AActor
 	, public IMovieSceneSequenceActor
diff --git a/Engine/Source/Runtime/LevelSequence/Public/SequenceMediaController.h b/Engine/Source/Runtime/LevelSequence/Public/SequenceMediaController.h
--- a/Engine/Source/Runtime/LevelSequence/Public/SequenceMediaController.h
+++ b/Engine/Source/Runtime/LevelSequence/Public/SequenceMediaController.h
@@ -17,7 +17,7 @@ class UMediaPlayer;
 /**
  * Replicated actor class that is responsible for instigating various cinematic assets (Media, Audio, Level Sequences) in a synchronized fasion
  */
-UCLASS(hideCategories=(Rendering, Physics, LOD, Activation, Input, Collision))
+UCLASS(hideCategories=(Rendering, Physics, HLOD, Activation, Input, Collision))
 class LEVELSEQUENCE_API ALevelSequenceMediaController
 	: public AActor
 	, public IMovieSceneCustomClockSource
diff --git a/Engine/Source/Runtime/NavigationSystem/Public/NavigationTestingActor.h b/Engine/Source/Runtime/NavigationSystem/Public/NavigationTestingActor.h
--- a/Engine/Source/Runtime/NavigationSystem/Public/NavigationTestingActor.h
+++ b/Engine/Source/Runtime/NavigationSystem/Public/NavigationTestingActor.h
@@ -40,7 +40,7 @@ namespace ENavCostDisplay
 	};
 }
 
-UCLASS(hidecategories=(Object, Actor, Input, Rendering, Replication, LOD, Cooking), showcategories=("Input|MouseInput", "Input|TouchInput"), Blueprintable)
+UCLASS(hidecategories=(Object, Actor, Input, Rendering, Replication, HLOD, Cooking), showcategories=("Input|MouseInput", "Input|TouchInput"), Blueprintable)
 class NAVIGATIONSYSTEM_API ANavigationTestingActor : public AActor, public INavAgentInterface, public INavPathObserverInterface
 {
 	GENERATED_UCLASS_BODY()
