Update the VIEWNAMES command to support FILENAME option which outputs dumped names into a text file in the project directory saved folder, with the provided filename.

diff --git a/Engine/Source/Runtime/Engine/Private/UnrealEngine.cpp b/Engine/Source/Runtime/Engine/Private/UnrealEngine.cpp
--- a/Engine/Source/Runtime/Engine/Private/UnrealEngine.cpp
+++ b/Engine/Source/Runtime/Engine/Private/UnrealEngine.cpp
@@ -255,6 +255,8 @@ UnrealEngine.cpp: Implements the UEngine class and helpers.
 #include "Engine/InstancedStaticMesh.h"
 #include "IDeviceProfileSelectorModule.h"
 
+#include "ObjectTools.h"
+
 #if WITH_DUMPGPU
 	#include "RenderGraph.h"
 #endif
@@ -5269,9 +5271,25 @@ bool UEngine::HandleViewnamesCommand( const TCHAR* Cmd, FOutputDevice& Ar )
 	int32 NumLast = 0;
 	int32 BeginIdx = FParse::Value(Cmd, TEXT("NUM="), NumLast) ? FMath::Max(Entries.Num() - NumLast, 0) : 0;
 
+	FString Filename;
+	TUniquePtr<FArchive> FileWriter;
+	if (FParse::Value(Cmd, TEXT("FILENAME="), Filename))
+	{
+		const FString OutputFilename = FPaths::ProjectSavedDir() / FString::Printf(TEXT("Names-%s-%s.log"), *ObjectTools::SanitizeObjectPath(Filename), *FDateTime::Now().ToString());
+		FileWriter = TUniquePtr<FArchive>(IFileManager::Get().CreateFileWriter(*OutputFilename));
+	}
+
 	for (int32 I = BeginIdx; I < Entries.Num(); ++I)
 	{
-		Ar.Log(*Entries[I]->GetPlainNameString());	
+		const FString PlainNameString = Entries[I]->GetPlainNameString();
+
+		Ar.Log(*PlainNameString);
+
+		if (FileWriter)
+		{
+			FString LogEntry = FString::Printf(TEXT("%s") LINE_TERMINATOR, *PlainNameString);
+			FileWriter->Serialize(TCHAR_TO_ANSI(*LogEntry), LogEntry.Len());
+		}
 	}
 
 	return true;
