Fix ObjectTools.h include

diff --git a/Engine/Source/Runtime/Engine/Private/UnrealEngine.cpp b/Engine/Source/Runtime/Engine/Private/UnrealEngine.cpp
--- a/Engine/Source/Runtime/Engine/Private/UnrealEngine.cpp
+++ b/Engine/Source/Runtime/Engine/Private/UnrealEngine.cpp
@@ -255,8 +255,6 @@ UnrealEngine.cpp: Implements the UEngine class and helpers.
 #include "Engine/InstancedStaticMesh.h"
 #include "IDeviceProfileSelectorModule.h"
 
-#include "ObjectTools.h"
-
 #if WITH_DUMPGPU
 	#include "RenderGraph.h"
 #endif
@@ -5275,7 +5273,7 @@ bool UEngine::HandleViewnamesCommand( const TCHAR* Cmd, FOutputDevice& Ar )
 	TUniquePtr<FArchive> FileWriter;
 	if (FParse::Value(Cmd, TEXT("FILENAME="), Filename))
 	{
-		const FString OutputFilename = FPaths::ProjectSavedDir() / FString::Printf(TEXT("Names-%s-%s.log"), *ObjectTools::SanitizeObjectPath(Filename), *FDateTime::Now().ToString());
+		const FString OutputFilename = FPaths::ProjectSavedDir() / FString::Printf(TEXT("Names-%s-%s.log"), *Filename, *FDateTime::Now().ToString());
 		FileWriter = TUniquePtr<FArchive>(IFileManager::Get().CreateFileWriter(*OutputFilename));
 	}
 
