Fix issue with inline reduction and postload usersectiondata correction.
#rnx

diff --git a/Engine/Source/Developer/MeshUtilities/MeshUtilities.Build.cs b/Engine/Source/Developer/MeshUtilities/MeshUtilities.Build.cs
--- a/Engine/Source/Developer/MeshUtilities/MeshUtilities.Build.cs
+++ b/Engine/Source/Developer/MeshUtilities/MeshUtilities.Build.cs
@@ -40,7 +40,9 @@ public class MeshUtilities : ModuleRules
 				"StaticMeshDescription",
 				"ToolMenus",
 				"MeshUtilitiesEngine",
-            }
+				"SkeletalMeshUtilitiesCommon",
+				"ClothingSystemRuntimeCommon",
+			}
 		);
 
         PublicIncludePathModuleNames.AddRange(
diff --git a/Engine/Source/Developer/MeshUtilities/Private/MeshUtilities.cpp b/Engine/Source/Developer/MeshUtilities/Private/MeshUtilities.cpp
--- a/Engine/Source/Developer/MeshUtilities/Private/MeshUtilities.cpp
+++ b/Engine/Source/Developer/MeshUtilities/Private/MeshUtilities.cpp
@@ -101,6 +101,9 @@
 #include "StaticMeshOperations.h"
 #include "MeshUtilitiesEngine.h"
 
+#include "LODUtilities.h"
+#include "ClothingAsset.h"
+
 #if WITH_EDITOR
 #include "Editor.h"
 #include "UnrealEdMisc.h"
@@ -5830,6 +5833,8 @@ void FMeshUtilities::GenerateRuntimeSkinWeightData(const FSkeletalMeshLODModel*
 
 void FMeshUtilities::CreateImportDataFromLODModel(USkeletalMesh* SkeletalMesh) const
 {
+	check(IsInGameThread());
+
 	for (int32 LodIndex = 0; LodIndex < SkeletalMesh->GetLODNum(); LodIndex++)
 	{
 		const FSkeletalMeshLODInfo* ThisLODInfo = SkeletalMesh->GetLODInfo(LodIndex);
@@ -5839,10 +5844,12 @@ void FMeshUtilities::CreateImportDataFromLODModel(USkeletalMesh* SkeletalMesh) c
 		const bool bRawBuildDataAvailable = SkeletalMesh->IsLODImportedDataBuildAvailable(LodIndex);
 		if (!bRawDataEmpty && bRawBuildDataAvailable)
 		{
+			//No need to create import data if we already have some
 			continue;
 		}
-		const bool bInlineReduction = SkeletalMesh->IsReductionActive(LodIndex) && (ThisLODInfo->ReductionSettings.BaseLOD == LodIndex);
-		if (!bInlineReduction)
+		const bool bReductionActive = SkeletalMesh->IsReductionActive(LodIndex);
+		const bool bInlineReduction = (ThisLODInfo->ReductionSettings.BaseLOD == LodIndex);
+		if (bReductionActive && !bInlineReduction)
 		{
 			//Generated LOD (not inline) do not need imported data
 			continue;
@@ -5851,39 +5858,62 @@ void FMeshUtilities::CreateImportDataFromLODModel(USkeletalMesh* SkeletalMesh) c
 		//Find the reduction data to restore the imported LODModel by using the deprecated OriginalReductionSourceMeshData_DEPRECATED data
 		if (SkeletalMesh->GetImportedModel()->OriginalReductionSourceMeshData_DEPRECATED.IsValidIndex(LodIndex))
 		{
-			FSkeletalMeshLODModel& ToUpdateLODModel = SkeletalMesh->GetImportedModel()->LODModels[LodIndex];
+			FSkeletalMeshLODModel TempLODModel;
 			TMap<FString, TArray<FMorphTargetDelta>> LODMorphTargetData;
 			
 			//Old inline reduced assets do not have the original imported skin weight data, it must be re-import if we want to convert it to import data
-			//Swap the LODModel pointer to the one we store when reducing, so we add the true import data
-			SkeletalMesh->GetImportedModel()->OriginalReductionSourceMeshData_DEPRECATED[LodIndex]->LoadReductionData(ToUpdateLODModel, LODMorphTargetData, SkeletalMesh);
+			SkeletalMesh->GetImportedModel()->OriginalReductionSourceMeshData_DEPRECATED[LodIndex]->LoadReductionData(TempLODModel, LODMorphTargetData, SkeletalMesh);
 
-			TArray<UMorphTarget*> OriginalMorphTargets;
-			OriginalMorphTargets.Reserve(LODMorphTargetData.Num());
-			for (const TPair<FString, TArray<FMorphTargetDelta>>& MorphTargetChannel : LODMorphTargetData)
+			if (TempLODModel.NumVertices > 0)
 			{
-				UMorphTarget* ExistingMorphTarget = nullptr;
-				for (UMorphTarget* MorphTarget : SkeletalMesh->GetMorphTargets())
+				FSkeletalMeshLODModel& ToUpdateLODModel = SkeletalMesh->GetImportedModel()->LODModels[LodIndex];
+
+				// Unbind any existing clothing assets before loading the reduction data
+				TArray<ClothingAssetUtils::FClothingAssetMeshBinding> ClothingBindings;
+				FLODUtilities::UnbindClothingAndBackup(SkeletalMesh, ClothingBindings, LodIndex);
+
+				//Swap the LODModel pointer to the one we store when reducing, so we add the true import data
+				FSkeletalMeshLODModel::CopyStructure(&ToUpdateLODModel, &TempLODModel);
+
+				//Rebind the cloth asset on the LODModel restore by the reduction data
+				FLODUtilities::RestoreClothingFromBackup(SkeletalMesh, ClothingBindings, LodIndex);
+
+				//Since we load the reduction LODModel we have to validate the material indexes.
+				SkeletalMesh->ValidateAllLodMaterialIndexes();
+
+				TArray<UMorphTarget*> OriginalMorphTargets;
+				OriginalMorphTargets.Reserve(LODMorphTargetData.Num());
+				for (const TPair<FString, TArray<FMorphTargetDelta>>& MorphTargetChannel : LODMorphTargetData)
 				{
-					if (MorphTarget->GetName() == MorphTargetChannel.Key)
+					UMorphTarget* ExistingMorphTarget = nullptr;
+					for (UMorphTarget* MorphTarget : SkeletalMesh->GetMorphTargets())
+					{
+						if (MorphTarget->GetName() == MorphTargetChannel.Key)
+						{
+							ExistingMorphTarget = MorphTarget;
+						}
+					}
+					if (!ExistingMorphTarget)
 					{
-						ExistingMorphTarget = MorphTarget;
+						FName MorphTargetName = *(MorphTargetChannel.Key);
+						ExistingMorphTarget = NewObject<UMorphTarget>(SkeletalMesh, MorphTargetName);
+						check(ExistingMorphTarget);
 					}
-				}
-				if (!ExistingMorphTarget)
-				{
-					FName MorphTargetName = *(MorphTargetChannel.Key);
-					ExistingMorphTarget = NewObject<UMorphTarget>(SkeletalMesh, MorphTargetName);
-					check(ExistingMorphTarget);
-				}
 
-				const TArray<FMorphTargetDelta>& MorphTargetDeltas = MorphTargetChannel.Value;
-				ExistingMorphTarget->PopulateDeltas(MorphTargetDeltas, LodIndex, ToUpdateLODModel.Sections, ThisLODInfo->BuildSettings.bRecomputeNormals, false, ThisLODInfo->BuildSettings.MorphThresholdPosition);
+					const TArray<FMorphTargetDelta>& MorphTargetDeltas = MorphTargetChannel.Value;
+					ExistingMorphTarget->PopulateDeltas(MorphTargetDeltas, LodIndex, ToUpdateLODModel.Sections, ThisLODInfo->BuildSettings.bRecomputeNormals, false, ThisLODInfo->BuildSettings.MorphThresholdPosition);
+				}
 			}
 		}
 
 		FSkeletalMeshLODModel* LODModel = &(SkeletalMesh->GetImportedModel()->LODModels[LodIndex]);
 
+		//Make sure the inline cache fit the LODModel
+		if (SkeletalMesh->GetImportedModel()->InlineReductionCacheDatas.IsValidIndex(LodIndex))
+		{
+			SkeletalMesh->GetImportedModel()->InlineReductionCacheDatas[LodIndex].SetCacheGeometryInfo(*LODModel);
+		}
+
 		FSkeletalMeshImportData ImportData;
 		ImportData.bDiffPose = false;
 		ImportData.bHasNormals = true;
diff --git a/Engine/Source/Developer/SkeletalMeshUtilitiesCommon/Private/LODUtilities.cpp b/Engine/Source/Developer/SkeletalMeshUtilitiesCommon/Private/LODUtilities.cpp
--- a/Engine/Source/Developer/SkeletalMeshUtilitiesCommon/Private/LODUtilities.cpp
+++ b/Engine/Source/Developer/SkeletalMeshUtilitiesCommon/Private/LODUtilities.cpp
@@ -2380,7 +2380,7 @@ void FLODUtilities::RegenerateDependentLODs(USkeletalMesh* SkeletalMesh, int32 L
 				int32 DependentLODIndex = DependentLODs[IterationIndex];
 				check(SkeletalMesh->GetLODInfo(DependentLODIndex)); //We cannot add a LOD when reducing with multi thread, so check we already have one
 				FLODUtilities::SimplifySkeletalMeshLOD(SkeletalMesh, DependentLODIndex, TargetPlatform, false, &bNeedsPackageDirtied);
-			});
+			}, IsInGameThread() ? EParallelForFlags::None : EParallelForFlags::ForceSingleThread);
 
 			if (bNeedsPackageDirtied && IsInGameThread())
 			{
diff --git a/Engine/Source/Runtime/Engine/Classes/Engine/SkeletalMesh.h b/Engine/Source/Runtime/Engine/Classes/Engine/SkeletalMesh.h
--- a/Engine/Source/Runtime/Engine/Classes/Engine/SkeletalMesh.h
+++ b/Engine/Source/Runtime/Engine/Classes/Engine/SkeletalMesh.h
@@ -2663,6 +2663,11 @@ private:
 	void PostLoadVerifyAndFixBadTangent();
 
 public:
+	/*
+	 * This function will enforce valid material index in the sections and the LODMaterialMap of all LOD.
+	 */
+	void ValidateAllLodMaterialIndexes();
+
 	//We want to avoid calling post edit change multiple time during import and build process.
 
 	/*
diff --git a/Engine/Source/Runtime/Engine/Private/SkeletalMesh.cpp b/Engine/Source/Runtime/Engine/Private/SkeletalMesh.cpp
--- a/Engine/Source/Runtime/Engine/Private/SkeletalMesh.cpp
+++ b/Engine/Source/Runtime/Engine/Private/SkeletalMesh.cpp
@@ -684,7 +684,9 @@ bool USkeletalMesh::IsReductionActive(int32 LODIndex) const
 		uint32 LODTriNumber = MAX_uint32;
 		const FSkeletalMeshLODInfo* LODInfoPtr = GetLODInfo(LODIndex);
 		const bool bLODHasBeenSimplified = LODInfoPtr && LODInfoPtr->bHasBeenSimplified;
-		if (GetImportedModel() && GetImportedModel()->LODModels.IsValidIndex(LODIndex))
+		//If we are not inline reduced, we wont set the LODVertexNumber and LODTriNumber from the LODModel or from the cache.
+		const bool bInlineReduction = LODInfoPtr && (LODInfoPtr->ReductionSettings.BaseLOD == LODIndex);
+		if (bInlineReduction && GetImportedModel() && GetImportedModel()->LODModels.IsValidIndex(LODIndex))
 		{
 			if (!bLODHasBeenSimplified)
 			{
@@ -2869,6 +2871,91 @@ void USkeletalMesh::CreateUserSectionsDataForLegacyAssets()
 	}
 }
 
+void USkeletalMesh::ValidateAllLodMaterialIndexes()
+{
+	TRACE_CPUPROFILER_EVENT_SCOPE(USkeletalMesh::ValidateAllLodMaterialIndexes);
+
+	for (int32 LodIndex = 0; LodIndex < GetLODNum(); LodIndex++)
+	{
+		FSkeletalMeshLODInfo* LODInfoPtr = GetLODInfo(LodIndex);
+		if (!LODInfoPtr)
+		{
+			continue;
+		}
+		FSkeletalMeshLODModel& LODModel = GetImportedModel()->LODModels[LodIndex];
+		const int32 SectionNum = LODModel.Sections.Num();
+		const TArray<FSkeletalMaterial>& ListOfMaterials = GetMaterials();
+		//See if more then one section use the same UserSectionData
+		for (int32 SectionIndex = 0; SectionIndex < SectionNum; ++SectionIndex)
+		{
+			int32 LodMaterialMapOverride = INDEX_NONE;
+			FSkelMeshSection& Section = LODModel.Sections[SectionIndex];
+			//Validate and fix the LODMaterialMap override
+			if (LODInfoPtr->LODMaterialMap.IsValidIndex(SectionIndex) && LODInfoPtr->LODMaterialMap[SectionIndex] != INDEX_NONE)
+			{
+				LodMaterialMapOverride = LODInfoPtr->LODMaterialMap[SectionIndex];
+				if (!ListOfMaterials.IsValidIndex(LodMaterialMapOverride))
+				{
+					UE_ASSET_LOG(LogSkeletalMesh
+						, Display
+						, this
+						, TEXT("Fix LOD %d Section %d LODMaterialMap override material index from %d to INDEX_NONE. The value is not pointing on a valid Material slot index.")
+						, LodIndex
+						, SectionIndex
+						, LodMaterialMapOverride);
+					LODInfoPtr->LODMaterialMap[SectionIndex] = INDEX_NONE;
+				}
+			}
+			//Validate and fix the section material index
+			{
+				if (!ListOfMaterials.IsValidIndex(Section.MaterialIndex))
+				{
+					if (LodMaterialMapOverride != INDEX_NONE)
+					{
+						UE_ASSET_LOG(LogSkeletalMesh
+							, Display
+							, this
+							, TEXT("Fix LOD %d Section %d Material index from %d to %d. The fallback value is from the LODMaterialMap Override. The value is not pointing on a valid Material slot index.")
+							, LodIndex
+							, SectionIndex
+							, Section.MaterialIndex
+							, LodMaterialMapOverride);
+						Section.MaterialIndex = LodMaterialMapOverride;
+					}
+					else
+					{
+						//Fall back on the original section index
+						if (ListOfMaterials.IsValidIndex(Section.OriginalDataSectionIndex))
+						{
+							UE_ASSET_LOG(LogSkeletalMesh
+								, Display
+								, this
+								, TEXT("Fix LOD %d Section %d Material index from %d to %d. The fallback value is from the OriginalDataSectionIndex. The value is not pointing on a valid Material slot index.")
+								, LodIndex
+								, SectionIndex
+								, Section.MaterialIndex
+								, Section.OriginalDataSectionIndex);
+							Section.MaterialIndex = Section.OriginalDataSectionIndex;
+						}
+						else
+						{
+							UE_ASSET_LOG(LogSkeletalMesh
+								, Display
+								, this
+								, TEXT("Fix LOD %d Section %d Material index from %d to 0. The fallback value is 0. The value is not pointing on a valid Material slot index.")
+								, LodIndex
+								, SectionIndex
+								, Section.MaterialIndex);
+							//Fallback on material index 0
+							Section.MaterialIndex = 0;
+						}
+					}
+				}
+			}
+		}
+	}
+}
+
 void USkeletalMesh::PostLoadValidateUserSectionData()
 {
 	TRACE_CPUPROFILER_EVENT_SCOPE(USkeletalMesh::PostLoadValidateUserSectionData);
@@ -2892,15 +2979,12 @@ void USkeletalMesh::PostLoadValidateUserSectionData()
 		FSkeletalMeshLODModel& BaseReductionLODModel = GetImportedModel()->LODModels[ReductionBaseLOD];
 		FSkeletalMeshLODModel& ThisLODModel = GetImportedModel()->LODModels[LodIndex];
 		const int32 SectionNum = ThisLODModel.Sections.Num();
-		
+		const int32 UserSectionsDataNum = ThisLODModel.UserSectionsData.Num();
+		const int32 BaseUserSectionsDataNum = BaseReductionLODModel.UserSectionsData.Num();
 		//We must make sure the result is similar to what the reduction will give. So we will not have more user section data then the number we have for the base LOD.
 		//Because reduction reset the UserSectionData to the number of parent section after the reduction.
-		bool UserSectionDataBadCount = ThisLODModel.UserSectionsData.Num() > SectionNum;
-		if (!UserSectionDataBadCount && LodIndex != ReductionBaseLOD)
-		{
-			UserSectionDataBadCount = (ThisLODModel.UserSectionsData.Num() > BaseReductionLODModel.UserSectionsData.Num());
-		}
-		bool bLODHaveSectionIssue = UserSectionDataBadCount;
+		const bool bIsInlineReduction = LodIndex == ReductionBaseLOD;
+		bool bLODHaveSectionIssue = !bIsInlineReduction && ( (UserSectionsDataNum > SectionNum) || (UserSectionsDataNum > BaseUserSectionsDataNum) );
 		if (!bLODHaveSectionIssue)
 		{
 			//See if more then one section use the same UserSectionData
@@ -3046,13 +3130,13 @@ void USkeletalMesh::PostLoadEnsureImportDataExist()
 		//We create the import data only if the mesh was simplified and its doing inline reduction
 		//Non simplified mesh without import data should not be converted because the simplification will kick in and the asset will be simplified (asset vertex and face count will change)
 		const bool bReductionActive = IsReductionActive(LodIndex);
-		const bool bInlineReduction = bReductionActive && (ThisLODInfo->ReductionSettings.BaseLOD == LodIndex);
+		const bool bInlineReduction = (ThisLODInfo->ReductionSettings.BaseLOD == LodIndex);
 		if (bReductionActive && !bInlineReduction)
 		{
 			//Generated LOD (not inline) do not need imported data
 			continue;
 		}
-		else if (bInlineReduction && ThisLODInfo->bHasBeenSimplified)
+		else if (bInlineReduction && (bReductionActive || ThisLODInfo->bHasBeenSimplified))
 		{
 			//If we cannot convert alternate skin weight data, we do not want to allow 
 			if (GetImportedModel()->OriginalReductionSourceMeshData_DEPRECATED.IsValidIndex(LodIndex))
@@ -3099,7 +3183,10 @@ void USkeletalMesh::PostLoadEnsureImportDataExist()
 		}
 		IMeshUtilities& MeshUtilities = FModuleManager::Get().LoadModuleChecked<IMeshUtilities>("MeshUtilities");
 		//We create the import data for all LOD that do not have import data except for the generated LODs.
-		MeshUtilities.CreateImportDataFromLODModel(this);
+		{
+			FScopedSkeletalMeshPostEditChange ScopedPostEditChange(this, false, false);
+			MeshUtilities.CreateImportDataFromLODModel(this);
+		}
 #if WITH_EDITORONLY_DATA
 		//If the import data is existing we want to turn use legacy derive data key to false
 		SetUseLegacyMeshDerivedDataKey(false);
@@ -3448,6 +3535,8 @@ void USkeletalMesh::BeginPostLoadInternal(FSkeletalMeshPostLoadContext& Context)
 			CreateUserSectionsDataForLegacyAssets();
 		}
 
+		ValidateAllLodMaterialIndexes();
+
 		PostLoadEnsureImportDataExist();
 	}
 
diff --git a/Engine/Source/Runtime/Engine/Private/SkeletalMeshLODModel.cpp b/Engine/Source/Runtime/Engine/Private/SkeletalMeshLODModel.cpp
--- a/Engine/Source/Runtime/Engine/Private/SkeletalMeshLODModel.cpp
+++ b/Engine/Source/Runtime/Engine/Private/SkeletalMeshLODModel.cpp
@@ -1264,6 +1264,7 @@ void FSkeletalMeshLODModel::UpdateChunkedSectionInfo(const FString& SkeletalMesh
 	//Fill the ChunkedParentSectionIndex data, we assume that every section using the same material are chunked
 	int32 LastMaterialIndex = INDEX_NONE;
 	uint32 LastBoneCount = 0;
+	int32 LastOriginalDataSectionIndex = INDEX_NONE;
 	int32 CurrentParentChunkIndex = INDEX_NONE;
 	int32 OriginalIndex = 0;
 	//We assume here that if the project use per platform chunking, the minimum value will be the same has the prior project settings.
@@ -1276,8 +1277,10 @@ void FSkeletalMeshLODModel::UpdateChunkedSectionInfo(const FString& SkeletalMesh
 	{
 		FSkelMeshSection& Section = Sections[LODModelSectionIndex];
 		
+		//If we have already chunked data in this LODModel use it to know if we need to chunk a section or not, this can happen when we load reduction data.
+		const bool bIsOldChunkingSection = LastOriginalDataSectionIndex != INDEX_NONE && Section.OriginalDataSectionIndex == LastOriginalDataSectionIndex;
 		//If we have cloth on a chunked section we treat the chunked section has a parent section (this is to get the same result has before the refactor)
-		if (LastBoneCount >= MaxGPUSkinBones && Section.MaterialIndex == LastMaterialIndex && !Section.ClothingData.AssetGuid.IsValid())
+		if ((bIsOldChunkingSection || LastBoneCount >= MaxGPUSkinBones) && Section.MaterialIndex == LastMaterialIndex && !Section.ClothingData.AssetGuid.IsValid())
 		{
 			Section.ChunkedParentSectionIndex = CurrentParentChunkIndex;
 			Section.OriginalDataSectionIndex = Sections[CurrentParentChunkIndex].OriginalDataSectionIndex;
@@ -1319,7 +1322,7 @@ void FSkeletalMeshLODModel::UpdateChunkedSectionInfo(const FString& SkeletalMesh
 		}
 
 		LastMaterialIndex = Section.MaterialIndex;
-		//Set the last bone count
+		LastOriginalDataSectionIndex = Section.OriginalDataSectionIndex;
 		LastBoneCount = (uint32)Sections[LODModelSectionIndex].BoneMap.Num();
 	}
 }
diff --git a/Engine/Source/Runtime/Engine/Private/SkeletalMeshRenderData.cpp b/Engine/Source/Runtime/Engine/Private/SkeletalMeshRenderData.cpp
--- a/Engine/Source/Runtime/Engine/Private/SkeletalMeshRenderData.cpp
+++ b/Engine/Source/Runtime/Engine/Private/SkeletalMeshRenderData.cpp
@@ -428,14 +428,20 @@ void FSkeletalMeshRenderData::Cache(const ITargetPlatform* TargetPlatform, USkel
 			for (int32 LODIndex = 0; LODIndex < SkelMeshModel->LODModels.Num(); LODIndex++)
 			{
 				FSkeletalMeshLODModel* LODModel = &(SkelMeshModel->LODModels[LODIndex]);
-				const FSkeletalMeshLODInfo* LODInfo = Owner->GetLODInfo(LODIndex);
+				FSkeletalMeshLODInfo* LODInfo = Owner->GetLODInfo(LODIndex);
 				check(LODInfo);
-				bool bRawDataEmpty = Owner->IsLODImportedDataEmpty(LODIndex);
-				bool bRawBuildDataAvailable = Owner->IsLODImportedDataBuildAvailable(LODIndex);
+				//We want to avoid building a LOD if the LOD was generated from a previous LODIndex.
+				const bool bIsGeneratedLodNotInline = (LODInfo->bHasBeenSimplified && Owner->IsReductionActive(LODIndex) && Owner->GetReductionSettings(LODIndex).BaseLOD < LODIndex);
+				
+				//Make sure the LOD have all the data needed to be build
+				const bool bRawDataEmpty = Owner->IsLODImportedDataEmpty(LODIndex);
+				const bool bRawBuildDataAvailable = Owner->IsLODImportedDataBuildAvailable(LODIndex);
+				
 				//Build the source model before the render data, if we are a purely generated LOD we do not need to be build
 				IMeshBuilderModule& MeshBuilderModule = IMeshBuilderModule::GetForPlatform(TargetPlatform);
-				if (!bRawDataEmpty && bRawBuildDataAvailable)
+				if (!bIsGeneratedLodNotInline && !bRawDataEmpty && bRawBuildDataAvailable)
 				{
+					LODInfo->bHasBeenSimplified = false;
 					const bool bRegenDepLODs = true;
 					FSkeletalMeshBuildParameters BuildParameters(Owner, TargetPlatform, LODIndex, bRegenDepLODs);
 					MeshBuilderModule.BuildSkeletalMesh(BuildParameters);
@@ -522,7 +528,7 @@ void FSkeletalMeshRenderData::Cache(const ITargetPlatform* TargetPlatform, USkel
 					//If we are in this case we should resave the asset so the source data will be the same and we can use this DDC. Reduction can change the number of sections and the user section data is in the DDC key.
 					//So if we change the reduction algorithm, its possible we fall in this situation.
 					//We save the real data key which force the asset to always rebuild when the editor is loading it until the user save it
-					UE_LOG(LogSkeletalMesh, Log, TEXT("Skeletal mesh [%s]: The derived data key is different after the build. Resave the asset to avoid rebuilding it everytime the editor load it."), *Owner->GetPathName());
+					UE_LOG(LogSkeletalMesh, Log, TEXT("Skeletal mesh [%s]: The derived data key is different after the build. Save the asset to avoid rebuilding it everytime the editor load it."), *Owner->GetPathName());
 				}
 			}
 			else
