#UE Fix crash when loading some AnimBPs

diff --git a/Engine/Source/Runtime/Engine/Private/Animation/AnimBlueprintGeneratedClass.cpp b/Engine/Source/Runtime/Engine/Private/Animation/AnimBlueprintGeneratedClass.cpp
--- a/Engine/Source/Runtime/Engine/Private/Animation/AnimBlueprintGeneratedClass.cpp
+++ b/Engine/Source/Runtime/Engine/Private/Animation/AnimBlueprintGeneratedClass.cpp
@@ -700,7 +700,7 @@ void UAnimBlueprintGeneratedClass::InitializeAnimNodeData(UObject* DefaultObject
 #endif
 
 	// Link functions to their nodes
-	for(int32 AnimNodeIndex = 0; AnimNodeIndex < AnimNodeProperties.Num(); ++AnimNodeIndex)
+	for(int32 AnimNodeIndex = 0; AnimNodeIndex < AnimNodeProperties.Num() && AnimNodeIndex < AnimNodeData.Num(); ++AnimNodeIndex)
 	{
 		FStructProperty* StructProperty = AnimNodeProperties[AnimNodeIndex];
 
