Fixing tool tips on FInteriorSettings
- Interior time only matters when exiting a volume for sounds on the interior of said volume, and exterior time only matters when entering a volume for sounds on the exterior of said volume

diff --git a/Engine/Source/Runtime/Engine/Classes/Sound/AudioVolume.h b/Engine/Source/Runtime/Engine/Classes/Sound/AudioVolume.h
--- a/Engine/Source/Runtime/Engine/Classes/Sound/AudioVolume.h
+++ b/Engine/Source/Runtime/Engine/Classes/Sound/AudioVolume.h
@@ -84,7 +84,7 @@ struct ENGINE_API FInteriorSettings
 	UPROPERTY(EditAnywhere, BlueprintReadWrite, Category=InteriorSettings)
 	float ExteriorVolume;
 
-	// The time over which to interpolate from the current volume to the desired volume of sounds outside the volume when the player enters or exits the volume
+	// The time over which to interpolate from the current volume to the desired volume of sounds outside the volume when the player enters the volume
 	UPROPERTY(EditAnywhere, BlueprintReadWrite, Category=InteriorSettings)
 	float ExteriorTime;
 
@@ -92,7 +92,7 @@ struct ENGINE_API FInteriorSettings
 	UPROPERTY(EditAnywhere, BlueprintReadWrite, Category = InteriorSettings)
 	float ExteriorLPF;
 
-	// The time over which to interpolate from the current LPF to the desired LPF of sounds outside the volume when the player enters or exits the volume
+	// The time over which to interpolate from the current LPF to the desired LPF of sounds outside the volume when the player enters the volume
 	UPROPERTY(EditAnywhere, BlueprintReadWrite, Category=InteriorSettings)
 	float ExteriorLPFTime;
 
@@ -100,7 +100,7 @@ struct ENGINE_API FInteriorSettings
 	UPROPERTY(EditAnywhere, BlueprintReadWrite, Category=InteriorSettings)
 	float InteriorVolume;
 
-	// The time over which to interpolate from the current volume to the desired volume of sounds inside the volume when the player enters or exits the volume
+	// The time over which to interpolate from the current volume to the desired volume of sounds inside the volume when the player exits the volume
 	UPROPERTY(EditAnywhere, BlueprintReadWrite, Category=InteriorSettings)
 	float InteriorTime;
 
@@ -108,7 +108,7 @@ struct ENGINE_API FInteriorSettings
 	UPROPERTY(EditAnywhere, BlueprintReadWrite, Category = InteriorSettings)
 	float InteriorLPF;
 
-	// The time over which to interpolate from the current LPF to the desired LPF of sounds inside the volume when the player enters or exits the volume
+	// The time over which to interpolate from the current LPF to the desired LPF of sounds inside the volume when the player exits the volume
 	UPROPERTY(EditAnywhere, BlueprintReadWrite, Category=InteriorSettings)
 	float InteriorLPFTime;
 
