Cleaning up actor references when unloading levels and making sure metadata has the standalone flag cleared so that it does not prevent levels from being GC'd.

diff --git a/Engine/Source/Runtime/Engine/Classes/Engine/Level.h b/Engine/Source/Runtime/Engine/Classes/Engine/Level.h
--- a/Engine/Source/Runtime/Engine/Classes/Engine/Level.h
+++ b/Engine/Source/Runtime/Engine/Classes/Engine/Level.h
@@ -868,6 +868,11 @@ public:
 	 */
 	ENGINE_API void CleanupLevel(bool bCleanupResources);
 
+	/**
+	 * Empties the Actors array after everything's been cleaned up
+	 */
+	ENGINE_API void CleanupActors();
+
 	/**
 	 * Clears all components of actors associated with this level (aka in Actors array) and 
 	 * also the BSP model components.
diff --git a/Engine/Source/Runtime/Engine/Private/Level.cpp b/Engine/Source/Runtime/Engine/Private/Level.cpp
--- a/Engine/Source/Runtime/Engine/Private/Level.cpp
+++ b/Engine/Source/Runtime/Engine/Private/Level.cpp
@@ -413,6 +413,12 @@ void ULevel::CleanupLevel(bool bCleanupResources)
 	}
 }
 
+void ULevel::CleanupActors()
+{
+	Actors.Empty();
+	ActorsForGC.Empty();
+}
+
 void ULevel::PostInitProperties()
 {
 	Super::PostInitProperties();
diff --git a/Engine/Source/Runtime/Engine/Private/World.cpp b/Engine/Source/Runtime/Engine/Private/World.cpp
--- a/Engine/Source/Runtime/Engine/Private/World.cpp
+++ b/Engine/Source/Runtime/Engine/Private/World.cpp
@@ -2007,8 +2007,19 @@ void UWorld::MarkObjectsPendingKill()
 	auto MarkObjectPendingKill = [](UObject* Object)
 	{
 		Object->MarkAsGarbage();
+
+		if (ULevel* Level = Cast<ULevel>(Object))
+		{
+			Level->CleanupActors();
+#if WITH_EDITORONLY_DATA
+			if (Level->GetOutermost()->MetaData)
+			{
+				Level->GetOutermost()->MetaData->ClearFlags(RF_Standalone);
+			}
+#endif
+		}
 	};
-	ForEachObjectWithOuter(this, MarkObjectPendingKill, true, RF_NoFlags, EInternalObjectFlags::PendingKill);
+	ForEachObjectWithOuter(this, MarkObjectPendingKill, true, RF_NoFlags, EInternalObjectFlags::PendingKill | EInternalObjectFlags::Garbage);
 
 	MarkAsGarbage();
 	bMarkedObjectsPendingKill = true;
@@ -3230,7 +3241,15 @@ void FLevelStreamingGCHelper::PrepareStreamedOutLevelsForGC()
 				{
 					ForEachObjectWithPackage(Package, [](UObject* PackageObject) { PackageObject->MarkAsGarbage(); return true; }, true, RF_NoFlags, EInternalObjectFlags::PendingKill);
 				}
-			}, true, RF_NoFlags, EInternalObjectFlags::PendingKill);
+			}, true, RF_NoFlags, EInternalObjectFlags::PendingKill | EInternalObjectFlags::Garbage);
+
+			Level->CleanupActors();
+#if WITH_EDITORONLY_DATA
+			if (Level->GetOutermost()->MetaData)
+			{
+				Level->GetOutermost()->MetaData->ClearFlags(RF_Standalone);
+			}
+#endif
 		}
 	}
 
@@ -4962,6 +4981,7 @@ void UWorld::CleanupWorldInternal(bool bSessionEnded, bool bCleanupResources, UW
 		if (PersistentLevel)
 		{
 			PersistentLevel->CleanupLevel(bCleanupResources);
+			PersistentLevel->CleanupActors();
 		}
 
 		if (GetNumLevels() > 1)
@@ -4969,7 +4989,9 @@ void UWorld::CleanupWorldInternal(bool bSessionEnded, bool bCleanupResources, UW
 			check(GetLevel(0) == PersistentLevel);
 			for (int32 LevelIndex = 1; LevelIndex < GetNumLevels(); ++LevelIndex)
 			{
-				GetLevel(LevelIndex)->CleanupLevel(bCleanupResources);
+				ULevel* Level = GetLevel(LevelIndex);
+				Level->CleanupLevel(bCleanupResources);
+				Level->CleanupActors();
 			}
 		}
 	}
