Backing out changelist 18496927

diff --git a/Engine/Source/Runtime/Engine/Classes/Engine/Level.h b/Engine/Source/Runtime/Engine/Classes/Engine/Level.h
--- a/Engine/Source/Runtime/Engine/Classes/Engine/Level.h
+++ b/Engine/Source/Runtime/Engine/Classes/Engine/Level.h
@@ -868,11 +868,6 @@ public:
 	 */
 	ENGINE_API void CleanupLevel(bool bCleanupResources);
 
-	/**
-	 * Empties the Actors array after everything's been cleaned up
-	 */
-	ENGINE_API void CleanupActors();
-
 	/**
 	 * Clears all components of actors associated with this level (aka in Actors array) and 
 	 * also the BSP model components.
diff --git a/Engine/Source/Runtime/Engine/Private/Level.cpp b/Engine/Source/Runtime/Engine/Private/Level.cpp
--- a/Engine/Source/Runtime/Engine/Private/Level.cpp
+++ b/Engine/Source/Runtime/Engine/Private/Level.cpp
@@ -413,12 +413,6 @@ void ULevel::CleanupLevel(bool bCleanupResources)
 	}
 }
 
-void ULevel::CleanupActors()
-{
-	Actors.Empty();
-	ActorsForGC.Empty();
-}
-
 void ULevel::PostInitProperties()
 {
 	Super::PostInitProperties();
diff --git a/Engine/Source/Runtime/Engine/Private/World.cpp b/Engine/Source/Runtime/Engine/Private/World.cpp
--- a/Engine/Source/Runtime/Engine/Private/World.cpp
+++ b/Engine/Source/Runtime/Engine/Private/World.cpp
@@ -2007,19 +2007,8 @@ void UWorld::MarkObjectsPendingKill()
 	auto MarkObjectPendingKill = [](UObject* Object)
 	{
 		Object->MarkAsGarbage();
-
-		if (ULevel* Level = Cast<ULevel>(Object))
-		{
-			Level->CleanupActors();
-#if WITH_EDITORONLY_DATA
-			if (Level->GetOutermost()->MetaData)
-			{
-				Level->GetOutermost()->MetaData->ClearFlags(RF_Standalone);
-			}
-#endif
-		}
 	};
-	ForEachObjectWithOuter(this, MarkObjectPendingKill, true, RF_NoFlags, EInternalObjectFlags::PendingKill | EInternalObjectFlags::Garbage);
+	ForEachObjectWithOuter(this, MarkObjectPendingKill, true, RF_NoFlags, EInternalObjectFlags::PendingKill);
 
 	MarkAsGarbage();
 	bMarkedObjectsPendingKill = true;
@@ -3241,15 +3230,7 @@ void FLevelStreamingGCHelper::PrepareStreamedOutLevelsForGC()
 				{
 					ForEachObjectWithPackage(Package, [](UObject* PackageObject) { PackageObject->MarkAsGarbage(); return true; }, true, RF_NoFlags, EInternalObjectFlags::PendingKill);
 				}
-			}, true, RF_NoFlags, EInternalObjectFlags::PendingKill | EInternalObjectFlags::Garbage);
-
-			Level->CleanupActors();
-#if WITH_EDITORONLY_DATA
-			if (Level->GetOutermost()->MetaData)
-			{
-				Level->GetOutermost()->MetaData->ClearFlags(RF_Standalone);
-			}
-#endif
+			}, true, RF_NoFlags, EInternalObjectFlags::PendingKill);
 		}
 	}
 
@@ -4981,7 +4962,6 @@ void UWorld::CleanupWorldInternal(bool bSessionEnded, bool bCleanupResources, UW
 		if (PersistentLevel)
 		{
 			PersistentLevel->CleanupLevel(bCleanupResources);
-			PersistentLevel->CleanupActors();
 		}
 
 		if (GetNumLevels() > 1)
@@ -4989,9 +4969,7 @@ void UWorld::CleanupWorldInternal(bool bSessionEnded, bool bCleanupResources, UW
 			check(GetLevel(0) == PersistentLevel);
 			for (int32 LevelIndex = 1; LevelIndex < GetNumLevels(); ++LevelIndex)
 			{
-				ULevel* Level = GetLevel(LevelIndex);
-				Level->CleanupLevel(bCleanupResources);
-				Level->CleanupActors();
+				GetLevel(LevelIndex)->CleanupLevel(bCleanupResources);
 			}
 		}
 	}
