CheatManager will now empty its extensions list when its outer (PlayerController) is about to be destroyed to prevent issues with Garbage Collecting references

diff --git a/Engine/Source/Runtime/Engine/Classes/GameFramework/CheatManager.h b/Engine/Source/Runtime/Engine/Classes/GameFramework/CheatManager.h
--- a/Engine/Source/Runtime/Engine/Classes/GameFramework/CheatManager.h
+++ b/Engine/Source/Runtime/Engine/Classes/GameFramework/CheatManager.h
@@ -359,6 +359,9 @@ class ENGINE_API UCheatManager : public UObject
 	UFUNCTION(exec)
 	void UpdateSafeArea();
 
+	UFUNCTION()
+	void OnPlayerEndPlayed(AActor* Player, EEndPlayReason::Type EndPlayReason);
+
 	/**
 	 * This will move the player and set their rotation to the passed in values.
 	 * This actually does the location / rotation setting.  Additionally it will set you as ghost as the level may have
diff --git a/Engine/Source/Runtime/Engine/Private/CheatManager.cpp b/Engine/Source/Runtime/Engine/Private/CheatManager.cpp
--- a/Engine/Source/Runtime/Engine/Private/CheatManager.cpp
+++ b/Engine/Source/Runtime/Engine/Private/CheatManager.cpp
@@ -76,6 +76,11 @@ UCheatManager::UCheatManager(const FObjectInitializer& ObjectInitializer)
 	bDebugCapsuleTraceComplex = false;
 }
 
+void UCheatManager::OnPlayerEndPlayed(AActor* Player, EEndPlayReason::Type EndPlayReason)
+{
+	CheatManagerExtensions.Empty();
+}
+
 bool UCheatManager::ProcessConsoleExec(const TCHAR* Cmd, FOutputDevice& Ar, UObject* Executor)
 {
 #if UE_WITH_CHEAT_MANAGER
@@ -673,6 +678,7 @@ void UCheatManager::InitCheatManager()
 {
 	ReceiveInitCheatManager(); //BP Initialization event
 	OnCheatManagerCreatedDelegate.Broadcast(this);
+	GetOuterAPlayerController()->OnEndPlay.AddDynamic(this, &UCheatManager::OnPlayerEndPlayed);
 }
 
 void UCheatManager::BeginDestroy()
