Keeping track of UCameraModifier's CameraOwner reference and nulling it when it gets destroyed

diff --git a/Engine/Source/Runtime/Engine/Classes/Camera/CameraModifier.h b/Engine/Source/Runtime/Engine/Classes/Camera/CameraModifier.h
--- a/Engine/Source/Runtime/Engine/Classes/Camera/CameraModifier.h
+++ b/Engine/Source/Runtime/Engine/Classes/Camera/CameraModifier.h
@@ -62,6 +62,9 @@ protected:
 	UPROPERTY(transient, BlueprintReadOnly, Category = CameraModifier)
 	float Alpha;
 
+	UFUNCTION()
+	void OnCameraOwnerDestroyed(AActor* InOwner);
+
 protected:
 	/** @return Returns the ideal blend alpha for this modifier. Interpolation will seek this value. */
 	virtual float GetTargetAlpha();
diff --git a/Engine/Source/Runtime/Engine/Private/Camera/CameraModifier.cpp b/Engine/Source/Runtime/Engine/Private/Camera/CameraModifier.cpp
--- a/Engine/Source/Runtime/Engine/Private/Camera/CameraModifier.cpp
+++ b/Engine/Source/Runtime/Engine/Private/Camera/CameraModifier.cpp
@@ -100,7 +100,23 @@ AActor* UCameraModifier::GetViewTarget() const
 
 void UCameraModifier::AddedToCamera( APlayerCameraManager* Camera ) 
 {
+	if (CameraOwner)
+	{
+		CameraOwner->OnDestroyed.RemoveDynamic(this, &UCameraModifier::OnCameraOwnerDestroyed);
+	}
 	CameraOwner = Camera;
+	if (CameraOwner)
+	{
+		CameraOwner->OnDestroyed.AddDynamic(this, &UCameraModifier::OnCameraOwnerDestroyed);
+	}
+}
+
+void UCameraModifier::OnCameraOwnerDestroyed(AActor* InOwner)
+{
+	if (InOwner == CameraOwner)
+	{
+		CameraOwner = nullptr;
+	}
 }
 
 UWorld* UCameraModifier::GetWorld() const
