Keeping track of PlayerState's Pawn reference to make sure it doesn't prevent levels from being GC'd

diff --git a/Engine/Source/Runtime/Engine/Classes/GameFramework/PlayerState.h b/Engine/Source/Runtime/Engine/Classes/GameFramework/PlayerState.h
--- a/Engine/Source/Runtime/Engine/Classes/GameFramework/PlayerState.h
+++ b/Engine/Source/Runtime/Engine/Classes/GameFramework/PlayerState.h
@@ -176,6 +176,11 @@ private:
 	UPROPERTY(BlueprintReadOnly, Category=PlayerState, meta=(AllowPrivateAccess="true"))
 	TObjectPtr<APawn> PawnPrivate;
 
+	void SetPawnPrivate(APawn* InPawn);
+
+	UFUNCTION()
+	void OnPawnPrivateDestroyed(AActor* InActor);
+
 	/**
 	 * Stores the last 4 seconds worth of ping data (one second per 'bucket').
 	 * It is stored in this manner, to allow calculating a moving average,
@@ -475,7 +480,7 @@ private:
 	FSetPlayerStatePawn(APlayerState* PlayerState, APawn* Pawn)
 	{
 		APawn* OldPawn = PlayerState->PawnPrivate;
-		PlayerState->PawnPrivate = Pawn;
+		PlayerState->SetPawnPrivate(Pawn);
 		PlayerState->OnPawnSet.Broadcast(PlayerState, PlayerState->PawnPrivate, OldPawn); 
 	}
-};
+};
\ No newline at end of file
diff --git a/Engine/Source/Runtime/Engine/Private/PlayerState.cpp b/Engine/Source/Runtime/Engine/Private/PlayerState.cpp
--- a/Engine/Source/Runtime/Engine/Private/PlayerState.cpp
+++ b/Engine/Source/Runtime/Engine/Private/PlayerState.cpp
@@ -568,4 +568,28 @@ void APlayerState::OnSetUniqueId()
 
 }
 
-PRAGMA_ENABLE_DEPRECATION_WARNINGS
+void APlayerState::SetPawnPrivate(APawn* InPawn)
+{
+	if (InPawn != PawnPrivate)
+	{
+		if (PawnPrivate)
+		{
+			PawnPrivate->OnDestroyed.RemoveDynamic(this, &APlayerState::OnPawnPrivateDestroyed);
+		}
+		PawnPrivate = InPawn;
+		if (PawnPrivate)
+		{
+			PawnPrivate->OnDestroyed.AddDynamic(this, &APlayerState::OnPawnPrivateDestroyed);
+		}
+	}
+}
+
+void APlayerState::OnPawnPrivateDestroyed(AActor* InActor)
+{
+	if (InActor == PawnPrivate)
+	{
+		PawnPrivate = nullptr;
+	}
+}
+
+PRAGMA_ENABLE_DEPRECATION_WARNINGS
\ No newline at end of file
