Making InputComponent clear its reference to PlayerInput to not prevent levels from being GC'd

diff --git a/Engine/Source/Runtime/Engine/Classes/Components/InputComponent.h b/Engine/Source/Runtime/Engine/Classes/Components/InputComponent.h
--- a/Engine/Source/Runtime/Engine/Classes/Components/InputComponent.h
+++ b/Engine/Source/Runtime/Engine/Classes/Components/InputComponent.h
@@ -754,6 +754,9 @@ private:
 	UPROPERTY(Transient, DuplicateTransient)
 	TArray<FCachedKeyToActionInfo> CachedKeyToActionInfo;
 
+	UFUNCTION()
+	void OnInputOwnerEndPlayed(AActor* InOwner, EEndPlayReason::Type EndPlayReason);
+
 public:
 	/** The priority of this input component when pushed in to the stack. */
 	int32 Priority;
diff --git a/Engine/Source/Runtime/Engine/Private/Components/InputComponent.cpp b/Engine/Source/Runtime/Engine/Private/Components/InputComponent.cpp
--- a/Engine/Source/Runtime/Engine/Private/Components/InputComponent.cpp
+++ b/Engine/Source/Runtime/Engine/Private/Components/InputComponent.cpp
@@ -57,6 +57,10 @@ void UInputComponent::ConditionalBuildKeyMap(UPlayerInput* PlayerInput)
 		CachedKeyToActionInfo.AddDefaulted();
 		CachedInfoToPopulate = &CachedKeyToActionInfo.Last();
 		CachedInfoToPopulate->PlayerInput = PlayerInput;
+		if (AActor* OwnerActor = PlayerInput->GetTypedOuter<AActor>())
+		{
+			OwnerActor->OnEndPlay.AddUniqueDynamic(this, &UInputComponent::OnInputOwnerEndPlayed);
+		}
 	}
 
 	// Reset the map and AnyKey array
@@ -86,6 +90,18 @@ void UInputComponent::ConditionalBuildKeyMap(UPlayerInput* PlayerInput)
 	CachedInfoToPopulate->KeyMapBuiltForIndex = PlayerInput->GetKeyMapBuildIndex();
 }
 
+void UInputComponent::OnInputOwnerEndPlayed(AActor* InOwner, EEndPlayReason::Type EndPlayReason)
+{
+	for (int32 Index = CachedKeyToActionInfo.Num() - 1; Index >= 0; --Index)
+	{
+		FCachedKeyToActionInfo& CachedInfo = CachedKeyToActionInfo[Index];
+		if (CachedInfo.PlayerInput->GetTypedOuter<AActor>() == InOwner)
+		{
+			CachedKeyToActionInfo.RemoveAtSwap(Index, 1, false);
+		}
+	}
+}
+
 void UInputComponent::GetActionsBoundToKey(UPlayerInput* PlayerInput, const FKey Key, TArray<TSharedPtr<FInputActionBinding>>& Actions) const
 {
 	for (const FCachedKeyToActionInfo& CachedInfo : CachedKeyToActionInfo)
