AudioComponent will now clear references to other objects when its EndPlay is called to make sure it does not prevent levels from being GC'd

diff --git a/Engine/Source/Runtime/Engine/Classes/Components/AudioComponent.h b/Engine/Source/Runtime/Engine/Classes/Components/AudioComponent.h
--- a/Engine/Source/Runtime/Engine/Classes/Components/AudioComponent.h
+++ b/Engine/Source/Runtime/Engine/Classes/Components/AudioComponent.h
@@ -759,6 +759,7 @@ public:
 	virtual void OnUnregister() override;
 	virtual const UObject* AdditionalStatObject() const override;
 	virtual bool IsReadyForOwnerToAutoDestroy() const override;
+	virtual void EndPlay(const EEndPlayReason::Type EndPlayReason) override;
 	//~ End ActorComponent Interface.
 
 	void AdjustVolumeInternal(float AdjustVolumeDuration, float AdjustVolumeLevel, bool bIsFadeOut, EAudioFaderCurve FadeCurve);
diff --git a/Engine/Source/Runtime/Engine/Private/Components/AudioComponent.cpp b/Engine/Source/Runtime/Engine/Private/Components/AudioComponent.cpp
--- a/Engine/Source/Runtime/Engine/Private/Components/AudioComponent.cpp
+++ b/Engine/Source/Runtime/Engine/Private/Components/AudioComponent.cpp
@@ -274,6 +274,13 @@ void UAudioComponent::OnUnregister()
 	{
 		Stop();
 	}
+	Shutdown();
+}
+
+void UAudioComponent::EndPlay(const EEndPlayReason::Type EndPlayReason)
+{
+	Super::EndPlay(EndPlayReason);
+	Shutdown();
 }
 
 const UObject* UAudioComponent::AdditionalStatObject() const
