Miscellaneous World Partition properties display/naming/grouping fixes
* World Partition properties are now grouped together in the World Settings Detailed View
* Hide actors world partition related properties when their owning world isn't a WP'd world (previously only disabled)
* Hide AWorldSettings::bEnableWorldOriginRebasing unless using World Composition
* Hide AWorldSettings::bGenerateSingleClusterForLevel/HLODSetupAsset/OverrideBaseMaterial/HierarchicalLODSetup/HLODBakingTransform when in WP (irrelevant)

diff --git a/Engine/Source/Editor/DetailCustomizations/Private/ActorDetails.cpp b/Engine/Source/Editor/DetailCustomizations/Private/ActorDetails.cpp
--- a/Engine/Source/Editor/DetailCustomizations/Private/ActorDetails.cpp
+++ b/Engine/Source/Editor/DetailCustomizations/Private/ActorDetails.cpp
@@ -56,6 +56,7 @@
 #include "ScopedTransaction.h"
 #include "WorldPartition/WorldPartitionSubsystem.h"
 #include "Settings/EditorExperimentalSettings.h"
+#include "Algo/AnyOf.h"
 
 #include "Kismet2/BlueprintEditorUtils.h"
 #include "K2Node_AddDelegate.h"
@@ -144,6 +145,21 @@ void FActorDetails::CustomizeDetails( IDetailLayoutBuilder& DetailLayout )
 			AddActorCategory(DetailLayout, ActorsPerLevelCount);
 		}
 
+		// Hide World Partition specific properties in non WP levels
+		const bool bShouldDisplayWorldPartitionProperties = Algo::AnyOf(SelectedActors, [](const TWeakObjectPtr<AActor> Actor)
+		{
+			UWorld* World = Actor.IsValid() ? Actor->GetTypedOuter<UWorld>() : nullptr;
+			return World != nullptr && UWorld::HasSubsystem<UWorldPartitionSubsystem>(World);
+		});
+
+		if (!bShouldDisplayWorldPartitionProperties)
+		{
+			DetailLayout.HideProperty(DetailLayout.GetProperty(AActor::GetRuntimeGridPropertyName(), AActor::StaticClass()));
+			DetailLayout.HideProperty(DetailLayout.GetProperty(AActor::GetIsSpatiallyLoadedPropertyName(), AActor::StaticClass()));
+			DetailLayout.HideProperty(DetailLayout.GetProperty(AActor::GetDataLayersPropertyName(), AActor::StaticClass()));
+			DetailLayout.HideProperty(DetailLayout.GetProperty(AActor::GetHLODLayerPropertyName(), AActor::StaticClass()));
+		}
+
 		OnExtendActorDetails.Broadcast(DetailLayout, FGetSelectedActors::CreateSP(this, &FActorDetails::GetSelectedActors));
 	}
 
diff --git a/Engine/Source/Editor/DetailCustomizations/Private/DetailCustomizations.cpp b/Engine/Source/Editor/DetailCustomizations/Private/DetailCustomizations.cpp
--- a/Engine/Source/Editor/DetailCustomizations/Private/DetailCustomizations.cpp
+++ b/Engine/Source/Editor/DetailCustomizations/Private/DetailCustomizations.cpp
@@ -744,7 +744,7 @@ void FDetailCustomizationsModule::RegisterSectionMappings()
 	{
 		{
 			TSharedRef<FPropertySection> Section = PropertyModule.FindOrCreateSection("WorldPartition", "Streaming", LOCTEXT("Streaming", "Streaming"));
-			Section->AddCategory("RuntimeHash");
+			Section->AddCategory("WorldPartition");
 		}
 	}
 
diff --git a/Engine/Source/Editor/DetailCustomizations/Private/WorldSettingsDetails.cpp b/Engine/Source/Editor/DetailCustomizations/Private/WorldSettingsDetails.cpp
--- a/Engine/Source/Editor/DetailCustomizations/Private/WorldSettingsDetails.cpp
+++ b/Engine/Source/Editor/DetailCustomizations/Private/WorldSettingsDetails.cpp
@@ -94,14 +94,23 @@ void FWorldSettingsDetails::AddWorldCustomization(IDetailLayoutBuilder& DetailBu
 		}
 	}
 
+	// Hide some of the WorldPartition properties found in AActor
+	TArray<TSharedRef<IPropertyHandle>> PropertiesToHide;
+	PropertiesToHide.Add(DetailBuilder.GetProperty(AActor::GetRuntimeGridPropertyName(), AActor::StaticClass()));
+	PropertiesToHide.Add(DetailBuilder.GetProperty(AActor::GetIsSpatiallyLoadedPropertyName(), AActor::StaticClass()));
+	for (TSharedRef<IPropertyHandle> Property : PropertiesToHide)
+	{
+		DetailBuilder.HideProperty(Property);
+	}
+
 	if (CustomizedLevel)
 	{
 		const bool bIsPartitionedWorld = UWorld::HasSubsystem<UWorldPartitionSubsystem>(CustomizedLevel->GetWorld());
 
-		IDetailCategoryBuilder& WorldCategory = DetailBuilder.EditCategory("World");
+		IDetailCategoryBuilder& WorldPartitionCategory = DetailBuilder.EditCategory("WorldPartition");
 		if (GetDefault<UEditorExperimentalSettings>()->bEnableOneFilePerActorSupport)
 		{
-			WorldCategory.AddCustomRow(LOCTEXT("LevelUseExternalActorsRow", "LevelUseExternalActors"), true)
+			WorldPartitionCategory.AddCustomRow(LOCTEXT("LevelUseExternalActorsRow", "LevelUseExternalActors"), true)
 				.NameContent()
 				[
 					SNew(STextBlock)
@@ -121,7 +130,7 @@ void FWorldSettingsDetails::AddWorldCustomization(IDetailLayoutBuilder& DetailBu
 
 		if (bIsPartitionedWorld)
 		{
-			WorldCategory.AddCustomRow(LOCTEXT("DefaultWorldPartitionSettingsRow", "DefaultWorldPartitionSettings"), true)
+			WorldPartitionCategory.AddCustomRow(LOCTEXT("DefaultWorldPartitionSettingsRow", "DefaultWorldPartitionSettings"), true)
 				.NameContent()
 				[
 					SNew(STextBlock)
@@ -168,7 +177,7 @@ void FWorldSettingsDetails::AddWorldCustomization(IDetailLayoutBuilder& DetailBu
 					]
 				];
 
-			WorldCategory.AddCustomRow(LOCTEXT("WorldPartitionEditorCellSizeRow", "WorldPartitionEditorCellSize"), true)
+			WorldPartitionCategory.AddCustomRow(LOCTEXT("WorldPartitionEditorCellSizeRow", "WorldPartitionEditorCellSize"), true)
 				.NameContent()
 				[
 					SNew(STextBlock)
diff --git a/Engine/Source/Runtime/Engine/Classes/GameFramework/Actor.h b/Engine/Source/Runtime/Engine/Classes/GameFramework/Actor.h
--- a/Engine/Source/Runtime/Engine/Classes/GameFramework/Actor.h
+++ b/Engine/Source/Runtime/Engine/Classes/GameFramework/Actor.h
@@ -764,7 +764,11 @@ public:
 
 	/** Specify in which HLOD layer this actor should be included. */
 	void SetHLODLayer(class UHLODLayer* InHLODLayer);
+
+	/** Gets the property name for HLODLayer. */
+	static const FName GetHLODLayerPropertyName() { return GET_MEMBER_NAME_CHECKED(AActor, HLODLayer); }
 #endif
+
 	/** Specify a RayTracingGroupId for this actors. Components with invalid RayTracingGroupId will inherit the actors. */
 	UFUNCTION(BlueprintCallable, Category = RayTracing)
 	void SetRayTracingGroupId(int32 InRaytracingGroupId);
@@ -860,6 +864,9 @@ public:
 	/** Sets this actor's current target runtime grid. */
 	void SetRuntimeGrid(FName InRuntimeGrid) { RuntimeGrid = InRuntimeGrid; }
 
+	/** Gets the property name for RuntimeGrid. */
+	static const FName GetRuntimeGridPropertyName()	{ return GET_MEMBER_NAME_CHECKED(AActor, RuntimeGrid); }
+
 	/** Returns this actor's current Guid. Actor Guids are only available in development builds. */
 	inline const FGuid& GetActorGuid() const { return ActorGuid; }
 
@@ -1146,6 +1153,7 @@ public:
 	bool IsPropertyChangedAffectingDataLayers(FPropertyChangedEvent& PropertyChangedEvent) const;
 	bool IsValidForDataLayer() const;
 	void FixupDataLayers(bool bRevertChangesOnLockedDataLayer = false);
+	static const FName GetDataLayersPropertyName() { return GET_MEMBER_NAME_CHECKED(AActor, DataLayers); }
 #endif
 
 	//~=============================================================================
@@ -2233,6 +2241,9 @@ public:
 	/** Returns true if this actor allows changing the spatially loaded flag.  */
 	virtual bool CanChangeIsSpatiallyLoadedFlag() const { return true; }
 
+	/** Gets the property name for bIsSpatiallyLoaded. */
+	static const FName GetIsSpatiallyLoadedPropertyName() { return GET_MEMBER_NAME_CHECKED(AActor, bIsSpatiallyLoaded); }
+
 	/**
 	 * Returns whether or not this actor was explicitly hidden in the editor for the duration of the current editor session
 	 * @param bIncludeParent - Whether to recurse up child actor hierarchy or not
diff --git a/Engine/Source/Runtime/Engine/Classes/GameFramework/WorldSettings.h b/Engine/Source/Runtime/Engine/Classes/GameFramework/WorldSettings.h
--- a/Engine/Source/Runtime/Engine/Classes/GameFramework/WorldSettings.h
+++ b/Engine/Source/Runtime/Engine/Classes/GameFramework/WorldSettings.h
@@ -407,7 +407,7 @@ struct FBroadphaseSettings
 /**
  * Actor containing all script accessible world properties.
  */
-UCLASS(config=game, hidecategories=(Actor, Advanced, Display, Events, Object, Attachment, Info, Input, Blueprint, Layers, Tags, Replication), showcategories=(Rendering, "Input|MouseInput", "Input|TouchInput"), notplaceable)
+UCLASS(config=game, hidecategories=(Actor, Advanced, Display, Events, Object, Attachment, Info, Input, Blueprint, Layers, Tags, Replication), showcategories=(Rendering, WorldPartition, "Input|MouseInput", "Input|TouchInput"), notplaceable)
 class ENGINE_API AWorldSettings : public AInfo, public IInterface_AssetUserData
 {
 	GENERATED_UCLASS_BODY()
@@ -478,7 +478,7 @@ public:
 	uint8 bUseClientSideLevelStreamingVolumes:1;
 
 	/** World origin will shift to a camera position when camera goes far away from current origin */
-	UPROPERTY(EditAnywhere, BlueprintReadOnly, Category=World, AdvancedDisplay, meta=(editcondition = "bEnableWorldComposition"))
+	UPROPERTY(EditAnywhere, BlueprintReadOnly, Category=World, AdvancedDisplay, meta=(EditConditionHides, EditCondition = "bEnableWorldComposition"))
 	uint8 bEnableWorldOriginRebasing:1;
 		
 	/** if set to true, when we call GetGravityZ we assume WorldGravityZ has already been initialized and skip the lookup of DefaultGravityZ and GlobalGravityZ */
@@ -517,7 +517,7 @@ public:
 	uint8 bOverrideDefaultBroadphaseSettings:1;
 
 	/** If set to true, all eligible actors in this level will be added to a single cluster representing the entire level (used for small sublevels) */
-	UPROPERTY(EditAnywhere, config, Category = HLODSystem, AdvancedDisplay, meta=(DisplayAfter="HierarchicalLODSetup"))
+	UPROPERTY(EditAnywhere, config, Category = HLODSystem, AdvancedDisplay, meta=(EditConditionHides, EditCondition="WorldPartition == nullptr", DisplayAfter="HierarchicalLODSetup"))
 	uint8 bGenerateSingleClusterForLevel : 1;
 
 #if WITH_EDITORONLY_DATA
@@ -666,16 +666,16 @@ public:
 
 #if WITH_EDITORONLY_DATA
 	/** If set overrides the level settings and global project settings */
-	UPROPERTY(EditAnywhere, config, Category = HLODSystem)
+	UPROPERTY(EditAnywhere, config, Category = HLODSystem, meta=(EditConditionHides, EditCondition = "WorldPartition == nullptr"))
 	TSoftClassPtr<class UHierarchicalLODSetup> HLODSetupAsset;
 
 	/** If set overrides the project-wide base material used for Proxy Materials */
-	UPROPERTY(EditAnywhere, config, Category = HLODSystem, meta=(editcondition = "HLODSetupAsset == nullptr"))
+	UPROPERTY(EditAnywhere, config, Category = HLODSystem, meta=(EditConditionHides, EditCondition = "WorldPartition == nullptr && HLODSetupAsset == nullptr"))
 	TSoftObjectPtr<class UMaterialInterface> OverrideBaseMaterial;
 
 protected:
 	/** Hierarchical LOD Setup */
-	UPROPERTY(EditAnywhere, Category = HLODSystem, config, meta=(editcondition = "HLODSetupAsset == nullptr"))
+	UPROPERTY(EditAnywhere, Category = HLODSystem, config, meta=(EditConditionHides, EditCondition = "WorldPartition == nullptr && HLODSetupAsset == nullptr"))
 	TArray<struct FHierarchicalSimplification> HierarchicalLODSetup;
 
 public:
@@ -683,7 +683,7 @@ public:
 	int32 NumHLODLevels;
 
 	/** Specify the transform to apply to the source meshes when building HLODs. */
-	UPROPERTY(EditAnywhere, config, Category = HLODSystem, AdvancedDisplay, meta=(DisplayName = "HLOD Baking Transform"))
+	UPROPERTY(EditAnywhere, config, Category = HLODSystem, AdvancedDisplay, meta=(EditConditionHides, EditCondition = "WorldPartition == nullptr", DisplayName = "HLOD Baking Transform"))
 	FTransform HLODBakingTransform;
 
 	/************************************/
diff --git a/Engine/Source/Runtime/Engine/Public/WorldPartition/WorldPartition.h b/Engine/Source/Runtime/Engine/Public/WorldPartition/WorldPartition.h
--- a/Engine/Source/Runtime/Engine/Public/WorldPartition/WorldPartition.h
+++ b/Engine/Source/Runtime/Engine/Public/WorldPartition/WorldPartition.h
@@ -181,11 +181,11 @@ public:
 	IWorldPartitionEditor* WorldPartitionEditor;
 
 	/** Class of WorldPartitionStreamingPolicy to be used to manage world partition streaming. */
-	UPROPERTY(EditAnywhere, Category=RuntimeHash, NoClear, meta = (NoResetToDefault))
+	UPROPERTY(EditAnywhere, Category=WorldPartition, NoClear, meta = (NoResetToDefault))
 	TSubclassOf<UWorldPartitionStreamingPolicy> WorldPartitionStreamingPolicyClass;
 #endif
 
-	UPROPERTY(EditAnywhere, Category=RuntimeHash, NoClear, meta = (NoResetToDefault), Instanced)
+	UPROPERTY(EditAnywhere, Category=WorldPartition, NoClear, meta = (NoResetToDefault), Instanced)
 	TObjectPtr<UWorldPartitionRuntimeHash> RuntimeHash;
 
 #if WITH_EDITOR
@@ -196,7 +196,7 @@ public:
 
 #if WITH_EDITORONLY_DATA
 	// Default HLOD layer
-	UPROPERTY(EditAnywhere, Category=HLOD)
+	UPROPERTY(EditAnywhere, Category=WorldPartition, meta = (DisplayName = "Default HLOD Layer"))
 	TObjectPtr<class UHLODLayer> DefaultHLODLayer;
 
 	TArray<FWorldPartitionReference> LoadedSubobjects;
