MeshDescription support for EditorDomain: Do not resave the bulkdata if the MeshDescription bytes have not changed. Keeping the bulkdata unmodified allows EditorDomain to refer to it by reference.

diff --git a/Engine/Source/Runtime/Engine/Private/Rendering/SkeletalMeshLODImporterData.cpp b/Engine/Source/Runtime/Engine/Private/Rendering/SkeletalMeshLODImporterData.cpp
--- a/Engine/Source/Runtime/Engine/Private/Rendering/SkeletalMeshLODImporterData.cpp
+++ b/Engine/Source/Runtime/Engine/Private/Rendering/SkeletalMeshLODImporterData.cpp
@@ -13,6 +13,7 @@
 #include "Rendering/SkeletalMeshModel.h"
 #include "Serialization/BulkDataReader.h"
 #include "Serialization/BulkDataWriter.h"
+#include "Serialization/LargeMemoryWriter.h"
 #include "SkeletalMeshAttributes.h"
 #include "SkeletalMeshOperations.h"
 
@@ -776,9 +777,7 @@ void FRawSkeletalMeshBulkData::Serialize(FArchive& Ar, UObject* Owner)
 		if (Ar.IsSaving() && bUseSerializeLoadingCustomVersion == true)
 		{
 			//We need to update the FReductionSkeletalMeshData serialize version to the latest in case we save the Parent bulkdata
-			FSkeletalMeshImportData MeshImportData;
-			LoadRawMesh(MeshImportData);
-			SaveRawMesh(MeshImportData);
+			UpdateRawMeshFormat();
 		}
 	}
 
@@ -884,6 +883,91 @@ void FRawSkeletalMeshBulkData::LoadRawMesh(FSkeletalMeshImportData& OutMesh)
 	}
 }
 
+
+void FRawSkeletalMeshBulkData::UpdateRawMeshFormat()
+{
+	if (BulkData.GetElementCount() == 0)
+	{
+		return;
+	}
+
+	// An exclusive lock is required so we can safely load the raw data from multiple threads
+	FWriteScopeLock ScopeLock(BulkDataLock.Get());
+#if WITH_EDITOR
+	// This allows any thread to be able to deserialize from the RawMesh directly
+	// from disk so we can unload bulk data from memory.
+	bool bHasBeenLoadedFromFileReader = false;
+	if (BulkData.IsAsyncLoadingComplete() && !BulkData.IsBulkDataLoaded())
+	{
+		// This can't be called in -game mode because we're not allowed to load bulk data outside of EDL.
+		bHasBeenLoadedFromFileReader = BulkData.LoadBulkDataWithFileReader();
+	}
+#endif
+
+	bool bModified = false;
+	uint64 NumBytes = 0;
+	FLargeMemoryWriter NewBytes(BulkData.GetBulkDataSize(), true /* bIsPersistent */);
+	FSkeletalMeshImportData MeshImportData;
+	// This is in a scope because the BulkData needs to be unlocked for UnloadBulkData.
+	{
+		FMemoryView OldBytes(BulkData.Lock(LOCK_READ_ONLY), BulkData.GetBulkDataSize());
+		ON_SCOPE_EXIT
+		{
+			BulkData.Unlock();
+		};
+		FMemoryReaderView Reader(OldBytes, true /* bIsPersistent */);
+
+		// Propagate the custom version information from the package to the bulk data, so that the MeshDescription
+		// is serialized with the same versioning.
+		Reader.SetCustomVersions(SerializeLoadingCustomVersionContainer);
+		Reader << MeshImportData;
+
+		NewBytes << MeshImportData;
+		NumBytes = static_cast<uint64>(NewBytes.TotalSize());
+		bModified = NumBytes != OldBytes.GetSize() ||
+			0 != FMemory::Memcmp(OldBytes.GetData(), NewBytes.GetData(), NumBytes);
+	}
+
+#if WITH_EDITOR
+	// Throw away the bulk data allocation only in the case we can safely reload it from disk
+	// and if BulkData.LoadBulkDataWithFileReader() is allowed to work from any thread.
+	// This saves a significant amount of memory during map loading of Nanite Meshes.
+	if (bHasBeenLoadedFromFileReader)
+	{
+		verify(BulkData.UnloadBulkData());
+	}
+#endif
+
+	if (bModified)
+	{
+		//Clear the bulk data before writing it
+		BulkData.RemoveBulkData();
+		BulkData.StoreCompressedOnDisk(NAME_Zlib);
+		{
+			BulkData.Lock(LOCK_READ_WRITE);
+			ON_SCOPE_EXIT
+			{
+				BulkData.Unlock();
+			};
+			void* Buffer = BulkData.Realloc(NumBytes);
+			FMemory::Memcpy(Buffer, NewBytes.GetData(), NumBytes);
+		}
+		// Preserve CustomVersions at save time so we can reuse the same ones when reloading direct from memory
+		SerializeLoadingCustomVersionContainer = NewBytes.GetCustomVersions();
+
+		//Create the guid from the content, this allow to use the data into the ddc key
+		FSHA1 Sha;
+		if (NumBytes > 0)
+		{
+			Sha.Update(NewBytes.GetData(), NumBytes);
+		}
+		Sha.Final();
+		uint32 Hash[5];
+		Sha.GetHash((uint8*)Hash);
+		Guid = FGuid(Hash[0] ^ Hash[4], Hash[1], Hash[2], Hash[3]);
+	}
+}
+
 FString FRawSkeletalMeshBulkData::GetIdString() const
 {
 	FString GuidString = Guid.ToString();
diff --git a/Engine/Source/Runtime/Engine/Public/Rendering/SkeletalMeshLODImporterData.h b/Engine/Source/Runtime/Engine/Public/Rendering/SkeletalMeshLODImporterData.h
--- a/Engine/Source/Runtime/Engine/Public/Rendering/SkeletalMeshLODImporterData.h
+++ b/Engine/Source/Runtime/Engine/Public/Rendering/SkeletalMeshLODImporterData.h
@@ -596,6 +596,9 @@ public:
 		return GeoImportVersion >= ESkeletalMeshGeoImportVersions::SkeletalMeshBuildRefactor &&
 			SkinningImportVersion >= ESkeletalMeshSkinningImportVersions::SkeletalMeshBuildRefactor;
 	}
+
+private:
+	ENGINE_API void UpdateRawMeshFormat();
 };
 
 namespace FWedgePositionHelper
diff --git a/Engine/Source/Runtime/MeshDescription/Private/MeshDescription.cpp b/Engine/Source/Runtime/MeshDescription/Private/MeshDescription.cpp
--- a/Engine/Source/Runtime/MeshDescription/Private/MeshDescription.cpp
+++ b/Engine/Source/Runtime/MeshDescription/Private/MeshDescription.cpp
@@ -6,6 +6,8 @@
 #include "MeshAttributes.h"
 #include "Misc/SecureHash.h"
 #include "Serialization/BulkData.h"
+#include "Serialization/LargeMemoryWriter.h"
+#include "Serialization/MemoryReader.h"
 #include "Serialization/MemoryWriter.h"
 #include "Serialization/NameAsStringProxyArchive.h"
 #include "Serialization/VirtualizedBulkDataReader.h"
@@ -2006,20 +2008,7 @@ void FMeshDescriptionBulkData::Serialize( FArchive& Ar, UObject* Owner )
 		{
 			// If the bulk data hasn't been updated since this was loaded, there's a possibility that it has old versioning.
 			// Explicitly load and resave the FMeshDescription so that its version is in sync with the FMeshDescriptionBulkData.
-			if( !bBulkDataUpdated )
-			{
-				const FGuid OriginalGuid = Guid;
-
-				FMeshDescription MeshDescription;
-				LoadMeshDescription( MeshDescription );
-				SaveMeshDescription( MeshDescription );
-
-				// Maintain the original guid; this is a format change only. GetIdString will change because it adds the CustomVersions to the key.
-				if (!bGuidIsHash)
-				{
-					Guid = OriginalGuid;
-				}
-			}
+			UpdateMeshDescriptionFormat();
 		}
 	}
 
@@ -2138,6 +2127,64 @@ void FMeshDescriptionBulkData::LoadMeshDescription( FMeshDescription& MeshDescri
 	}
 }
 
+void FMeshDescriptionBulkData::UpdateMeshDescriptionFormat()
+{
+	TRACE_CPUPROFILER_EVENT_SCOPE(FMeshDescriptionBulkData::UpdateMeshDescriptionFormat);
+	if (bBulkDataUpdated)
+	{
+		return;
+	}
+	if (BulkData.GetPayloadSize() == 0)
+	{
+		return;
+	}
+
+#if WITH_EDITOR
+	// TODO: Remove this once VirtualizedBulkData can be shown to be thread safe
+	// A lock is required so we can clone the mesh description from multiple threads
+	FRWScopeLock ScopeLock(BulkDataLock, SLT_Write);
+#endif //WITH_EDITOR
+
+	FSharedBuffer OldBytes = BulkData.GetPayload().Get();
+	FMemoryReaderView Reader(OldBytes.GetView(), true /* bIsPersistent */);
+	BulkData.UnloadData();
+	// Propagate the custom version information from the package to the bulk data, so that the MeshDescription
+	// is serialized with the same versioning.
+	Reader.SetCustomVersions(CustomVersions);
+	FMeshDescription MeshDescription;
+	MeshDescription.Empty();
+	Reader << MeshDescription;
+
+	FLargeMemoryWriter NewBytes(OldBytes.GetSize(), true /* bIsPersistent */);
+	if (!MeshDescription.IsEmpty())
+	{
+		NewBytes << MeshDescription;
+	}
+	uint64 NumBytes = static_cast<uint64>(NewBytes.TotalSize());
+
+	if (NumBytes != OldBytes.GetSize() ||
+		0 != FMemory::Memcmp(OldBytes.GetData(), NewBytes.GetData(), NumBytes))
+	{
+		// The MeshDescription format has changed, so save it to this's BulkData
+		BulkData.UpdatePayload(FSharedBuffer::TakeOwnership(NewBytes.ReleaseOwnership(), NumBytes, FMemory::Free));
+		// Preserve CustomVersions at save time so we can reuse the same ones when reloading direct from memory
+		CustomVersions = NewBytes.GetCustomVersions();
+		if (bGuidIsHash)
+		{
+			UseHashAsGuid();
+		}
+		else
+		{
+			// Maintain the original guid; this is a format change only. GetIdString will change because it adds the CustomVersions to the key.
+			FPlatformMisc::CreateGuid(Guid);
+		}
+	}
+
+	// Mark the MeshDescriptionBulkData as having been updated.
+	// This means we know that its version is up-to-date.
+	bBulkDataUpdated = true;
+}
+
 void FMeshDescriptionBulkData::Empty()
 {
 	BulkData.Reset();
diff --git a/Engine/Source/Runtime/MeshDescription/Public/MeshDescription.h b/Engine/Source/Runtime/MeshDescription/Public/MeshDescription.h
--- a/Engine/Source/Runtime/MeshDescription/Public/MeshDescription.h
+++ b/Engine/Source/Runtime/MeshDescription/Public/MeshDescription.h
@@ -1333,6 +1333,8 @@ public:
 	static TConstArrayView<FGuid> GetMeshDescriptionCustomVersions();
 
 private:
+	void UpdateMeshDescriptionFormat();
+
 #if WITH_EDITOR
 	/** Protects simultaneous access to BulkData */
 	FRWLock       BulkDataLock;
