Migrate behavior for SetTriggerParameter back to just executing on command (and not caching in InstanceParameters). Rename to "ExecuteTriggerParameter" & add tooltip explaining the difference and when to use "SetBooleanParameter" for triggers instead
#rnx

diff --git a/Engine/Source/Runtime/AudioExtensions/Public/AudioParameterControllerInterface.h b/Engine/Source/Runtime/AudioExtensions/Public/AudioParameterControllerInterface.h
--- a/Engine/Source/Runtime/AudioExtensions/Public/AudioParameterControllerInterface.h
+++ b/Engine/Source/Runtime/AudioExtensions/Public/AudioParameterControllerInterface.h
@@ -18,7 +18,7 @@ class AUDIOEXTENSIONS_API UAudioParameterControllerInterface : public UInterface
 	GENERATED_UINTERFACE_BODY()
 };
 
-// Base interface for any object implementing a
+// Base interface for any object implementing parameter control for a given sound instance controller.
 class AUDIOEXTENSIONS_API IAudioParameterControllerInterface : public IInterface
 {
 	GENERATED_IINTERFACE_BODY()
@@ -28,8 +28,11 @@ public:
 	UFUNCTION(BlueprintCallable, Category = "Audio|Parameter")
 	virtual void ResetParameters() = 0;
 
-	// Triggers a named trigger
-	UFUNCTION(BlueprintCallable, meta = (DisplayName = "Set Trigger Parameter"), Category = "Audio|Parameter")
+	// Executes a named trigger.  Does *not* cache trigger value, so only executes if the sound
+	// is already playing.  If the intent is for the trigger to execute immediately (if playing)
+	// and be called on initialization for all future instances, call 'SetBoolParameter' with the
+	// intended initial trigger behavior (true if trigger desired on initialization, false if not).
+	UFUNCTION(BlueprintCallable, meta = (DisplayName = "Execute Trigger Parameter"), Category = "Audio|Parameter")
 	virtual void SetTriggerParameter(FName InName) = 0;
 
 	// Sets a named Boolean
diff --git a/Engine/Source/Runtime/Engine/Private/Audio/SoundParameterControllerInterface.cpp b/Engine/Source/Runtime/Engine/Private/Audio/SoundParameterControllerInterface.cpp
--- a/Engine/Source/Runtime/Engine/Private/Audio/SoundParameterControllerInterface.cpp
+++ b/Engine/Source/Runtime/Engine/Private/Audio/SoundParameterControllerInterface.cpp
@@ -24,7 +24,7 @@ void ISoundParameterControllerInterface::ResetParameters()
 	{
 		if (IsPlaying() && !GetDisableParameterUpdatesWhilePlaying())
 		{
-			DECLARE_CYCLE_STAT(TEXT("FAudioThreadTask.SoundParameterController.ResetParameters"), STAT_AudioResetParameters, STATGROUP_AudioThreadCommands);
+			DECLARE_CYCLE_STAT(TEXT("FAudioThreadTask.SoundParameterControllerInterface.ResetParameters"), STAT_AudioResetParameters, STATGROUP_AudioThreadCommands);
 			AudioDevice->SendCommandToActiveSounds(GetInstanceOwnerID(), [] (FActiveSound& ActiveSound)
 			{
 				if (Audio::IParameterTransmitter* Transmitter = ActiveSound.GetTransmitter())
@@ -38,8 +38,41 @@ void ISoundParameterControllerInterface::ResetParameters()
 
 void ISoundParameterControllerInterface::SetTriggerParameter(FName InName)
 {
-	// Trigger is just a (true) bool param currently.
-	SetParameterInternal(FAudioParameter(InName, true));
+	if (InName.IsNone())
+	{
+		return;
+	}
+
+	if (IsPlaying() && !GetDisableParameterUpdatesWhilePlaying())
+	{
+		if (FAudioDevice* AudioDevice = GetAudioDevice())
+		{
+			FAudioParameter ParamToSet = FAudioParameter(InName, true);
+			if (USoundBase* Sound = GetSound())
+			{
+				TArray<FAudioParameter> Params = { MoveTemp(ParamToSet) };
+				Sound->InitParameters(Params, SoundParameterControllerInterfacePrivate::ProxyFeatureName);
+				ParamToSet = MoveTemp(Params[0]);
+			}
+
+			DECLARE_CYCLE_STAT(TEXT("FAudioThreadTask.SoundParameterControllerInterface.ExecuteTriggerParameter"), STAT_AudioExecuteTriggerParameter, STATGROUP_AudioThreadCommands);
+
+			AudioDevice->SendCommandToActiveSounds(GetInstanceOwnerID(), [AudioDevice, Param = MoveTemp(ParamToSet)](FActiveSound& ActiveSound)
+			{
+				if (Audio::IParameterTransmitter* Transmitter = ActiveSound.GetTransmitter())
+				{
+					const FName ParamName = Param.ParamName;
+
+					// Must be copied as original version must be preserved in case command is called on multiple ActiveSounds.
+					FAudioParameter TempParam = Param;
+					if (!Transmitter->SetParameter(MoveTemp(TempParam)))
+					{
+						UE_LOG(LogAudio, Warning, TEXT("Failed to execute trigger parameter '%s'"), *ParamName.ToString());
+					}
+				}
+			}, GET_STATID(STAT_AudioExecuteTriggerParameter));
+		}
+	}
 }
 
 void ISoundParameterControllerInterface::SetBoolParameter(FName InName, bool InValue)
@@ -120,7 +153,7 @@ void ISoundParameterControllerInterface::SetParameters(TArray<FAudioParameter>&&
 
 			ParamsToSet = MoveTemp(InValues);
 
-			DECLARE_CYCLE_STAT(TEXT("FAudioThreadTask.SoundParameterController.SetParameters"), STAT_AudioSetParameters, STATGROUP_AudioThreadCommands);
+			DECLARE_CYCLE_STAT(TEXT("FAudioThreadTask.SoundParameterControllerInterface.SetParameters"), STAT_AudioSetParameters, STATGROUP_AudioThreadCommands);
 			AudioDevice->SendCommandToActiveSounds(GetInstanceOwnerID(), [AudioDevice, Params = MoveTemp(ParamsToSet)](FActiveSound& ActiveSound)
 			{
 				if (Audio::IParameterTransmitter* Transmitter = ActiveSound.GetTransmitter())
@@ -173,7 +206,7 @@ void ISoundParameterControllerInterface::SetParameterInternal(FAudioParameter&&
 				ParamToSet = MoveTemp(InParam);
 			}
 
-			DECLARE_CYCLE_STAT(TEXT("FAudioThreadTask.SoundParameterController.SetParameter"), STAT_AudioSetParameter, STATGROUP_AudioThreadCommands);
+			DECLARE_CYCLE_STAT(TEXT("FAudioThreadTask.SoundParameterControllerInterface.SetParameter"), STAT_AudioSetParameter, STATGROUP_AudioThreadCommands);
 
 			AudioDevice->SendCommandToActiveSounds(GetInstanceOwnerID(), [AudioDevice, Param = MoveTemp(ParamToSet)](FActiveSound& ActiveSound)
 			{
