WorldPartition: Fix for non-unity bulds.

diff --git a/Engine/Source/Runtime/Engine/Public/WorldPartition/WorldPartitionSubsystem.h b/Engine/Source/Runtime/Engine/Public/WorldPartition/WorldPartitionSubsystem.h
--- a/Engine/Source/Runtime/Engine/Public/WorldPartition/WorldPartitionSubsystem.h
+++ b/Engine/Source/Runtime/Engine/Public/WorldPartition/WorldPartitionSubsystem.h
@@ -2,6 +2,7 @@
 
 #pragma once
 
+#include "WorldPartitionStreamingSource.h"
 #include "Subsystems/WorldSubsystem.h"
 #include "WorldPartitionSubsystem.generated.h"
 
