RPC DoS stability tweaks.

diff --git a/Engine/Source/Runtime/Engine/Private/Net/RPCDoSDetection.cpp b/Engine/Source/Runtime/Engine/Private/Net/RPCDoSDetection.cpp
--- a/Engine/Source/Runtime/Engine/Private/Net/RPCDoSDetection.cpp
+++ b/Engine/Source/Runtime/Engine/Private/Net/RPCDoSDetection.cpp
@@ -1020,10 +1020,10 @@ ERPCNotifyResult FRPCDoSDetection::CheckRPCTracking(UFunction* Function, FName F
 {
 	ERPCNotifyResult Result = ERPCNotifyResult::ExecuteRPC;
 	bool bNewTracking = false;
-	FRPCTrackingInfo& Tracking = FindOrAddRPCTracking(Function, bNewTracking);
 
 	if (!bHitchSuspendDetection)
 	{
+		FRPCTrackingInfo& Tracking = FindOrAddRPCTracking(Function, bNewTracking);
 		const uint8 OldTrackedSecondIncrement = Tracking.LastTrackedSecondIncrement;
 
 		Tracking.LastTrackedSecondIncrement = SecondsIncrementer;
@@ -1128,13 +1128,20 @@ ERPCNotifyResult FRPCDoSDetection::CheckRPCTracking(UFunction* Function, FName F
 		if (Tracking.BlockState == ERPCBlockState::Blocked)
 		{
 			PostReceivedRPCBlockCount++;
+
+			Result = ERPCNotifyResult::BlockRPC;
 		}
 	}
+	else
+	{
+		TSharedPtr<FRPCTrackingInfo>* TrackingPtr = RPCTracking.Find(Function);
 
+		if (TrackingPtr != nullptr && TrackingPtr->IsValid() && TrackingPtr->Get()->BlockState == ERPCBlockState::Blocked)
+		{
+			PostReceivedRPCBlockCount++;
 
-	if (Tracking.BlockState == ERPCBlockState::Blocked)
-	{
-		Result = ERPCNotifyResult::BlockRPC;
+			Result = ERPCNotifyResult::BlockRPC;
+		}
 	}
 
 	return Result;
diff --git a/Engine/Source/Runtime/Engine/Private/NetConnection.cpp b/Engine/Source/Runtime/Engine/Private/NetConnection.cpp
--- a/Engine/Source/Runtime/Engine/Private/NetConnection.cpp
+++ b/Engine/Source/Runtime/Engine/Private/NetConnection.cpp
@@ -2647,555 +2647,593 @@ void UNetConnection::ReceivedPacket( FBitReader& Reader, bool bIsReinjectedPacke
 	// which would cause a disconnect serverside
 	TArray<int32> RejectedChans;
 	const bool bIsServer = Driver->ServerConnection == nullptr;
+	double PostReceiveTime = 0.0;
 
-	if (bIsServer && !IsReplay())
 	{
-		RPCDoS.PreReceivedPacket(CurrentReceiveTimeInS);
-	}
-
-	// Disassemble and dispatch all bunches in the packet.
-	while( !Reader.AtEnd() && GetConnectionState()!=USOCK_Closed )
-	{
-		// For demo backwards compatibility, old replays still have this bit
-		if (IsInternalAck() && EngineNetworkProtocolVersion < EEngineNetworkVersionHistory::HISTORY_ACKS_INCLUDED_IN_HEADER)
+		if (bIsServer && !IsReplay())
 		{
-			const bool IsAckDummy = Reader.ReadBit() == 1u;
+			RPCDoS.PreReceivedPacket(CurrentReceiveTimeInS);
 		}
 
-		// Parse the bunch.
-		int32 StartPos = Reader.GetPosBits();
-		
-		// Process Received data
-		{
-			// Parse the incoming data.
-			FInBunch Bunch( this );
-			int32 IncomingStartPos		= Reader.GetPosBits();
-			uint8 bControl				= Reader.ReadBit();
-			Bunch.PacketId				= InPacketId;
-			Bunch.bOpen					= bControl ? Reader.ReadBit() : 0;
-			Bunch.bClose				= bControl ? Reader.ReadBit() : 0;
-			
-			if (Bunch.EngineNetVer() < HISTORY_CHANNEL_CLOSE_REASON)
-			{
-				const uint8 bDormant = Bunch.bClose ? Reader.ReadBit() : 0;
-				Bunch.CloseReason = bDormant ? EChannelCloseReason::Dormancy : EChannelCloseReason::Destroyed;
-			}
-			else
+		ON_SCOPE_EXIT
+		{
+			PostReceiveTime = FPlatformTime::Seconds();
+
+			if (bIsServer && !IsReplay())
 			{
-				Bunch.CloseReason = Bunch.bClose ? (EChannelCloseReason)Reader.ReadInt((uint32)EChannelCloseReason::MAX) : EChannelCloseReason::Destroyed;
+				RPCDoS.PostReceivedPacket(PostReceiveTime);
 			}
+		};
 
-			Bunch.bIsReplicationPaused  = Reader.ReadBit();
-			Bunch.bReliable				= Reader.ReadBit();
-
-			if (Bunch.EngineNetVer() < HISTORY_MAX_ACTOR_CHANNELS_CUSTOMIZATION)
+		// Disassemble and dispatch all bunches in the packet.
+		while( !Reader.AtEnd() && GetConnectionState()!=USOCK_Closed )
+		{
+			// For demo backwards compatibility, old replays still have this bit
+			if (IsInternalAck() && EngineNetworkProtocolVersion < EEngineNetworkVersionHistory::HISTORY_ACKS_INCLUDED_IN_HEADER)
 			{
-				static const int OLD_MAX_ACTOR_CHANNELS = 10240;
-				Bunch.ChIndex = Reader.ReadInt(OLD_MAX_ACTOR_CHANNELS);
+				const bool IsAckDummy = Reader.ReadBit() == 1u;
 			}
-			else
-			{
-				uint32 ChIndex;
-				Reader.SerializeIntPacked(ChIndex);
 
-				if (ChIndex >= (uint32)MaxChannelSize)
+			// Parse the bunch.
+			int32 StartPos = Reader.GetPosBits();
+		
+			// Process Received data
+			{
+				// Parse the incoming data.
+				FInBunch Bunch( this );
+				int32 IncomingStartPos		= Reader.GetPosBits();
+				uint8 bControl				= Reader.ReadBit();
+				Bunch.PacketId				= InPacketId;
+				Bunch.bOpen					= bControl ? Reader.ReadBit() : 0;
+				Bunch.bClose				= bControl ? Reader.ReadBit() : 0;
+			
+				if (Bunch.EngineNetVer() < HISTORY_CHANNEL_CLOSE_REASON)
+				{
+					const uint8 bDormant = Bunch.bClose ? Reader.ReadBit() : 0;
+					Bunch.CloseReason = bDormant ? EChannelCloseReason::Dormancy : EChannelCloseReason::Destroyed;
+				}
+				else
 				{
-					UE_LOG(LogNet, Warning, TEXT("Bunch channel index exceeds channel limit"));
+					Bunch.CloseReason = Bunch.bClose ? (EChannelCloseReason)Reader.ReadInt((uint32)EChannelCloseReason::MAX) :
+														EChannelCloseReason::Destroyed;
+				}
 
-					Close(ENetCloseResult::BunchBadChannelIndex);
+				Bunch.bIsReplicationPaused  = Reader.ReadBit();
+				Bunch.bReliable				= Reader.ReadBit();
 
-					return;
+				if (Bunch.EngineNetVer() < HISTORY_MAX_ACTOR_CHANNELS_CUSTOMIZATION)
+				{
+					static const int OLD_MAX_ACTOR_CHANNELS = 10240;
+					Bunch.ChIndex = Reader.ReadInt(OLD_MAX_ACTOR_CHANNELS);
 				}
+				else
+				{
+					uint32 ChIndex;
+					Reader.SerializeIntPacked(ChIndex);
 
-				Bunch.ChIndex = ChIndex;
-			}
+					if (ChIndex >= (uint32)MaxChannelSize)
+					{
+						UE_LOG(LogNet, Warning, TEXT("Bunch channel index exceeds channel limit"));
 
-			const int32 CachedChIndex = Bunch.ChIndex;
+						Close(ENetCloseResult::BunchBadChannelIndex);
 
-			// if flag is set, remap channel index values, we're fast forwarding a replay checkpoint
-			// and there should be no bunches for existing channels
-			if (IsInternalAck() && bAllowExistingChannelIndex && (Bunch.EngineNetVer() >= HISTORY_REPLAY_DORMANCY))
-			{
-				if (ChannelIndexMap.Contains(Bunch.ChIndex))
-				{
-					Bunch.ChIndex = ChannelIndexMap[Bunch.ChIndex];
+						return;
+					}
+
+					Bunch.ChIndex = ChIndex;
 				}
-				else
+
+				const int32 CachedChIndex = Bunch.ChIndex;
+
+				// if flag is set, remap channel index values, we're fast forwarding a replay checkpoint
+				// and there should be no bunches for existing channels
+				if (IsInternalAck() && bAllowExistingChannelIndex && (Bunch.EngineNetVer() >= HISTORY_REPLAY_DORMANCY))
 				{
-					if (Channels[Bunch.ChIndex])
+					if (ChannelIndexMap.Contains(Bunch.ChIndex))
 					{
-						// this should be an open bunch if its the first time we've seen it
-						if (Bunch.bOpen && (!Bunch.bPartial || Bunch.bPartialInitial))
+						Bunch.ChIndex = ChannelIndexMap[Bunch.ChIndex];
+					}
+					else
+					{
+						if (Channels[Bunch.ChIndex])
 						{
-							FBitReaderMark Mark(Reader);
+							// this should be an open bunch if its the first time we've seen it
+							if (Bunch.bOpen && (!Bunch.bPartial || Bunch.bPartialInitial))
+							{
+								FBitReaderMark Mark(Reader);
 
-							Reader.ReadBit(); // bHasPackageMapExports
-							Reader.ReadBit(); // bHasMustBeMappedGUIDs
+								Reader.ReadBit(); // bHasPackageMapExports
+								Reader.ReadBit(); // bHasMustBeMappedGUIDs
 
-							const uint8 bPartial = Reader.ReadBit(); // bPartial
+								const uint8 bPartial = Reader.ReadBit(); // bPartial
 
-							if (bPartial)
-							{
-								Reader.ReadBit(); // bPartialInitial
-								Reader.ReadBit(); // bPartialFinal
-							}
+								if (bPartial)
+								{
+									Reader.ReadBit(); // bPartialInitial
+									Reader.ReadBit(); // bPartialFinal
+								}
 
-							FName ChName;
-							UPackageMap::StaticSerializeName(Reader, ChName);
+								FName ChName;
+								UPackageMap::StaticSerializeName(Reader, ChName);
 
-							Mark.Pop(Reader);
+								Mark.Pop(Reader);
 
-							int32 FreeIndex = GetFreeChannelIndex(ChName);
+								int32 FreeIndex = GetFreeChannelIndex(ChName);
 
-							if (FreeIndex != INDEX_NONE)
-							{
-								UE_LOG(LogNetTraffic, Verbose, TEXT("Adding channel mapping %d to %d"), Bunch.ChIndex, FreeIndex);
-								ChannelIndexMap.Add(Bunch.ChIndex, FreeIndex);
-								Bunch.ChIndex = FreeIndex;
-							}
-							else
-							{
-								UE_LOG(LogNetTraffic, Warning, TEXT("Unable to find free channel index"));
-								continue;
+								if (FreeIndex != INDEX_NONE)
+								{
+									UE_LOG(LogNetTraffic, Verbose, TEXT("Adding channel mapping %d to %d"), Bunch.ChIndex, FreeIndex);
+									ChannelIndexMap.Add(Bunch.ChIndex, FreeIndex);
+									Bunch.ChIndex = FreeIndex;
+								}
+								else
+								{
+									UE_LOG(LogNetTraffic, Warning, TEXT("Unable to find free channel index"));
+									continue;
+								}
 							}
 						}
 					}
 				}
-			}
 
-			Bunch.bHasPackageMapExports	= Reader.ReadBit();
-			Bunch.bHasMustBeMappedGUIDs	= Reader.ReadBit();
-			Bunch.bPartial				= Reader.ReadBit();
+				Bunch.bHasPackageMapExports	= Reader.ReadBit();
+				Bunch.bHasMustBeMappedGUIDs	= Reader.ReadBit();
+				Bunch.bPartial				= Reader.ReadBit();
 
-			if ( Bunch.bReliable )
-			{
-				if ( IsInternalAck() )
+				if ( Bunch.bReliable )
 				{
-					// We can derive the sequence for 100% reliable connections
-					Bunch.ChSequence = InReliable[Bunch.ChIndex] + 1;
+					if ( IsInternalAck() )
+					{
+						// We can derive the sequence for 100% reliable connections
+						Bunch.ChSequence = InReliable[Bunch.ChIndex] + 1;
+					}
+					else
+					{
+						// If this is a reliable bunch, use the last processed reliable sequence to read the new reliable sequence
+						Bunch.ChSequence = MakeRelative( Reader.ReadInt( MAX_CHSEQUENCE ), InReliable[Bunch.ChIndex], MAX_CHSEQUENCE );
+					}
+				} 
+				else if ( Bunch.bPartial )
+				{
+					// If this is an unreliable partial bunch, we simply use packet sequence since we already have it
+					Bunch.ChSequence = InPacketId;
 				}
 				else
 				{
-					// If this is a reliable bunch, use the last processed reliable sequence to read the new reliable sequence
-					Bunch.ChSequence = MakeRelative( Reader.ReadInt( MAX_CHSEQUENCE ), InReliable[Bunch.ChIndex], MAX_CHSEQUENCE );
+					Bunch.ChSequence = 0;
 				}
-			} 
-			else if ( Bunch.bPartial )
-			{
-				// If this is an unreliable partial bunch, we simply use packet sequence since we already have it
-				Bunch.ChSequence = InPacketId;
-			}
-			else
-			{
-				Bunch.ChSequence = 0;
-			}
 
-			Bunch.bPartialInitial = Bunch.bPartial ? Reader.ReadBit() : 0;
-			Bunch.bPartialFinal = Bunch.bPartial ? Reader.ReadBit() : 0;
+				Bunch.bPartialInitial = Bunch.bPartial ? Reader.ReadBit() : 0;
+				Bunch.bPartialFinal = Bunch.bPartial ? Reader.ReadBit() : 0;
 
-			if (Bunch.EngineNetVer() < HISTORY_CHANNEL_NAMES)
-			{
-				uint32 ChType = (Bunch.bReliable || Bunch.bOpen) ? Reader.ReadInt(CHTYPE_MAX) : CHTYPE_None;
-				switch (ChType)
+				if (Bunch.EngineNetVer() < HISTORY_CHANNEL_NAMES)
 				{
-					case CHTYPE_Control:
-						Bunch.ChName = NAME_Control;
-						break;
-					case CHTYPE_Voice:
-						Bunch.ChName = NAME_Voice;
-						break;
-					case CHTYPE_Actor:
-						Bunch.ChName = NAME_Actor;
+					uint32 ChType = (Bunch.bReliable || Bunch.bOpen) ? Reader.ReadInt(CHTYPE_MAX) : CHTYPE_None;
+					switch (ChType)
+					{
+						case CHTYPE_Control:
+							Bunch.ChName = NAME_Control;
+							break;
+						case CHTYPE_Voice:
+							Bunch.ChName = NAME_Voice;
+							break;
+						case CHTYPE_Actor:
+							Bunch.ChName = NAME_Actor;
+							break;
 						break;
-					break;
+					}
 				}
-			}
-			else
-			{
-				if (Bunch.bReliable || Bunch.bOpen)
+				else
 				{
-					UPackageMap::StaticSerializeName(Reader, Bunch.ChName);
-
-					if( Reader.IsError() )
+					if (Bunch.bReliable || Bunch.bOpen)
 					{
-						UE_LOG(LogNet, Warning, TEXT("Channel name serialization failed."));
+						UPackageMap::StaticSerializeName(Reader, Bunch.ChName);
 
-						HandleNetResultOrClose(ENetCloseResult::BunchChannelNameFail);
+						if( Reader.IsError() )
+						{
+							UE_LOG(LogNet, Warning, TEXT("Channel name serialization failed."));
 
-						return;
+							HandleNetResultOrClose(ENetCloseResult::BunchChannelNameFail);
+
+							return;
+						}
+					}
+					else
+					{
+						Bunch.ChName = NAME_None;
 					}
 				}
-				else
-				{
-					Bunch.ChName = NAME_None;
-				}
-			}
 
-			UChannel* Channel = Channels[Bunch.ChIndex];
+				UChannel* Channel = Channels[Bunch.ChIndex];
 
-			// If there's an existing channel and the bunch specified it's channel type, make sure they match.
-			if (Channel && (Bunch.ChName != NAME_None) && (Bunch.ChName != Channel->ChName))
-			{
-				UE_LOG(LogNet, Error, TEXT("Existing channel at index %d with type \"%s\" differs from the incoming bunch's expected channel type, \"%s\"."),
-					Bunch.ChIndex, *Channel->ChName.ToString(), *Bunch.ChName.ToString());
+				// If there's an existing channel and the bunch specified it's channel type, make sure they match.
+				if (Channel && (Bunch.ChName != NAME_None) && (Bunch.ChName != Channel->ChName))
+				{
+					UE_LOG(LogNet, Error,
+							TEXT("Existing channel at index %d with type \"%s\" differs from the incoming bunch's expected channel type, \"%s\"."),
+							Bunch.ChIndex, *Channel->ChName.ToString(), *Bunch.ChName.ToString());
 
-				Close(ENetCloseResult::BunchWrongChannelType);
+					Close(ENetCloseResult::BunchWrongChannelType);
 
-				return;
-			}
-
-			int32 BunchDataBits  = Reader.ReadInt( UNetConnection::MaxPacket*8 );
+					return;
+				}
 
-			if ((Bunch.bClose || Bunch.bOpen) && UE_LOG_ACTIVE(LogNetDormancy,VeryVerbose) )
-			{
-				UE_LOG(LogNetDormancy, VeryVerbose, TEXT("Received: %s"), *Bunch.ToString());
-			}
+				int32 BunchDataBits  = Reader.ReadInt( UNetConnection::MaxPacket*8 );
 
-			if (UE_LOG_ACTIVE(LogNetTraffic,VeryVerbose))
-			{
-				UE_LOG(LogNetTraffic, VeryVerbose, TEXT("Received: %s"), *Bunch.ToString());
-			}
+				if ((Bunch.bClose || Bunch.bOpen) && UE_LOG_ACTIVE(LogNetDormancy,VeryVerbose) )
+				{
+					UE_LOG(LogNetDormancy, VeryVerbose, TEXT("Received: %s"), *Bunch.ToString());
+				}
 
-			const int32 HeaderPos = Reader.GetPosBits();
+				if (UE_LOG_ACTIVE(LogNetTraffic,VeryVerbose))
+				{
+					UE_LOG(LogNetTraffic, VeryVerbose, TEXT("Received: %s"), *Bunch.ToString());
+				}
 
-			if( Reader.IsError() )
-			{
-				UE_LOG(LogNet, Warning, TEXT("Bunch header overflowed"));
+				const int32 HeaderPos = Reader.GetPosBits();
 
-				HandleNetResultOrClose(ENetCloseResult::BunchHeaderOverflow);
+				if( Reader.IsError() )
+				{
+					UE_LOG(LogNet, Warning, TEXT("Bunch header overflowed"));
 
-				return;
-			}
+					HandleNetResultOrClose(ENetCloseResult::BunchHeaderOverflow);
 
-			// Trace bunch read
-			UE_NET_TRACE_BUNCH_SCOPE(InTraceCollector, Bunch, StartPos, HeaderPos - StartPos);
+					return;
+				}
 
-			Bunch.SetData( Reader, BunchDataBits );
-			if( Reader.IsError() )
-			{
-				// Bunch claims it's larger than the enclosing packet.
-				UE_LOG(LogNet, Warning, TEXT("Bunch data overflowed (%i %i+%i/%i)"), IncomingStartPos, HeaderPos, BunchDataBits, Reader.GetNumBits());
+				// Trace bunch read
+				UE_NET_TRACE_BUNCH_SCOPE(InTraceCollector, Bunch, StartPos, HeaderPos - StartPos);
 
-				HandleNetResultOrClose(ENetCloseResult::BunchDataOverflow);
+				Bunch.SetData( Reader, BunchDataBits );
+				if( Reader.IsError() )
+				{
+					// Bunch claims it's larger than the enclosing packet.
+					UE_LOG(LogNet, Warning, TEXT("Bunch data overflowed (%i %i+%i/%i)"), IncomingStartPos, HeaderPos, BunchDataBits,
+							Reader.GetNumBits());
 
-				return;
-			}
+					HandleNetResultOrClose(ENetCloseResult::BunchDataOverflow);
 
-			if ( Bunch.bHasPackageMapExports )
-			{
-				Driver->NetGUIDInBytes += (BunchDataBits + (HeaderPos - IncomingStartPos)) >> 3;
+					return;
+				}
 
-				if ( IsInternalAck() )
+				if ( Bunch.bHasPackageMapExports )
 				{
-					// NOTE - For replays, we do this even earlier, to try and load this as soon as possible, in case there is an issue creating the channel
-					// If a replay fails to create a channel, we want to salvage as much as possible
-					Cast<UPackageMapClient>( PackageMap )->ReceiveNetGUIDBunch( Bunch );
+					Driver->NetGUIDInBytes += (BunchDataBits + (HeaderPos - IncomingStartPos)) >> 3;
 
-					if ( Bunch.IsError() )
+					if (IsInternalAck())
 					{
-						UE_LOG( LogNetTraffic, Error, TEXT( "UNetConnection::ReceivedPacket: Bunch.IsError() after ReceiveNetGUIDBunch. ChIndex: %i" ), Bunch.ChIndex );
+						// NOTE - For replays, we do this even earlier, to try and load this as soon as possible,
+						// in case there's an issue creating the channel. If a replay fails to create a channel, we want to salvage as much as possible
+						Cast<UPackageMapClient>(PackageMap)->ReceiveNetGUIDBunch(Bunch);
+
+						if (Bunch.IsError())
+						{
+							UE_LOG(LogNetTraffic, Error, TEXT("UNetConnection::ReceivedPacket: Bunch.IsError() after ReceiveNetGUIDBunch. ChIndex: %i"),
+									Bunch.ChIndex);
+						}
 					}
 				}
-			}
 
-			if( Bunch.bReliable )
-			{
-				UE_LOG(LogNetTraffic, Verbose, TEXT("   Reliable Bunch, Channel %i Sequence %i: Size %.1f+%.1f"), Bunch.ChIndex, Bunch.ChSequence, (HeaderPos-IncomingStartPos)/8.f, (Reader.GetPosBits()-HeaderPos)/8.f );
-			}
-			else
-			{
-				UE_LOG(LogNetTraffic, Verbose, TEXT("   Unreliable Bunch, Channel %i: Size %.1f+%.1f"), Bunch.ChIndex, (HeaderPos-IncomingStartPos)/8.f, (Reader.GetPosBits()-HeaderPos)/8.f );
-			}
+				if (Bunch.bReliable)
+				{
+					UE_LOG(LogNetTraffic, Verbose, TEXT("   Reliable Bunch, Channel %i Sequence %i: Size %.1f+%.1f"),
+							Bunch.ChIndex, Bunch.ChSequence, (HeaderPos-IncomingStartPos)/8.f, (Reader.GetPosBits()-HeaderPos)/8.f);
+				}
+				else
+				{
+					UE_LOG(LogNetTraffic, Verbose, TEXT("   Unreliable Bunch, Channel %i: Size %.1f+%.1f"),
+							Bunch.ChIndex, (HeaderPos-IncomingStartPos)/8.f, (Reader.GetPosBits()-HeaderPos)/8.f);
+				}
 
-			if ( Bunch.bOpen )
-			{
-				UE_LOG(LogNetTraffic, Verbose, TEXT("   bOpen Bunch, Channel %i Sequence %i: Size %.1f+%.1f"), Bunch.ChIndex, Bunch.ChSequence, (HeaderPos-IncomingStartPos)/8.f, (Reader.GetPosBits()-HeaderPos)/8.f );
-			}
+				if (Bunch.bOpen)
+				{
+					UE_LOG(LogNetTraffic, Verbose, TEXT("   bOpen Bunch, Channel %i Sequence %i: Size %.1f+%.1f"),
+							Bunch.ChIndex, Bunch.ChSequence, (HeaderPos-IncomingStartPos)/8.f, (Reader.GetPosBits()-HeaderPos)/8.f );
+				}
 
-			if ( Channels[Bunch.ChIndex] == NULL && ( Bunch.ChIndex != 0 || Bunch.ChName != NAME_Control ) )
-			{
-				// Can't handle other channels until control channel exists.
-				if ( Channels[0] == NULL )
+				if ( Channels[Bunch.ChIndex] == NULL && ( Bunch.ChIndex != 0 || Bunch.ChName != NAME_Control ) )
 				{
-					UE_LOG(LogNetTraffic, Log,
-							TEXT("ReceivedPacket: Received non-control bunch before control channel was created. ChIndex: %i, ChName: %s"),
-							Bunch.ChIndex, *Bunch.ChName.ToString());
+					// Can't handle other channels until control channel exists.
+					if ( Channels[0] == NULL )
+					{
+						UE_LOG(LogNetTraffic, Log,
+								TEXT("ReceivedPacket: Received non-control bunch before control channel was created. ChIndex: %i, ChName: %s"),
+								Bunch.ChIndex, *Bunch.ChName.ToString());
 
-					Close(ENetCloseResult::BunchPrematureControlChannel);
+						Close(ENetCloseResult::BunchPrematureControlChannel);
 
-					return;
+						return;
+					}
+					// on the server, if we receive bunch data for a channel that doesn't exist while we're still logging in,
+					// it's either a broken client or a new instance of a previous connection,
+					// so reject it
+					else if ( PlayerController == NULL && Driver->ClientConnections.Contains( this ) )
+					{
+						UE_LOG(LogNet, Warning,
+								TEXT("ReceivedPacket: Received non-control bunch before player controller was assigned. ChIndex: %i, ChName: %s" ),
+								Bunch.ChIndex, *Bunch.ChName.ToString());
+
+						Close(ENetCloseResult::BunchPrematureChannel);
+
+						return;
+					}
 				}
-				// on the server, if we receive bunch data for a channel that doesn't exist while we're still logging in,
-				// it's either a broken client or a new instance of a previous connection,
-				// so reject it
-				else if ( PlayerController == NULL && Driver->ClientConnections.Contains( this ) )
+				// ignore control channel close if it hasn't been opened yet
+				if ( Bunch.ChIndex == 0 && Channels[0] == NULL && Bunch.bClose && Bunch.ChName == NAME_Control )
 				{
-					UE_LOG(LogNet, Warning,
-							TEXT("ReceivedPacket: Received non-control bunch before player controller was assigned. ChIndex: %i, ChName: %s" ),
-							Bunch.ChIndex, *Bunch.ChName.ToString());
+					UE_LOG(LogNetTraffic, Log, TEXT("ReceivedPacket: Received control channel close before open"));
 
-					Close(ENetCloseResult::BunchPrematureChannel);
+					Close(ENetCloseResult::BunchPrematureControlClose);
 
 					return;
 				}
-			}
-			// ignore control channel close if it hasn't been opened yet
-			if ( Bunch.ChIndex == 0 && Channels[0] == NULL && Bunch.bClose && Bunch.ChName == NAME_Control )
-			{
-				UE_LOG(LogNetTraffic, Log, TEXT("ReceivedPacket: Received control channel close before open"));
-
-				Close(ENetCloseResult::BunchPrematureControlClose);
-
-				return;
-			}
 
-			// Receiving data.
+				// Receiving data.
 
-			// We're on a 100% reliable connection and we are rolling back some data.
-			// In that case, we can generally ignore these bunches.
-			if (IsInternalAck() && bAllowExistingChannelIndex)
-			{
-				if (Bunch.EngineNetVer() < HISTORY_REPLAY_DORMANCY)
+				// We're on a 100% reliable connection and we are rolling back some data.
+				// In that case, we can generally ignore these bunches.
+				if (IsInternalAck() && bAllowExistingChannelIndex)
 				{
-					if (Channel)
+					if (Bunch.EngineNetVer() < HISTORY_REPLAY_DORMANCY)
 					{
-						// This was an open bunch for a channel that's already opened.
-						// We can ignore future bunches from this channel.
-						const bool bNewlyOpenedActorChannel = Bunch.bOpen && (Bunch.ChName == NAME_Actor) && (!Bunch.bPartial || Bunch.bPartialInitial);
-
-						if (bNewlyOpenedActorChannel)
+						if (Channel)
 						{
-							FNetworkGUID ActorGUID = GetActorGUIDFromOpenBunch(Bunch);
+							// This was an open bunch for a channel that's already opened.
+							// We can ignore future bunches from this channel.
+							const bool bNewlyOpenedActorChannel = Bunch.bOpen && (Bunch.ChName == NAME_Actor) &&
+																	(!Bunch.bPartial || Bunch.bPartialInitial);
 
-							if (!Bunch.IsError())
+							if (bNewlyOpenedActorChannel)
 							{
-								IgnoringChannels.Add(Bunch.ChIndex, ActorGUID);
-							}
-							else
-							{
-								UE_LOG(LogNetTraffic, Error, TEXT("UNetConnection::ReceivedPacket: Unable to read actor GUID for ignored bunch. (Channel %d)"), Bunch.ChIndex);
+								FNetworkGUID ActorGUID = GetActorGUIDFromOpenBunch(Bunch);
+
+								if (!Bunch.IsError())
+								{
+									IgnoringChannels.Add(Bunch.ChIndex, ActorGUID);
+								}
+								else
+								{
+									UE_LOG(LogNetTraffic, Error,
+											TEXT("UNetConnection::ReceivedPacket: Unable to read actor GUID for ignored bunch. (Channel %d)"),
+											Bunch.ChIndex);
+								}
 							}
-						}
 
-						if (IgnoringChannels.Contains(Bunch.ChIndex))
-						{
-							if (Bunch.bClose && (!Bunch.bPartial || Bunch.bPartialFinal))
+							if (IgnoringChannels.Contains(Bunch.ChIndex))
 							{
-								FNetworkGUID ActorGUID = IgnoringChannels.FindAndRemoveChecked(Bunch.ChIndex);
-								if (ActorGUID.IsStatic())
+								if (Bunch.bClose && (!Bunch.bPartial || Bunch.bPartialFinal))
 								{
-									UObject* FoundObject = Driver->GuidCache->GetObjectFromNetGUID(ActorGUID, false);
-									if (AActor* StaticActor = Cast<AActor>(FoundObject))
+									FNetworkGUID ActorGUID = IgnoringChannels.FindAndRemoveChecked(Bunch.ChIndex);
+									if (ActorGUID.IsStatic())
 									{
-										DestroyIgnoredActor(StaticActor);
-									}
-									else
-									{
-										ensure(FoundObject == nullptr);
-										UE_LOG(LogNetTraffic, Log, TEXT("UNetConnection::ReceivedPacket: Unable to find static actor to cleanup for ignored bunch. (Channel %d NetGUID %lu)"), Bunch.ChIndex, ActorGUID.Value);
+										UObject* FoundObject = Driver->GuidCache->GetObjectFromNetGUID(ActorGUID, false);
+										if (AActor* StaticActor = Cast<AActor>(FoundObject))
+										{
+											DestroyIgnoredActor(StaticActor);
+										}
+										else
+										{
+											ensure(FoundObject == nullptr);
+
+											UE_LOG(LogNetTraffic, Log,
+													TEXT("UNetConnection::ReceivedPacket: Unable to find static actor to cleanup for ignored bunch. ")
+													TEXT("(Channel %d NetGUID %lu)"), Bunch.ChIndex, ActorGUID.Value);
+										}
 									}
 								}
-							}
 
-							UE_LOG(LogNetTraffic, Log, TEXT("Ignoring bunch for already open channel: %i"), Bunch.ChIndex);
-							continue;
+								UE_LOG(LogNetTraffic, Log, TEXT("Ignoring bunch for already open channel: %i"), Bunch.ChIndex);
+								continue;
+							}
 						}
 					}
-				}
-				else
-				{
-					const bool bCloseBunch = Bunch.bClose && (!Bunch.bPartial || Bunch.bPartialFinal);
-
-					if (bCloseBunch && ChannelIndexMap.Contains(CachedChIndex))
+					else
 					{
-						check(ChannelIndexMap[CachedChIndex] == Bunch.ChIndex);
+						const bool bCloseBunch = Bunch.bClose && (!Bunch.bPartial || Bunch.bPartialFinal);
+
+						if (bCloseBunch && ChannelIndexMap.Contains(CachedChIndex))
+						{
+							check(ChannelIndexMap[CachedChIndex] == Bunch.ChIndex);
 
-						UE_LOG(LogNetTraffic, Verbose, TEXT("Removing channel mapping %d to %d"), CachedChIndex, Bunch.ChIndex);
-						ChannelIndexMap.Remove(CachedChIndex);
+							UE_LOG(LogNetTraffic, Verbose, TEXT("Removing channel mapping %d to %d"), CachedChIndex, Bunch.ChIndex);
+							ChannelIndexMap.Remove(CachedChIndex);
+						}
 					}
 				}
-			}
-
-			// Ignore if reliable packet has already been processed.
-			if ( Bunch.bReliable && Bunch.ChSequence <= InReliable[Bunch.ChIndex] )
-			{
-				UE_LOG( LogNetTraffic, Log, TEXT( "UNetConnection::ReceivedPacket: Received outdated bunch (Channel %d Current Sequence %i)" ), Bunch.ChIndex, InReliable[Bunch.ChIndex] );
-				check( !IsInternalAck() );		// Should be impossible with 100% reliable connections
-				continue;
-			}
-			
-			// If opening the channel with an unreliable packet, check that it is "bNetTemporary", otherwise discard it
-			if( !Channel && !Bunch.bReliable )
-			{
-				// Unreliable bunches that open channels should be bOpen && (bClose || bPartial)
-				// NetTemporary usually means one bunch that is unreliable (bOpen and bClose):	1(bOpen, bClose)
-				// But if that bunch export NetGUIDs, it will get split into 2 bunches:			1(bOpen, bPartial) - 2(bClose).
-				// (the initial actor bunch itself could also be split into multiple bunches. So bPartial is the right check here)
 
-				const bool ValidUnreliableOpen = Bunch.bOpen && (Bunch.bClose || Bunch.bPartial);
-				if (!ValidUnreliableOpen)
+				// Ignore if reliable packet has already been processed.
+				if ( Bunch.bReliable && Bunch.ChSequence <= InReliable[Bunch.ChIndex] )
 				{
-					if ( IsInternalAck() )
-					{
-						// Should be impossible with 100% reliable connections
-						UE_LOG( LogNetTraffic, Error, TEXT( "      Received unreliable bunch before open with reliable connection (Channel %d Current Sequence %i)" ), Bunch.ChIndex, InReliable[Bunch.ChIndex] );
-					}
-					else
-					{
-						// Simply a log (not a warning, since this can happen under normal conditions, like from a re-join, etc)
-						UE_LOG( LogNetTraffic, Log, TEXT( "      Received unreliable bunch before open (Channel %d Current Sequence %i)" ), Bunch.ChIndex, InReliable[Bunch.ChIndex] );
-					}
+					UE_LOG(LogNetTraffic, Log, TEXT("UNetConnection::ReceivedPacket: Received outdated bunch (Channel %d Current Sequence %i)"),
+							Bunch.ChIndex, InReliable[Bunch.ChIndex]);
 
-					// Since we won't be processing this packet, don't ack it
-					// We don't want the sender to think this bunch was processed when it really wasn't
-					bSkipAck = true;
+					check( !IsInternalAck() );		// Should be impossible with 100% reliable connections
 					continue;
 				}
-			}
-
-			// Create channel if necessary.
-			if (Channel == nullptr)
-			{
-				if (RejectedChans.Contains(Bunch.ChIndex))
+			
+				// If opening the channel with an unreliable packet, check that it is "bNetTemporary", otherwise discard it
+				if( !Channel && !Bunch.bReliable )
 				{
-					UE_LOG(LogNetTraffic, Log, TEXT("      Ignoring Bunch for ChIndex %i, as the channel was already rejected while processing this packet."), Bunch.ChIndex);
+					// Unreliable bunches that open channels should be bOpen && (bClose || bPartial)
+					// NetTemporary usually means one bunch that is unreliable (bOpen and bClose):	1(bOpen, bClose)
+					// But if that bunch export NetGUIDs, it will get split into 2 bunches:			1(bOpen, bPartial) - 2(bClose).
+					// (the initial actor bunch itself could also be split into multiple bunches. So bPartial is the right check here)
 
-					continue;
+					const bool ValidUnreliableOpen = Bunch.bOpen && (Bunch.bClose || Bunch.bPartial);
+					if (!ValidUnreliableOpen)
+					{
+						if ( IsInternalAck() )
+						{
+							// Should be impossible with 100% reliable connections
+							UE_LOG(LogNetTraffic, Error,
+									TEXT("      Received unreliable bunch before open with reliable connection (Channel %d Current Sequence %i)" ),
+									Bunch.ChIndex, InReliable[Bunch.ChIndex]);
+						}
+						else
+						{
+							// Simply a log (not a warning, since this can happen under normal conditions, like from a re-join, etc)
+							UE_LOG(LogNetTraffic, Log, TEXT("      Received unreliable bunch before open (Channel %d Current Sequence %i)"),
+									Bunch.ChIndex, InReliable[Bunch.ChIndex]);
+						}
+
+						// Since we won't be processing this packet, don't ack it
+						// We don't want the sender to think this bunch was processed when it really wasn't
+						bSkipAck = true;
+						continue;
+					}
 				}
 
-				// Validate channel type.
-				if ( !Driver->IsKnownChannelName( Bunch.ChName ) )
+				// Create channel if necessary.
+				if (Channel == nullptr)
 				{
-					UE_LOG(LogNet, Warning, TEXT("ReceivedPacket: Connection unknown channel type (%s)"), *Bunch.ChName.ToString());
+					if (RejectedChans.Contains(Bunch.ChIndex))
+					{
+						UE_LOG(LogNetTraffic, Log,
+								TEXT("      Ignoring Bunch for ChIndex %i, as the channel was already rejected while processing this packet."),
+								Bunch.ChIndex);
 
-					Close(ENetCloseResult::UnknownChannelType);
+						continue;
+					}
 
-					return;
-				}
+					// Validate channel type.
+					if ( !Driver->IsKnownChannelName( Bunch.ChName ) )
+					{
+						UE_LOG(LogNet, Warning, TEXT("ReceivedPacket: Connection unknown channel type (%s)"), *Bunch.ChName.ToString());
 
-				// Ignore incoming data on channel types that the client are not allowed to create. This can occur if we have in-flight data when server is closing a channel
-				if ( Driver->IsServer() && (Driver->ChannelDefinitionMap[Bunch.ChName].bClientOpen == false) )
-				{
-					UE_LOG(LogNetTraffic, Warning, TEXT("      Ignoring Bunch Create received from client since only server is allowed to create this type of channel: Bunch  %i: ChName %s, ChSequence: %i, bReliable: %i, bPartial: %i, bPartialInitial: %i, bPartialFinal: %i"), Bunch.ChIndex, *Bunch.ChName.ToString(), Bunch.ChSequence, (int)Bunch.bReliable, (int)Bunch.bPartial, (int)Bunch.bPartialInitial, (int)Bunch.bPartialFinal );
-					RejectedChans.AddUnique(Bunch.ChIndex);
-					continue;
-				}
+						Close(ENetCloseResult::UnknownChannelType);
 
-				// peek for guid
-				if (IsInternalAck() && bIgnoreActorBunches)
-				{
-					if (Bunch.bOpen && (!Bunch.bPartial || Bunch.bPartialInitial) && (Bunch.ChName == NAME_Actor))
+						return;
+					}
+
+					// Ignore incoming data on channel types clients can't create. Can occur if we've incoming data when server is closing a channel
+					if ( Driver->IsServer() && (Driver->ChannelDefinitionMap[Bunch.ChName].bClientOpen == false) )
 					{
-						FBitReaderMark Mark(Bunch);
-						FNetworkGUID ActorGUID = GetActorGUIDFromOpenBunch(Bunch);
-						Mark.Pop(Bunch);
+						UE_LOG(LogNetTraffic, Warning, TEXT("      Ignoring Bunch Create received from client since only server is allowed to ")
+								TEXT("create this type of channel: Bunch  %i: ChName %s, ChSequence: %i, bReliable: %i, bPartial: %i, ")
+								TEXT("bPartialInitial: %i, bPartialFinal: %i"), Bunch.ChIndex, *Bunch.ChName.ToString(), Bunch.ChSequence,
+								(int)Bunch.bReliable, (int)Bunch.bPartial, (int)Bunch.bPartialInitial, (int)Bunch.bPartialFinal);
+
+						RejectedChans.AddUnique(Bunch.ChIndex);
+
+						continue;
+					}
 
-						if (ActorGUID.IsValid() && !ActorGUID.IsDefault())
+					// peek for guid
+					if (IsInternalAck() && bIgnoreActorBunches)
+					{
+						if (Bunch.bOpen && (!Bunch.bPartial || Bunch.bPartialInitial) && (Bunch.ChName == NAME_Actor))
 						{
-							if (IgnoredBunchGuids.Contains(ActorGUID))
+							FBitReaderMark Mark(Bunch);
+							FNetworkGUID ActorGUID = GetActorGUIDFromOpenBunch(Bunch);
+							Mark.Pop(Bunch);
+
+							if (ActorGUID.IsValid() && !ActorGUID.IsDefault())
 							{
-								UE_LOG(LogNetTraffic, Verbose, TEXT("Adding Channel: %i to ignore list, ignoring guid: %s"), Bunch.ChIndex, *ActorGUID.ToString());
-								IgnoredBunchChannels.Add(Bunch.ChIndex);
-								continue;
+								if (IgnoredBunchGuids.Contains(ActorGUID))
+								{
+									UE_LOG(LogNetTraffic, Verbose, TEXT("Adding Channel: %i to ignore list, ignoring guid: %s"),
+											Bunch.ChIndex, *ActorGUID.ToString());
+
+									IgnoredBunchChannels.Add(Bunch.ChIndex);
+
+									continue;
+								}
+								else
+								{
+									if (IgnoredBunchChannels.Remove(Bunch.ChIndex))
+									{
+										UE_LOG(LogNetTraffic, Verbose, TEXT("Removing Channel: %i from ignore list, got new guid: %s"),
+												Bunch.ChIndex, *ActorGUID.ToString());
+									}
+								}
 							}
 							else
 							{
-								if (IgnoredBunchChannels.Remove(Bunch.ChIndex))
-								{
-									UE_LOG(LogNetTraffic, Verbose, TEXT("Removing Channel: %i from ignore list, got new guid: %s"), Bunch.ChIndex, *ActorGUID.ToString());
-								}
+								UE_LOG(LogNetTraffic, Warning, TEXT("Open bunch with invalid actor guid, Channel: %i"), Bunch.ChIndex);
 							}
 						}
 						else
 						{
-							UE_LOG(LogNetTraffic, Warning, TEXT("Open bunch with invalid actor guid, Channel: %i"), Bunch.ChIndex);
+							if (IgnoredBunchChannels.Contains(Bunch.ChIndex))
+							{
+								UE_LOG(LogNetTraffic, Verbose, TEXT("Ignoring bunch on channel: %i"), Bunch.ChIndex);
+								continue;
+							}
 						}
 					}
-					else
+
+					// Reliable (either open or later), so create new channel.
+					UE_LOG(LogNetTraffic, Log, TEXT("      Bunch Create %i: ChName %s, ChSequence: %i, bReliable: %i, bPartial: %i, ")
+							TEXT("bPartialInitial: %i, bPartialFinal: %i"), Bunch.ChIndex, *Bunch.ChName.ToString(), Bunch.ChSequence,
+							(int)Bunch.bReliable, (int)Bunch.bPartial, (int)Bunch.bPartialInitial, (int)Bunch.bPartialFinal);
+
+					Channel = CreateChannelByName( Bunch.ChName, EChannelCreateFlags::None, Bunch.ChIndex );
+
+					// Notify the server of the new channel.
+					if( !Driver->Notify->NotifyAcceptingChannel( Channel ) )
 					{
-						if (IgnoredBunchChannels.Contains(Bunch.ChIndex))
+						// Channel refused, so close it, flush it, and delete it.
+						UE_LOG(LogNet, Verbose, TEXT("      NotifyAcceptingChannel Failed! Channel: %s"), *Channel->Describe() );
+
+						RejectedChans.AddUnique(Bunch.ChIndex);
+
+						FOutBunch CloseBunch( Channel, true );
+						check(!CloseBunch.IsError());
+						check(CloseBunch.bClose);
+						CloseBunch.bReliable = 1;
+						Channel->SendBunch( &CloseBunch, false );
+						FlushNet();
+						Channel->ConditionalCleanUp(false, EChannelCloseReason::Destroyed);
+						if( Bunch.ChIndex==0 )
 						{
-							UE_LOG(LogNetTraffic, Verbose, TEXT("Ignoring bunch on channel: %i"), Bunch.ChIndex);
-							continue;
+							UE_LOG(LogNetTraffic, Log, TEXT("Channel 0 create failed") );
+							SetConnectionState(EConnectionState::USOCK_Closed);
 						}
+						continue;
 					}
 				}
 
-				// Reliable (either open or later), so create new channel.
-				UE_LOG(LogNetTraffic, Log, TEXT("      Bunch Create %i: ChName %s, ChSequence: %i, bReliable: %i, bPartial: %i, bPartialInitial: %i, bPartialFinal: %i"), Bunch.ChIndex, *Bunch.ChName.ToString(), Bunch.ChSequence, (int)Bunch.bReliable, (int)Bunch.bPartial, (int)Bunch.bPartialInitial, (int)Bunch.bPartialFinal );
-				Channel = CreateChannelByName( Bunch.ChName, EChannelCreateFlags::None, Bunch.ChIndex );
+				Bunch.bIgnoreRPCs = bIgnoreRPCs;
 
-				// Notify the server of the new channel.
-				if( !Driver->Notify->NotifyAcceptingChannel( Channel ) )
+				// Dispatch the raw, unsequenced bunch to the channel.
+				bool bLocalSkipAck = false;
+				Channel->ReceivedRawBunch( Bunch, bLocalSkipAck ); //warning: May destroy channel.
+				if ( bLocalSkipAck )
 				{
-					// Channel refused, so close it, flush it, and delete it.
-					UE_LOG(LogNet, Verbose, TEXT("      NotifyAcceptingChannel Failed! Channel: %s"), *Channel->Describe() );
-
-					RejectedChans.AddUnique(Bunch.ChIndex);
-
-					FOutBunch CloseBunch( Channel, true );
-					check(!CloseBunch.IsError());
-					check(CloseBunch.bClose);
-					CloseBunch.bReliable = 1;
-					Channel->SendBunch( &CloseBunch, false );
-					FlushNet();
-					Channel->ConditionalCleanUp(false, EChannelCloseReason::Destroyed);
-					if( Bunch.ChIndex==0 )
-					{
-						UE_LOG(LogNetTraffic, Log, TEXT("Channel 0 create failed") );
-						SetConnectionState(EConnectionState::USOCK_Closed);
-					}
-					continue;
+					bSkipAck = true;
 				}
-			}
+				Driver->InBunches++;
+				Driver->InTotalBunches++;
 
-			Bunch.bIgnoreRPCs = bIgnoreRPCs;
-
-			// Dispatch the raw, unsequenced bunch to the channel.
-			bool bLocalSkipAck = false;
-			Channel->ReceivedRawBunch( Bunch, bLocalSkipAck ); //warning: May destroy channel.
-			if ( bLocalSkipAck )
-			{
-				bSkipAck = true;
-			}
-			Driver->InBunches++;
-			Driver->InTotalBunches++;
-
-			// Disconnect if we received a corrupted packet from the client (eg server crash attempt).
-			if (bIsServer && (Bunch.IsCriticalError() || Bunch.IsError()))
-			{
-				UE_LOG( LogNetTraffic, Error, TEXT("Received corrupted packet data from client %s.  Disconnecting."), *LowLevelGetRemoteAddress() );
+				// Disconnect if we received a corrupted packet from the client (eg server crash attempt).
+				if (bIsServer && (Bunch.IsCriticalError() || Bunch.IsError()))
+				{
+					UE_LOG(LogNetTraffic, Error, TEXT("Received corrupted packet data from client %s. Disconnecting."), *LowLevelGetRemoteAddress());
 
-				Close(AddToAndConsumeChainResultPtr(Bunch.ExtendedError, ENetCloseResult::CorruptData));
+					Close(AddToAndConsumeChainResultPtr(Bunch.ExtendedError, ENetCloseResult::CorruptData));
 
-				return;
-			}
+					return;
+				}
 
-			// In replay, if the bunch generated an error but the channel isn't actually open, clean it up so the channel index remains free
-			if (IsInternalAck() && Bunch.IsError() && Channel && !Channel->OpenedLocally && !Channel->OpenAcked)
-			{
-				UE_LOG(LogNetTraffic, Warning, TEXT("Replay cleaning up channel that couldn't be opened: %s"), *Channel->Describe());
-				Channel->ConditionalCleanUp(true, EChannelCloseReason::Destroyed);
+				// In replay, if the bunch generated an error but the channel isn't actually open, clean it up so the channel index remains free
+				if (IsInternalAck() && Bunch.IsError() && Channel && !Channel->OpenedLocally && !Channel->OpenAcked)
+				{
+					UE_LOG(LogNetTraffic, Warning, TEXT("Replay cleaning up channel that couldn't be opened: %s"), *Channel->Describe());
+					Channel->ConditionalCleanUp(true, EChannelCloseReason::Destroyed);
+				}
 			}
 		}
-	}
 
-	// Close/clean-up channels pending close due to received acks.
-	for (FChannelCloseInfo& Info : ChannelsToClose)
-	{
-		if (UChannel* Channel = Channels[Info.Id])
+		// Close/clean-up channels pending close due to received acks.
+		for (FChannelCloseInfo& Info : ChannelsToClose)
 		{
-			Channel->ConditionalCleanUp(false, Info.CloseReason);
+			if (UChannel* Channel = Channels[Info.Id])
+			{
+				Channel->ConditionalCleanUp(false, Info.CloseReason);
+			}
 		}
-	}
-
-	ValidateSendBuffer();
 
-	double PostReceiveTime = FPlatformTime::Seconds();
-
-	if (bIsServer && !IsReplay())
-	{
-		RPCDoS.PostReceivedPacket(PostReceiveTime);
+		ValidateSendBuffer();
 	}
 
 	// Acknowledge the packet.
