Level Instance:
- Track Move Actors to prevent discarding changes
- Track deleted packages through move operations so they get saved when committing Level instance
- Do not prompt for save anymore when committing existing Level instance
- Consolidate package save/delete tracking inside new LevelInstanceEditorObject

diff --git a/Engine/Source/Editor/LevelInstanceEditor/Private/LevelInstanceActorDetails.cpp b/Engine/Source/Editor/LevelInstanceEditor/Private/LevelInstanceActorDetails.cpp
--- a/Engine/Source/Editor/LevelInstanceEditor/Private/LevelInstanceActorDetails.cpp
+++ b/Engine/Source/Editor/LevelInstanceEditor/Private/LevelInstanceActorDetails.cpp
@@ -108,7 +108,7 @@ void FLevelInstanceActorDetails::CustomizeDetails(IDetailLayoutBuilder& DetailBu
 		return;
 	}
 
-	IDetailCategoryBuilder& LevelInstanceEditingCategory = DetailBuilder.EditCategory("LevelInstance Editing", FText::GetEmpty(), ECategoryPriority::Transform);
+	IDetailCategoryBuilder& LevelInstanceEditingCategory = DetailBuilder.EditCategory("Level Edit", FText::GetEmpty(), ECategoryPriority::Transform);
 
 	LevelInstanceEditingCategory.AddCustomRow(FText::GetEmpty())
 	[
diff --git a/Engine/Source/Editor/LevelInstanceEditor/Private/LevelInstanceEditorMode.cpp b/Engine/Source/Editor/LevelInstanceEditor/Private/LevelInstanceEditorMode.cpp
--- a/Engine/Source/Editor/LevelInstanceEditor/Private/LevelInstanceEditorMode.cpp
+++ b/Engine/Source/Editor/LevelInstanceEditor/Private/LevelInstanceEditorMode.cpp
@@ -24,8 +24,6 @@ ULevelInstanceEditorMode::ULevelInstanceEditorMode()
 		false);
 
 	bContextRestriction = true;
-
-	FEditorDelegates::PreBeginPIE.AddUObject(this, &ULevelInstanceEditorMode::OnPreBeginPIE);
 }
 
 ULevelInstanceEditorMode::~ULevelInstanceEditorMode()
@@ -59,6 +57,8 @@ void ULevelInstanceEditorMode::Enter()
 	UEdMode::Enter();
 
 	UpdateEngineShowFlags();
+
+	FEditorDelegates::PreBeginPIE.AddUObject(this, &ULevelInstanceEditorMode::OnPreBeginPIE);
 }
 
 void ULevelInstanceEditorMode::Exit()
@@ -66,6 +66,10 @@ void ULevelInstanceEditorMode::Exit()
 	UEdMode::Exit();
 		
 	UpdateEngineShowFlags();
+
+	bContextRestriction = true;
+
+	FEditorDelegates::PreBeginPIE.RemoveAll(this);
 }
 
 void ULevelInstanceEditorMode::CreateToolkit()
diff --git a/Engine/Source/Editor/LevelInstanceEditor/Private/LevelInstanceEditorModule.cpp b/Engine/Source/Editor/LevelInstanceEditor/Private/LevelInstanceEditorModule.cpp
--- a/Engine/Source/Editor/LevelInstanceEditor/Private/LevelInstanceEditorModule.cpp
+++ b/Engine/Source/Editor/LevelInstanceEditor/Private/LevelInstanceEditorModule.cpp
@@ -136,7 +136,7 @@ namespace LevelInstanceMenuUtils
 		}
 	}
 		
-	void CreateCommitMenu(UToolMenu* Menu, AActor* ContextActor)
+	void CreateCommitDiscardMenu(UToolMenu* Menu, AActor* ContextActor)
 	{
 		ALevelInstance* LevelInstanceEdit = nullptr;
 		if (ContextActor)
@@ -157,35 +157,23 @@ namespace LevelInstanceMenuUtils
 
 		if (LevelInstanceEdit)
 		{
-			FText Tooltip;
-			const bool bCanCommit = LevelInstanceEdit->CanCommit(&Tooltip);
+			FToolMenuSection& Section = CreateLevelSection(Menu);
 
-			auto CreateMenuEntry = [Menu, LevelInstanceEdit, Tooltip](bool bDiscard, bool bCanCommit)
-			{
-				FToolUIAction MenuAction;
+			FText CommitTooltip;
+			const bool bCanCommit = LevelInstanceEdit->CanCommit(&CommitTooltip);
 
-				MenuAction.ExecuteAction.BindLambda([bDiscard, LevelInstanceEdit](const FToolMenuContext&)
-				{
-					if (bDiscard)
-					{
-						LevelInstanceEdit->Discard();
-					}
-					else
-					{
-						LevelInstanceEdit->Commit();
-					}
-				});
-				MenuAction.CanExecuteAction.BindLambda([bCanCommit, LevelInstanceEdit](const FToolMenuContext&)
-				{
-					return bCanCommit;
-				});
+			FToolUIAction CommitAction;
+			CommitAction.ExecuteAction.BindLambda([LevelInstanceEdit](const FToolMenuContext&) { LevelInstanceEdit->Commit(); });
+			CommitAction.CanExecuteAction.BindLambda([bCanCommit](const FToolMenuContext&) { return bCanCommit; });
+			Section.AddMenuEntry(NAME_None, LOCTEXT("LevelInstanceCommitLabel", "Commit"), CommitTooltip, FSlateIcon(), CommitAction);
 
-				FToolMenuSection& Section = CreateLevelSection(Menu);
-				Section.AddMenuEntry(NAME_None, bDiscard ? LOCTEXT("LevelInstanceDiscardLabel", "Discard") : LOCTEXT("LevelInstanceCommitLabel", "Commit"), Tooltip, FSlateIcon(), MenuAction);
-			};
-			
-			CreateMenuEntry(/*bDiscard=*/ false, bCanCommit);
-			CreateMenuEntry(/*bDiscard=*/ true, bCanCommit);
+			FText DiscardTooltip;
+			const bool bCanDiscard = LevelInstanceEdit->CanDiscard(&DiscardTooltip);
+
+			FToolUIAction DiscardAction;
+			DiscardAction.ExecuteAction.BindLambda([LevelInstanceEdit](const FToolMenuContext&) { LevelInstanceEdit->Discard(); });
+			DiscardAction.CanExecuteAction.BindLambda([bCanDiscard](const FToolMenuContext&) { return bCanDiscard; });
+			Section.AddMenuEntry(NAME_None, LOCTEXT("LevelInstanceDiscardLabel", "Discard"), DiscardTooltip, FSlateIcon(), DiscardAction);
 		}
 	}
 
@@ -730,7 +718,7 @@ void FLevelInstanceEditorModule::ExtendContextMenu()
 			if (ContextActor)
 			{
 				LevelInstanceMenuUtils::CreateEditMenu(ToolMenu, ContextActor);
-				LevelInstanceMenuUtils::CreateCommitMenu(ToolMenu, ContextActor);
+				LevelInstanceMenuUtils::CreateCommitDiscardMenu(ToolMenu, ContextActor);
 				LevelInstanceMenuUtils::CreateSaveAsMenu(ToolMenu, ContextActor);
 				LevelInstanceMenuUtils::CreateBreakMenu(ToolMenu, ContextActor);
 				LevelInstanceMenuUtils::CreatePackedBlueprintMenu(ToolMenu, ContextActor);
diff --git a/Engine/Source/Editor/UnrealEd/Private/EditorLevelUtils.cpp b/Engine/Source/Editor/UnrealEd/Private/EditorLevelUtils.cpp
--- a/Engine/Source/Editor/UnrealEd/Private/EditorLevelUtils.cpp
+++ b/Engine/Source/Editor/UnrealEd/Private/EditorLevelUtils.cpp
@@ -68,6 +68,7 @@ static TAutoConsoleVariable<int32>  CVarReflectEditorLevelVisibilityWithGame(
 	TEXT("1 - game state is relfected with editor.\n"), ECVF_Default);
 
 UEditorLevelUtils::FCanMoveActorToLevelDelegate UEditorLevelUtils::CanMoveActorToLevelDelegate;
+UEditorLevelUtils::FOnMoveActorsToLevelEvent UEditorLevelUtils::OnMoveActorsToLevelEvent;
 
 int32 UEditorLevelUtils::MoveActorsToLevel(const TArray<AActor*>& ActorsToMove, ULevelStreaming* DestStreamingLevel, bool bWarnAboutReferences, bool bWarnAboutRenaming)
 {
@@ -76,19 +77,17 @@ int32 UEditorLevelUtils::MoveActorsToLevel(const TArray<AActor*>& ActorsToMove,
 
 int32 UEditorLevelUtils::MoveActorsToLevel(const TArray<AActor*>& ActorsToMove, ULevel* DestLevel, bool bWarnAboutReferences, bool bWarnAboutRenaming, bool bMoveAllOrFail, TArray<AActor*>* OutActors /*=nullptr*/)
 {
-	const bool bCutActors = true;
 	const bool bMoveActors = true;
-	return CopyOrMoveActorsToLevel(ActorsToMove, DestLevel, bCutActors, bMoveActors, bWarnAboutReferences, bWarnAboutRenaming, bMoveAllOrFail, OutActors);
+	return CopyOrMoveActorsToLevel(ActorsToMove, DestLevel, bMoveActors, bWarnAboutReferences, bWarnAboutRenaming, bMoveAllOrFail, OutActors);
 }
 
 int32 UEditorLevelUtils::CopyActorsToLevel(const TArray<AActor*>& ActorsToMove, ULevel* DestLevel, bool bWarnAboutReferences, bool bWarnAboutRenaming, bool bMoveAllOrFail, TArray<AActor*>* OutActors /*=nullptr*/)
 {
-	const bool bCutActors = false;
 	const bool bMoveActors = false;
-	return CopyOrMoveActorsToLevel(ActorsToMove, DestLevel, bCutActors, bMoveActors, bWarnAboutReferences, bWarnAboutRenaming, bMoveAllOrFail);
+	return CopyOrMoveActorsToLevel(ActorsToMove, DestLevel, bMoveActors, bWarnAboutReferences, bWarnAboutRenaming, bMoveAllOrFail);
 }
 
-int32 UEditorLevelUtils::CopyOrMoveActorsToLevel(const TArray<AActor*>& ActorsToMove, ULevel* DestLevel, bool bCutActors, bool bMoveActors, bool bWarnAboutReferences /*= true*/, bool bWarnAboutRenaming /*= true*/, bool bMoveAllOrFail /*= false*/, TArray<AActor*>* OutActors /*=nullptr*/)
+int32 UEditorLevelUtils::CopyOrMoveActorsToLevel(const TArray<AActor*>& ActorsToMove, ULevel* DestLevel, bool bMoveActors, bool bWarnAboutReferences /*= true*/, bool bWarnAboutRenaming /*= true*/, bool bMoveAllOrFail /*= false*/, TArray<AActor*>* OutActors /*=nullptr*/)
 {
 	int32 NumMovedActors = 0;
 
@@ -102,7 +101,8 @@ int32 UEditorLevelUtils::CopyOrMoveActorsToLevel(const TArray<AActor*>& ActorsTo
 		FPlatformApplicationMisc::ClipboardPaste(OriginalClipboardContent);
 
 		// The final list of actors to move after invalid actors were removed
-		TArray<TWeakObjectPtr<AActor>> FinalMoveList;
+		TArray<AActor*> FinalMoveList;
+		TArray<TWeakObjectPtr<AActor>> FinalWeakMoveList;
 		FinalMoveList.Reserve(ActorsToMove.Num());
 
 		bool bIsDestLevelLocked = FLevelUtils::IsLevelLocked(DestLevel);
@@ -153,10 +153,10 @@ int32 UEditorLevelUtils::CopyOrMoveActorsToLevel(const TArray<AActor*>& ActorsTo
 
 			USelection* ActorSelection = GEditor->GetSelectedActors();
 			ActorSelection->BeginBatchSelectOperation();
-			for (TWeakObjectPtr<AActor> ActorPtr : FinalMoveList)
+			FinalWeakMoveList.Reserve(FinalMoveList.Num());
+			for (AActor* Actor : FinalMoveList)
 			{
-				AActor* Actor = ActorPtr.Get();
-				check(Actor);
+				FinalWeakMoveList.Add(Actor);
 				check(Actor->CopyPasteId == INDEX_NONE);
 				Actor->CopyPasteId = ActorPathMapping.Add(TTuple<FSoftObjectPath, FSoftObjectPath>(FSoftObjectPath(Actor), FSoftObjectPath()));
 				GEditor->SelectActor(Actor, true, false);
@@ -168,11 +168,17 @@ int32 UEditorLevelUtils::CopyOrMoveActorsToLevel(const TArray<AActor*>& ActorsTo
 				// Start the transaction
 				FScopedTransaction Transaction(NSLOCTEXT("UnrealEd", "MoveSelectedActorsToSelectedLevel", "Move Actors To Level"));
 
+				// Broadcast event
+				if(bMoveActors)
+				{
+					OnMoveActorsToLevelEvent.Broadcast(ActorsToMove, DestLevel);
+				}
+
 				// Cache the old level
 				ULevel* OldCurrentLevel = OwningWorld->GetCurrentLevel();
 
 				FString DestinationData;
-				GEditor->CopySelectedActorsToClipboard(OwningWorld, bCutActors, bMoveActors, bWarnAboutReferences, &DestinationData);
+				GEditor->CopySelectedActorsToClipboard(OwningWorld, bMoveActors, bMoveActors, bWarnAboutReferences, &DestinationData);
 
 				const bool bLevelVisible = DestLevel->bIsVisible;
 				if (!bLevelVisible)
@@ -252,7 +258,7 @@ int32 UEditorLevelUtils::CopyOrMoveActorsToLevel(const TArray<AActor*>& ActorsTo
 			// The moved (pasted) actors will now be selected
 			NumMovedActors += FinalMoveList.Num();
 
-			for (TWeakObjectPtr<AActor> ActorPtr : FinalMoveList)
+			for (TWeakObjectPtr<AActor> ActorPtr : FinalWeakMoveList)
 			{
 				// It is possible a GC happens because of RenameAssets being called so here we want to update the CopyPasteId only on reachable actors
 				if (AActor* Actor = ActorPtr.Get(/*bEvenIfPendingKill=*/true))
diff --git a/Engine/Source/Editor/UnrealEd/Public/EditorLevelUtils.h b/Engine/Source/Editor/UnrealEd/Public/EditorLevelUtils.h
--- a/Engine/Source/Editor/UnrealEd/Public/EditorLevelUtils.h
+++ b/Engine/Source/Editor/UnrealEd/Public/EditorLevelUtils.h
@@ -129,6 +129,12 @@ public:
 	DECLARE_MULTICAST_DELEGATE_ThreeParams(FCanMoveActorToLevelDelegate, const AActor* /* ActorToMove */, const ULevel* /* DestLevel */, bool& /* bOutCanMove */);
 	static UNREALED_API FCanMoveActorToLevelDelegate CanMoveActorToLevelDelegate;
 		
+	/** 
+	* Delegate used by MoveActorsToLevel() to notify about actors being moved 
+	*/
+	DECLARE_EVENT_TwoParams(UEditorLevelUtils, FOnMoveActorsToLevelEvent, const TArray<AActor*>& /* ActorsToMove */, const ULevel* /* DestLevel */);
+	static UNREALED_API FOnMoveActorsToLevelEvent OnMoveActorsToLevelEvent;
+
 	/**
 	 * Creates a new streaming level and adds it to a world
 	 *
@@ -213,7 +219,7 @@ private:
 
 	static UNREALED_API ULevelStreaming* AddLevelToWorld_Internal(UWorld* InWorld, const TCHAR* LevelPackageName, TSubclassOf<ULevelStreaming> LevelStreamingClass, const FTransform& LevelTransform = FTransform::Identity);
 
-	static UNREALED_API int32 CopyOrMoveActorsToLevel(const TArray<AActor*>& ActorsToMove, ULevel* DestLevel, bool bCutActors, bool bMoveActors, bool bWarnAboutReferences = true, bool bWarnAboutRenaming = true, bool bMoveAllOrFail = false, TArray<AActor*>* OutActors = nullptr);
+	static UNREALED_API int32 CopyOrMoveActorsToLevel(const TArray<AActor*>& ActorsToMove, ULevel* DestLevel, bool bMoveActors, bool bWarnAboutReferences = true, bool bWarnAboutRenaming = true, bool bMoveAllOrFail = false, TArray<AActor*>* OutActors = nullptr);
 
 public:
 	/** Sets the LevelStreamingClass for the specified Level 
diff --git a/Engine/Source/Runtime/Engine/Private/LevelInstance/LevelInstanceActor.cpp b/Engine/Source/Runtime/Engine/Private/LevelInstance/LevelInstanceActor.cpp
--- a/Engine/Source/Runtime/Engine/Private/LevelInstance/LevelInstanceActor.cpp
+++ b/Engine/Source/Runtime/Engine/Private/LevelInstance/LevelInstanceActor.cpp
@@ -718,6 +718,19 @@ bool ALevelInstance::CanCommit(FText* OutReason) const
 	return false;
 }
 
+bool ALevelInstance::CanDiscard(FText* OutReason) const
+{
+	if (HasValidLevelInstanceID())
+	{
+		if (ULevelInstanceSubsystem* LevelInstanceSubsystem = GetLevelInstanceSubsystem())
+		{
+			return LevelInstanceSubsystem->CanDiscardLevelInstance(this, OutReason);
+		}
+	}
+
+	return false;
+}
+
 bool ALevelInstance::IsEditing() const
 {
 	if (HasValidLevelInstanceID())
diff --git a/Engine/Source/Runtime/Engine/Private/LevelInstance/LevelInstanceEditorObject.cpp b/Engine/Source/Runtime/Engine/Private/LevelInstance/LevelInstanceEditorObject.cpp
new file mode 100644
--- /dev/null
+++ b/Engine/Source/Runtime/Engine/Private/LevelInstance/LevelInstanceEditorObject.cpp
@@ -0,0 +1,166 @@
+// Copyright Epic Games, Inc. All Rights Reserved.
+#include "LevelInstance/LevelInstanceEditorObject.h"
+
+#if WITH_EDITOR
+#include "Editor.h"
+#include "EditorLevelUtils.h"
+#include "LevelUtils.h"
+#include "LevelInstance/LevelInstanceEditorLevelStreaming.h"
+#include "Engine/World.h"
+#endif
+
+#define LOCTEXT_NAMESPACE "LevelInstanceSubsystem"
+
+ULevelInstanceEditorObject::ULevelInstanceEditorObject(const FObjectInitializer& ObjectInitializer)
+	: Super(ObjectInitializer)
+#if WITH_EDITORONLY_DATA
+	, bCommittedChanges(false)
+	, bCreatingChildLevelInstance(false)
+	, EditWorld(nullptr)
+#endif
+{
+
+}
+
+#if WITH_EDITOR
+
+bool ULevelInstanceEditorObject::CanDiscard(FText* OutReason /*= nullptr*/) const
+{
+	if (bCommittedChanges)
+	{
+		if (OutReason)
+		{
+			*OutReason = LOCTEXT("CantDiscardCommittedChanges", "Can't discard edit because some packages were saved while editing.");
+		}
+
+		return false;
+	}
+
+	if (bMovedActors)
+	{
+		if(OutReason)
+		{
+			*OutReason = LOCTEXT("CantDiscardMovedActors", "Can't discard edit because some actors were moved from/to level while editing.");
+		}
+
+		return false;
+	}
+	
+	return true;
+}
+
+void ULevelInstanceEditorObject::EnterEdit(UWorld* InEditWorld)
+{
+	check(!bMovedActors);
+	check(!bCommittedChanges);
+	check(!EditWorld);
+
+	EditWorld = InEditWorld;
+	EditorLevelUtils::OnMoveActorsToLevelEvent.AddUObject(this, &ULevelInstanceEditorObject::OnMoveActorsToLevel);
+	FCoreUObjectDelegates::OnObjectPreSave.AddUObject(this, &ULevelInstanceEditorObject::OnObjectPreSave);
+	FEditorDelegates::OnPackageDeleted.AddUObject(this, &ULevelInstanceEditorObject::OnPackageDeleted);
+}
+
+void ULevelInstanceEditorObject::ExitEdit()
+{
+	EditorLevelUtils::OnMoveActorsToLevelEvent.RemoveAll(this);
+	FCoreUObjectDelegates::OnObjectPreSave.RemoveAll(this);
+	FEditorDelegates::OnPackageDeleted.RemoveAll(this);
+
+	EditWorld = nullptr;
+	bMovedActors = false;
+	OtherPackagesToSave.Empty();
+}
+
+void ULevelInstanceEditorObject::OnPackageChanged(UPackage* Package)
+{
+	if (bCommittedChanges)
+	{
+		return;
+	}
+
+	check(EditWorld);
+	if (EditWorld->GetPackage() == Package)
+	{
+		bCommittedChanges = true;
+	}
+	else
+	{
+		TSet<UPackage*> Packages;
+		Packages.Append(EditWorld->GetPackage()->GetExternalPackages());
+		if (Packages.Contains(Package))
+		{
+			bCommittedChanges = true;
+		}
+	}
+}
+
+void ULevelInstanceEditorObject::OnObjectPreSave(UObject* Object, FObjectPreSaveContext SaveContext)
+{
+	if (!SaveContext.IsProceduralSave() && !(SaveContext.GetSaveFlags() & SAVE_FromAutosave))
+	{
+		OnPackageChanged(Object->GetPackage());
+	}
+}
+
+void ULevelInstanceEditorObject::OnPackageDeleted(UPackage* Package)
+{
+	OnPackageChanged(Package);
+}
+
+void ULevelInstanceEditorObject::OnMoveActorsToLevel(const TArray<AActor*>& ActorsToMove, const ULevel* DestLevel)
+{	
+	// When Child Level is being created we don't need to listen to this event because proper packages will be saved by the ULevelInstanceSubsystem code.
+	if (bCreatingChildLevelInstance)
+	{
+		return;
+	}
+
+	auto DoesMoveInvolveEditingLevelInstance = [this](const ULevel* Level)
+	{
+		if (ULevelStreamingLevelInstanceEditor* LevelStreaming = Cast<ULevelStreamingLevelInstanceEditor>(FLevelUtils::FindStreamingLevel(Level)))
+		{
+			return true;
+		}
+
+		return false;
+	};
+
+	bool bNewlyMovedActors = false;
+		
+	if(DoesMoveInvolveEditingLevelInstance(DestLevel))
+	{
+		bNewlyMovedActors = true;
+	}
+
+	if (!bNewlyMovedActors)
+	{
+		for (const AActor* ActorToMove : ActorsToMove)
+		{
+			bNewlyMovedActors = DoesMoveInvolveEditingLevelInstance(ActorToMove->GetLevel());
+			if (bNewlyMovedActors)
+			{
+				break;
+			}
+		}
+	}
+
+	if (bNewlyMovedActors)
+	{
+		Modify();
+		bMovedActors = true;
+
+		for (AActor* ActorToMove : ActorsToMove)
+		{
+			UPackage* PackageToSave = ActorToMove->GetPackage();
+			if (!FPackageName::IsTempPackage(PackageToSave->GetName()))
+			{
+				OtherPackagesToSave.AddUnique(ActorToMove->GetPackage());
+			}
+		}
+	}
+}
+
+#endif // WITH_EDITOR
+
+#undef LOCTEXT_NAMESPACE
\ No newline at end of file
diff --git a/Engine/Source/Runtime/Engine/Private/LevelInstance/LevelInstanceEditorObject.h b/Engine/Source/Runtime/Engine/Private/LevelInstance/LevelInstanceEditorObject.h
new file mode 100644
--- /dev/null
+++ b/Engine/Source/Runtime/Engine/Private/LevelInstance/LevelInstanceEditorObject.h
@@ -0,0 +1,49 @@
+// Copyright Epic Games, Inc. All Rights Reserved.
+#pragma once
+
+#include "CoreMinimal.h"
+#include "UObject/ObjectMacros.h"
+#include "UObject/Object.h"
+#include "UObject/ObjectSaveContext.h"
+#include "LevelInstanceEditorObject.generated.h"
+
+class AActor;
+class ULevel;
+class UWorld;
+class UPackage;
+class FText;
+
+UCLASS()
+class ULevelInstanceEditorObject : public UObject
+{
+	GENERATED_UCLASS_BODY()
+
+#if WITH_EDITORONLY_DATA
+public:
+	bool bCommittedChanges;
+	bool bCreatingChildLevelInstance;
+	TObjectPtr<UWorld> EditWorld;
+	
+	// If current Level edit has moved actors in/out of the level
+	UPROPERTY()
+	bool bMovedActors;
+
+	// List of packages to save when committing level 
+	UPROPERTY()
+	TArray<TWeakObjectPtr<UPackage>> OtherPackagesToSave;
+#endif
+
+#if WITH_EDITOR
+
+	bool CanDiscard(FText* OutReason = nullptr) const;
+
+	void EnterEdit(UWorld* InEditWorld);
+	void ExitEdit();
+
+private:
+	void OnPackageChanged(UPackage* Package);
+	void OnObjectPreSave(UObject* Object, FObjectPreSaveContext SaveContext);
+	void OnPackageDeleted(UPackage* Package);
+	void OnMoveActorsToLevel(const TArray<AActor*>& ActorsToMove, const ULevel* DestLevel);
+#endif
+};
\ No newline at end of file
diff --git a/Engine/Source/Runtime/Engine/Private/LevelInstance/LevelInstanceSubsystem.cpp b/Engine/Source/Runtime/Engine/Private/LevelInstance/LevelInstanceSubsystem.cpp
--- a/Engine/Source/Runtime/Engine/Private/LevelInstance/LevelInstanceSubsystem.cpp
+++ b/Engine/Source/Runtime/Engine/Private/LevelInstance/LevelInstanceSubsystem.cpp
@@ -17,6 +17,9 @@
 #if WITH_EDITOR
 #include "Settings/LevelEditorMiscSettings.h"
 #include "LevelInstance/LevelInstanceEditorLevelStreaming.h"
+#include "LevelInstance/ILevelInstanceEditorModule.h"
+#include "LevelInstance/LevelInstanceEditorInstanceActor.h"
+#include "LevelInstance/LevelInstanceEditorObject.h"
 #include "Misc/ScopedSlowTask.h"
 #include "Misc/ITransaction.h"
 #include "Misc/Paths.h"
@@ -26,11 +29,9 @@
 #include "FileHelpers.h"
 #include "Editor.h"
 #include "EditorLevelUtils.h"
-#include "LevelInstance/ILevelInstanceEditorModule.h"
 #include "HAL/PlatformTime.h"
 #include "Engine/Selection.h"
 #include "Engine/LevelBounds.h"
-#include "LevelInstance/LevelInstanceEditorInstanceActor.h"
 #include "Modules/ModuleManager.h"
 #include "Engine/Blueprint.h"
 #include "PackedLevelActor/PackedLevelActor.h"
@@ -69,12 +70,6 @@ void ULevelInstanceSubsystem::Initialize(FSubsystemCollectionBase& Collection)
 	{
 		ILevelInstanceEditorModule& EditorModule = FModuleManager::LoadModuleChecked<ILevelInstanceEditorModule>("LevelInstanceEditor");
 		EditorModule.OnExitEditorMode().AddUObject(this, &ULevelInstanceSubsystem::OnExitEditorMode);
-
-		if (!GetWorld()->IsGameWorld())
-		{
-			FCoreUObjectDelegates::OnObjectPreSave.AddUObject(this, &ULevelInstanceSubsystem::OnObjectPreSave);
-			FEditorDelegates::OnPackageDeleted.AddUObject(this, &ULevelInstanceSubsystem::OnPackageDeleted);
-		}
 	}
 #endif
 }
@@ -88,12 +83,6 @@ void ULevelInstanceSubsystem::Deinitialize()
 	{
 		EditorModule->OnExitEditorMode().RemoveAll(this);
 	}
-
-	if (!GetWorld()->IsGameWorld())
-	{
-		FCoreUObjectDelegates::OnObjectPreSave.RemoveAll(this);
-		FEditorDelegates::OnPackageDeleted.RemoveAll(this);
-	}
 #endif
 }
 
@@ -392,55 +381,21 @@ void ULevelInstanceSubsystem::Tick()
 	}
 }
 
-void ULevelInstanceSubsystem::OnPackageChanged(UPackage* Package)
-{
-	if (LevelInstanceEdit)
-	{
-		if (UWorld* EditWorld = LevelInstanceEdit->GetEditWorld())
-		{
-			if (EditWorld->GetPackage() == Package)
-			{
-				LevelInstanceEdit->bCommittedChanges = true;
-			}
-			else
-			{
-				TSet<UPackage*> Packages;
-				Packages.Append(EditWorld->GetPackage()->GetExternalPackages());
-				if (Packages.Contains(Package))
-				{
-					LevelInstanceEdit->bCommittedChanges = true;
-				}
-			}
-		}
-	}
-}
-
-void ULevelInstanceSubsystem::OnObjectPreSave(UObject* Object, FObjectPreSaveContext SaveContext)
-{
-	if (!SaveContext.IsProceduralSave() && !(SaveContext.GetSaveFlags() & SAVE_FromAutosave))
-	{
-		OnPackageChanged(Object->GetPackage());
-	}
-}
-
-void ULevelInstanceSubsystem::OnPackageDeleted(UPackage* Package)
-{
-	OnPackageChanged(Package);
-}
-
 void ULevelInstanceSubsystem::OnExitEditorMode()
 {
 	if (LevelInstanceEdit && !bCreatingLevelInstance)
 	{
-		bool bDiscard = true;
+		ALevelInstance* LevelInstance = GetEditingLevelInstance();
+
+		bool bDiscard = false;
 		bool bIsDirty = IsLevelInstanceEditDirty(LevelInstanceEdit.Get());
-		if (bIsDirty)
+		if (bIsDirty && CanDiscardLevelInstance(LevelInstance))
 		{
 			FText Title = LOCTEXT("CommitOrDiscardChangesTitle", "Save changes?");
 			bDiscard = (FMessageDialog::Open(EAppMsgType::YesNo, LOCTEXT("CommitOrDiscardChangesMsg", "Unsaved Level changes will get discarded. Do you want to save them now?"), &Title) == EAppReturnType::No);
 		}
 
-		CommitLevelInstanceInternal(LevelInstanceEdit, bDiscard, /*bPromptForSave=*/false);
+		CommitLevelInstanceInternal(LevelInstanceEdit, bDiscard);
 	}
 }
 
@@ -811,6 +766,19 @@ ALevelInstance* ULevelInstanceSubsystem::CreateLevelInstanceFrom(const TArray<AA
 		LevelFilename = FPackageName::LongPackageNameToFilename(CreationParams.LevelPackageName, FPackageName::GetMapPackageExtension());
 	}
 
+	// Tell current level edit to stop listening because management of packages to save is done here (operation is atomic and can't be undone)
+	if (LevelInstanceEdit)
+	{
+		LevelInstanceEdit->EditorObject->bCreatingChildLevelInstance = true;
+	}
+	ON_SCOPE_EXIT
+	{
+		if (LevelInstanceEdit)
+		{
+			LevelInstanceEdit->EditorObject->bCreatingChildLevelInstance = false;
+		}
+	};
+	
 	TSet<FName> DirtyPackages;
 
 	// Capture Packages before Moving actors as they can get GCed in the process
@@ -916,7 +884,7 @@ ALevelInstance* ULevelInstanceSubsystem::CreateLevelInstanceFrom(const TArray<AA
 	// New level instance
 	TUniquePtr<FLevelInstanceEdit> TempLevelInstanceEdit = MakeUnique<FLevelInstanceEdit>(LevelStreaming, NewLevelInstanceActor->GetLevelInstanceID());
 	// Force mark it as changed
-	TempLevelInstanceEdit->bCommittedChanges = true;
+	TempLevelInstanceEdit->MarkCommittedChanges();
 
 	GetWorld()->SetCurrentLevel(LoadedLevel);
 
@@ -926,7 +894,7 @@ ALevelInstance* ULevelInstanceSubsystem::CreateLevelInstanceFrom(const TArray<AA
 		DirtyPackages.Add(NewLevelInstanceActor->GetPackage()->GetFName());
 	}
 			
-	ALevelInstance* CommittedLevelInstance = CommitLevelInstanceInternal(TempLevelInstanceEdit, /*bDiscardEdits*/false, CreationParams.bPromptForSave, &DirtyPackages);
+	ALevelInstance* CommittedLevelInstance = CommitLevelInstanceInternal(TempLevelInstanceEdit, /*bDiscardEdits*/false, &DirtyPackages);
 	check(!TempLevelInstanceEdit);
 
 	// EditorLevelUtils::CreateNewStreamingLevelForWorld deactivates all modes. Re-activate if needed
@@ -1317,10 +1285,13 @@ ULevelInstanceSubsystem::FLevelInstanceEdit::FLevelInstanceEdit(ULevelStreamingL
 	: LevelStreaming(InLevelStreaming)
 {
 	LevelStreaming->LevelInstanceID = InLevelInstanceID;
+	EditorObject = NewObject<ULevelInstanceEditorObject>(GetTransientPackage(), NAME_None, RF_Transactional);
+	EditorObject->EnterEdit(GetEditWorld());
 }
 
 ULevelInstanceSubsystem::FLevelInstanceEdit::~FLevelInstanceEdit()
 {
+	EditorObject->ExitEdit();
 	ULevelStreamingLevelInstanceEditor::Unload(LevelStreaming);
 }
 
@@ -1334,6 +1305,51 @@ UWorld* ULevelInstanceSubsystem::FLevelInstanceEdit::GetEditWorld() const
 	return nullptr;
 }
 
+void ULevelInstanceSubsystem::FLevelInstanceEdit::AddReferencedObjects(FReferenceCollector& Collector)
+{
+	Collector.AddReferencedObject(EditorObject);
+	Collector.AddReferencedObject(LevelStreaming);
+}
+
+FString ULevelInstanceSubsystem::FLevelInstanceEdit::GetReferencerName() const
+{
+	return TEXT("FLevelInstanceEdit");
+}
+
+bool ULevelInstanceSubsystem::FLevelInstanceEdit::CanDiscard(FText* OutReason) const
+{
+	return EditorObject->CanDiscard(OutReason);
+}
+
+bool ULevelInstanceSubsystem::FLevelInstanceEdit::HasCommittedChanges() const
+{
+	return EditorObject->bCommittedChanges;
+}
+
+void ULevelInstanceSubsystem::FLevelInstanceEdit::MarkCommittedChanges()
+{
+	EditorObject->bCommittedChanges = true;
+}
+
+void ULevelInstanceSubsystem::FLevelInstanceEdit::GetPackagesToSave(TArray<UPackage*>& OutPackagesToSave) const
+{
+	const UWorld* EditingWorld = GetEditWorld();
+	check(EditingWorld);
+
+	FEditorFileUtils::GetDirtyPackages(OutPackagesToSave, [EditingWorld](UPackage* DirtyPackage)
+	{
+		return ULevelInstanceSubsystem::ShouldIgnoreDirtyPackage(DirtyPackage, EditingWorld);
+	});
+
+	for (const TWeakObjectPtr<UPackage>& WeakPackageToSave : EditorObject->OtherPackagesToSave)
+	{
+		if (UPackage* PackageToSave = WeakPackageToSave.Get())
+		{
+			OutPackagesToSave.Add(PackageToSave);
+		}
+	}
+}
+
 FLevelInstanceID ULevelInstanceSubsystem::FLevelInstanceEdit::GetLevelInstanceID() const
 {
 	return LevelStreaming ? LevelStreaming->GetLevelInstanceID() : FLevelInstanceID();
@@ -1362,16 +1378,9 @@ bool ULevelInstanceSubsystem::IsEditingLevelInstanceDirty(const ALevelInstance*
 
 bool ULevelInstanceSubsystem::IsLevelInstanceEditDirty(const FLevelInstanceEdit* InLevelInstanceEdit) const
 {
-	const UWorld* EditingWorld = InLevelInstanceEdit->GetEditWorld();
-	check(EditingWorld);
-
-	TArray<UPackage*> OutDirtyPackages;
-	FEditorFileUtils::GetDirtyPackages(OutDirtyPackages, [EditingWorld](UPackage* DirtyPackage)
-	{
-		return ULevelInstanceSubsystem::ShouldIgnoreDirtyPackage(DirtyPackage, EditingWorld);
-	});
-
-	return OutDirtyPackages.Num() > 0;
+	TArray<UPackage*> OutPackagesToSave;
+	InLevelInstanceEdit->GetPackagesToSave(OutPackagesToSave);
+	return OutPackagesToSave.Num() > 0;
 }
 
 ALevelInstance* ULevelInstanceSubsystem::GetEditingLevelInstance() const
@@ -1469,6 +1478,20 @@ bool ULevelInstanceSubsystem::CanCommitLevelInstance(const ALevelInstance* Level
 	return true;
 }
 
+bool ULevelInstanceSubsystem::CanDiscardLevelInstance(const ALevelInstance* LevelInstanceActor, FText* OutReason) const
+{
+	if (const FLevelInstanceEdit* CurrentEdit = GetLevelInstanceEdit(LevelInstanceActor))
+	{
+		return LevelInstanceEdit->CanDiscard(OutReason);
+	}
+
+	if (OutReason)
+	{
+		*OutReason = LOCTEXT("CanCommitLevelInstanceNotEditing", "Level Instance is not currently being edited");
+	}
+	return false;
+}
+
 void ULevelInstanceSubsystem::EditLevelInstance(ALevelInstance* LevelInstanceActor, TWeakObjectPtr<AActor> ContextActorPtr)
 {
 	if (EditLevelInstanceInternal(LevelInstanceActor, ContextActorPtr, false))
@@ -1587,11 +1610,11 @@ bool ULevelInstanceSubsystem::EditLevelInstanceInternal(ALevelInstance* LevelIns
 	return true;
 }
 
-ALevelInstance* ULevelInstanceSubsystem::CommitLevelInstance(ALevelInstance* LevelInstanceActor, bool bDiscardEdits, bool bPromptForSave, TSet<FName>* DirtyPackages)
+ALevelInstance* ULevelInstanceSubsystem::CommitLevelInstance(ALevelInstance* LevelInstanceActor, bool bDiscardEdits, TSet<FName>* DirtyPackages)
 {
 	check(LevelInstanceEdit.Get() == GetLevelInstanceEdit(LevelInstanceActor));
 	check(CanCommitLevelInstance(LevelInstanceActor));
-	ALevelInstance* CommittedLevelInstance = CommitLevelInstanceInternal(LevelInstanceEdit, bDiscardEdits, bPromptForSave, DirtyPackages);
+	ALevelInstance* CommittedLevelInstance = CommitLevelInstanceInternal(LevelInstanceEdit, bDiscardEdits, DirtyPackages);
 
 	ILevelInstanceEditorModule& EditorModule = FModuleManager::GetModuleChecked<ILevelInstanceEditorModule>("LevelInstanceEditor");
 	EditorModule.DeactivateEditorMode();
@@ -1599,16 +1622,40 @@ ALevelInstance* ULevelInstanceSubsystem::CommitLevelInstance(ALevelInstance* Lev
 	return CommittedLevelInstance;
 }
 
-ALevelInstance* ULevelInstanceSubsystem::CommitLevelInstanceInternal(TUniquePtr<FLevelInstanceEdit>& InLevelInstanceEdit, bool bDiscardEdits, bool bPromptForSave, TSet<FName>* DirtyPackages)
+ALevelInstance* ULevelInstanceSubsystem::CommitLevelInstanceInternal(TUniquePtr<FLevelInstanceEdit>& InLevelInstanceEdit, bool bDiscardEdits, TSet<FName>* DirtyPackages)
 {
 	ALevelInstance* LevelInstanceActor = GetLevelInstance(InLevelInstanceEdit->GetLevelInstanceID());
 	check(InLevelInstanceEdit);
 	UWorld* EditingWorld = InLevelInstanceEdit->GetEditWorld();
 	check(EditingWorld);
 			
-	if (IsLevelInstanceEditDirty(InLevelInstanceEdit.Get()) && !bDiscardEdits)
+	// Check with EditorObject if Discard is possible
+	if (!InLevelInstanceEdit->CanDiscard())
+	{
+		bDiscardEdits = false;
+	}
+
+	// Build list of Packages to save
+	TSet<FName> PackagesToSave;
+
+	// First dirty packages belonging to the edit Level or external level actors that were moved into the level
+	TArray<UPackage*> EditPackagesToSave;
+	InLevelInstanceEdit->GetPackagesToSave(EditPackagesToSave);
+	for (UPackage* PackageToSave : EditPackagesToSave)
+	{
+		PackagesToSave.Add(PackageToSave->GetFName());
+	}
+
+	// Second dirty packages passed in to the Commit method
+	if (DirtyPackages)
+	{
+		PackagesToSave.Append(*DirtyPackages);
+	}
+		
+	bool bChangesCommitted = InLevelInstanceEdit->HasCommittedChanges();
+	if (PackagesToSave.Num() > 0 && !bDiscardEdits)
 	{
-		const bool bPromptUserToSave = bPromptForSave;
+		const bool bPromptUserToSave = false;
 		const bool bSaveMapPackages = true;
 		const bool bSaveContentPackages = true;
 		const bool bFastSave = false;
@@ -1616,9 +1663,9 @@ ALevelInstance* ULevelInstanceSubsystem::CommitLevelInstanceInternal(TUniquePtr<
 		const bool bCanBeDeclined = true;
 
 		if (!FEditorFileUtils::SaveDirtyPackages(bPromptUserToSave, bSaveMapPackages, bSaveContentPackages, bFastSave, bNotifyNoPackagesSaved, bCanBeDeclined, nullptr,
-			[=](UPackage* DirtyPackage)
+			[&PackagesToSave,EditingWorld](UPackage* DirtyPackage)
 			{
-				if (DirtyPackages && DirtyPackages->Contains(DirtyPackage->GetFName()))
+				if (PackagesToSave.Contains(DirtyPackage->GetFName()))
 				{
 					return false;
 				}
@@ -1627,6 +1674,7 @@ ALevelInstance* ULevelInstanceSubsystem::CommitLevelInstanceInternal(TUniquePtr<
 		{
 			return LevelInstanceActor;
 		}
+		bChangesCommitted = true;
 	}
 
 	FScopedSlowTask SlowTask(0, LOCTEXT("EndEditLevelInstance", "Unloading Level..."), !GetWorld()->IsGameWorld());
@@ -1637,7 +1685,6 @@ ALevelInstance* ULevelInstanceSubsystem::CommitLevelInstanceInternal(TUniquePtr<
 	const FString EditPackage = LevelInstanceActor->GetWorldAssetPackage();
 
 	// Remove from streaming level...
-	const bool bChangesCommitted = InLevelInstanceEdit->bCommittedChanges;
 	InLevelInstanceEdit.Reset();
 
 	if (bChangesCommitted)
@@ -1649,7 +1696,7 @@ ALevelInstance* ULevelInstanceSubsystem::CommitLevelInstanceInternal(TUniquePtr<
 	const FLevelInstanceID LevelInstanceID = LevelInstanceActor->GetLevelInstanceID();
 
 	// Notify (Actor might get destroyed by this call if its a packed bp)
-	LevelInstanceActor->OnCommit(bChangesCommitted, bPromptForSave);
+	LevelInstanceActor->OnCommit(bChangesCommitted);
 
 	// Update pointer since BP Compilation might have invalidated LevelInstanceActor
 	LevelInstanceActor = GetLevelInstance(LevelInstanceID);
diff --git a/Engine/Source/Runtime/Engine/Private/PackedLevelActor/PackedLevelActor.cpp b/Engine/Source/Runtime/Engine/Private/PackedLevelActor/PackedLevelActor.cpp
--- a/Engine/Source/Runtime/Engine/Private/PackedLevelActor/PackedLevelActor.cpp
+++ b/Engine/Source/Runtime/Engine/Private/PackedLevelActor/PackedLevelActor.cpp
@@ -170,18 +170,17 @@ void APackedLevelActor::OnEdit()
 	MarkComponentsRenderStateDirty();
 }
 
-void APackedLevelActor::OnCommit(bool bChanged, bool bPromptForSave)
+void APackedLevelActor::OnCommit(bool bChanged)
 {
-	Super::OnCommit(bChanged, bPromptForSave);
+	Super::OnCommit(bChanged);
 
 	if (bChanged)
 	{
 		if (UBlueprint* GeneratedBy = Cast<UBlueprint>(GetClass()->ClassGeneratedBy))
 		{
 			check(GeneratedBy == BlueprintAsset.Get());
-			const bool bCheckoutAndSave = true;
 			TSharedPtr<FPackedLevelActorBuilder> Builder = FPackedLevelActorBuilder::CreateDefaultBuilder();
-			Builder->UpdateBlueprint(GeneratedBy, bCheckoutAndSave, bPromptForSave);
+			Builder->UpdateBlueprint(GeneratedBy);
 			return; // return here because Actor might have been reinstanced
 		}
 	}
diff --git a/Engine/Source/Runtime/Engine/Private/PackedLevelActor/PackedLevelActorBuilder.cpp b/Engine/Source/Runtime/Engine/Private/PackedLevelActor/PackedLevelActorBuilder.cpp
--- a/Engine/Source/Runtime/Engine/Private/PackedLevelActor/PackedLevelActorBuilder.cpp
+++ b/Engine/Source/Runtime/Engine/Private/PackedLevelActor/PackedLevelActorBuilder.cpp
@@ -326,12 +326,12 @@ bool FPackedLevelActorBuilder::PackActor(APackedLevelActor* InActor, TSoftObject
 	return PackActor(InActor, TransientLevelInstance);
 }
 
-void FPackedLevelActorBuilder::UpdateBlueprint(UBlueprint* Blueprint, bool bCheckoutAndSave, bool bPromptForSave)
+void FPackedLevelActorBuilder::UpdateBlueprint(UBlueprint* Blueprint, bool bCheckoutAndSave)
 {
 	APackedLevelActor* CDO = CastChecked<APackedLevelActor>(Blueprint->GeneratedClass->GetDefaultObject());
 	check(CDO);
 
-	CreateOrUpdateBlueprint(CDO->GetWorldAsset(), Blueprint, bCheckoutAndSave, bPromptForSave);
+	CreateOrUpdateBlueprint(CDO->GetWorldAsset(), Blueprint, bCheckoutAndSave, /*bPromptForSave=*/false);
 }
 
 bool FPackedLevelActorBuilder::CreateOrUpdateBlueprint(TSoftObjectPtr<UWorld> InWorldAsset, TSoftObjectPtr<UBlueprint> InBlueprintAsset, bool bCheckoutAndSave, bool bPromptForSave)
diff --git a/Engine/Source/Runtime/Engine/Private/PackedLevelActor/PackedLevelActorBuilder.h b/Engine/Source/Runtime/Engine/Private/PackedLevelActor/PackedLevelActorBuilder.h
--- a/Engine/Source/Runtime/Engine/Private/PackedLevelActor/PackedLevelActorBuilder.h
+++ b/Engine/Source/Runtime/Engine/Private/PackedLevelActor/PackedLevelActorBuilder.h
@@ -40,7 +40,7 @@ public:
 	/* Creates/Updates a APackedLeveInstance Blueprint from InWorldAsset (will overwrite existing asset or show a dialog if InBlueprintAsset is null) */
 	bool CreateOrUpdateBlueprint(TSoftObjectPtr<UWorld> InWorldAsset, TSoftObjectPtr<UBlueprint> InBlueprintAsset, bool bCheckoutAndSave = true, bool bPromptForSave = true);
 	/* Update existing Blueprint */
-	void UpdateBlueprint(UBlueprint* Blueprint, bool bCheckoutAndSave = true, bool bPromptForSave = true);
+	void UpdateBlueprint(UBlueprint* Blueprint, bool bCheckoutAndSave = true);
 
 	static const FString& GetPackedBPPrefix();
 	/* Creates a new APackedLevelActor Blueprint using InPackagePath/InAssetName as hint for path. Prompts the user to input the final asset name. */
diff --git a/Engine/Source/Runtime/Engine/Public/LevelInstance/LevelInstanceActor.h b/Engine/Source/Runtime/Engine/Public/LevelInstance/LevelInstanceActor.h
--- a/Engine/Source/Runtime/Engine/Public/LevelInstance/LevelInstanceActor.h
+++ b/Engine/Source/Runtime/Engine/Public/LevelInstance/LevelInstanceActor.h
@@ -88,6 +88,7 @@ public:
 
 	bool CanEdit(FText* OutReason = nullptr) const;
 	bool CanCommit(FText* OutReason = nullptr) const;
+	bool CanDiscard(FText* OutReason = nullptr) const;
 	bool IsEditing() const;
 	bool HasChildEdit() const;
 	void Edit(AActor* ContextActor = nullptr);
@@ -108,7 +109,7 @@ public:
 	virtual void OnWorldAssetChanged() { UpdateFromLevel(); }
 	virtual void OnEdit() {}
 	virtual void OnEditChild() {}
-	virtual void OnCommit(bool bChanged, bool bPromptForSave) {}
+	virtual void OnCommit(bool bChanged) {}
 	virtual void OnCommitChild(bool bChanged) {}
 	virtual void UpdateFromLevel();
 		
diff --git a/Engine/Source/Runtime/Engine/Public/LevelInstance/LevelInstanceSubsystem.h b/Engine/Source/Runtime/Engine/Public/LevelInstance/LevelInstanceSubsystem.h
--- a/Engine/Source/Runtime/Engine/Public/LevelInstance/LevelInstanceSubsystem.h
+++ b/Engine/Source/Runtime/Engine/Public/LevelInstance/LevelInstanceSubsystem.h
@@ -5,12 +5,14 @@
 #include "Containers/Array.h"
 #include "Containers/Map.h"
 #include "UObject/NameTypes.h"
+#include "UObject/GCObject.h"
 #include "Folder.h"
 #include "LevelInstance/LevelInstanceTypes.h"
 
 #include "LevelInstanceSubsystem.generated.h"
 
 class ALevelInstance;
+class ULevelInstanceEditorObject;
 class ULevelStreamingLevelInstance;
 class ULevelStreamingLevelInstanceEditor;
 class UWorldPartitionSubsystem;
@@ -48,8 +50,6 @@ public:
 
 #if WITH_EDITOR
 	void Tick();
-	void OnObjectPreSave(UObject* Object, FObjectPreSaveContext SaveContext);
-	void OnPackageDeleted(UPackage* Package);
 	void OnExitEditorMode();
 	void PackAllLoadedActors();
 	bool CanPackAllLoadedActors() const;
@@ -57,8 +57,9 @@ public:
 	ALevelInstance* GetEditingLevelInstance() const;
 	bool CanEditLevelInstance(const ALevelInstance* LevelInstanceActor, FText* OutReason = nullptr) const;
 	bool CanCommitLevelInstance(const ALevelInstance* LevelInstanceActor, FText* OutReason = nullptr) const;
+	bool CanDiscardLevelInstance(const ALevelInstance* LevelInstanceActor, FText* OutReason = nullptr) const;
 	void EditLevelInstance(ALevelInstance* LevelInstanceActor, TWeakObjectPtr<AActor> ContextActorPtr = nullptr);
-	ALevelInstance* CommitLevelInstance(ALevelInstance* LevelInstanceActor, bool bDiscardEdits = false, bool bPromptForSave = true, TSet<FName>* DirtyPackages = nullptr);
+	ALevelInstance* CommitLevelInstance(ALevelInstance* LevelInstanceActor, bool bDiscardEdits = false, TSet<FName>* DirtyPackages = nullptr);
 	void SaveLevelInstanceAs(ALevelInstance* LevelInstanceActor);
 	bool IsEditingLevelInstanceDirty(const ALevelInstance* LevelInstanceActor) const;
 	bool IsEditingLevelInstance(const ALevelInstance* LevelInstanceActor) const { return GetLevelInstanceEdit(LevelInstanceActor) != nullptr; }
@@ -114,22 +115,30 @@ private:
 	void BreakLevelInstance_Impl(ALevelInstance* LevelInstanceActor, uint32 Levels, TArray<AActor*>& OutMovedActors);
 
 	static bool ShouldIgnoreDirtyPackage(UPackage* DirtyPackage, const UWorld* EditingWorld);
-	void OnPackageChanged(UPackage* Package);
 
-	struct FLevelInstanceEdit
+	class FLevelInstanceEdit : public FGCObject
 	{
-		ULevelStreamingLevelInstanceEditor* LevelStreaming;
-		bool bCommittedChanges = false;
+	public:
+		TObjectPtr<ULevelStreamingLevelInstanceEditor> LevelStreaming;
+		TObjectPtr<ULevelInstanceEditorObject> EditorObject;
 
 		FLevelInstanceEdit(ULevelStreamingLevelInstanceEditor* InLevelStreaming, FLevelInstanceID InLevelInstanceID);
-		~FLevelInstanceEdit();
+		virtual ~FLevelInstanceEdit();
 
 		UWorld* GetEditWorld() const;
 		FLevelInstanceID GetLevelInstanceID() const;
+
+		virtual void AddReferencedObjects(FReferenceCollector& Collector) override;
+		virtual FString GetReferencerName() const override;
+
+		void GetPackagesToSave(TArray<UPackage*>& OutPackagesToSave) const;
+		bool CanDiscard(FText* OutReason = nullptr) const;
+		bool HasCommittedChanges() const;
+		void MarkCommittedChanges();
 	};
 
 	bool EditLevelInstanceInternal(ALevelInstance* LevelInstanceActor, TWeakObjectPtr<AActor> ContextActorPtr, bool bRecursive);
-	ALevelInstance* CommitLevelInstanceInternal(TUniquePtr<FLevelInstanceEdit>& InLevelInstanceEdit, bool bDiscardEdits = false, bool bPromptForSave = true, TSet<FName>* DirtyPackages = nullptr);
+	ALevelInstance* CommitLevelInstanceInternal(TUniquePtr<FLevelInstanceEdit>& InLevelInstanceEdit, bool bDiscardEdits = false, TSet<FName>* DirtyPackages = nullptr);
 	
 	const FLevelInstanceEdit* GetLevelInstanceEdit(const ALevelInstance* LevelInstanceActor) const;
 	bool IsLevelInstanceEditDirty(const FLevelInstanceEdit* LevelInstanceEdit) const;
diff --git a/Engine/Source/Runtime/Engine/Public/PackedLevelActor/PackedLevelActor.h b/Engine/Source/Runtime/Engine/Public/PackedLevelActor/PackedLevelActor.h
--- a/Engine/Source/Runtime/Engine/Public/PackedLevelActor/PackedLevelActor.h
+++ b/Engine/Source/Runtime/Engine/Public/PackedLevelActor/PackedLevelActor.h
@@ -40,7 +40,7 @@ public:
 
 	virtual void PostLoad() override;
 	virtual void UpdateFromLevel() override;
-	virtual void OnCommit(bool bChanged, bool bPromptForSave) override;
+	virtual void OnCommit(bool bChanged) override;
 	virtual void OnCommitChild(bool bChanged) override;
 	virtual void OnEdit() override;
 	virtual void OnEditChild() override;
