WorldPartiton: Add flag to distinguish generated runtime cell levels
- Flag is set in FWorldPartitionLevelHelper::CreateEmptyLevelForRuntimeCell
- Use flag in AInstancedFoliageActor::PostLoad to avoid invalid warning

diff --git a/Engine/Source/Runtime/Engine/Classes/Engine/Level.h b/Engine/Source/Runtime/Engine/Classes/Engine/Level.h
--- a/Engine/Source/Runtime/Engine/Classes/Engine/Level.h
+++ b/Engine/Source/Runtime/Engine/Classes/Engine/Level.h
@@ -658,6 +658,9 @@ public:
 	/** Whether the level is partitioned or not. */
     UPROPERTY()
 	uint8										bIsPartitioned : 1;
+	/** Whether the level is a world partition runtime cell. */
+	UPROPERTY()
+	uint8										bIsWorldPartitionRuntimeCell;
 
 	enum class EIncrementalComponentState : uint8
 	{
diff --git a/Engine/Source/Runtime/Engine/Private/WorldPartition/WorldPartitionLevelHelper.cpp b/Engine/Source/Runtime/Engine/Private/WorldPartition/WorldPartitionLevelHelper.cpp
--- a/Engine/Source/Runtime/Engine/Private/WorldPartition/WorldPartitionLevelHelper.cpp
+++ b/Engine/Source/Runtime/Engine/Private/WorldPartition/WorldPartitionLevelHelper.cpp
@@ -203,6 +203,9 @@ ULevel* FWorldPartitionLevelHelper::CreateEmptyLevelForRuntimeCell(const UWorld*
 	check(NewLevel->Model);
 	check(!NewLevel->bIsVisible);
 
+	// Mark the level as a runtime cell
+	NewLevel->bIsWorldPartitionRuntimeCell = true;
+
 	// Set the guids on the constructed level to something based on the generator rather than allowing indeterminism by
 	// constructing new Guids on every cook
 	// @todo_ow: revisit for static lighting support. We need to base the LevelBuildDataId on the relevant information from the
diff --git a/Engine/Source/Runtime/Foliage/Private/InstancedFoliage.cpp b/Engine/Source/Runtime/Foliage/Private/InstancedFoliage.cpp
--- a/Engine/Source/Runtime/Foliage/Private/InstancedFoliage.cpp
+++ b/Engine/Source/Runtime/Foliage/Private/InstancedFoliage.cpp
@@ -4303,7 +4303,8 @@ void AInstancedFoliageActor::PostLoad()
 	ULevel* OwningLevel = GetLevel();
 	// We can't check the ActorPartitionSubsystem here because World is not initialized yet. So we fallback on the bIsPartitioned
 	// to know if multiple InstanceFoliageActors is valid or not.
-	if (OwningLevel && !OwningLevel->bIsPartitioned)
+	// For levels that are World Partition runtime cells. Having multiple IFAs is valid so skip validation.
+	if (OwningLevel && !OwningLevel->bIsPartitioned && !OwningLevel->bIsWorldPartitionRuntimeCell)
 	{
 		if (!OwningLevel->InstancedFoliageActor.IsValid())
 		{
