Update error handling in new HLSL translator, errors now propogate from HLSLTree back to the correct UMaterialExpression
Fix problem with preshader constants potentially saving the wrong number of components

diff --git a/Engine/Source/Runtime/Engine/Classes/Materials/MaterialExpression.h b/Engine/Source/Runtime/Engine/Classes/Materials/MaterialExpression.h
--- a/Engine/Source/Runtime/Engine/Classes/Materials/MaterialExpression.h
+++ b/Engine/Source/Runtime/Engine/Classes/Materials/MaterialExpression.h
@@ -143,12 +143,6 @@ struct FExpressionExecOutputEntry
 	FExpressionExecOutput* Output = nullptr;
 };
 
-enum class EMaterialGenerateHLSLStatus : uint8
-{
-	Success,
-	Error,
-};
-
 enum class EMaterialExpressionSetParameterValueFlags : uint32
 {
 	None = 0u,
@@ -310,9 +304,9 @@ class ENGINE_API UMaterialExpression : public UObject
 	 * For example, a for-loop expression might generate a statement for the execution input, but also generate an expression to access the loop index
 	 * These methods replace the Compile() method; once we switch over to the new system, Compile() will be removed
 	 */
-	virtual EMaterialGenerateHLSLStatus GenerateHLSLStatements(FMaterialHLSLGenerator& Generator, UE::HLSLTree::FScope& Scope);
-	virtual EMaterialGenerateHLSLStatus GenerateHLSLExpression(FMaterialHLSLGenerator& Generator, UE::HLSLTree::FScope& Scope, int32 OutputIndex, UE::HLSLTree::FExpression* &OutExpression);
-	virtual EMaterialGenerateHLSLStatus GenerateHLSLTexture(FMaterialHLSLGenerator& Generator, UE::HLSLTree::FScope& Scope, int32 OutputIndex, UE::HLSLTree::FTextureParameterDeclaration*& OutTexture);
+	virtual bool GenerateHLSLStatements(FMaterialHLSLGenerator& Generator, UE::HLSLTree::FScope& Scope);
+	virtual bool GenerateHLSLExpression(FMaterialHLSLGenerator& Generator, UE::HLSLTree::FScope& Scope, int32 OutputIndex, UE::HLSLTree::FExpression* &OutExpression);
+	virtual bool GenerateHLSLTexture(FMaterialHLSLGenerator& Generator, UE::HLSLTree::FScope& Scope, int32 OutputIndex, UE::HLSLTree::FTextureParameterDeclaration*& OutTexture);
 
 #endif // WITH_EDITOR
 
diff --git a/Engine/Source/Runtime/Engine/Classes/Materials/MaterialExpressionAdd.h b/Engine/Source/Runtime/Engine/Classes/Materials/MaterialExpressionAdd.h
--- a/Engine/Source/Runtime/Engine/Classes/Materials/MaterialExpressionAdd.h
+++ b/Engine/Source/Runtime/Engine/Classes/Materials/MaterialExpressionAdd.h
@@ -35,7 +35,7 @@ class UMaterialExpressionAdd : public UMaterialExpression
 	virtual void GetCaption(TArray<FString>& OutCaptions) const override;
 	virtual FText GetKeywords() const override {return FText::FromString(TEXT("+"));}
 
-	virtual EMaterialGenerateHLSLStatus GenerateHLSLExpression(FMaterialHLSLGenerator& Generator, UE::HLSLTree::FScope& Scope, int32 OutputIndex, UE::HLSLTree::FExpression*& OutExpression) override;
+	virtual bool GenerateHLSLExpression(FMaterialHLSLGenerator& Generator, UE::HLSLTree::FScope& Scope, int32 OutputIndex, UE::HLSLTree::FExpression*& OutExpression) override;
 #endif // WITH_EDITOR
 	//~ End UMaterialExpression Interface
 
diff --git a/Engine/Source/Runtime/Engine/Classes/Materials/MaterialExpressionAppendVector.h b/Engine/Source/Runtime/Engine/Classes/Materials/MaterialExpressionAppendVector.h
--- a/Engine/Source/Runtime/Engine/Classes/Materials/MaterialExpressionAppendVector.h
+++ b/Engine/Source/Runtime/Engine/Classes/Materials/MaterialExpressionAppendVector.h
@@ -25,7 +25,7 @@ class UMaterialExpressionAppendVector : public UMaterialExpression
 #if WITH_EDITOR
 	virtual int32 Compile(class FMaterialCompiler* Compiler, int32 OutputIndex) override;
 	virtual void GetCaption(TArray<FString>& OutCaptions) const override;
-	virtual EMaterialGenerateHLSLStatus GenerateHLSLExpression(FMaterialHLSLGenerator& Generator, UE::HLSLTree::FScope& Scope, int32 OutputIndex, UE::HLSLTree::FExpression*& OutExpression) override;
+	virtual bool GenerateHLSLExpression(FMaterialHLSLGenerator& Generator, UE::HLSLTree::FScope& Scope, int32 OutputIndex, UE::HLSLTree::FExpression*& OutExpression) override;
 #endif
 	//~ End UMaterialExpression Interface
 
diff --git a/Engine/Source/Runtime/Engine/Classes/Materials/MaterialExpressionBinaryOp.h b/Engine/Source/Runtime/Engine/Classes/Materials/MaterialExpressionBinaryOp.h
--- a/Engine/Source/Runtime/Engine/Classes/Materials/MaterialExpressionBinaryOp.h
+++ b/Engine/Source/Runtime/Engine/Classes/Materials/MaterialExpressionBinaryOp.h
@@ -35,7 +35,7 @@ class UMaterialExpressionBinaryOp : public UMaterialExpression
 #if WITH_EDITOR
 	virtual FText GetKeywords() const override;
 	virtual void GetCaption(TArray<FString>& OutCaptions) const override;
-	virtual EMaterialGenerateHLSLStatus GenerateHLSLExpression(FMaterialHLSLGenerator& Generator, UE::HLSLTree::FScope& Scope, int32 OutputIndex, UE::HLSLTree::FExpression*& OutExpression) override;
+	virtual bool GenerateHLSLExpression(FMaterialHLSLGenerator& Generator, UE::HLSLTree::FScope& Scope, int32 OutputIndex, UE::HLSLTree::FExpression*& OutExpression) override;
 	virtual int32 Compile(class FMaterialCompiler* Compiler, int32 OutputIndex) override;
 #endif // WITH_EDITOR
 	//~ End UMaterialExpression Interface
diff --git a/Engine/Source/Runtime/Engine/Classes/Materials/MaterialExpressionComponentMask.h b/Engine/Source/Runtime/Engine/Classes/Materials/MaterialExpressionComponentMask.h
--- a/Engine/Source/Runtime/Engine/Classes/Materials/MaterialExpressionComponentMask.h
+++ b/Engine/Source/Runtime/Engine/Classes/Materials/MaterialExpressionComponentMask.h
@@ -35,7 +35,7 @@ class UMaterialExpressionComponentMask : public UMaterialExpression
 	virtual int32 Compile(class FMaterialCompiler* Compiler, int32 OutputIndex) override;
 	virtual void GetCaption(TArray<FString>& OutCaptions) const override;
 	virtual FText GetKeywords() const override { return FText::FromString(TEXT("component mask")); }
-	virtual EMaterialGenerateHLSLStatus GenerateHLSLExpression(FMaterialHLSLGenerator& Generator, UE::HLSLTree::FScope& Scope, int32 OutputIndex, UE::HLSLTree::FExpression*& OutExpression) override;
+	virtual bool GenerateHLSLExpression(FMaterialHLSLGenerator& Generator, UE::HLSLTree::FScope& Scope, int32 OutputIndex, UE::HLSLTree::FExpression*& OutExpression) override;
 #endif
 	//~ End UMaterialExpression Interface
 };
diff --git a/Engine/Source/Runtime/Engine/Classes/Materials/MaterialExpressionConstant.h b/Engine/Source/Runtime/Engine/Classes/Materials/MaterialExpressionConstant.h
--- a/Engine/Source/Runtime/Engine/Classes/Materials/MaterialExpressionConstant.h
+++ b/Engine/Source/Runtime/Engine/Classes/Materials/MaterialExpressionConstant.h
@@ -22,7 +22,7 @@ class UMaterialExpressionConstant : public UMaterialExpression
 	virtual int32 Compile(class FMaterialCompiler* Compiler, int32 OutputIndex) override;
 	virtual void GetCaption(TArray<FString>& OutCaptions) const override;
 	virtual FString GetDescription() const override;
-	virtual EMaterialGenerateHLSLStatus GenerateHLSLExpression(FMaterialHLSLGenerator& Generator, UE::HLSLTree::FScope& Scope, int32 OutputIndex, UE::HLSLTree::FExpression*& OutExpression) override;
+	virtual bool GenerateHLSLExpression(FMaterialHLSLGenerator& Generator, UE::HLSLTree::FScope& Scope, int32 OutputIndex, UE::HLSLTree::FExpression*& OutExpression) override;
 #endif // WITH_EDITOR
 
 	//~ End UMaterialExpression Interface
diff --git a/Engine/Source/Runtime/Engine/Classes/Materials/MaterialExpressionConstant2Vector.h b/Engine/Source/Runtime/Engine/Classes/Materials/MaterialExpressionConstant2Vector.h
--- a/Engine/Source/Runtime/Engine/Classes/Materials/MaterialExpressionConstant2Vector.h
+++ b/Engine/Source/Runtime/Engine/Classes/Materials/MaterialExpressionConstant2Vector.h
@@ -26,7 +26,7 @@ class UMaterialExpressionConstant2Vector : public UMaterialExpression
 	virtual void GetCaption(TArray<FString>& OutCaptions) const override;
 	virtual FString GetDescription() const override;
 	virtual uint32 GetOutputType(int32 OutputIndex) override {return MCT_Float2;}
-	virtual EMaterialGenerateHLSLStatus GenerateHLSLExpression(FMaterialHLSLGenerator& Generator, UE::HLSLTree::FScope& Scope, int32 OutputIndex, UE::HLSLTree::FExpression*& OutExpression) override;
+	virtual bool GenerateHLSLExpression(FMaterialHLSLGenerator& Generator, UE::HLSLTree::FScope& Scope, int32 OutputIndex, UE::HLSLTree::FExpression*& OutExpression) override;
 #endif // WITH_EDITOR
 	//~ End UMaterialExpression Interface
 };
diff --git a/Engine/Source/Runtime/Engine/Classes/Materials/MaterialExpressionConstant3Vector.h b/Engine/Source/Runtime/Engine/Classes/Materials/MaterialExpressionConstant3Vector.h
--- a/Engine/Source/Runtime/Engine/Classes/Materials/MaterialExpressionConstant3Vector.h
+++ b/Engine/Source/Runtime/Engine/Classes/Materials/MaterialExpressionConstant3Vector.h
@@ -25,6 +25,6 @@ public:
 
 	virtual FString GetDescription() const override;
 	virtual uint32 GetOutputType(int32 OutputIndex) override {return MCT_Float3;}
-	virtual EMaterialGenerateHLSLStatus GenerateHLSLExpression(FMaterialHLSLGenerator& Generator, UE::HLSLTree::FScope& Scope, int32 OutputIndex, UE::HLSLTree::FExpression*& OutExpression) override;
+	virtual bool GenerateHLSLExpression(FMaterialHLSLGenerator& Generator, UE::HLSLTree::FScope& Scope, int32 OutputIndex, UE::HLSLTree::FExpression*& OutExpression) override;
 #endif // WITH_EDITOR
 };
diff --git a/Engine/Source/Runtime/Engine/Classes/Materials/MaterialExpressionConstant4Vector.h b/Engine/Source/Runtime/Engine/Classes/Materials/MaterialExpressionConstant4Vector.h
--- a/Engine/Source/Runtime/Engine/Classes/Materials/MaterialExpressionConstant4Vector.h
+++ b/Engine/Source/Runtime/Engine/Classes/Materials/MaterialExpressionConstant4Vector.h
@@ -22,7 +22,7 @@ class UMaterialExpressionConstant4Vector : public UMaterialExpression
 	virtual void GetCaption(TArray<FString>& OutCaptions) const override;
 	virtual FString GetDescription() const override;
 	virtual uint32 GetOutputType(int32 OutputIndex) override {return MCT_Float4;}
-	virtual EMaterialGenerateHLSLStatus GenerateHLSLExpression(FMaterialHLSLGenerator& Generator, UE::HLSLTree::FScope& Scope, int32 OutputIndex, UE::HLSLTree::FExpression*& OutExpression) override;
+	virtual bool GenerateHLSLExpression(FMaterialHLSLGenerator& Generator, UE::HLSLTree::FScope& Scope, int32 OutputIndex, UE::HLSLTree::FExpression*& OutExpression) override;
 #endif // WITH_EDITOR
 	//~ End UMaterialExpression Interface
 };
diff --git a/Engine/Source/Runtime/Engine/Classes/Materials/MaterialExpressionDivide.h b/Engine/Source/Runtime/Engine/Classes/Materials/MaterialExpressionDivide.h
--- a/Engine/Source/Runtime/Engine/Classes/Materials/MaterialExpressionDivide.h
+++ b/Engine/Source/Runtime/Engine/Classes/Materials/MaterialExpressionDivide.h
@@ -32,7 +32,7 @@ class UMaterialExpressionDivide : public UMaterialExpression
 	//~ Begin UMaterialExpression Interface
 #if WITH_EDITOR
 	virtual int32 Compile(class FMaterialCompiler* Compiler, int32 OutputIndex) override;
-	virtual EMaterialGenerateHLSLStatus GenerateHLSLExpression(FMaterialHLSLGenerator& Generator, UE::HLSLTree::FScope& Scope, int32 OutputIndex, UE::HLSLTree::FExpression*& OutExpression) override;
+	virtual bool GenerateHLSLExpression(FMaterialHLSLGenerator& Generator, UE::HLSLTree::FScope& Scope, int32 OutputIndex, UE::HLSLTree::FExpression*& OutExpression) override;
 	virtual void GetCaption(TArray<FString>& OutCaptions) const override;
 	virtual FText GetKeywords() const override {return FText::FromString(TEXT("/"));}
 #endif // WITH_EDITOR
diff --git a/Engine/Source/Runtime/Engine/Classes/Materials/MaterialExpressionDoubleVectorParameter.h b/Engine/Source/Runtime/Engine/Classes/Materials/MaterialExpressionDoubleVectorParameter.h
--- a/Engine/Source/Runtime/Engine/Classes/Materials/MaterialExpressionDoubleVectorParameter.h
+++ b/Engine/Source/Runtime/Engine/Classes/Materials/MaterialExpressionDoubleVectorParameter.h
@@ -19,7 +19,7 @@ class UMaterialExpressionDoubleVectorParameter : public UMaterialExpressionParam
 
 	//~ Begin UMaterialExpression Interface
 #if WITH_EDITOR
-	EMaterialGenerateHLSLStatus GenerateHLSLExpression(FMaterialHLSLGenerator& Generator, UE::HLSLTree::FScope& Scope, int32 OutputIndex, UE::HLSLTree::FExpression*& OutExpression);
+	bool GenerateHLSLExpression(FMaterialHLSLGenerator& Generator, UE::HLSLTree::FScope& Scope, int32 OutputIndex, UE::HLSLTree::FExpression*& OutExpression);
 	virtual int32 Compile(class FMaterialCompiler* Compiler, int32 OutputIndex) override;
 	virtual void GetCaption(TArray<FString>& OutCaptions) const override;
 	virtual bool GetParameterValue(FMaterialParameterMetadata& OutMeta) const override
diff --git a/Engine/Source/Runtime/Engine/Classes/Materials/MaterialExpressionExecBegin.h b/Engine/Source/Runtime/Engine/Classes/Materials/MaterialExpressionExecBegin.h
--- a/Engine/Source/Runtime/Engine/Classes/Materials/MaterialExpressionExecBegin.h
+++ b/Engine/Source/Runtime/Engine/Classes/Materials/MaterialExpressionExecBegin.h
@@ -23,7 +23,7 @@ class UMaterialExpressionExecBegin : public UMaterialExpression
 	virtual void GetCaption(TArray<FString>& OutCaptions) const override;
 	virtual bool CanUserDeleteExpression() const override { return false; }
 
-	virtual EMaterialGenerateHLSLStatus GenerateHLSLStatements(FMaterialHLSLGenerator& Generator, UE::HLSLTree::FScope& Scope) override;
+	virtual bool GenerateHLSLStatements(FMaterialHLSLGenerator& Generator, UE::HLSLTree::FScope& Scope) override;
 #endif
 	//~ End UMaterialExpression Interface
 };
diff --git a/Engine/Source/Runtime/Engine/Classes/Materials/MaterialExpressionExecEnd.h b/Engine/Source/Runtime/Engine/Classes/Materials/MaterialExpressionExecEnd.h
--- a/Engine/Source/Runtime/Engine/Classes/Materials/MaterialExpressionExecEnd.h
+++ b/Engine/Source/Runtime/Engine/Classes/Materials/MaterialExpressionExecEnd.h
@@ -18,7 +18,7 @@ class UMaterialExpressionExecEnd : public UMaterialExpression
 	virtual bool CanUserDeleteExpression() const override { return false; }
 	virtual bool HasExecInput() override { return true; }
 
-	virtual EMaterialGenerateHLSLStatus GenerateHLSLStatements(FMaterialHLSLGenerator& Generator, UE::HLSLTree::FScope& Scope) override;
+	virtual bool GenerateHLSLStatements(FMaterialHLSLGenerator& Generator, UE::HLSLTree::FScope& Scope) override;
 #endif
 	//~ End UMaterialExpression Interface
 };
diff --git a/Engine/Source/Runtime/Engine/Classes/Materials/MaterialExpressionForLoop.h b/Engine/Source/Runtime/Engine/Classes/Materials/MaterialExpressionForLoop.h
--- a/Engine/Source/Runtime/Engine/Classes/Materials/MaterialExpressionForLoop.h
+++ b/Engine/Source/Runtime/Engine/Classes/Materials/MaterialExpressionForLoop.h
@@ -37,8 +37,8 @@ class UMaterialExpressionForLoop : public UMaterialExpression
 	virtual uint32 GetOutputType(int32 InputIndex) override;
 	virtual bool HasExecInput() override { return true; }
 
-	virtual EMaterialGenerateHLSLStatus GenerateHLSLExpression(FMaterialHLSLGenerator& Generator, UE::HLSLTree::FScope& Scope, int32 OutputIndex, UE::HLSLTree::FExpression*& OutExpression) override;
-	virtual EMaterialGenerateHLSLStatus GenerateHLSLStatements(FMaterialHLSLGenerator& Generator, UE::HLSLTree::FScope& Scope) override;
+	virtual bool GenerateHLSLExpression(FMaterialHLSLGenerator& Generator, UE::HLSLTree::FScope& Scope, int32 OutputIndex, UE::HLSLTree::FExpression*& OutExpression) override;
+	virtual bool GenerateHLSLStatements(FMaterialHLSLGenerator& Generator, UE::HLSLTree::FScope& Scope) override;
 #endif
 	//~ End UMaterialExpression Interface
 };
diff --git a/Engine/Source/Runtime/Engine/Classes/Materials/MaterialExpressionFunctionInput.h b/Engine/Source/Runtime/Engine/Classes/Materials/MaterialExpressionFunctionInput.h
--- a/Engine/Source/Runtime/Engine/Classes/Materials/MaterialExpressionFunctionInput.h
+++ b/Engine/Source/Runtime/Engine/Classes/Materials/MaterialExpressionFunctionInput.h
@@ -110,7 +110,7 @@ class UMaterialExpressionFunctionInput : public UMaterialExpression
 	virtual bool IsResultMaterialAttributes(int32 OutputIndex) override;
 	virtual uint32 GetInputType(int32 InputIndex) override;
 	virtual uint32 GetOutputType(int32 OutputIndex) override;
-	virtual EMaterialGenerateHLSLStatus GenerateHLSLExpression(FMaterialHLSLGenerator& Generator, UE::HLSLTree::FScope& Scope, int32 OutputIndex, UE::HLSLTree::FExpression*& OutExpression) override;
+	virtual bool GenerateHLSLExpression(FMaterialHLSLGenerator& Generator, UE::HLSLTree::FScope& Scope, int32 OutputIndex, UE::HLSLTree::FExpression*& OutExpression) override;
 #endif // WITH_EDITOR
 	//~ End UMaterialExpression Interface
 
diff --git a/Engine/Source/Runtime/Engine/Classes/Materials/MaterialExpressionGenericConstant.h b/Engine/Source/Runtime/Engine/Classes/Materials/MaterialExpressionGenericConstant.h
--- a/Engine/Source/Runtime/Engine/Classes/Materials/MaterialExpressionGenericConstant.h
+++ b/Engine/Source/Runtime/Engine/Classes/Materials/MaterialExpressionGenericConstant.h
@@ -18,7 +18,7 @@ class UMaterialExpressionGenericConstant : public UMaterialExpression
 	virtual int32 Compile(class FMaterialCompiler* Compiler, int32 OutputIndex) override;
 	virtual void GetCaption(TArray<FString>& OutCaptions) const override;
 	virtual FString GetDescription() const override;
-	virtual EMaterialGenerateHLSLStatus GenerateHLSLExpression(FMaterialHLSLGenerator& Generator, UE::HLSLTree::FScope& Scope, int32 OutputIndex, UE::HLSLTree::FExpression*& OutExpression) override;
+	virtual bool GenerateHLSLExpression(FMaterialHLSLGenerator& Generator, UE::HLSLTree::FScope& Scope, int32 OutputIndex, UE::HLSLTree::FExpression*& OutExpression) override;
 #endif // WITH_EDITOR
 
 	//~ End UMaterialExpression Interface
diff --git a/Engine/Source/Runtime/Engine/Classes/Materials/MaterialExpressionGetLocal.h b/Engine/Source/Runtime/Engine/Classes/Materials/MaterialExpressionGetLocal.h
--- a/Engine/Source/Runtime/Engine/Classes/Materials/MaterialExpressionGetLocal.h
+++ b/Engine/Source/Runtime/Engine/Classes/Materials/MaterialExpressionGetLocal.h
@@ -24,7 +24,7 @@ class UMaterialExpressionGetLocal : public UMaterialExpression
 	virtual void GetCaption(TArray<FString>& OutCaptions) const override;
 	virtual uint32 GetOutputType(int32 InputIndex) override;
 	virtual bool IsResultMaterialAttributes(int32 OutputIndex) override;
-	virtual EMaterialGenerateHLSLStatus GenerateHLSLExpression(FMaterialHLSLGenerator& Generator, UE::HLSLTree::FScope& Scope, int32 OutputIndex, UE::HLSLTree::FExpression*& OutExpression) override;
+	virtual bool GenerateHLSLExpression(FMaterialHLSLGenerator& Generator, UE::HLSLTree::FScope& Scope, int32 OutputIndex, UE::HLSLTree::FExpression*& OutExpression) override;
 #endif
 	//~ End UMaterialExpression Interface
 };
diff --git a/Engine/Source/Runtime/Engine/Classes/Materials/MaterialExpressionGetMaterialAttributes.h b/Engine/Source/Runtime/Engine/Classes/Materials/MaterialExpressionGetMaterialAttributes.h
--- a/Engine/Source/Runtime/Engine/Classes/Materials/MaterialExpressionGetMaterialAttributes.h
+++ b/Engine/Source/Runtime/Engine/Classes/Materials/MaterialExpressionGetMaterialAttributes.h
@@ -36,7 +36,7 @@ class UMaterialExpressionGetMaterialAttributes : public UMaterialExpression
 	//~ Begin UMaterialExpression Interface
 #if WITH_EDITOR
 	virtual int32 Compile(class FMaterialCompiler* Compiler, int32 OutputIndex) override;
-	virtual EMaterialGenerateHLSLStatus GenerateHLSLExpression(FMaterialHLSLGenerator& Generator, UE::HLSLTree::FScope& Scope, int32 OutputIndex, UE::HLSLTree::FExpression*& OutExpression) override;
+	virtual bool GenerateHLSLExpression(FMaterialHLSLGenerator& Generator, UE::HLSLTree::FScope& Scope, int32 OutputIndex, UE::HLSLTree::FExpression*& OutExpression) override;
 
 	virtual void GetCaption(TArray<FString>& OutCaptions) const override;
 	virtual const TArray<FExpressionInput*> GetInputs()override;
diff --git a/Engine/Source/Runtime/Engine/Classes/Materials/MaterialExpressionIfThenElse.h b/Engine/Source/Runtime/Engine/Classes/Materials/MaterialExpressionIfThenElse.h
--- a/Engine/Source/Runtime/Engine/Classes/Materials/MaterialExpressionIfThenElse.h
+++ b/Engine/Source/Runtime/Engine/Classes/Materials/MaterialExpressionIfThenElse.h
@@ -30,7 +30,7 @@ class UMaterialExpressionIfThenElse : public UMaterialExpression
 	virtual uint32 GetInputType(int32 InputIndex) override;
 	virtual bool HasExecInput() override { return true; }
 
-	virtual EMaterialGenerateHLSLStatus GenerateHLSLStatements(FMaterialHLSLGenerator& Generator, UE::HLSLTree::FScope& Scope) override;
+	virtual bool GenerateHLSLStatements(FMaterialHLSLGenerator& Generator, UE::HLSLTree::FScope& Scope) override;
 #endif
 	//~ End UMaterialExpression Interface
 };
diff --git a/Engine/Source/Runtime/Engine/Classes/Materials/MaterialExpressionMaterialFunctionCall.h b/Engine/Source/Runtime/Engine/Classes/Materials/MaterialExpressionMaterialFunctionCall.h
--- a/Engine/Source/Runtime/Engine/Classes/Materials/MaterialExpressionMaterialFunctionCall.h
+++ b/Engine/Source/Runtime/Engine/Classes/Materials/MaterialExpressionMaterialFunctionCall.h
@@ -162,7 +162,7 @@ class UMaterialExpressionMaterialFunctionCall : public UMaterialExpression
 		SharedCompileState = SharedState;
 	}
 
-	virtual EMaterialGenerateHLSLStatus GenerateHLSLExpression(FMaterialHLSLGenerator& Generator, UE::HLSLTree::FScope& Scope, int32 OutputIndex, UE::HLSLTree::FExpression*& OutExpression) override;
+	virtual bool GenerateHLSLExpression(FMaterialHLSLGenerator& Generator, UE::HLSLTree::FScope& Scope, int32 OutputIndex, UE::HLSLTree::FExpression*& OutExpression) override;
 
 private:	
 	/** Helper that fixes up expression links where possible. */
diff --git a/Engine/Source/Runtime/Engine/Classes/Materials/MaterialExpressionMultiply.h b/Engine/Source/Runtime/Engine/Classes/Materials/MaterialExpressionMultiply.h
--- a/Engine/Source/Runtime/Engine/Classes/Materials/MaterialExpressionMultiply.h
+++ b/Engine/Source/Runtime/Engine/Classes/Materials/MaterialExpressionMultiply.h
@@ -34,7 +34,7 @@ class UMaterialExpressionMultiply : public UMaterialExpression
 	virtual int32 Compile(class FMaterialCompiler* Compiler, int32 OutputIndex) override;
 	virtual void GetCaption(TArray<FString>& OutCaptions) const override;
 	virtual FText GetKeywords() const override {return FText::FromString(TEXT("*"));}
-	virtual EMaterialGenerateHLSLStatus GenerateHLSLExpression(FMaterialHLSLGenerator& Generator, UE::HLSLTree::FScope& Scope, int32 OutputIndex, UE::HLSLTree::FExpression*& OutExpression) override;
+	virtual bool GenerateHLSLExpression(FMaterialHLSLGenerator& Generator, UE::HLSLTree::FScope& Scope, int32 OutputIndex, UE::HLSLTree::FExpression*& OutExpression) override;
 #endif // WITH_EDITOR
 	//~ End UMaterialExpression Interface
 };
diff --git a/Engine/Source/Runtime/Engine/Classes/Materials/MaterialExpressionReflectionVectorWS.h b/Engine/Source/Runtime/Engine/Classes/Materials/MaterialExpressionReflectionVectorWS.h
--- a/Engine/Source/Runtime/Engine/Classes/Materials/MaterialExpressionReflectionVectorWS.h
+++ b/Engine/Source/Runtime/Engine/Classes/Materials/MaterialExpressionReflectionVectorWS.h
@@ -26,7 +26,7 @@ class UMaterialExpressionReflectionVectorWS : public UMaterialExpression
 #if WITH_EDITOR
 	virtual int32 Compile(class FMaterialCompiler* Compiler, int32 OutputIndex) override;
 	virtual void GetCaption(TArray<FString>& OutCaptions) const override;
-	virtual EMaterialGenerateHLSLStatus GenerateHLSLExpression(FMaterialHLSLGenerator& Generator, UE::HLSLTree::FScope& Scope, int32 OutputIndex, UE::HLSLTree::FExpression*& OutExpression) override;
+	virtual bool GenerateHLSLExpression(FMaterialHLSLGenerator& Generator, UE::HLSLTree::FScope& Scope, int32 OutputIndex, UE::HLSLTree::FExpression*& OutExpression) override;
 #endif
 	//~ End UMaterialExpression Interface
 };
diff --git a/Engine/Source/Runtime/Engine/Classes/Materials/MaterialExpressionScalarParameter.h b/Engine/Source/Runtime/Engine/Classes/Materials/MaterialExpressionScalarParameter.h
--- a/Engine/Source/Runtime/Engine/Classes/Materials/MaterialExpressionScalarParameter.h
+++ b/Engine/Source/Runtime/Engine/Classes/Materials/MaterialExpressionScalarParameter.h
@@ -43,7 +43,7 @@ class UMaterialExpressionScalarParameter : public UMaterialExpressionParameter
 #if WITH_EDITOR
 	virtual int32 Compile(class FMaterialCompiler* Compiler, int32 OutputIndex) override;
 	virtual void GetCaption(TArray<FString>& OutCaptions) const override;
-	EMaterialGenerateHLSLStatus GenerateHLSLExpression(FMaterialHLSLGenerator& Generator, UE::HLSLTree::FScope& Scope, int32 OutputIndex, UE::HLSLTree::FExpression*& OutExpression) override;
+	bool GenerateHLSLExpression(FMaterialHLSLGenerator& Generator, UE::HLSLTree::FScope& Scope, int32 OutputIndex, UE::HLSLTree::FExpression*& OutExpression) override;
 	virtual bool GetParameterValue(FMaterialParameterMetadata& OutMeta) const override
 	{
 		OutMeta.Value = DefaultValue;
diff --git a/Engine/Source/Runtime/Engine/Classes/Materials/MaterialExpressionSetLocal.h b/Engine/Source/Runtime/Engine/Classes/Materials/MaterialExpressionSetLocal.h
--- a/Engine/Source/Runtime/Engine/Classes/Materials/MaterialExpressionSetLocal.h
+++ b/Engine/Source/Runtime/Engine/Classes/Materials/MaterialExpressionSetLocal.h
@@ -29,7 +29,7 @@ class UMaterialExpressionSetLocal : public UMaterialExpression
 	virtual void GetCaption(TArray<FString>& OutCaptions) const override;
 	virtual uint32 GetInputType(int32 InputIndex) override;
 	virtual bool HasExecInput() override { return true; }
-	virtual EMaterialGenerateHLSLStatus GenerateHLSLStatements(FMaterialHLSLGenerator& Generator, UE::HLSLTree::FScope& Scope) override;
+	virtual bool GenerateHLSLStatements(FMaterialHLSLGenerator& Generator, UE::HLSLTree::FScope& Scope) override;
 #endif
 	//~ End UMaterialExpression Interface
 };
diff --git a/Engine/Source/Runtime/Engine/Classes/Materials/MaterialExpressionSetMaterialAttributes.h b/Engine/Source/Runtime/Engine/Classes/Materials/MaterialExpressionSetMaterialAttributes.h
--- a/Engine/Source/Runtime/Engine/Classes/Materials/MaterialExpressionSetMaterialAttributes.h
+++ b/Engine/Source/Runtime/Engine/Classes/Materials/MaterialExpressionSetMaterialAttributes.h
@@ -44,7 +44,7 @@ class UMaterialExpressionSetMaterialAttributes : public UMaterialExpression
 	virtual bool IsResultMaterialAttributes(int32 OutputIndex) override {return true;}
 	virtual void GetExpressionToolTip(TArray<FString>& OutToolTip) override;
 
-	virtual EMaterialGenerateHLSLStatus GenerateHLSLExpression(FMaterialHLSLGenerator& Generator, UE::HLSLTree::FScope& Scope, int32 OutputIndex, UE::HLSLTree::FExpression*& OutExpression) override;
+	virtual bool GenerateHLSLExpression(FMaterialHLSLGenerator& Generator, UE::HLSLTree::FScope& Scope, int32 OutputIndex, UE::HLSLTree::FExpression*& OutExpression) override;
 #endif
 	//~ End UMaterialExpression Interface
 };
diff --git a/Engine/Source/Runtime/Engine/Classes/Materials/MaterialExpressionStaticBool.h b/Engine/Source/Runtime/Engine/Classes/Materials/MaterialExpressionStaticBool.h
--- a/Engine/Source/Runtime/Engine/Classes/Materials/MaterialExpressionStaticBool.h
+++ b/Engine/Source/Runtime/Engine/Classes/Materials/MaterialExpressionStaticBool.h
@@ -24,7 +24,7 @@ class UMaterialExpressionStaticBool : public UMaterialExpression
 	virtual void GetCaption(TArray<FString>& OutCaptions) const override;
 	virtual uint32 GetOutputType(int32 OutputIndex) override {return MCT_StaticBool;}
 
-	virtual EMaterialGenerateHLSLStatus GenerateHLSLExpression(FMaterialHLSLGenerator& Generator, UE::HLSLTree::FScope& Scope, int32 OutputIndex, UE::HLSLTree::FExpression*& OutExpression) override;
+	virtual bool GenerateHLSLExpression(FMaterialHLSLGenerator& Generator, UE::HLSLTree::FScope& Scope, int32 OutputIndex, UE::HLSLTree::FExpression*& OutExpression) override;
 #endif
 	//~ End UMaterialExpression Interface
 };
diff --git a/Engine/Source/Runtime/Engine/Classes/Materials/MaterialExpressionStaticBoolParameter.h b/Engine/Source/Runtime/Engine/Classes/Materials/MaterialExpressionStaticBoolParameter.h
--- a/Engine/Source/Runtime/Engine/Classes/Materials/MaterialExpressionStaticBoolParameter.h
+++ b/Engine/Source/Runtime/Engine/Classes/Materials/MaterialExpressionStaticBoolParameter.h
@@ -23,7 +23,7 @@ public:
 	virtual int32 CompilePreview(class FMaterialCompiler* Compiler, int32 OutputIndex) override;
 	virtual void GetCaption(TArray<FString>& OutCaptions) const override;
 	virtual uint32 GetOutputType(int32 OutputIndex) override {return MCT_StaticBool;}
-	virtual EMaterialGenerateHLSLStatus GenerateHLSLExpression(FMaterialHLSLGenerator& Generator, UE::HLSLTree::FScope& Scope, int32 OutputIndex, UE::HLSLTree::FExpression*& OutExpression) override;
+	virtual bool GenerateHLSLExpression(FMaterialHLSLGenerator& Generator, UE::HLSLTree::FScope& Scope, int32 OutputIndex, UE::HLSLTree::FExpression*& OutExpression) override;
 	virtual bool GetParameterValue(FMaterialParameterMetadata& OutMeta) const override
 	{
 		OutMeta.Value = (bool)DefaultValue;
diff --git a/Engine/Source/Runtime/Engine/Classes/Materials/MaterialExpressionStaticSwitch.h b/Engine/Source/Runtime/Engine/Classes/Materials/MaterialExpressionStaticSwitch.h
--- a/Engine/Source/Runtime/Engine/Classes/Materials/MaterialExpressionStaticSwitch.h
+++ b/Engine/Source/Runtime/Engine/Classes/Materials/MaterialExpressionStaticSwitch.h
@@ -36,7 +36,7 @@ class UMaterialExpressionStaticSwitch : public UMaterialExpression
 	virtual uint32 GetInputType(int32 InputIndex) override;
 	virtual uint32 GetOutputType(int32 OutputIndex) override {return MCT_Unknown;}
 
-	virtual EMaterialGenerateHLSLStatus GenerateHLSLExpression(FMaterialHLSLGenerator& Generator, UE::HLSLTree::FScope& Scope, int32 OutputIndex, UE::HLSLTree::FExpression*& OutExpression) override;
+	virtual bool GenerateHLSLExpression(FMaterialHLSLGenerator& Generator, UE::HLSLTree::FScope& Scope, int32 OutputIndex, UE::HLSLTree::FExpression*& OutExpression) override;
 #endif // WITH_EDITOR
 	//~ End UMaterialExpression Interface
 };
diff --git a/Engine/Source/Runtime/Engine/Classes/Materials/MaterialExpressionTextureCoordinate.h b/Engine/Source/Runtime/Engine/Classes/Materials/MaterialExpressionTextureCoordinate.h
--- a/Engine/Source/Runtime/Engine/Classes/Materials/MaterialExpressionTextureCoordinate.h
+++ b/Engine/Source/Runtime/Engine/Classes/Materials/MaterialExpressionTextureCoordinate.h
@@ -40,7 +40,7 @@ class UMaterialExpressionTextureCoordinate : public UMaterialExpression
 	virtual void GetCaption(TArray<FString>& OutCaptions) const override;
 	virtual FText GetKeywords() const override { return FText::FromString(TEXT("texcoord")); }
 	virtual bool MatchesSearchQuery(const TCHAR* SearchQuery) override;
-	virtual EMaterialGenerateHLSLStatus GenerateHLSLExpression(FMaterialHLSLGenerator& Generator, UE::HLSLTree::FScope& Scope, int32 OutputIndex, UE::HLSLTree::FExpression*& OutExpression) override;
+	virtual bool GenerateHLSLExpression(FMaterialHLSLGenerator& Generator, UE::HLSLTree::FScope& Scope, int32 OutputIndex, UE::HLSLTree::FExpression*& OutExpression) override;
 #endif
 	//~ End UMaterialExpression Interface
 };
diff --git a/Engine/Source/Runtime/Engine/Classes/Materials/MaterialExpressionTextureObject.h b/Engine/Source/Runtime/Engine/Classes/Materials/MaterialExpressionTextureObject.h
--- a/Engine/Source/Runtime/Engine/Classes/Materials/MaterialExpressionTextureObject.h
+++ b/Engine/Source/Runtime/Engine/Classes/Materials/MaterialExpressionTextureObject.h
@@ -28,7 +28,7 @@ class UMaterialExpressionTextureObject : public UMaterialExpressionTextureBase
 	virtual int32 CompilePreview(class FMaterialCompiler* Compiler, int32 OutputIndex) override;
 	virtual void GetCaption(TArray<FString>& OutCaptions) const override;
 	virtual uint32 GetOutputType(int32 OutputIndex) override;
-	virtual EMaterialGenerateHLSLStatus GenerateHLSLTexture(FMaterialHLSLGenerator& Generator, UE::HLSLTree::FScope& Scope, int32 OutputIndex, UE::HLSLTree::FTextureParameterDeclaration*& OutTexture) override;
+	virtual bool GenerateHLSLTexture(FMaterialHLSLGenerator& Generator, UE::HLSLTree::FScope& Scope, int32 OutputIndex, UE::HLSLTree::FTextureParameterDeclaration*& OutTexture) override;
 	//~ End UMaterialExpression Interface
 #endif // WITH_EDITOR
 };
diff --git a/Engine/Source/Runtime/Engine/Classes/Materials/MaterialExpressionTextureObjectParameter.h b/Engine/Source/Runtime/Engine/Classes/Materials/MaterialExpressionTextureObjectParameter.h
--- a/Engine/Source/Runtime/Engine/Classes/Materials/MaterialExpressionTextureObjectParameter.h
+++ b/Engine/Source/Runtime/Engine/Classes/Materials/MaterialExpressionTextureObjectParameter.h
@@ -30,7 +30,7 @@ class UMaterialExpressionTextureObjectParameter : public UMaterialExpressionText
 	virtual int32 CompilePreview(class FMaterialCompiler* Compiler, int32 OutputIndex) override;
 	virtual uint32 GetOutputType(int32 InputIndex) override {return MCT_Texture;}
 	virtual const TArray<FExpressionInput*> GetInputs() override;
-	virtual EMaterialGenerateHLSLStatus GenerateHLSLTexture(FMaterialHLSLGenerator& Generator, UE::HLSLTree::FScope& Scope, int32 OutputIndex, UE::HLSLTree::FTextureParameterDeclaration*& OutTexture) override;
+	virtual bool GenerateHLSLTexture(FMaterialHLSLGenerator& Generator, UE::HLSLTree::FScope& Scope, int32 OutputIndex, UE::HLSLTree::FTextureParameterDeclaration*& OutTexture) override;
 #endif
 	//~ End UMaterialExpression Interface
 };
diff --git a/Engine/Source/Runtime/Engine/Classes/Materials/MaterialExpressionTextureSample.h b/Engine/Source/Runtime/Engine/Classes/Materials/MaterialExpressionTextureSample.h
--- a/Engine/Source/Runtime/Engine/Classes/Materials/MaterialExpressionTextureSample.h
+++ b/Engine/Source/Runtime/Engine/Classes/Materials/MaterialExpressionTextureSample.h
@@ -64,7 +64,7 @@ protected:
 	uint8 bShowTextureInputPin : 1;
 
 #if WITH_EDITOR
-	EMaterialGenerateHLSLStatus GenerateHLSLExpressionBase(FMaterialHLSLGenerator& Generator, UE::HLSLTree::FScope& Scope, UE::HLSLTree::FTextureParameterDeclaration* TextureDeclaration, UE::HLSLTree::FExpression*& OutExpression);
+	bool GenerateHLSLExpressionBase(FMaterialHLSLGenerator& Generator, UE::HLSLTree::FScope& Scope, UE::HLSLTree::FTextureParameterDeclaration* TextureDeclaration, UE::HLSLTree::FExpression*& OutExpression);
 #endif
 
 public:
@@ -99,7 +99,7 @@ public:
 	virtual bool MatchesSearchQuery( const TCHAR* SearchQuery ) override;
 	virtual uint32 GetInputType(int32 InputIndex) override;
 	virtual bool CanIgnoreOutputIndex() { return true; }
-	virtual EMaterialGenerateHLSLStatus GenerateHLSLExpression(FMaterialHLSLGenerator& Generator, UE::HLSLTree::FScope& Scope, int32 OutputIndex, UE::HLSLTree::FExpression*& OutExpression) override;
+	virtual bool GenerateHLSLExpression(FMaterialHLSLGenerator& Generator, UE::HLSLTree::FScope& Scope, int32 OutputIndex, UE::HLSLTree::FExpression*& OutExpression) override;
 #endif // WITH_EDITOR
 	//~ End UMaterialExpression Interface
 
diff --git a/Engine/Source/Runtime/Engine/Classes/Materials/MaterialExpressionTextureSampleParameter.h b/Engine/Source/Runtime/Engine/Classes/Materials/MaterialExpressionTextureSampleParameter.h
--- a/Engine/Source/Runtime/Engine/Classes/Materials/MaterialExpressionTextureSampleParameter.h
+++ b/Engine/Source/Runtime/Engine/Classes/Materials/MaterialExpressionTextureSampleParameter.h
@@ -50,7 +50,7 @@ class ENGINE_API UMaterialExpressionTextureSampleParameter : public UMaterialExp
 	virtual FName GetParameterName() const override { return ParameterName; }
 	virtual void SetParameterName(const FName& Name) override { ParameterName = Name; }
 	virtual void ValidateParameterName(const bool bAllowDuplicateName) override;
-	virtual EMaterialGenerateHLSLStatus GenerateHLSLExpression(FMaterialHLSLGenerator& Generator, UE::HLSLTree::FScope& Scope, int32 OutputIndex, UE::HLSLTree::FExpression*& OutExpression) override;
+	virtual bool GenerateHLSLExpression(FMaterialHLSLGenerator& Generator, UE::HLSLTree::FScope& Scope, int32 OutputIndex, UE::HLSLTree::FExpression*& OutExpression) override;
 	virtual bool GetParameterValue(FMaterialParameterMetadata& OutMeta) const override
 	{
 		OutMeta.Value = Texture;
diff --git a/Engine/Source/Runtime/Engine/Classes/Materials/MaterialExpressionVectorParameter.h b/Engine/Source/Runtime/Engine/Classes/Materials/MaterialExpressionVectorParameter.h
--- a/Engine/Source/Runtime/Engine/Classes/Materials/MaterialExpressionVectorParameter.h
+++ b/Engine/Source/Runtime/Engine/Classes/Materials/MaterialExpressionVectorParameter.h
@@ -32,7 +32,7 @@ class UMaterialExpressionVectorParameter : public UMaterialExpressionParameter
 #if WITH_EDITOR
 	virtual int32 Compile(class FMaterialCompiler* Compiler, int32 OutputIndex) override;
 	virtual void GetCaption(TArray<FString>& OutCaptions) const override;
-	EMaterialGenerateHLSLStatus GenerateHLSLExpression(FMaterialHLSLGenerator& Generator, UE::HLSLTree::FScope& Scope, int32 OutputIndex, UE::HLSLTree::FExpression*& OutExpression) override;
+	bool GenerateHLSLExpression(FMaterialHLSLGenerator& Generator, UE::HLSLTree::FScope& Scope, int32 OutputIndex, UE::HLSLTree::FExpression*& OutExpression) override;
 	virtual bool GetParameterValue(FMaterialParameterMetadata& OutMeta) const override
 	{
 		OutMeta.Value = DefaultValue;
diff --git a/Engine/Source/Runtime/Engine/Classes/Materials/MaterialExpressionWhileLoop.h b/Engine/Source/Runtime/Engine/Classes/Materials/MaterialExpressionWhileLoop.h
--- a/Engine/Source/Runtime/Engine/Classes/Materials/MaterialExpressionWhileLoop.h
+++ b/Engine/Source/Runtime/Engine/Classes/Materials/MaterialExpressionWhileLoop.h
@@ -29,7 +29,7 @@ class UMaterialExpressionWhileLoop : public UMaterialExpression
 	virtual void GetCaption(TArray<FString>& OutCaptions) const override;
 	virtual bool HasExecInput() override { return true; }
 
-	virtual EMaterialGenerateHLSLStatus GenerateHLSLStatements(FMaterialHLSLGenerator& Generator, UE::HLSLTree::FScope& Scope) override;
+	virtual bool GenerateHLSLStatements(FMaterialHLSLGenerator& Generator, UE::HLSLTree::FScope& Scope) override;
 #endif
 	//~ End UMaterialExpression Interface
 };
diff --git a/Engine/Source/Runtime/Engine/Classes/Materials/MaterialExpressionWorldPosition.h b/Engine/Source/Runtime/Engine/Classes/Materials/MaterialExpressionWorldPosition.h
--- a/Engine/Source/Runtime/Engine/Classes/Materials/MaterialExpressionWorldPosition.h
+++ b/Engine/Source/Runtime/Engine/Classes/Materials/MaterialExpressionWorldPosition.h
@@ -37,7 +37,7 @@ class UMaterialExpressionWorldPosition : public UMaterialExpression
 	
 	//~ Begin UMaterialExpression Interface
 #if WITH_EDITOR
-	EMaterialGenerateHLSLStatus GenerateHLSLExpression(FMaterialHLSLGenerator& Generator, UE::HLSLTree::FScope& Scope, int32 OutputIndex, UE::HLSLTree::FExpression*& OutExpression);
+	bool GenerateHLSLExpression(FMaterialHLSLGenerator& Generator, UE::HLSLTree::FScope& Scope, int32 OutputIndex, UE::HLSLTree::FExpression*& OutExpression);
 	virtual int32 Compile(class FMaterialCompiler* Compiler, int32 OutputIndex) override;
 	virtual void GetCaption(TArray<FString>& OutCaptions) const override;
 	virtual FText GetKeywords() const override {return FText::FromString(TEXT("position"));}
diff --git a/Engine/Source/Runtime/Engine/Private/HLSLTree/HLSLTree.cpp b/Engine/Source/Runtime/Engine/Private/HLSLTree/HLSLTree.cpp
--- a/Engine/Source/Runtime/Engine/Private/HLSLTree/HLSLTree.cpp
+++ b/Engine/Source/Runtime/Engine/Private/HLSLTree/HLSLTree.cpp
@@ -54,31 +54,10 @@ EExpressionEvaluation CombineEvaluations(EExpressionEvaluation Lhs, EExpressionE
 	return EExpressionEvaluation::Preshader;
 }
 
-FErrors::FErrors(FMemStackBase& InAllocator)
-	: Allocator(&InAllocator)
-{
-}
-
-void FErrors::AddError(const FNode* InNode, FStringView InError)
-{
-	const int32 SizeofString = InError.Len() * sizeof(TCHAR);
-	void* Memory = Allocator->Alloc(sizeof(FError) + SizeofString, alignof(FError));
-	FError* Error = new(Memory) FError();
-	FMemory::Memcpy(Error->Message, InError.GetData(), SizeofString);
-	Error->Message[InError.Len()] = 0;
-	Error->MessageLength = InError.Len();
-	Error->Node = InNode;
-	Error->Next = FirstError;
-	FirstError = Error;
-	NumErrors++;
-
-	ensureMsgf(false, TEXT("%s"), Error->Message);
-}
-
-FEmitContext::FEmitContext(FMemStackBase& InAllocator, const Shader::FStructTypeRegistry& InTypeRegistry)
+FEmitContext::FEmitContext(FMemStackBase& InAllocator, FErrorHandlerInterface& InErrors, const Shader::FStructTypeRegistry& InTypeRegistry)
 	: Allocator(&InAllocator)
+	, Errors(&InErrors)
 	, TypeRegistry(&InTypeRegistry)
-	, Errors(InAllocator)
 {
 }
 
@@ -555,7 +534,7 @@ void FExpressionLocalPHI::ComputeAnalyticDerivatives(FTree& Tree, FExpressionDer
 	OutResult.ExpressionDdy = ExpressionDdy;
 }
 
-void FExpressionLocalPHI::PrepareValue(FEmitContext& Context, const FRequestedType& RequestedType, FPrepareValueResult& OutResult) const
+bool FExpressionLocalPHI::PrepareValue(FEmitContext& Context, const FRequestedType& RequestedType, FPrepareValueResult& OutResult) const
 {
 	check(NumValues <= MaxNumPreviousScopes);
 	FExpression* ForwardExpression = Values[0];
@@ -623,13 +602,15 @@ void FExpressionLocalPHI::PrepareValue(FEmitContext& Context, const FRequestedTy
 					CurrentType = MergePreparedTypes(CurrentType, ValueType);
 					if (CurrentType.IsVoid())
 					{
-						return Context.Errors.AddError(this, TEXT("Type mismatch"));
+						return Context.Errors->AddError(TEXT("Type mismatch"));
 					}
 					check(NumValidTypes < NumValues);
 					NumValidTypes++;
 				}
 			}
 		}
+
+		return true;
 	};
 
 	// First try to assign all the values we can
@@ -637,7 +618,10 @@ void FExpressionLocalPHI::PrepareValue(FEmitContext& Context, const FRequestedTy
 
 	// Assuming we have at least one value with a valid type, we use that to initialize our type
 	CurrentType.SetEvaluation(EExpressionEvaluation::Shader); // TODO - No support for preshader flow control
-	OutResult.SetType(Context, RequestedType, CurrentType);
+	if (!OutResult.SetType(Context, RequestedType, CurrentType))
+	{
+		return false;
+	}
 
 	if (NumValidTypes < NumValues)
 	{
@@ -645,9 +629,11 @@ void FExpressionLocalPHI::PrepareValue(FEmitContext& Context, const FRequestedTy
 		UpdateValueTypes();
 		if (NumValidTypes < NumValues)
 		{
-			return Context.Errors.AddError(this, TEXT("Failed to compute all types for LocalPHI"));
+			return Context.Errors->AddError(TEXT("Failed to compute all types for LocalPHI"));
 		}
 	}
+
+	return true;
 }
 
 void FExpressionLocalPHI::EmitValueShader(FEmitContext& Context, const FRequestedType& RequestedType, FEmitValueShaderResult& OutResult) const
@@ -666,7 +652,8 @@ void FExpressionLocalPHI::EmitValueShader(FEmitContext& Context, const FRequeste
 			DeclarationScope = FScope::FindSharedParent(DeclarationScope, Scopes[i]);
 			if (!DeclarationScope)
 			{
-				return Context.Errors.AddError(this, TEXT("Invalid LocalPHI"));
+				Context.Errors->AddError(TEXT("Invalid LocalPHI"));
+				return;
 			}
 		}
 
@@ -730,6 +717,7 @@ const FPreparedType& PrepareExpressionValue(FEmitContext& Context, FExpression*
 		return VoidType;
 	}
 
+	FOwnerScope OwnerScope(*Context.Errors, InExpression->GetOwner());
 	if (InExpression->bReentryFlag)
 	{
 		// Valid for this to be called reentrantly
@@ -746,7 +734,7 @@ const FPreparedType& PrepareExpressionValue(FEmitContext& Context, FExpression*
 	}
 	else if (InExpression->CurrentRequestedType.GetStructType() != RequestedType.GetStructType())
 	{
-		Context.Errors.AddError(InExpression, TEXT("Type mismatch"));
+		Context.Errors->AddError(TEXT("Type mismatch"));
 		return VoidType;
 	}
 	else
@@ -769,15 +757,22 @@ const FPreparedType& PrepareExpressionValue(FEmitContext& Context, FExpression*
 	{
 		check(!InExpression->CurrentRequestedType.IsVoid());
 
-		InExpression->bReentryFlag = true;
-		InExpression->PrepareValue(Context, InExpression->CurrentRequestedType, InExpression->PrepareValueResult);
-		InExpression->bReentryFlag = false;
+		bool bResult = false;
+		{
+			FExpressionReentryScope ReentryScope(InExpression);
+			bResult = InExpression->PrepareValue(Context, InExpression->CurrentRequestedType, InExpression->PrepareValueResult);
+		}
 
-		if (InExpression->PrepareValueResult.GetPreparedType().IsVoid())
+		if (!bResult)
 		{
 			// If we failed to assign a valid type, reset the requested type as well
 			// This ensures we'll try to compute a type again the next time we're called
 			InExpression->CurrentRequestedType.Reset();
+			InExpression->PrepareValueResult.SetTypeVoid();
+		}
+		else
+		{
+			check(!InExpression->PrepareValueResult.GetPreparedType().IsVoid());
 		}
 	}
 
@@ -1107,20 +1102,31 @@ bool FPrepareValueResult::TryMergePreparedType(FEmitContext& Context, const Shad
 		check(ComponentType == Shader::EValueComponentType::Void);
 		if (StructType != PreparedType.StructType)
 		{
-			Context.Errors.AddError(nullptr, TEXT("Invalid type"));
-			return false;
+			return Context.Errors->AddError(TEXT("Invalid type"));
 		}
 	}
 	else
 	{
-		check(ComponentType != Shader::EValueComponentType::Void);
+		if (ComponentType == Shader::EValueComponentType::Void)
+		{
+			return false;
+		}
 		PreparedType.ValueComponentType = Shader::CombineComponentTypes(PreparedType.ValueComponentType, ComponentType);
 	}
 
 	return true;
 }
 
-void FPrepareValueResult::SetType(FEmitContext& Context, const FRequestedType& RequestedType, EExpressionEvaluation Evaluation, const Shader::FType& Type)
+bool FPrepareValueResult::SetTypeVoid()
+{
+	PreparedType.PreparedComponents.Reset();
+	PreparedType.ValueComponentType = Shader::EValueComponentType::Void;
+	PreparedType.StructType = nullptr;
+	ForwardValue = nullptr;
+	return false;
+}
+
+bool FPrepareValueResult::SetType(FEmitContext& Context, const FRequestedType& RequestedType, EExpressionEvaluation Evaluation, const Shader::FType& Type)
 {
 	if (TryMergePreparedType(Context, Type.StructType, Shader::GetValueTypeDescription(Type.ValueType).ComponentType))
 	{
@@ -1135,10 +1141,12 @@ void FPrepareValueResult::SetType(FEmitContext& Context, const FRequestedType& R
 				}
 			}
 		}
+		return true;
 	}
+	return false;
 }
 
-void FPrepareValueResult::SetType(FEmitContext& Context, const FRequestedType& RequestedType, const FPreparedType& Type)
+bool FPrepareValueResult::SetType(FEmitContext& Context, const FRequestedType& RequestedType, const FPreparedType& Type)
 {
 	if (TryMergePreparedType(Context, Type.StructType, Type.ValueComponentType))
 	{
@@ -1150,10 +1158,12 @@ void FPrepareValueResult::SetType(FEmitContext& Context, const FRequestedType& R
 				PreparedType.MergeComponentEvaluation(Index, Type.GetComponentEvaluation(Index));
 			}
 		}
+		return true;
 	}
+	return false;
 }
 
-void FPrepareValueResult::SetForwardValue(FEmitContext& Context, const FRequestedType& RequestedType, FExpression* InForwardValue)
+bool FPrepareValueResult::SetForwardValue(FEmitContext& Context, const FRequestedType& RequestedType, FExpression* InForwardValue)
 {
 	check(InForwardValue);
 	if (InForwardValue != ForwardValue)
@@ -1161,6 +1171,7 @@ void FPrepareValueResult::SetForwardValue(FEmitContext& Context, const FRequeste
 		PreparedType = PrepareExpressionValue(Context, InForwardValue, RequestedType);
 		ForwardValue = InForwardValue;
 	}
+	return true;
 }
 
 void FExpression::ComputeAnalyticDerivatives(FTree& Tree, FExpressionDerivatives& OutResult) const
@@ -1180,6 +1191,7 @@ void FExpression::EmitValuePreshader(FEmitContext& Context, const FRequestedType
 
 FEmitShaderCode* FExpression::GetValueShader(FEmitContext& Context, const FRequestedType& RequestedType, const Shader::FType& ResultType)
 {
+	FOwnerScope OwnerScope(*Context.Errors, GetOwner());
 	if (PrepareValueResult.ForwardValue)
 	{
 		return PrepareValueResult.ForwardValue->GetValueShader(Context, RequestedType, ResultType);
@@ -1217,6 +1229,7 @@ FEmitShaderCode* FExpression::GetValueShader(FEmitContext& Context)
 
 void FExpression::GetValuePreshader(FEmitContext& Context, const FRequestedType& RequestedType, Shader::FPreshaderData& OutPreshader)
 {
+	FOwnerScope OwnerScope(*Context.Errors, GetOwner());
 	if (PrepareValueResult.ForwardValue)
 	{
 		return PrepareValueResult.ForwardValue->GetValuePreshader(Context, RequestedType, OutPreshader);
@@ -1240,6 +1253,7 @@ void FExpression::GetValuePreshader(FEmitContext& Context, const FRequestedType&
 
 Shader::FValue FExpression::GetValueConstant(FEmitContext& Context, const FRequestedType& RequestedType)
 {
+	FOwnerScope OwnerScope(*Context.Errors, GetOwner());
 	if (PrepareValueResult.ForwardValue)
 	{
 		return PrepareValueResult.ForwardValue->GetValueConstant(Context, RequestedType);
@@ -1329,6 +1343,7 @@ bool PrepareScope(FEmitContext& Context, FScope* InScope)
 		{
 			if (InScope->OwnerStatement)
 			{
+				FOwnerScope OwnerScope(*Context.Errors, InScope->OwnerStatement->GetOwner());
 				InScope->OwnerStatement->Prepare(Context);
 			}
 			else
@@ -1454,6 +1469,21 @@ void FScope::WriteHLSL(int32 Indent, FStringBuilderBase& OutString) const
 	}
 }
 
+void FOwnerContext::PushOwner(UObject* Owner)
+{
+	OwnerStack.Add(Owner);
+}
+
+UObject* FOwnerContext::PopOwner()
+{
+	return OwnerStack.Pop(false);
+}
+
+UObject* FOwnerContext::GetCurrentOwner() const
+{
+	return (OwnerStack.Num() > 0) ? OwnerStack.Last() : nullptr;
+}
+
 void FTree::ResetNodes()
 {
 	FNode* Node = Nodes;
@@ -1502,16 +1532,18 @@ bool FTree::EmitShader(FEmitContext& Context, FStringBuilderBase& OutCode) const
 		Context.ScopeStack.Add(RootScope);
 		RootScope->ContainedStatement->EmitShader(Context);
 		Context.ScopeStack.Pop(false);
-
-		if (Context.Errors.Num() > 0)
-		{
-			return false;
-		}
 	}
 
 	Context.Finalize();
 	RootScope->WriteHLSL(1, OutCode);
-	return Context.Errors.Num() == 0;
+	return true;
+}
+
+void FTree::RegisterNode(FNode* Node)
+{
+	Node->Owner = GetCurrentOwner();
+	Node->NextNode = Nodes;
+	Nodes = Node;
 }
 
 void FTree::RegisterExpression(FExpression* Expression)
@@ -1573,6 +1605,8 @@ const FExpressionDerivatives& FTree::GetAnalyticDerivatives(FExpression* InExpre
 	if (!InExpression->bComputedDerivatives)
 	{
 		FExpressionReentryScope ReentryScope(InExpression);
+		FOwnerScope OwnerScope(*this, InExpression->GetOwner()); // Associate any newly created nodes with the same owner as the input expression
+
 		InExpression->ComputeAnalyticDerivatives(*this, InExpression->Derivatives);
 		InExpression->bComputedDerivatives = true;
 	}
diff --git a/Engine/Source/Runtime/Engine/Private/HLSLTree/HLSLTreeCommon.cpp b/Engine/Source/Runtime/Engine/Private/HLSLTree/HLSLTreeCommon.cpp
--- a/Engine/Source/Runtime/Engine/Private/HLSLTree/HLSLTreeCommon.cpp
+++ b/Engine/Source/Runtime/Engine/Private/HLSLTree/HLSLTreeCommon.cpp
@@ -140,9 +140,9 @@ void FExpressionConstant::ComputeAnalyticDerivatives(FTree& Tree, FExpressionDer
 	}
 }
 
-void FExpressionConstant::PrepareValue(FEmitContext& Context, const FRequestedType& RequestedType, FPrepareValueResult& OutResult) const
+bool FExpressionConstant::PrepareValue(FEmitContext& Context, const FRequestedType& RequestedType, FPrepareValueResult& OutResult) const
 {
-	OutResult.SetType(Context, RequestedType, EExpressionEvaluation::Constant, Value.Type);
+	return OutResult.SetType(Context, RequestedType, EExpressionEvaluation::Constant, Value.Type);
 }
 
 void FExpressionConstant::EmitValuePreshader(FEmitContext& Context, const FRequestedType& RequestedType, Shader::FPreshaderData& OutPreshader) const
@@ -162,10 +162,10 @@ void FExpressionMaterialParameter::ComputeAnalyticDerivatives(FTree& Tree, FExpr
 	}
 }
 
-void FExpressionMaterialParameter::PrepareValue(FEmitContext& Context, const FRequestedType& RequestedType, FPrepareValueResult& OutResult) const
+bool FExpressionMaterialParameter::PrepareValue(FEmitContext& Context, const FRequestedType& RequestedType, FPrepareValueResult& OutResult) const
 {
 	const EExpressionEvaluation Evaluation = IsStaticMaterialParameter(ParameterType) ? EExpressionEvaluation::Constant : EExpressionEvaluation::Preshader;
-	OutResult.SetType(Context, RequestedType, Evaluation, GetShaderValueType(ParameterType));
+	return OutResult.SetType(Context, RequestedType, Evaluation, GetShaderValueType(ParameterType));
 }
 
 void FExpressionMaterialParameter::EmitValuePreshader(FEmitContext& Context, const FRequestedType& RequestedType, Shader::FPreshaderData& OutPreshader) const
@@ -214,10 +214,10 @@ void FExpressionExternalInput::ComputeAnalyticDerivatives(FTree& Tree, FExpressi
 	}
 }
 
-void FExpressionExternalInput::PrepareValue(FEmitContext& Context, const FRequestedType& RequestedType, FPrepareValueResult& OutResult) const
+bool FExpressionExternalInput::PrepareValue(FEmitContext& Context, const FRequestedType& RequestedType, FPrepareValueResult& OutResult) const
 {
 	const FExternalInputDescription InputDesc = GetExternalInputDescription(InputType);
-	OutResult.SetType(Context, RequestedType, EExpressionEvaluation::Shader, InputDesc.Type);
+	return OutResult.SetType(Context, RequestedType, EExpressionEvaluation::Shader, InputDesc.Type);
 }
 
 void FExpressionExternalInput::EmitValueShader(FEmitContext& Context, const FRequestedType& RequestedType, FEmitValueShaderResult& OutResult) const
@@ -271,13 +271,13 @@ void FExpressionExternalInput::EmitValueShader(FEmitContext& Context, const FReq
 	}
 }
 
-void FExpressionTextureSample::PrepareValue(FEmitContext& Context, const FRequestedType& RequestedType, FPrepareValueResult& OutResult) const
+bool FExpressionTextureSample::PrepareValue(FEmitContext& Context, const FRequestedType& RequestedType, FPrepareValueResult& OutResult) const
 {
 	PrepareExpressionValue(Context, TexCoordExpression, Shader::EValueType::Float2);
 	PrepareExpressionValue(Context, TexCoordDerivatives.ExpressionDdx, Shader::EValueType::Float2);
 	PrepareExpressionValue(Context, TexCoordDerivatives.ExpressionDdy, Shader::EValueType::Float2);
 
-	OutResult.SetType(Context, RequestedType, EExpressionEvaluation::Shader, Shader::EValueType::Float4);
+	return OutResult.SetType(Context, RequestedType, EExpressionEvaluation::Shader, Shader::EValueType::Float4);
 }
 
 void FExpressionTextureSample::EmitValueShader(FEmitContext& Context, const FRequestedType& RequestedType, FEmitValueShaderResult& OutResult) const
@@ -432,7 +432,7 @@ void FExpressionGetStructField::ComputeAnalyticDerivatives(FTree& Tree, FExpress
 	}
 }
 
-void FExpressionGetStructField::PrepareValue(FEmitContext& Context, const FRequestedType& RequestedType, FPrepareValueResult& OutResult) const
+bool FExpressionGetStructField::PrepareValue(FEmitContext& Context, const FRequestedType& RequestedType, FPrepareValueResult& OutResult) const
 {
 	FRequestedType RequestedStructType(StructType, false);
 	RequestedStructType.SetField(Field, RequestedType);
@@ -440,10 +440,10 @@ void FExpressionGetStructField::PrepareValue(FEmitContext& Context, const FReque
 	const FPreparedType& StructPreparedType = PrepareExpressionValue(Context, StructExpression, RequestedStructType);
 	if (!StructPreparedType.IsVoid() && StructPreparedType.StructType != StructType)
 	{
-		return Context.Errors.AddErrorf(this, TEXT("Expected type %s"), StructType->Name);
+		return Context.Errors->AddErrorf(TEXT("Expected type %s"), StructType->Name);
 	}
 
-	OutResult.SetType(Context, RequestedType, StructPreparedType.GetFieldType(Field));
+	return OutResult.SetType(Context, RequestedType, StructPreparedType.GetFieldType(Field));
 }
 
 void FExpressionGetStructField::EmitValueShader(FEmitContext& Context, const FRequestedType& RequestedType, FEmitValueShaderResult& OutResult) const
@@ -482,7 +482,7 @@ void FExpressionSetStructField::ComputeAnalyticDerivatives(FTree& Tree, FExpress
 	}
 }
 
-void FExpressionSetStructField::PrepareValue(FEmitContext& Context, const FRequestedType& RequestedType, FPrepareValueResult& OutResult) const
+bool FExpressionSetStructField::PrepareValue(FEmitContext& Context, const FRequestedType& RequestedType, FPrepareValueResult& OutResult) const
 {
 	FRequestedType RequestedStructType(RequestedType);
 	RequestedStructType.ClearFieldRequested(Field);
@@ -490,13 +490,13 @@ void FExpressionSetStructField::PrepareValue(FEmitContext& Context, const FReque
 	const FPreparedType& StructPreparedType = PrepareExpressionValue(Context, StructExpression, RequestedStructType);
 	if (!StructPreparedType.IsVoid() && StructPreparedType.StructType != StructType)
 	{
-		return Context.Errors.AddErrorf(this, TEXT("Expected type %s"), StructType->Name);
+		return Context.Errors->AddErrorf(TEXT("Expected type %s"), StructType->Name);
 	}
 
 	const FPreparedType& FieldPreparedType = PrepareExpressionValue(Context, FieldExpression, RequestedType.GetField(Field));
 	if (FieldPreparedType.IsVoid())
 	{
-		OutResult.SetForwardValue(Context, RequestedType, StructExpression);
+		return OutResult.SetForwardValue(Context, RequestedType, StructExpression);
 	}
 	else
 	{
@@ -506,7 +506,7 @@ void FExpressionSetStructField::PrepareValue(FEmitContext& Context, const FReque
 			ResultType = StructType;
 		}
 		ResultType.SetField(Field, FieldPreparedType);
-		OutResult.SetType(Context, RequestedType, ResultType);
+		return OutResult.SetType(Context, RequestedType, ResultType);
 	}
 }
 
@@ -565,7 +565,7 @@ void FExpressionSetStructField::EmitValuePreshader(FEmitContext& Context, const
 	OutPreshader.WriteOpcode(Shader::EPreshaderOpcode::SetField).Write(Field->ComponentIndex).Write(Field->GetNumComponents());
 }
 
-void FExpressionSelect::PrepareValue(FEmitContext& Context, const FRequestedType& RequestedType, FPrepareValueResult& OutResult) const
+bool FExpressionSelect::PrepareValue(FEmitContext& Context, const FRequestedType& RequestedType, FPrepareValueResult& OutResult) const
 {
 	const FPreparedType& ConditionType = PrepareExpressionValue(Context, ConditionExpression, 1);
 	if (ConditionType.GetEvaluation(Shader::EValueType::Bool1) == EExpressionEvaluation::Constant)
@@ -580,10 +580,10 @@ void FExpressionSelect::PrepareValue(FEmitContext& Context, const FRequestedType
 	if (LhsType.ValueComponentType != RhsType.ValueComponentType ||
 		LhsType.StructType != RhsType.StructType)
 	{
-		return Context.Errors.AddError(this, TEXT("Type mismatch"));
+		return Context.Errors->AddError(TEXT("Type mismatch"));
 	}
 
-	OutResult.SetType(Context, RequestedType, MergePreparedTypes(LhsType, RhsType));
+	return OutResult.SetType(Context, RequestedType, MergePreparedTypes(LhsType, RhsType));
 }
 
 void FExpressionSelect::EmitValueShader(FEmitContext& Context, const FRequestedType& RequestedType, FEmitValueShaderResult& OutResult) const
@@ -633,17 +633,17 @@ void FExpressionUnaryOp::ComputeAnalyticDerivatives(FTree& Tree, FExpressionDeri
 	}
 }
 
-void FExpressionUnaryOp::PrepareValue(FEmitContext& Context, const FRequestedType& RequestedType, FPrepareValueResult& OutResult) const
+bool FExpressionUnaryOp::PrepareValue(FEmitContext& Context, const FRequestedType& RequestedType, FPrepareValueResult& OutResult) const
 {
 	const FPreparedType& InputType = PrepareExpressionValue(Context, Input, RequestedType);
 	if (InputType.IsVoid())
 	{
-		return;
+		return false;
 	}
 
 	if (!InputType.IsNumeric())
 	{
-		return Context.Errors.AddError(this, TEXT("Invalid operation on non-numeric type"));
+		return Context.Errors->AddError(TEXT("Invalid operation on non-numeric type"));
 	}
 
 	const FUnaryOpDescription OpDesc = GetUnaryOpDesription(Op);
@@ -654,7 +654,7 @@ void FExpressionUnaryOp::PrepareValue(FEmitContext& Context, const FRequestedTyp
 		ResultType.SetEvaluation(EExpressionEvaluation::Shader);
 	}
 
-	OutResult.SetType(Context, RequestedType, ResultType);
+	return OutResult.SetType(Context, RequestedType, ResultType);
 }
 
 void FExpressionUnaryOp::EmitValueShader(FEmitContext& Context, const FRequestedType& RequestedType, FEmitValueShaderResult& OutResult) const
@@ -758,18 +758,23 @@ void FExpressionBinaryOp::ComputeAnalyticDerivatives(FTree& Tree, FExpressionDer
 	}
 }
 
-void FExpressionBinaryOp::PrepareValue(FEmitContext& Context, const FRequestedType& RequestedType, FPrepareValueResult& OutResult) const
+bool FExpressionBinaryOp::PrepareValue(FEmitContext& Context, const FRequestedType& RequestedType, FPrepareValueResult& OutResult) const
 {
 	const FPreparedType& LhsType = PrepareExpressionValue(Context, Lhs, RequestedType);
 	const FPreparedType& RhsType = PrepareExpressionValue(Context, Rhs, RequestedType);
 	if (LhsType.IsVoid() || RhsType.IsVoid())
 	{
-		return;
+		return false;
 	}
 
 	if (!LhsType.IsNumeric() || !RhsType.IsNumeric())
 	{
-		return Context.Errors.AddError(this, TEXT("Invalid arithmetic between non-numeric types"));
+		return Context.Errors->AddError(TEXT("Invalid arithmetic between non-numeric types"));
+	}
+
+	if (LhsType.GetNumComponents() != 1 && RhsType.GetNumComponents() != 1 && LhsType.GetNumComponents() != RhsType.GetNumComponents())
+	{
+		return Context.Errors->AddErrorf(TEXT("Invalid arithmetic between %s and %s"), LhsType.GetType().GetName(), RhsType.GetType().GetName());
 	}
 
 	FPreparedType ResultType = MergePreparedTypes(LhsType, RhsType);
@@ -785,7 +790,7 @@ void FExpressionBinaryOp::PrepareValue(FEmitContext& Context, const FRequestedTy
 		ResultType.SetEvaluation(EExpressionEvaluation::Shader);
 	}
 
-	OutResult.SetType(Context, RequestedType, ResultType);
+	return OutResult.SetType(Context, RequestedType, ResultType);
 }
 
 namespace Private
@@ -886,7 +891,7 @@ void FExpressionSwizzle::ComputeAnalyticDerivatives(FTree& Tree, FExpressionDeri
 	}
 }
 
-void FExpressionSwizzle::PrepareValue(FEmitContext& Context, const FRequestedType& RequestedType, FPrepareValueResult& OutResult) const
+bool FExpressionSwizzle::PrepareValue(FEmitContext& Context, const FRequestedType& RequestedType, FPrepareValueResult& OutResult) const
 {
 	const FRequestedType RequestedInputType = Parameters.GetRequestedInputType(RequestedType);
 	const FPreparedType& InputType = PrepareExpressionValue(Context, Input, RequestedInputType);
@@ -898,7 +903,7 @@ void FExpressionSwizzle::PrepareValue(FEmitContext& Context, const FRequestedTyp
 		ResultType.SetComponentEvaluation(ComponentIndex, InputType.GetComponentEvaluation(SwizzledComponentIndex));
 	}
 
-	OutResult.SetType(Context, RequestedType, ResultType);
+	return OutResult.SetType(Context, RequestedType, ResultType);
 }
 
 void FExpressionSwizzle::EmitValueShader(FEmitContext& Context, const FRequestedType& RequestedType, FEmitValueShaderResult& OutResult) const
@@ -956,7 +961,7 @@ void FExpressionAppend::ComputeAnalyticDerivatives(FTree& Tree, FExpressionDeriv
 	}
 }
 
-void FExpressionAppend::PrepareValue(FEmitContext& Context, const FRequestedType& RequestedType, FPrepareValueResult& OutResult) const
+bool FExpressionAppend::PrepareValue(FEmitContext& Context, const FRequestedType& RequestedType, FPrepareValueResult& OutResult) const
 {
 	const FPreparedType& LhsType = PrepareExpressionValue(Context, Lhs, RequestedType);
 	const int32 NumRequestedComponents = RequestedType.GetNumComponents();
@@ -973,7 +978,7 @@ void FExpressionAppend::PrepareValue(FEmitContext& Context, const FRequestedType
 
 	if (LhsType.ValueComponentType != RhsType.ValueComponentType)
 	{
-		return Context.Errors.AddError(this, TEXT("Type mismatch"));
+		return Context.Errors->AddError(TEXT("Type mismatch"));
 	}
 
 	FPreparedType ResultType(LhsType.ValueComponentType);
@@ -986,7 +991,7 @@ void FExpressionAppend::PrepareValue(FEmitContext& Context, const FRequestedType
 		ResultType.SetComponentEvaluation(NumLhsComponents + Index, LhsType.GetComponentEvaluation(Index));
 	}
 
-	OutResult.SetType(Context, RequestedType, ResultType);
+	return OutResult.SetType(Context, RequestedType, ResultType);
 }
 
 namespace Private
@@ -994,6 +999,7 @@ namespace Private
 struct FAppendTypes
 {
 	Shader::FType ResultType;
+	Shader::EValueComponentType ResultComponentType;
 	FRequestedType LhsType;
 	FRequestedType RhsType;
 };
@@ -1004,6 +1010,7 @@ FAppendTypes GetAppendTypes(const FRequestedType& RequestedType, const FRequeste
 	const int32 NumRhsComponents = FMath::Min<int32>(NumResultComponents - NumLhsComponents, RhsType.GetNumComponents());
 
 	FAppendTypes Types;
+	Types.ResultComponentType = RequestedType.ValueComponentType;
 	Types.ResultType = Shader::MakeValueType(RequestedType.ValueComponentType, NumLhsComponents + NumRhsComponents);
 	Types.LhsType.ValueComponentType = RequestedType.ValueComponentType;
 	Types.RhsType.ValueComponentType = RequestedType.ValueComponentType;
@@ -1031,10 +1038,19 @@ void FExpressionAppend::EmitValueShader(FEmitContext& Context, const FRequestedT
 	else
 	{
 		FEmitShaderCode* RhsValue = Rhs->GetValueShader(Context, Types.RhsType);
-		OutResult.Code = Context.EmitInlineCode(Types.ResultType, TEXT("%(%, %)"),
-			RequestedType.GetName(),
-			LhsValue,
-			RhsValue);
+		if (Types.ResultComponentType == Shader::EValueComponentType::Double)
+		{
+			OutResult.Code = Context.EmitInlineCode(Types.ResultType, TEXT("MakeLWCVector(%, %)"),
+				LhsValue,
+				RhsValue);
+		}
+		else
+		{
+			OutResult.Code = Context.EmitInlineCode(Types.ResultType, TEXT("%(%, %)"),
+				Types.ResultType.GetName(),
+				LhsValue,
+				RhsValue);
+		}
 	}
 }
 
@@ -1049,10 +1065,10 @@ void FExpressionAppend::EmitValuePreshader(FEmitContext& Context, const FRequest
 	}
 }
 
-void FExpressionReflectionVector::PrepareValue(FEmitContext& Context, const FRequestedType& RequestedType, FPrepareValueResult& OutResult) const
+bool FExpressionReflectionVector::PrepareValue(FEmitContext& Context, const FRequestedType& RequestedType, FPrepareValueResult& OutResult) const
 {
 	Context.bReadMaterialNormal = true;
-	OutResult.SetType(Context, RequestedType, EExpressionEvaluation::Shader, Shader::EValueType::Float3);
+	return OutResult.SetType(Context, RequestedType, EExpressionEvaluation::Shader, Shader::EValueType::Float3);
 }
 
 void FExpressionReflectionVector::EmitValueShader(FEmitContext& Context, const FRequestedType& RequestedType, FEmitValueShaderResult& OutResult) const
diff --git a/Engine/Source/Runtime/Engine/Private/HLSLTree/HLSLTreeIntrinsics.cpp b/Engine/Source/Runtime/Engine/Private/HLSLTree/HLSLTreeIntrinsics.cpp
--- a/Engine/Source/Runtime/Engine/Private/HLSLTree/HLSLTreeIntrinsics.cpp
+++ b/Engine/Source/Runtime/Engine/Private/HLSLTree/HLSLTreeIntrinsics.cpp
@@ -167,7 +167,7 @@ FEmitShaderCode* FEmitContext::EmitCast(FEmitShaderCode* ShaderValue, const Shad
 	}
 	else
 	{
-		Errors.AddErrorf(nullptr, TEXT("Cannot cast between non-numeric types %s to %s."), SourceTypeDesc.Name, DestTypeDesc.Name);
+		Errors->AddErrorf(TEXT("Cannot cast between non-numeric types %s to %s."), SourceTypeDesc.Name, DestTypeDesc.Name);
 		FormattedCode.Appendf(TEXT("((%s)0)"), DestType.GetName());
 	}
 
diff --git a/Engine/Source/Runtime/Engine/Private/Materials/MaterialExpressionHLSL.cpp b/Engine/Source/Runtime/Engine/Private/Materials/MaterialExpressionHLSL.cpp
--- a/Engine/Source/Runtime/Engine/Private/Materials/MaterialExpressionHLSL.cpp
+++ b/Engine/Source/Runtime/Engine/Private/Materials/MaterialExpressionHLSL.cpp
@@ -45,101 +45,101 @@
 #include "Materials/MaterialExpressionWhileLoop.h"
 #include "Materials/MaterialFunctionInterface.h"
 
-EMaterialGenerateHLSLStatus UMaterialExpression::GenerateHLSLExpression(FMaterialHLSLGenerator& Generator, UE::HLSLTree::FScope& Scope, int32 OutputIndex, UE::HLSLTree::FExpression*& OutExpression)
+bool UMaterialExpression::GenerateHLSLExpression(FMaterialHLSLGenerator& Generator, UE::HLSLTree::FScope& Scope, int32 OutputIndex, UE::HLSLTree::FExpression*& OutExpression)
 {
-	return Generator.Error(TEXT("Node does not support expressions"));
+	return Generator.GetErrors().AddError(TEXT("Node does not support expressions"));
 }
 
-EMaterialGenerateHLSLStatus UMaterialExpression::GenerateHLSLStatements(FMaterialHLSLGenerator& Generator, UE::HLSLTree::FScope& Scope)
+bool UMaterialExpression::GenerateHLSLStatements(FMaterialHLSLGenerator& Generator, UE::HLSLTree::FScope& Scope)
 {
-	return Generator.Error(TEXT("Node does not support statements"));
+	return Generator.GetErrors().AddError(TEXT("Node does not support statements"));
 }
 
-EMaterialGenerateHLSLStatus UMaterialExpression::GenerateHLSLTexture(FMaterialHLSLGenerator& Generator, UE::HLSLTree::FScope& Scope, int32 OutputIndex, UE::HLSLTree::FTextureParameterDeclaration*& OutTexture)
+bool UMaterialExpression::GenerateHLSLTexture(FMaterialHLSLGenerator& Generator, UE::HLSLTree::FScope& Scope, int32 OutputIndex, UE::HLSLTree::FTextureParameterDeclaration*& OutTexture)
 {
-	return Generator.Error(TEXT("Node does not support textures"));
+	return Generator.GetErrors().AddError(TEXT("Node does not support textures"));
 }
 
-EMaterialGenerateHLSLStatus UMaterialExpressionGenericConstant::GenerateHLSLExpression(FMaterialHLSLGenerator& Generator, UE::HLSLTree::FScope& Scope, int32 OutputIndex, UE::HLSLTree::FExpression*& OutExpression)
+bool UMaterialExpressionGenericConstant::GenerateHLSLExpression(FMaterialHLSLGenerator& Generator, UE::HLSLTree::FScope& Scope, int32 OutputIndex, UE::HLSLTree::FExpression*& OutExpression)
 {
 	OutExpression = Generator.NewConstant(GetConstantValue());
-	return EMaterialGenerateHLSLStatus::Success;
+	return true;
 }
 
-EMaterialGenerateHLSLStatus UMaterialExpressionConstant::GenerateHLSLExpression(FMaterialHLSLGenerator& Generator, UE::HLSLTree::FScope& Scope, int32 OutputIndex, UE::HLSLTree::FExpression*& OutExpression)
+bool UMaterialExpressionConstant::GenerateHLSLExpression(FMaterialHLSLGenerator& Generator, UE::HLSLTree::FScope& Scope, int32 OutputIndex, UE::HLSLTree::FExpression*& OutExpression)
 {
 	OutExpression = Generator.NewConstant(R);
-	return EMaterialGenerateHLSLStatus::Success;
+	return true;
 }
 
-EMaterialGenerateHLSLStatus UMaterialExpressionConstant2Vector::GenerateHLSLExpression(FMaterialHLSLGenerator& Generator, UE::HLSLTree::FScope& Scope, int32 OutputIndex, UE::HLSLTree::FExpression*& OutExpression)
+bool UMaterialExpressionConstant2Vector::GenerateHLSLExpression(FMaterialHLSLGenerator& Generator, UE::HLSLTree::FScope& Scope, int32 OutputIndex, UE::HLSLTree::FExpression*& OutExpression)
 {
 	OutExpression = Generator.NewConstant(UE::Shader::FValue(R, G));
-	return EMaterialGenerateHLSLStatus::Success;
+	return true;
 }
 
-EMaterialGenerateHLSLStatus UMaterialExpressionConstant3Vector::GenerateHLSLExpression(FMaterialHLSLGenerator& Generator, UE::HLSLTree::FScope& Scope, int32 OutputIndex, UE::HLSLTree::FExpression*& OutExpression)
+bool UMaterialExpressionConstant3Vector::GenerateHLSLExpression(FMaterialHLSLGenerator& Generator, UE::HLSLTree::FScope& Scope, int32 OutputIndex, UE::HLSLTree::FExpression*& OutExpression)
 {
 	OutExpression = Generator.NewConstant(UE::Shader::FValue(Constant.R, Constant.G, Constant.B));
-	return EMaterialGenerateHLSLStatus::Success;
+	return true;
 }
 
-EMaterialGenerateHLSLStatus UMaterialExpressionConstant4Vector::GenerateHLSLExpression(FMaterialHLSLGenerator& Generator, UE::HLSLTree::FScope& Scope, int32 OutputIndex, UE::HLSLTree::FExpression*& OutExpression)
+bool UMaterialExpressionConstant4Vector::GenerateHLSLExpression(FMaterialHLSLGenerator& Generator, UE::HLSLTree::FScope& Scope, int32 OutputIndex, UE::HLSLTree::FExpression*& OutExpression)
 {
 	OutExpression = Generator.NewConstant(UE::Shader::FValue(Constant.R, Constant.G, Constant.B, Constant.A));
-	return EMaterialGenerateHLSLStatus::Success;
+	return true;
 }
 
-EMaterialGenerateHLSLStatus UMaterialExpressionStaticBool::GenerateHLSLExpression(FMaterialHLSLGenerator& Generator, UE::HLSLTree::FScope& Scope, int32 OutputIndex, UE::HLSLTree::FExpression*& OutExpression)
+bool UMaterialExpressionStaticBool::GenerateHLSLExpression(FMaterialHLSLGenerator& Generator, UE::HLSLTree::FScope& Scope, int32 OutputIndex, UE::HLSLTree::FExpression*& OutExpression)
 {
 	OutExpression = Generator.NewConstant((bool)Value);
-	return EMaterialGenerateHLSLStatus::Success;
+	return true;
 }
 
-EMaterialGenerateHLSLStatus UMaterialExpressionStaticSwitch::GenerateHLSLExpression(FMaterialHLSLGenerator& Generator, UE::HLSLTree::FScope& Scope, int32 OutputIndex, UE::HLSLTree::FExpression*& OutExpression)
+bool UMaterialExpressionStaticSwitch::GenerateHLSLExpression(FMaterialHLSLGenerator& Generator, UE::HLSLTree::FScope& Scope, int32 OutputIndex, UE::HLSLTree::FExpression*& OutExpression)
 {
-	UE::HLSLTree::FExpression* ConditionExpression = Value.GetTracedInput().Expression ? Value.AcquireHLSLExpression(Generator, Scope) : Generator.NewConstant((bool)DefaultValue);
-	UE::HLSLTree::FExpression* TrueExpression = A.AcquireHLSLExpression(Generator, Scope);
-	UE::HLSLTree::FExpression* FalseExpression = B.AcquireHLSLExpression(Generator, Scope);
+	UE::HLSLTree::FExpression* ConditionExpression = Value.AcquireHLSLExpressionOrConstant(Generator, Scope, (bool)DefaultValue);
+	UE::HLSLTree::FExpression* TrueExpression = A.TryAcquireHLSLExpression(Generator, Scope);
+	UE::HLSLTree::FExpression* FalseExpression = B.TryAcquireHLSLExpression(Generator, Scope);
 
 	OutExpression = Generator.GetTree().NewExpression<UE::HLSLTree::FExpressionSelect>(ConditionExpression, TrueExpression, FalseExpression);
-	return EMaterialGenerateHLSLStatus::Success;
+	return true;
 }
 
-EMaterialGenerateHLSLStatus UMaterialExpressionGetLocal::GenerateHLSLExpression(FMaterialHLSLGenerator& Generator, UE::HLSLTree::FScope& Scope, int32 OutputIndex, UE::HLSLTree::FExpression*& OutExpression)
+bool UMaterialExpressionGetLocal::GenerateHLSLExpression(FMaterialHLSLGenerator& Generator, UE::HLSLTree::FScope& Scope, int32 OutputIndex, UE::HLSLTree::FExpression*& OutExpression)
 {
 	OutExpression = Generator.GetTree().AcquireLocal(Scope, LocalName);
 	if (!OutExpression)
 	{
-		return Generator.Error(TEXT("Local accessed before assigned"));
+		return Generator.GetErrors().AddError(TEXT("Local accessed before assigned"));
 	}
-	return EMaterialGenerateHLSLStatus::Success;
+	return true;
 }
 
-EMaterialGenerateHLSLStatus UMaterialExpressionVectorParameter::GenerateHLSLExpression(FMaterialHLSLGenerator& Generator, UE::HLSLTree::FScope& Scope, int32 OutputIndex, UE::HLSLTree::FExpression*& OutExpression)
+bool UMaterialExpressionVectorParameter::GenerateHLSLExpression(FMaterialHLSLGenerator& Generator, UE::HLSLTree::FScope& Scope, int32 OutputIndex, UE::HLSLTree::FExpression*& OutExpression)
 {
 	OutExpression = Generator.GetTree().NewExpression<UE::HLSLTree::FExpressionMaterialParameter>(EMaterialParameterType::Vector, ParameterName, DefaultValue);
-	return EMaterialGenerateHLSLStatus::Success;
+	return true;
 }
 
-EMaterialGenerateHLSLStatus UMaterialExpressionDoubleVectorParameter::GenerateHLSLExpression(FMaterialHLSLGenerator& Generator, UE::HLSLTree::FScope& Scope, int32 OutputIndex, UE::HLSLTree::FExpression*& OutExpression)
+bool UMaterialExpressionDoubleVectorParameter::GenerateHLSLExpression(FMaterialHLSLGenerator& Generator, UE::HLSLTree::FScope& Scope, int32 OutputIndex, UE::HLSLTree::FExpression*& OutExpression)
 {
 	OutExpression = Generator.GetTree().NewExpression<UE::HLSLTree::FExpressionMaterialParameter>(EMaterialParameterType::DoubleVector, ParameterName, DefaultValue);
-	return EMaterialGenerateHLSLStatus::Success;
+	return true;
 }
 
-EMaterialGenerateHLSLStatus UMaterialExpressionScalarParameter::GenerateHLSLExpression(FMaterialHLSLGenerator& Generator, UE::HLSLTree::FScope& Scope, int32 OutputIndex, UE::HLSLTree::FExpression*& OutExpression)
+bool UMaterialExpressionScalarParameter::GenerateHLSLExpression(FMaterialHLSLGenerator& Generator, UE::HLSLTree::FScope& Scope, int32 OutputIndex, UE::HLSLTree::FExpression*& OutExpression)
 {
 	OutExpression = Generator.GetTree().NewExpression<UE::HLSLTree::FExpressionMaterialParameter>(EMaterialParameterType::Scalar, ParameterName, DefaultValue);
-	return EMaterialGenerateHLSLStatus::Success;
+	return true;
 }
 
-EMaterialGenerateHLSLStatus UMaterialExpressionStaticBoolParameter::GenerateHLSLExpression(FMaterialHLSLGenerator& Generator, UE::HLSLTree::FScope& Scope, int32 OutputIndex, UE::HLSLTree::FExpression*& OutExpression)
+bool UMaterialExpressionStaticBoolParameter::GenerateHLSLExpression(FMaterialHLSLGenerator& Generator, UE::HLSLTree::FScope& Scope, int32 OutputIndex, UE::HLSLTree::FExpression*& OutExpression)
 {
 	OutExpression = Generator.GetTree().NewExpression<UE::HLSLTree::FExpressionMaterialParameter>(EMaterialParameterType::StaticSwitch, ParameterName, (bool)DefaultValue);
-	return EMaterialGenerateHLSLStatus::Success;
+	return true;
 }
-EMaterialGenerateHLSLStatus UMaterialExpressionWorldPosition::GenerateHLSLExpression(FMaterialHLSLGenerator& Generator, UE::HLSLTree::FScope& Scope, int32 OutputIndex, UE::HLSLTree::FExpression*& OutExpression)
+bool UMaterialExpressionWorldPosition::GenerateHLSLExpression(FMaterialHLSLGenerator& Generator, UE::HLSLTree::FScope& Scope, int32 OutputIndex, UE::HLSLTree::FExpression*& OutExpression)
 {
 	UE::HLSLTree::EExternalInput InputType = UE::HLSLTree::EExternalInput::None;
 
@@ -153,49 +153,49 @@ EMaterialGenerateHLSLStatus UMaterialExpressionWorldPosition::GenerateHLSLExpres
 	}
 
 	OutExpression = Generator.GetTree().NewExpression<UE::HLSLTree::FExpressionExternalInput>(InputType);
-	return EMaterialGenerateHLSLStatus::Success;
+	return true;
 }
 
-EMaterialGenerateHLSLStatus UMaterialExpressionTextureCoordinate::GenerateHLSLExpression(FMaterialHLSLGenerator& Generator, UE::HLSLTree::FScope& Scope, int32 OutputIndex, UE::HLSLTree::FExpression*& OutExpression)
+bool UMaterialExpressionTextureCoordinate::GenerateHLSLExpression(FMaterialHLSLGenerator& Generator, UE::HLSLTree::FScope& Scope, int32 OutputIndex, UE::HLSLTree::FExpression*& OutExpression)
 {
 	// TODO - tiling, mirroring
 	OutExpression = Generator.NewTexCoord(CoordinateIndex);
-	return EMaterialGenerateHLSLStatus::Success;
+	return true;
 }
 
-EMaterialGenerateHLSLStatus UMaterialExpressionTextureObject::GenerateHLSLTexture(FMaterialHLSLGenerator& Generator, UE::HLSLTree::FScope& Scope, int32 OutputIndex, UE::HLSLTree::FTextureParameterDeclaration*& OutTexture)
+bool UMaterialExpressionTextureObject::GenerateHLSLTexture(FMaterialHLSLGenerator& Generator, UE::HLSLTree::FScope& Scope, int32 OutputIndex, UE::HLSLTree::FTextureParameterDeclaration*& OutTexture)
 {
 	const UE::HLSLTree::FTextureDescription TextureDesc(Texture, SamplerType);
 	OutTexture = Generator.AcquireTextureDeclaration(TextureDesc);
-	return EMaterialGenerateHLSLStatus::Success;
+	return true;
 }
 
-EMaterialGenerateHLSLStatus UMaterialExpressionTextureObjectParameter::GenerateHLSLTexture(FMaterialHLSLGenerator& Generator, UE::HLSLTree::FScope& Scope, int32 OutputIndex, UE::HLSLTree::FTextureParameterDeclaration*& OutTexture)
+bool UMaterialExpressionTextureObjectParameter::GenerateHLSLTexture(FMaterialHLSLGenerator& Generator, UE::HLSLTree::FScope& Scope, int32 OutputIndex, UE::HLSLTree::FTextureParameterDeclaration*& OutTexture)
 {
 	const UE::HLSLTree::FTextureDescription TextureDesc(Texture, SamplerType);
 	OutTexture = Generator.AcquireTextureParameterDeclaration(ParameterName, TextureDesc);
-	return EMaterialGenerateHLSLStatus::Success;
+	return true;
 }
 
-EMaterialGenerateHLSLStatus UMaterialExpressionTextureSample::GenerateHLSLExpressionBase(FMaterialHLSLGenerator& Generator, UE::HLSLTree::FScope& Scope, UE::HLSLTree::FTextureParameterDeclaration* TextureDeclaration, UE::HLSLTree::FExpression*& OutExpression)
+bool UMaterialExpressionTextureSample::GenerateHLSLExpressionBase(FMaterialHLSLGenerator& Generator, UE::HLSLTree::FScope& Scope, UE::HLSLTree::FTextureParameterDeclaration* TextureDeclaration, UE::HLSLTree::FExpression*& OutExpression)
 {
 	using namespace UE::HLSLTree;
 	if (!TextureDeclaration)
 	{
-		return Generator.Error(TEXT("Missing input texture"));
+		return Generator.GetErrors().AddError(TEXT("Missing input texture"));
 	}
 
-	FExpression* TexCoordExpression = Coordinates.GetTracedInput().Expression ? Coordinates.AcquireHLSLExpression(Generator, Scope) : Generator.NewTexCoord(ConstCoordinate);
+	FExpression* TexCoordExpression = Coordinates.GetTracedInput().Expression ? Coordinates.TryAcquireHLSLExpression(Generator, Scope) : Generator.NewTexCoord(ConstCoordinate);
 	FExpressionTextureSample* ExpressionTextureSample = Generator.GetTree().NewExpression<FExpressionTextureSample>(TextureDeclaration, TexCoordExpression);
 	ExpressionTextureSample->SamplerSource = SamplerSource;
 	ExpressionTextureSample->MipValueMode = MipValueMode;
 	ExpressionTextureSample->TexCoordDerivatives = Generator.GetTree().GetAnalyticDerivatives(TexCoordExpression);
 
 	OutExpression = ExpressionTextureSample;
-	return EMaterialGenerateHLSLStatus::Success;
+	return true;
 }
 
-EMaterialGenerateHLSLStatus UMaterialExpressionTextureSample::GenerateHLSLExpression(FMaterialHLSLGenerator& Generator, UE::HLSLTree::FScope& Scope, int32 OutputIndex, UE::HLSLTree::FExpression*& OutExpression)
+bool UMaterialExpressionTextureSample::GenerateHLSLExpression(FMaterialHLSLGenerator& Generator, UE::HLSLTree::FScope& Scope, int32 OutputIndex, UE::HLSLTree::FExpression*& OutExpression)
 {
 	UE::HLSLTree::FTextureParameterDeclaration* TextureDeclaration = nullptr;
 	if (TextureObject.GetTracedInput().Expression)
@@ -211,7 +211,7 @@ EMaterialGenerateHLSLStatus UMaterialExpressionTextureSample::GenerateHLSLExpres
 	return GenerateHLSLExpressionBase(Generator, Scope, TextureDeclaration, OutExpression);
 }
 
-EMaterialGenerateHLSLStatus UMaterialExpressionTextureSampleParameter::GenerateHLSLExpression(FMaterialHLSLGenerator& Generator, UE::HLSLTree::FScope& Scope, int32 OutputIndex, UE::HLSLTree::FExpression*& OutExpression)
+bool UMaterialExpressionTextureSampleParameter::GenerateHLSLExpression(FMaterialHLSLGenerator& Generator, UE::HLSLTree::FScope& Scope, int32 OutputIndex, UE::HLSLTree::FExpression*& OutExpression)
 {
 	UE::HLSLTree::FTextureParameterDeclaration* TextureDeclaration = nullptr;
 	if (ParameterName.IsNone() && TextureObject.GetTracedInput().Expression)
@@ -227,118 +227,110 @@ EMaterialGenerateHLSLStatus UMaterialExpressionTextureSampleParameter::GenerateH
 	return GenerateHLSLExpressionBase(Generator, Scope, TextureDeclaration, OutExpression);
 }
 
-EMaterialGenerateHLSLStatus UMaterialExpressionBinaryOp::GenerateHLSLExpression(FMaterialHLSLGenerator& Generator, UE::HLSLTree::FScope& Scope, int32 OutputIndex, UE::HLSLTree::FExpression*& OutExpression)
+bool UMaterialExpressionBinaryOp::GenerateHLSLExpression(FMaterialHLSLGenerator& Generator, UE::HLSLTree::FScope& Scope, int32 OutputIndex, UE::HLSLTree::FExpression*& OutExpression)
 {
-	UE::HLSLTree::FExpression* Lhs = A.GetTracedInput().Expression ? A.AcquireHLSLExpression(Generator, Scope) : Generator.NewConstant(ConstA);
-	UE::HLSLTree::FExpression* Rhs = B.GetTracedInput().Expression ? B.AcquireHLSLExpression(Generator, Scope) : Generator.NewConstant(ConstB);
+	UE::HLSLTree::FExpression* Lhs = A.AcquireHLSLExpressionOrConstant(Generator, Scope, ConstA);
+	UE::HLSLTree::FExpression* Rhs = B.AcquireHLSLExpressionOrConstant(Generator, Scope, ConstB);
 	if (!Lhs || !Rhs)
 	{
-		return EMaterialGenerateHLSLStatus::Error;
+		return false;
 	}
 
 	OutExpression = Generator.GetTree().NewBinaryOp(GetBinaryOp(), Lhs, Rhs);
-	return EMaterialGenerateHLSLStatus::Success;
+	return true;
 }
 
-EMaterialGenerateHLSLStatus UMaterialExpressionAdd::GenerateHLSLExpression(FMaterialHLSLGenerator& Generator, UE::HLSLTree::FScope& Scope, int32 OutputIndex, UE::HLSLTree::FExpression*& OutExpression)
+bool UMaterialExpressionAdd::GenerateHLSLExpression(FMaterialHLSLGenerator& Generator, UE::HLSLTree::FScope& Scope, int32 OutputIndex, UE::HLSLTree::FExpression*& OutExpression)
 {
-	UE::HLSLTree::FExpression* Lhs = A.GetTracedInput().Expression ? A.AcquireHLSLExpression(Generator, Scope) : Generator.NewConstant(ConstA);
-	UE::HLSLTree::FExpression* Rhs = B.GetTracedInput().Expression ? B.AcquireHLSLExpression(Generator, Scope) : Generator.NewConstant(ConstB);
+	UE::HLSLTree::FExpression* Lhs = A.AcquireHLSLExpressionOrConstant(Generator, Scope, ConstA);
+	UE::HLSLTree::FExpression* Rhs = B.AcquireHLSLExpressionOrConstant(Generator, Scope, ConstB);
 	if (!Lhs || !Rhs)
 	{
-		return EMaterialGenerateHLSLStatus::Error;
+		return false;
 	}
 	OutExpression = Generator.GetTree().NewBinaryOp(UE::HLSLTree::EBinaryOp::Add, Lhs, Rhs);
-	return EMaterialGenerateHLSLStatus::Success;
+	return true;
 }
 
-EMaterialGenerateHLSLStatus UMaterialExpressionMultiply::GenerateHLSLExpression(FMaterialHLSLGenerator& Generator, UE::HLSLTree::FScope& Scope, int32 OutputIndex, UE::HLSLTree::FExpression*& OutExpression)
+bool UMaterialExpressionMultiply::GenerateHLSLExpression(FMaterialHLSLGenerator& Generator, UE::HLSLTree::FScope& Scope, int32 OutputIndex, UE::HLSLTree::FExpression*& OutExpression)
 {
-	UE::HLSLTree::FExpression* Lhs = A.GetTracedInput().Expression ? A.AcquireHLSLExpression(Generator, Scope) : Generator.NewConstant(ConstA);
-	UE::HLSLTree::FExpression* Rhs = B.GetTracedInput().Expression ? B.AcquireHLSLExpression(Generator, Scope) : Generator.NewConstant(ConstB);
+	UE::HLSLTree::FExpression* Lhs = A.AcquireHLSLExpressionOrConstant(Generator, Scope, ConstA);
+	UE::HLSLTree::FExpression* Rhs = B.AcquireHLSLExpressionOrConstant(Generator, Scope, ConstB);
 	if (!Lhs || !Rhs)
 	{
-		return EMaterialGenerateHLSLStatus::Error;
+		return false;
 	}
 	OutExpression = Generator.GetTree().NewBinaryOp(UE::HLSLTree::EBinaryOp::Mul, Lhs, Rhs);
-	return EMaterialGenerateHLSLStatus::Success;
+	return true;
 }
 
-EMaterialGenerateHLSLStatus UMaterialExpressionDivide::GenerateHLSLExpression(FMaterialHLSLGenerator& Generator, UE::HLSLTree::FScope& Scope, int32 OutputIndex, UE::HLSLTree::FExpression*& OutExpression)
+bool UMaterialExpressionDivide::GenerateHLSLExpression(FMaterialHLSLGenerator& Generator, UE::HLSLTree::FScope& Scope, int32 OutputIndex, UE::HLSLTree::FExpression*& OutExpression)
 {
-	UE::HLSLTree::FExpression* Lhs = A.GetTracedInput().Expression ? A.AcquireHLSLExpression(Generator, Scope) : Generator.NewConstant(ConstA);
-	UE::HLSLTree::FExpression* Rhs = B.GetTracedInput().Expression ? B.AcquireHLSLExpression(Generator, Scope) : Generator.NewConstant(ConstB);
+	UE::HLSLTree::FExpression* Lhs = A.AcquireHLSLExpressionOrConstant(Generator, Scope, ConstA);
+	UE::HLSLTree::FExpression* Rhs = B.AcquireHLSLExpressionOrConstant(Generator, Scope, ConstB);
 	if (!Lhs || !Rhs)
 	{
-		return EMaterialGenerateHLSLStatus::Error;
+		return false;
 	}
 	OutExpression = Generator.GetTree().NewBinaryOp(UE::HLSLTree::EBinaryOp::Div, Lhs, Rhs);
-	return EMaterialGenerateHLSLStatus::Success;
+	return true;
 }
 
-EMaterialGenerateHLSLStatus UMaterialExpressionAppendVector::GenerateHLSLExpression(FMaterialHLSLGenerator& Generator, UE::HLSLTree::FScope& Scope, int32 OutputIndex, UE::HLSLTree::FExpression*& OutExpression)
+bool UMaterialExpressionAppendVector::GenerateHLSLExpression(FMaterialHLSLGenerator& Generator, UE::HLSLTree::FScope& Scope, int32 OutputIndex, UE::HLSLTree::FExpression*& OutExpression)
 {
 	UE::HLSLTree::FExpression* Lhs = A.AcquireHLSLExpression(Generator, Scope);
 	UE::HLSLTree::FExpression* Rhs = B.AcquireHLSLExpression(Generator, Scope);
 	if (!Lhs || !Rhs)
 	{
-		return EMaterialGenerateHLSLStatus::Error;
+		return false;
 	}
 	OutExpression = Generator.GetTree().NewExpression<UE::HLSLTree::FExpressionAppend>(Lhs, Rhs);
-	return EMaterialGenerateHLSLStatus::Success;
+	return true;
 }
 
-EMaterialGenerateHLSLStatus UMaterialExpressionComponentMask::GenerateHLSLExpression(FMaterialHLSLGenerator& Generator, UE::HLSLTree::FScope& Scope, int32 OutputIndex, UE::HLSLTree::FExpression*& OutExpression)
+bool UMaterialExpressionComponentMask::GenerateHLSLExpression(FMaterialHLSLGenerator& Generator, UE::HLSLTree::FScope& Scope, int32 OutputIndex, UE::HLSLTree::FExpression*& OutExpression)
 {
 	UE::HLSLTree::FExpression* InputExpression = Input.AcquireHLSLExpression(Generator, Scope);
 	if (!InputExpression)
 	{
-		return EMaterialGenerateHLSLStatus::Error;
+		return false;
 	}
 
 	OutExpression = Generator.GetTree().NewExpression<UE::HLSLTree::FExpressionSwizzle>(UE::HLSLTree::MakeSwizzleMask(!!R, !!G, !!B, !!A), InputExpression);
-	return EMaterialGenerateHLSLStatus::Success;
+	return true;
 }
 
 
-EMaterialGenerateHLSLStatus UMaterialExpressionGetMaterialAttributes::GenerateHLSLExpression(FMaterialHLSLGenerator& Generator, UE::HLSLTree::FScope& Scope, int32 OutputIndex, UE::HLSLTree::FExpression*& OutExpression)
+bool UMaterialExpressionGetMaterialAttributes::GenerateHLSLExpression(FMaterialHLSLGenerator& Generator, UE::HLSLTree::FScope& Scope, int32 OutputIndex, UE::HLSLTree::FExpression*& OutExpression)
 {
 	UE::HLSLTree::FExpression* AttributesExpression = MaterialAttributes.AcquireHLSLExpression(Generator, Scope);
 	if (!AttributesExpression)
 	{
-		return EMaterialGenerateHLSLStatus::Error;
+		return false;
 	}
 	if (OutputIndex == 0)
 	{
 		OutExpression = AttributesExpression;
-		return EMaterialGenerateHLSLStatus::Success;
+		return true;
 	}
 	const int32 AttributeIndex = OutputIndex - 1;
 	if (!AttributeGetTypes.IsValidIndex(AttributeIndex))
 	{
-		return Generator.Error(TEXT("Invalid attribute"));
+		return Generator.GetErrors().AddError(TEXT("Invalid attribute"));
 	}
 
 	const FGuid& AttributeID = AttributeGetTypes[AttributeIndex];
 	const FString& AttributeName = FMaterialAttributeDefinitionMap::GetAttributeName(AttributeID);
 	OutExpression = Generator.GetTree().NewExpression<UE::HLSLTree::FExpressionGetStructField>(Generator.GetMaterialAttributesType(), *AttributeName, AttributesExpression);
 
-	return EMaterialGenerateHLSLStatus::Success;
+	return true;
 
 }
 
-EMaterialGenerateHLSLStatus UMaterialExpressionSetMaterialAttributes::GenerateHLSLExpression(FMaterialHLSLGenerator& Generator, UE::HLSLTree::FScope& Scope, int32 OutputIndex, UE::HLSLTree::FExpression*& OutExpression)
+bool UMaterialExpressionSetMaterialAttributes::GenerateHLSLExpression(FMaterialHLSLGenerator& Generator, UE::HLSLTree::FScope& Scope, int32 OutputIndex, UE::HLSLTree::FExpression*& OutExpression)
 {
-	UE::HLSLTree::FExpression* AttributesExpression = nullptr;
-	if (Inputs[0].GetTracedInput().Expression)
-	{
-		AttributesExpression = Inputs[0].AcquireHLSLExpression(Generator, Scope);
-	}
-	else
-	{
-		AttributesExpression = Generator.GetTree().NewExpression<UE::HLSLTree::FExpressionConstant>(Generator.GetMaterialAttributesDefaultValue());
-	}
-
+	UE::HLSLTree::FExpression* AttributesExpression = Inputs[0].AcquireHLSLExpressionOrConstant(Generator, Scope, Generator.GetMaterialAttributesDefaultValue());
+	
 	for (int32 PinIndex = 0; PinIndex < AttributeSetTypes.Num(); ++PinIndex)
 	{
 		const FExpressionInput& AttributeInput = Inputs[PinIndex + 1];
@@ -350,7 +342,7 @@ EMaterialGenerateHLSLStatus UMaterialExpressionSetMaterialAttributes::GenerateHL
 			const FString& AttributeName = FMaterialAttributeDefinitionMap::GetAttributeName(AttributeID);
 			//if (AttributeFrequency == Compiler->GetCurrentShaderFrequency())
 			{
-				UE::HLSLTree::FExpression* ValueExpression = AttributeInput.AcquireHLSLExpression(Generator, Scope);
+				UE::HLSLTree::FExpression* ValueExpression = AttributeInput.TryAcquireHLSLExpression(Generator, Scope);
 				if (ValueExpression)
 				{
 					AttributesExpression = Generator.GetTree().NewExpression<UE::HLSLTree::FExpressionSetStructField>(Generator.GetMaterialAttributesType(), *AttributeName, AttributesExpression, ValueExpression);
@@ -360,68 +352,68 @@ EMaterialGenerateHLSLStatus UMaterialExpressionSetMaterialAttributes::GenerateHL
 	}
 
 	OutExpression = AttributesExpression;
-	return EMaterialGenerateHLSLStatus::Success;
+	return true;
 }
 
-EMaterialGenerateHLSLStatus UMaterialExpressionReflectionVectorWS::GenerateHLSLExpression(FMaterialHLSLGenerator& Generator, UE::HLSLTree::FScope& Scope, int32 OutputIndex, UE::HLSLTree::FExpression*& OutExpression)
+bool UMaterialExpressionReflectionVectorWS::GenerateHLSLExpression(FMaterialHLSLGenerator& Generator, UE::HLSLTree::FScope& Scope, int32 OutputIndex, UE::HLSLTree::FExpression*& OutExpression)
 {
 	check(!CustomWorldNormal.GetTracedInput().Expression); // TODO
 
 	OutExpression = Generator.GetTree().NewExpression<UE::HLSLTree::FExpressionReflectionVector>();
-	return EMaterialGenerateHLSLStatus::Success;
+	return true;
 }
 
-EMaterialGenerateHLSLStatus UMaterialExpressionFunctionInput::GenerateHLSLExpression(FMaterialHLSLGenerator& Generator, UE::HLSLTree::FScope& Scope, int32 OutputIndex, UE::HLSLTree::FExpression*& OutExpression)
+bool UMaterialExpressionFunctionInput::GenerateHLSLExpression(FMaterialHLSLGenerator& Generator, UE::HLSLTree::FScope& Scope, int32 OutputIndex, UE::HLSLTree::FExpression*& OutExpression)
 {
 	OutExpression = Generator.AcquireFunctionInputExpression(Scope, this);
 	if (OutExpression)
 	{
-		return EMaterialGenerateHLSLStatus::Success;
+		return true;
 	}
-	return Generator.Error("Invalid material function input");
+	return Generator.GetErrors().AddError(TEXT("Invalid material function input"));
 }
 
-EMaterialGenerateHLSLStatus UMaterialExpressionMaterialFunctionCall::GenerateHLSLExpression(FMaterialHLSLGenerator& Generator, UE::HLSLTree::FScope& Scope, int32 OutputIndex, UE::HLSLTree::FExpression*& OutExpression)
+bool UMaterialExpressionMaterialFunctionCall::GenerateHLSLExpression(FMaterialHLSLGenerator& Generator, UE::HLSLTree::FScope& Scope, int32 OutputIndex, UE::HLSLTree::FExpression*& OutExpression)
 {
 	OutExpression = Generator.GenerateFunctionCall(Scope, MaterialFunction, FunctionInputs, OutputIndex);
 	if (OutExpression)
 	{
-		return EMaterialGenerateHLSLStatus::Success;
+		return true;
 	}
-	return EMaterialGenerateHLSLStatus::Error;
+	return false;
 }
 
-EMaterialGenerateHLSLStatus UMaterialExpressionExecBegin::GenerateHLSLStatements(FMaterialHLSLGenerator& Generator, UE::HLSLTree::FScope& Scope)
+bool UMaterialExpressionExecBegin::GenerateHLSLStatements(FMaterialHLSLGenerator& Generator, UE::HLSLTree::FScope& Scope)
 {
 	Exec.GenerateHLSLStatements(Generator, Scope);
-	return EMaterialGenerateHLSLStatus::Success;
+	return true;
 }
 
-EMaterialGenerateHLSLStatus UMaterialExpressionExecEnd::GenerateHLSLStatements(FMaterialHLSLGenerator& Generator, UE::HLSLTree::FScope& Scope)
+bool UMaterialExpressionExecEnd::GenerateHLSLStatements(FMaterialHLSLGenerator& Generator, UE::HLSLTree::FScope& Scope)
 {
 	Generator.GenerateResult(Scope);
-	return EMaterialGenerateHLSLStatus::Success;
+	return true;
 }
 
-EMaterialGenerateHLSLStatus UMaterialExpressionSetLocal::GenerateHLSLStatements(FMaterialHLSLGenerator& Generator, UE::HLSLTree::FScope& Scope)
+bool UMaterialExpressionSetLocal::GenerateHLSLStatements(FMaterialHLSLGenerator& Generator, UE::HLSLTree::FScope& Scope)
 {
 	UE::HLSLTree::FExpression* ValueExpression = Value.AcquireHLSLExpression(Generator, Scope);
 	if (!ValueExpression)
 	{
-		return Generator.Error(TEXT("Missing value connection"));
+		return false;
 	}
 
 	Generator.GetTree().AssignLocal(Scope, LocalName, ValueExpression);
 	Exec.GenerateHLSLStatements(Generator, Scope);
-	return EMaterialGenerateHLSLStatus::Success;
+	return true;
 }
 
-EMaterialGenerateHLSLStatus UMaterialExpressionIfThenElse::GenerateHLSLStatements(FMaterialHLSLGenerator& Generator, UE::HLSLTree::FScope& Scope)
+bool UMaterialExpressionIfThenElse::GenerateHLSLStatements(FMaterialHLSLGenerator& Generator, UE::HLSLTree::FScope& Scope)
 {
 	UE::HLSLTree::FExpression* ConditionExpression = Condition.AcquireHLSLExpression(Generator, Scope);
 	if (!ConditionExpression)
 	{
-		return Generator.Error(TEXT("Missing condition connection"));
+		return false;
 	}
 
 	UE::HLSLTree::FStatementIf* IfStatement = Generator.GetTree().NewStatement<UE::HLSLTree::FStatementIf>(Scope);
@@ -430,20 +422,20 @@ EMaterialGenerateHLSLStatus UMaterialExpressionIfThenElse::GenerateHLSLStatement
 	IfStatement->ThenScope = Then.NewOwnedScopeWithStatements(Generator, *IfStatement);
 	IfStatement->ElseScope = Else.NewOwnedScopeWithStatements(Generator, *IfStatement);
 
-	return EMaterialGenerateHLSLStatus::Success;
+	return true;
 }
 
-EMaterialGenerateHLSLStatus UMaterialExpressionWhileLoop::GenerateHLSLStatements(FMaterialHLSLGenerator& Generator, UE::HLSLTree::FScope& Scope)
+bool UMaterialExpressionWhileLoop::GenerateHLSLStatements(FMaterialHLSLGenerator& Generator, UE::HLSLTree::FScope& Scope)
 {
 	using namespace UE::HLSLTree;
 	if (!Condition.IsConnected())
 	{
-		return Generator.Error(TEXT("Missing condition connection"));
+		return Generator.GetErrors().AddError(TEXT("Missing condition connection"));
 	}
 
 	if (!LoopBody.GetExpression())
 	{
-		return Generator.Error(TEXT("Missing LoopBody connection"));
+		return Generator.GetErrors().AddError(TEXT("Missing LoopBody connection"));
 	}
 
 	FStatementLoop* LoopStatement = Generator.GetTree().NewStatement<FStatementLoop>(Scope);
@@ -462,7 +454,7 @@ EMaterialGenerateHLSLStatus UMaterialExpressionWhileLoop::GenerateHLSLStatements
 	LoopBody.GenerateHLSLStatements(Generator, *IfStatement->ThenScope);
 	Completed.GenerateHLSLStatements(Generator, *LoopStatement->NextScope);
 
-	return EMaterialGenerateHLSLStatus::Success;
+	return true;
 }
 
 struct FGlobalExpressionDataForLoop
@@ -478,43 +470,43 @@ struct FExpressionDataForLoop
 };
 DECLARE_MATERIAL_HLSLGENERATOR_DATA(FExpressionDataForLoop);
 
-EMaterialGenerateHLSLStatus UMaterialExpressionForLoop::GenerateHLSLExpression(FMaterialHLSLGenerator& Generator, UE::HLSLTree::FScope& Scope, int32 OutputIndex, UE::HLSLTree::FExpression*& OutExpression)
+bool UMaterialExpressionForLoop::GenerateHLSLExpression(FMaterialHLSLGenerator& Generator, UE::HLSLTree::FScope& Scope, int32 OutputIndex, UE::HLSLTree::FExpression*& OutExpression)
 {
 	FExpressionDataForLoop* ExpressionData = Generator.FindExpressionData<FExpressionDataForLoop>(this);
 	if (!ExpressionData || !Scope.HasParentScope(*ExpressionData->LoopScope))
 	{
-		return Generator.Error(TEXT("For loop index accessed outside loop scope"));
+		return Generator.GetErrors().AddError(TEXT("For loop index accessed outside loop scope"));
 	}
 
 	OutExpression = Generator.GetTree().AcquireLocal(Scope, ExpressionData->LocalName);
-	return EMaterialGenerateHLSLStatus::Success;
+	return true;
 }
 
-EMaterialGenerateHLSLStatus UMaterialExpressionForLoop::GenerateHLSLStatements(FMaterialHLSLGenerator& Generator, UE::HLSLTree::FScope& Scope)
+bool UMaterialExpressionForLoop::GenerateHLSLStatements(FMaterialHLSLGenerator& Generator, UE::HLSLTree::FScope& Scope)
 {
 	using namespace UE::HLSLTree;
 	if (!LoopBody.GetExpression())
 	{
-		return Generator.Error(TEXT("Missing LoopBody connection"));
+		return Generator.GetErrors().AddError(TEXT("Missing LoopBody connection"));
 	}
 
 	FExpression* StartExpression = StartIndex.AcquireHLSLExpression(Generator, Scope);
 	if (!StartExpression)
 	{
-		return Generator.Error(TEXT("Missing StartIndex connection"));
+		return false;
 	}
 
 	FExpression* EndExpression = EndIndex.AcquireHLSLExpression(Generator, Scope);
 	if (!EndExpression)
 	{
-		return Generator.Error(TEXT("Missing EndIndex connection"));
+		return false;
 	}
 
 	FGlobalExpressionDataForLoop* GlobalData = Generator.AcquireGlobalData<FGlobalExpressionDataForLoop>();
 	FExpressionDataForLoop* ExpressionData = Generator.NewExpressionData<FExpressionDataForLoop>(this);
 	ExpressionData->LocalName = *FString::Printf(TEXT("ForLoopControl%d"), GlobalData->NumLoops++);
 
-	FExpression* StepExpression = IndexStep.GetTracedInput().Expression ? IndexStep.AcquireHLSLExpression(Generator, Scope) : Generator.NewConstant(1);
+	FExpression* StepExpression = IndexStep.AcquireHLSLExpressionOrConstant(Generator, Scope, int32(1));
 
 	Generator.GetTree().AssignLocal(Scope, ExpressionData->LocalName, StartExpression);
 
@@ -541,7 +533,7 @@ EMaterialGenerateHLSLStatus UMaterialExpressionForLoop::GenerateHLSLStatements(F
 
 	Completed.GenerateHLSLStatements(Generator, *LoopStatement->NextScope);
 
-	return EMaterialGenerateHLSLStatus::Success;
+	return true;
 }
 
 #endif // WITH_EDITOR
diff --git a/Engine/Source/Runtime/Engine/Private/Materials/MaterialHLSLEmitter.cpp b/Engine/Source/Runtime/Engine/Private/Materials/MaterialHLSLEmitter.cpp
--- a/Engine/Source/Runtime/Engine/Private/Materials/MaterialHLSLEmitter.cpp
+++ b/Engine/Source/Runtime/Engine/Private/Materials/MaterialHLSLEmitter.cpp
@@ -714,38 +714,16 @@ bool MaterialEmitHLSL(const FMaterialCompileTargetParameters& InCompilerTarget,
 	SharedPixelProperties[MP_FrontMaterial] = true;
 
 	FMemStackBase Allocator;
-
-	UMaterialInterface* MaterialInterface = InOutMaterial.GetMaterialInterface();
-	UMaterial* BaseMaterial = MaterialInterface->GetMaterial();
-
 	FStructTypeRegistry TypeRegistry(Allocator);
 	FTree* HLSLTree = FTree::Create(Allocator);
-	FMaterialHLSLGenerator Generator(BaseMaterial, InCompilerTarget, TypeRegistry, *HLSLTree);
-	FScope& RootScope = HLSLTree->GetRootScope();
 
-	bool bGenerateResult = false;
-	if (BaseMaterial->IsCompiledWithExecutionFlow())
-	{
-		UMaterialExpression* BaseExpression = BaseMaterial->ExpressionExecBegin;
-		bGenerateResult = Generator.GenerateStatements(RootScope, BaseExpression);
-	}
-	else
-	{
-		bGenerateResult = Generator.GenerateResult(RootScope);
-	}
-
-	if (!Generator.Finalize())
-	{
-		bGenerateResult = false;
-	}
-
-	Generator.AcquireErrors(InOutMaterial);
-	if (!bGenerateResult)
+	FMaterialHLSLGenerator Generator(InCompilerTarget, InOutMaterial, TypeRegistry, *HLSLTree);
+	if (!Generator.Generate())
 	{
 		return false;
 	}
 
-	FEmitContext EmitContext(Allocator, TypeRegistry);
+	FEmitContext EmitContext(Allocator, Generator.GetErrors(), TypeRegistry);
 	EmitContext.Material = &InOutMaterial;
 	EmitContext.StaticParameters = &InStaticParameters;
 	EmitContext.MaterialCompilationOutput = &OutCompilationOutput;
@@ -761,7 +739,11 @@ bool MaterialEmitHLSL(const FMaterialCompileTargetParameters& InCompilerTarget,
 	// Prepare all fields *except* normal
 	FRequestedType RequestedMaterialAttributesType(Generator.GetMaterialAttributesType());
 	RequestedMaterialAttributesType.ClearFieldRequested(NormalField);
-	PrepareExpressionValue(EmitContext, Generator.GetResultExpression(), RequestedMaterialAttributesType);
+	if (PrepareExpressionValue(EmitContext, Generator.GetResultExpression(), RequestedMaterialAttributesType).IsVoid())
+	{
+		return false;
+	}
+
 
 	const bool bReadMaterialNormal = EmitContext.bReadMaterialNormal;
 	FStringBuilderMemstack CodePhase0(Allocator, 1024 * 1024);
@@ -770,7 +752,11 @@ bool MaterialEmitHLSL(const FMaterialCompileTargetParameters& InCompilerTarget,
 	{
 		// No access to material normal, can execute everything in phase0
 		RequestedMaterialAttributesType.SetFieldRequested(NormalField);
-		PrepareExpressionValue(EmitContext, Generator.GetResultExpression(), RequestedMaterialAttributesType);
+		if (PrepareExpressionValue(EmitContext, Generator.GetResultExpression(), RequestedMaterialAttributesType).IsVoid())
+		{
+			return false;
+		}
+
 		HLSLTree->EmitShader(EmitContext, CodePhase0);
 	}
 	else
@@ -788,12 +774,20 @@ bool MaterialEmitHLSL(const FMaterialCompileTargetParameters& InCompilerTarget,
 		// Prepare code for just the normal
 		FRequestedType RequestedMaterialNormal(Generator.GetMaterialAttributesType(), false);
 		RequestedMaterialNormal.SetFieldRequested(NormalField);
-		PrepareExpressionValue(EmitContext, Generator.GetResultExpression(), RequestedMaterialNormal);
+		if (PrepareExpressionValue(EmitContext, Generator.GetResultExpression(), RequestedMaterialNormal).IsVoid())
+		{
+			return false;
+		}
 
 		// Execute the normal in phase0
 		HLSLTree->EmitShader(EmitContext, CodePhase0);
 	}
 
+	if (EmitContext.Errors->HasErrors())
+	{
+		return false;
+	}
+
 	FStringBuilderMemstack Declarations(Allocator, 64 * 1024);
 	TypeRegistry.EmitDeclarationsCode(Declarations);
 
diff --git a/Engine/Source/Runtime/Engine/Private/Materials/MaterialHLSLGenerator.cpp b/Engine/Source/Runtime/Engine/Private/Materials/MaterialHLSLGenerator.cpp
--- a/Engine/Source/Runtime/Engine/Private/Materials/MaterialHLSLGenerator.cpp
+++ b/Engine/Source/Runtime/Engine/Private/Materials/MaterialHLSLGenerator.cpp
@@ -9,6 +9,7 @@
 #include "Materials/MaterialExpressionFunctionInput.h"
 #include "Materials/MaterialExpressionFunctionOutput.h"
 #include "Materials/MaterialExpressionVolumetricAdvancedMaterialOutput.h"
+#include "Materials/MaterialExpressionExecBegin.h"
 #include "Materials/Material.h"
 #include "ShaderCore.h"
 #include "HLSLTree/HLSLTree.h"
@@ -35,17 +36,62 @@ static UE::Shader::EValueType GetShaderType(EMaterialValueType MaterialType)
 	}
 }
 
-FMaterialHLSLGenerator::FMaterialHLSLGenerator(UMaterial* InTargetMaterial,
-	const FMaterialCompileTargetParameters& InCompileTarget,
+FMaterialHLSLErrorHandler::FMaterialHLSLErrorHandler(FMaterial& InOutMaterial)
+	: Material(&InOutMaterial)
+{
+	Material->CompileErrors.Reset();
+	Material->ErrorExpressions.Reset();
+}
+
+void FMaterialHLSLErrorHandler::AddErrorInternal(UObject* InOwner, FStringView InError)
+{
+	UMaterialExpression* MaterialExpressionOwner = Cast<UMaterialExpression>(InOwner);
+	UMaterialExpression* ExpressionToError = nullptr;
+	TStringBuilder<1024> FormattedError;
+
+	if (MaterialExpressionOwner)
+	{
+		if (MaterialExpressionOwner->GetClass() != UMaterialExpressionMaterialFunctionCall::StaticClass()
+			&& MaterialExpressionOwner->GetClass() != UMaterialExpressionFunctionInput::StaticClass()
+			&& MaterialExpressionOwner->GetClass() != UMaterialExpressionFunctionOutput::StaticClass())
+		{
+			// Add the expression currently being compiled to ErrorExpressions so we can draw it differently
+			ExpressionToError = MaterialExpressionOwner;
+
+			const int32 ChopCount = FCString::Strlen(TEXT("MaterialExpression"));
+			const FString ErrorClassName = MaterialExpressionOwner->GetClass()->GetName();
+
+			// Add the node type to the error message
+			FormattedError.Appendf(TEXT("(Node %s) "), *ErrorClassName.Right(ErrorClassName.Len() - ChopCount));
+		}
+	}
+
+	FormattedError.Append(InError);
+	const FString Error(FormattedError.ToView());
+
+	// Standard error handling, immediately append one-off errors and signal failure
+	Material->CompileErrors.AddUnique(Error);
+
+	if (ExpressionToError)
+	{
+		Material->ErrorExpressions.Add(ExpressionToError);
+		ExpressionToError->LastErrorText = Error;
+	}
+}
+
+FMaterialHLSLGenerator::FMaterialHLSLGenerator(const FMaterialCompileTargetParameters& InCompileTarget,
+	FMaterial& InOutMaterial,
 	UE::Shader::FStructTypeRegistry& InOutTypeRegistry,
 	UE::HLSLTree::FTree& InOutTree)
 	: CompileTarget(InCompileTarget)
-	, TargetMaterial(InTargetMaterial)
+	, Errors(InOutMaterial)
 	, HLSLTree(&InOutTree)
 	, TypeRegistry(&InOutTypeRegistry)
 	, bGeneratedResult(false)
 {
-	const EMaterialShadingModel DefaultShadingModel = InTargetMaterial->GetShadingModels().GetFirstShadingModel();
+	UMaterialInterface* MaterialInterface = InOutMaterial.GetMaterialInterface();
+	TargetMaterial = MaterialInterface->GetMaterial();
+	const EMaterialShadingModel DefaultShadingModel = TargetMaterial->GetShadingModels().GetFirstShadingModel();
 
 	FFunctionCallEntry* RootFunctionEntry = new(InOutTree.GetAllocator()) FFunctionCallEntry();
 	FunctionCallStack.Add(RootFunctionEntry);
@@ -86,26 +132,37 @@ FMaterialHLSLGenerator::FMaterialHLSLGenerator(UMaterial* InTargetMaterial,
 	MaterialAttributesDefaultValue.Type = MaterialAttributesType;
 }
 
-void FMaterialHLSLGenerator::AcquireErrors(FMaterial& OutMaterial)
+bool FMaterialHLSLGenerator::Generate()
 {
-	OutMaterial.CompileErrors = MoveTemp(CompileErrors);
-	OutMaterial.ErrorExpressions = MoveTemp(ErrorExpressions);
-}
+	using namespace UE::HLSLTree;
 
-bool FMaterialHLSLGenerator::Finalize()
-{
-	check(FunctionCallStack.Num() == 1);
+	FScope& RootScope = HLSLTree->GetRootScope();
 
-	if (!bGeneratedResult)
+	bool bResult = false;
+	if (TargetMaterial->IsCompiledWithExecutionFlow())
+	{
+		UMaterialExpression* BaseExpression = TargetMaterial->ExpressionExecBegin;
+		bResult = GenerateStatements(RootScope, BaseExpression);
+	}
+	else
+	{
+		bResult = GenerateResult(RootScope);
+	}
+
+	if (!bResult)
 	{
-		Error(TEXT("Missing connection to material output"));
 		return false;
 	}
 
+	check(FunctionCallStack.Num() == 1);
+	if (!bGeneratedResult)
+	{
+		return Errors.AddError(TEXT("Missing connection to material output"));
+	}
+
 	if (!ResultExpression || !ResultStatement)
 	{
-		Error(TEXT("Failed to initialize result"));
-		return false;
+		return Errors.AddError(TEXT("Failed to initialize result"));
 	}
 
 	for (const auto& It : StatementMap)
@@ -114,59 +171,18 @@ bool FMaterialHLSLGenerator::Finalize()
 		const FStatementEntry& Entry = It.Value;
 		if (Entry.NumInputs != Expression->NumExecutionInputs)
 		{
-			Error(TEXT("Invalid number of input connections"));
-			return false;
+			return Errors.AddError(TEXT("Invalid number of input connections"));
 		}
 	}
 
 	if (JoinedScopeStack.Num() != 0)
 	{
-		Error(TEXT("Invalid control flow"));
-		return false;
+		return Errors.AddError(TEXT("Invalid control flow"));
 	}
 
 	return HLSLTree->Finalize();
 }
 
-EMaterialGenerateHLSLStatus FMaterialHLSLGenerator::Error(const FString& Message)
-{
-	UMaterialExpression* ExpressionToError = nullptr;
-	FString ErrorString;
-
-	if (ExpressionStack.Num() > 0)
-	{
-		UMaterialExpression* ErrorExpression = ExpressionStack.Last();
-		check(ErrorExpression);
-
-		if (ErrorExpression->GetClass() != UMaterialExpressionMaterialFunctionCall::StaticClass()
-			&& ErrorExpression->GetClass() != UMaterialExpressionFunctionInput::StaticClass()
-			&& ErrorExpression->GetClass() != UMaterialExpressionFunctionOutput::StaticClass())
-		{
-			// Add the expression currently being compiled to ErrorExpressions so we can draw it differently
-			ExpressionToError = ErrorExpression;
-
-			const int32 ChopCount = FCString::Strlen(TEXT("MaterialExpression"));
-			const FString ErrorClassName = ErrorExpression->GetClass()->GetName();
-
-			// Add the node type to the error message
-			ErrorString += FString(TEXT("(Node ")) + ErrorClassName.Right(ErrorClassName.Len() - ChopCount) + TEXT(") ");
-		}
-	}
-
-	ErrorString += Message;
-
-	// Standard error handling, immediately append one-off errors and signal failure
-	CompileErrors.AddUnique(ErrorString);
-
-	if (ExpressionToError)
-	{
-		ErrorExpressions.Add(ExpressionToError);
-		ExpressionToError->LastErrorText = Message;
-	}
-
-	return EMaterialGenerateHLSLStatus::Error;
-}
-
 static UE::HLSLTree::FExpression* CompileMaterialInput(FMaterialHLSLGenerator& Generator,
 	UE::HLSLTree::FScope& Scope,
 	EMaterialProperty InputProperty,
@@ -235,7 +251,7 @@ bool FMaterialHLSLGenerator::GenerateResult(UE::HLSLTree::FScope& Scope)
 	bool bResult = false;
 	if (bGeneratedResult)
 	{
-		Error(TEXT("Multiple connections to execution output"));
+		return Errors.AddError(TEXT("Multiple connections to execution output"));
 	}
 	else
 	{
@@ -327,7 +343,7 @@ UE::HLSLTree::FTextureParameterDeclaration* FMaterialHLSLGenerator::AcquireTextu
 	FString SamplerTypeError;
 	if (!UMaterialExpressionTextureBase::VerifySamplerType(CompileTarget.FeatureLevel, CompileTarget.TargetPlatform, Value.Texture, Value.SamplerType, SamplerTypeError))
 	{
-		Errorf(TEXT("%s"), *SamplerTypeError);
+		Errors.AddErrorf(TEXT("%s"), *SamplerTypeError);
 		return nullptr;
 	}
 
@@ -344,7 +360,7 @@ UE::HLSLTree::FTextureParameterDeclaration* FMaterialHLSLGenerator::AcquireTextu
 	FString SamplerTypeError;
 	if (!UMaterialExpressionTextureBase::VerifySamplerType(CompileTarget.FeatureLevel, CompileTarget.TargetPlatform, DefaultValue.Texture, DefaultValue.SamplerType, SamplerTypeError))
 	{
-		Errorf(TEXT("%s"), *SamplerTypeError);
+		Errors.AddErrorf(TEXT("%s"), *SamplerTypeError);
 		return nullptr;
 	}
 
@@ -358,13 +374,16 @@ UE::HLSLTree::FTextureParameterDeclaration* FMaterialHLSLGenerator::AcquireTextu
 
 UE::HLSLTree::FExpression* FMaterialHLSLGenerator::AcquireExpression(UE::HLSLTree::FScope& Scope, UMaterialExpression* MaterialExpression, int32 OutputIndex)
 {
-	UE::HLSLTree::FExpression* Expression = nullptr;
-
-	ExpressionStack.Add(MaterialExpression);
-	const EMaterialGenerateHLSLStatus Status = MaterialExpression->GenerateHLSLExpression(*this, Scope, OutputIndex, Expression);
-	verify(ExpressionStack.Pop() == MaterialExpression);
+	using namespace UE::HLSLTree;
+	FOwnerScope TreeOwnerScope(GetTree(), MaterialExpression);
+	FOwnerScope ErrorOwnerScope(Errors, MaterialExpression);
 
-	return Expression;
+	FExpression* Expression = nullptr;
+	if (MaterialExpression->GenerateHLSLExpression(*this, Scope, OutputIndex, Expression))
+	{
+		return Expression;
+	}
+	return nullptr;
 }
 
 UE::HLSLTree::FExpression* FMaterialHLSLGenerator::AcquireFunctionInputExpression(UE::HLSLTree::FScope& Scope, const UMaterialExpressionFunctionInput* MaterialExpression)
@@ -390,21 +409,29 @@ UE::HLSLTree::FExpression* FMaterialHLSLGenerator::AcquireFunctionInputExpressio
 
 UE::HLSLTree::FTextureParameterDeclaration* FMaterialHLSLGenerator::AcquireTextureDeclaration(UE::HLSLTree::FScope& Scope, UMaterialExpression* MaterialExpression, int32 OutputIndex)
 {
-	// No need to cache at this level, TextureDeclarations are cached at a lower level, as they're generated by UMaterialExpression
-	UE::HLSLTree::FTextureParameterDeclaration* TextureDeclaration = nullptr;
-	const EMaterialGenerateHLSLStatus Status = MaterialExpression->GenerateHLSLTexture(*this, Scope, OutputIndex, TextureDeclaration);
-	return TextureDeclaration;
+	using namespace UE::HLSLTree;
+	FOwnerScope TreeOwnerScope(GetTree(), MaterialExpression);
+	FOwnerScope ErrorOwnerScope(Errors, MaterialExpression);
+
+	FTextureParameterDeclaration* TextureDeclaration = nullptr;
+	if (MaterialExpression->GenerateHLSLTexture(*this, Scope, OutputIndex, TextureDeclaration))
+	{
+		return TextureDeclaration;
+	}
+	return nullptr;
 }
 
 bool FMaterialHLSLGenerator::GenerateStatements(UE::HLSLTree::FScope& Scope, UMaterialExpression* MaterialExpression)
 {
+	using namespace UE::HLSLTree;
+	FOwnerScope ErrorOwnerScope(Errors, MaterialExpression);
+
 	FStatementEntry& Entry = StatementMap.FindOrAdd(MaterialExpression);
 	check(Entry.NumInputs >= 0);
 	check(Entry.NumInputs < MaterialExpression->NumExecutionInputs);
 	if (Entry.NumInputs == MaxNumPreviousScopes)
 	{
-		Error(TEXT("Bad control flow"));
-		return false;
+		return Errors.AddError(TEXT("Bad control flow"));
 	}
 
 	Entry.PreviousScope[Entry.NumInputs++] = &Scope;
@@ -412,13 +439,12 @@ bool FMaterialHLSLGenerator::GenerateStatements(UE::HLSLTree::FScope& Scope, UMa
 	bool bResult = true;
 	if (Entry.NumInputs == MaterialExpression->NumExecutionInputs)
 	{
-		UE::HLSLTree::FScope* ScopeToUse = &Scope;
+		FScope* ScopeToUse = &Scope;
 		if (MaterialExpression->NumExecutionInputs > 1u)
 		{
 			if (JoinedScopeStack.Num() == 0)
 			{
-				Error(TEXT("Bad control flow"));
-				return false;
+				return Errors.AddError(TEXT("Bad control flow"));
 			}
 
 			ScopeToUse = JoinedScopeStack.Pop(false);
@@ -428,9 +454,8 @@ bool FMaterialHLSLGenerator::GenerateStatements(UE::HLSLTree::FScope& Scope, UMa
 			}
 		}
 
-		ExpressionStack.Add(MaterialExpression);
-		const EMaterialGenerateHLSLStatus Status = MaterialExpression->GenerateHLSLStatements(*this, *ScopeToUse);
-		verify(ExpressionStack.Pop() == MaterialExpression);
+		FOwnerScope TreeOwnerScope(GetTree(), MaterialExpression);
+		bResult = MaterialExpression->GenerateHLSLStatements(*this, *ScopeToUse);
 	}
 
 	return bResult;
@@ -442,7 +467,7 @@ UE::HLSLTree::FExpression* FMaterialHLSLGenerator::GenerateFunctionCall(UE::HLSL
 
 	if (!Function)
 	{
-		Error(TEXT("Missing material function"));
+		Errors.AddError(TEXT("Missing material function"));
 		return nullptr;
 	}
 
@@ -452,14 +477,14 @@ UE::HLSLTree::FExpression* FMaterialHLSLGenerator::GenerateFunctionCall(UE::HLSL
 
 	if (FunctionInputs.Num() != ConnectedInputs.Num())
 	{
-		Error(TEXT("Mismatched function inputs"));
+		Errors.AddError(TEXT("Mismatched function inputs"));
 		return nullptr;
 	}
 
 	const UMaterialExpressionFunctionOutput* ExpressionOutput = FunctionOutputs.IsValidIndex(OutputIndex) ? FunctionOutputs[OutputIndex].ExpressionOutput.Get() : nullptr;
 	if (!ExpressionOutput)
 	{
-		Error(TEXT("Invalid function output"));
+		Errors.AddError(TEXT("Invalid function output"));
 		return nullptr;
 	}
 
diff --git a/Engine/Source/Runtime/Engine/Private/Materials/MaterialShared.cpp b/Engine/Source/Runtime/Engine/Private/Materials/MaterialShared.cpp
--- a/Engine/Source/Runtime/Engine/Private/Materials/MaterialShared.cpp
+++ b/Engine/Source/Runtime/Engine/Private/Materials/MaterialShared.cpp
@@ -272,7 +272,7 @@ int32 FExpressionInput::Compile(class FMaterialCompiler* Compiler)
 		return INDEX_NONE;
 }
 
-UE::HLSLTree::FExpression* FExpressionInput::AcquireHLSLExpression(FMaterialHLSLGenerator& Generator, UE::HLSLTree::FScope& Scope) const
+UE::HLSLTree::FExpression* FExpressionInput::TryAcquireHLSLExpression(FMaterialHLSLGenerator& Generator, UE::HLSLTree::FScope& Scope) const
 {
 	UE::HLSLTree::FExpression* Result = nullptr;
 	if (Expression)
@@ -285,14 +285,40 @@ UE::HLSLTree::FExpression* FExpressionInput::AcquireHLSLExpression(FMaterialHLSL
 			Result = Generator.NewSwizzle(SwizzleParams, Result);
 		}
 	}
-	
+
 	return Result;
 }
 
+UE::HLSLTree::FExpression* FExpressionInput::AcquireHLSLExpression(FMaterialHLSLGenerator& Generator, UE::HLSLTree::FScope& Scope) const
+{
+	const FExpressionInput TracedInput = GetTracedInput();
+	if (!TracedInput.Expression)
+	{
+		Generator.GetErrors().AddError(TEXT("Missing input"));
+		return nullptr;
+	}
+	return TryAcquireHLSLExpression(Generator, Scope);
+}
+
+UE::HLSLTree::FExpression* FExpressionInput::AcquireHLSLExpressionOrConstant(FMaterialHLSLGenerator& Generator, UE::HLSLTree::FScope& Scope, const UE::Shader::FValue& ConstantValue) const
+{
+	const FExpressionInput TracedInput = GetTracedInput();
+	if (!TracedInput.Expression)
+	{
+		return Generator.NewConstant(ConstantValue);
+	}
+	return TryAcquireHLSLExpression(Generator, Scope);
+}
+
 UE::HLSLTree::FTextureParameterDeclaration* FExpressionInput::AcquireHLSLTexture(FMaterialHLSLGenerator& Generator, UE::HLSLTree::FScope& Scope) const
 {
 	UE::HLSLTree::FTextureParameterDeclaration* Result = nullptr;
-	if (Expression)
+	const FExpressionInput TracedInput = GetTracedInput();
+	if (!TracedInput.Expression)
+	{
+		Generator.GetErrors().AddError(TEXT("Missing input"));
+	}
+	else if (Expression)
 	{
 		Expression->ValidateState();
 		Result = Generator.AcquireTextureDeclaration(Scope, Expression, OutputIndex);
diff --git a/Engine/Source/Runtime/Engine/Private/Shader/Preshader.cpp b/Engine/Source/Runtime/Engine/Private/Shader/Preshader.cpp
--- a/Engine/Source/Runtime/Engine/Private/Shader/Preshader.cpp
+++ b/Engine/Source/Runtime/Engine/Private/Shader/Preshader.cpp
@@ -165,12 +165,15 @@ void FPreshaderData::WriteType(const FType& Type)
 
 void FPreshaderData::WriteValue(const FValue& Value)
 {
+	const int32 NumComponents = Value.Type.GetNumComponents();
+
 	WriteType(Value.Type);
-	for (int32 Index = 0; Index < Value.Component.Num(); ++Index)
+	for (int32 Index = 0; Index < NumComponents; ++Index)
 	{
 		const EValueComponentType ComponentType = Value.Type.GetComponentType(Index);
 		const int32 ComponentSize = GetComponentTypeSizeInBytes(ComponentType);
-		WriteData(&Value.Component[Index].Packed, ComponentSize);
+		const FValueComponent Component = Value.TryGetComponent(Index);
+		WriteData(&Component.Packed, ComponentSize);
 	}
 }
 
diff --git a/Engine/Source/Runtime/Engine/Private/Shader/ShaderValue.cpp b/Engine/Source/Runtime/Engine/Private/Shader/ShaderValue.cpp
--- a/Engine/Source/Runtime/Engine/Private/Shader/ShaderValue.cpp
+++ b/Engine/Source/Runtime/Engine/Private/Shader/ShaderValue.cpp
@@ -230,12 +230,22 @@ void Cast(const Operation& Op, const FValue& Value, FValue& OutResult)
 		OutResult.Component.Add(Op(ValueTypeDesc.ComponentType, Value.Component[i]));
 	}
 
-	if (NumCopyComponents == 1 && ResultTypeDesc.NumComponents > 1)
+	if (NumCopyComponents < ResultTypeDesc.NumComponents)
 	{
-		const FValueComponent Component = OutResult.Component[0];
-		for (int32 i = 1; i < ResultTypeDesc.NumComponents; ++i)
+		if (NumCopyComponents == 1)
 		{
-			OutResult.Component.Add(Component);
+			const FValueComponent Component = OutResult.Component[0];
+			for (int32 i = 1; i < ResultTypeDesc.NumComponents; ++i)
+			{
+				OutResult.Component.Add(Component);
+			}
+		}
+		else
+		{
+			for (int32 i = NumCopyComponents; i < ResultTypeDesc.NumComponents; ++i)
+			{
+				OutResult.Component.AddDefaulted();
+			}
 		}
 	}
 }
diff --git a/Engine/Source/Runtime/Engine/Public/HLSLTree/HLSLTree.h b/Engine/Source/Runtime/Engine/Public/HLSLTree/HLSLTree.h
--- a/Engine/Source/Runtime/Engine/Public/HLSLTree/HLSLTree.h
+++ b/Engine/Source/Runtime/Engine/Public/HLSLTree/HLSLTree.h
@@ -39,34 +39,68 @@ struct FEmitShaderCode;
 
 static constexpr int32 MaxNumPreviousScopes = 2;
 
-struct FError
+class FOwnerContext
 {
-	const FError* Next;
-	const FNode* Node;
-	int32 MessageLength;
-	TCHAR Message[1];
+public:
+	virtual ~FOwnerContext() {}
+
+	void PushOwner(UObject* Owner);
+	UObject* PopOwner();
+	UObject* GetCurrentOwner() const;
+
+private:
+	TArray<UObject*, TInlineAllocator<8>> OwnerStack;
 };
 
-class FErrors
+struct FOwnerScope
+{
+	FOwnerScope(FOwnerContext& InContext, UObject* InOwner) : Context(&InContext), Owner(InOwner)
+	{
+		Context->PushOwner(InOwner);
+	}
+
+	~FOwnerScope()
+	{
+		verify(Context->PopOwner() == Owner);
+	}
+
+	FOwnerContext* Context;
+	UObject* Owner;
+};
+
+class FErrorHandlerInterface : public FOwnerContext
 {
 public:
-	explicit FErrors(FMemStackBase& InAllocator);
+	virtual ~FErrorHandlerInterface() {}
 
-	int32 Num() const { return NumErrors; }
+	int32 GetNumErrors() const { return NumErrors; }
+	inline bool HasErrors() const { return NumErrors > 0; }
 
-	void AddError(const FNode* InNode, FStringView InError);
+	template<typename StringType>
+	inline bool AddError(const StringType& InError)
+	{
+		return AddErrorInternalWrapper(FStringView(InError));
+	}
 
 	template<typename FormatType, typename... Types>
-	void AddErrorf(const FNode* InNode, const FormatType& Format, Types... Args)
+	inline bool AddErrorf(const FormatType& Format, Types... Args)
 	{
 		TStringBuilder<2048> String;
 		String.Appendf(Format, Forward<Types>(Args)...);
-		AddError(InNode, FStringView(String.ToString(), String.Len()));
+		return AddErrorInternalWrapper(FStringView(String.ToString(), String.Len()));
 	}
 
+protected:
+	bool AddErrorInternalWrapper(FStringView InError)
+	{
+		NumErrors++;
+		AddErrorInternal(GetCurrentOwner(), InError);
+		return false;
+	}
+
+	virtual void AddErrorInternal(UObject* InOwner, FStringView InError) = 0;
+
 private:
-	FMemStackBase* Allocator = nullptr;
-	const FError* FirstError = nullptr;
 	int32 NumErrors = 0;
 };
 
@@ -136,7 +170,7 @@ inline void BuildFormatArgList(FFormatArgList& OutList, Type Arg, Types... Args)
 class FEmitContext
 {
 public:
-	explicit FEmitContext(FMemStackBase& InAllocator, const Shader::FStructTypeRegistry& InTypeRegistry);
+	explicit FEmitContext(FMemStackBase& InAllocator, FErrorHandlerInterface& InErrors, const Shader::FStructTypeRegistry& InTypeRegistry);
 	~FEmitContext();
 
 	void Finalize();
@@ -169,12 +203,13 @@ public:
 	FEmitShaderCode* EmitCast(FEmitShaderCode* ShaderValue, const Shader::FType& DestType);
 	
 	FMemStackBase* Allocator = nullptr;
+	FErrorHandlerInterface* Errors = nullptr;
 	const Shader::FStructTypeRegistry* TypeRegistry = nullptr;
+
 	TArray<FScope*, TInlineAllocator<16>> ScopeStack;
 	TMap<FSHAHash, FEmitShaderCode*> ShaderValueMap;
 	TMap<FSHAHash, FEmitShaderCode*> PreshaderValueMap;
 	TArray<const FExpressionLocalPHI*> LocalPHIs;
-	FErrors Errors;
 
 	// TODO - remove preshader material dependency
 	const FMaterial* Material = nullptr;
@@ -194,7 +229,17 @@ class FNode
 public:
 	virtual ~FNode() {}
 	virtual void Reset() {}
+
+	UObject* GetOwner() const { return Owner; }
+
+private:
+	/** Next node in the FTree's list of all nodes */
 	FNode* NextNode = nullptr;
+
+	/** The UObject that owns/created this node, allows errors generated by this node to be properly associated with the owner */
+	UObject* Owner = nullptr;
+
+	friend class FTree;
 };
 
 /**
@@ -315,10 +360,12 @@ class FPrepareValueResult
 public:
 	const FPreparedType& GetPreparedType() const { return PreparedType; }
 
-	void SetType(FEmitContext& Context, const FRequestedType& RequestedType, EExpressionEvaluation Evaluation, const Shader::FType& Type);
-	void SetType(FEmitContext& Context, const FRequestedType& RequestedType, const FPreparedType& Type);
+	bool SetTypeVoid();
+
+	bool SetType(FEmitContext& Context, const FRequestedType& RequestedType, EExpressionEvaluation Evaluation, const Shader::FType& Type);
+	bool SetType(FEmitContext& Context, const FRequestedType& RequestedType, const FPreparedType& Type);
 
-	void SetForwardValue(FEmitContext& Context, const FRequestedType& RequestedType, FExpression* InValue);
+	bool SetForwardValue(FEmitContext& Context, const FRequestedType& RequestedType, FExpression* InValue);
 
 private:
 	bool TryMergePreparedType(FEmitContext& Context, const Shader::FStructType* StructType, Shader::EValueComponentType ComponentType);
@@ -377,7 +424,7 @@ public:
 
 protected:
 	virtual void ComputeAnalyticDerivatives(FTree& Tree, FExpressionDerivatives& OutResult) const;
-	virtual void PrepareValue(FEmitContext& Context, const FRequestedType& RequestedType, FPrepareValueResult& OutResult) const = 0;
+	virtual bool PrepareValue(FEmitContext& Context, const FRequestedType& RequestedType, FPrepareValueResult& OutResult) const = 0;
 	virtual void EmitValueShader(FEmitContext& Context, const FRequestedType& RequestedType, FEmitValueShaderResult& OutResult) const;
 	virtual void EmitValuePreshader(FEmitContext& Context, const FRequestedType& RequestedType, Shader::FPreshaderData& OutPreshader) const;
 
@@ -430,7 +477,7 @@ public:
 	FExpressionLocalPHI(const FExpressionLocalPHI* Source, EDerivativeCoordinate Coord);
 
 	virtual void ComputeAnalyticDerivatives(FTree& Tree, FExpressionDerivatives& OutResult) const override;
-	virtual void PrepareValue(FEmitContext& Context, const FRequestedType& RequestedType, FPrepareValueResult& OutResult) const override;
+	virtual bool PrepareValue(FEmitContext& Context, const FRequestedType& RequestedType, FPrepareValueResult& OutResult) const override;
 	virtual void EmitValueShader(FEmitContext& Context, const FRequestedType& RequestedType, FEmitValueShaderResult& OutResult) const override;
 
 	TArray<EDerivativeCoordinate, TInlineAllocator<8>> DerivativeChain;
@@ -587,7 +634,7 @@ inline void MarkScopeDead(FScope* InScope)
 /**
  * The HLSL AST.  Basically a wrapper around the root scope, with some helper methods
  */
-class FTree
+class FTree : public FOwnerContext
 {
 public:
 	static FTree* Create(FMemStackBase& Allocator);
@@ -648,11 +695,11 @@ private:
 	inline T* NewNode(ArgTypes&&... Args)
 	{
 		T* Node = new(*Allocator) T(Forward<ArgTypes>(Args)...);
-		Node->NextNode = Nodes;
-		Nodes = Node;
+		RegisterNode(Node);
 		return Node;
 	}
 
+	void RegisterNode(FNode* Node);
 	void RegisterExpression(FExpression* Expression);
 	void RegisterStatement(FScope& Scope, FStatement* Statement);
 
@@ -660,6 +707,7 @@ private:
 	FNode* Nodes = nullptr;
 	FScope* RootScope = nullptr;
 	TArray<FExpressionLocalPHI*> PHIExpressions;
+	TArray<UObject*, TInlineAllocator<8>> OwnerStack;
 
 	friend class FExpressionLocalPHI;
 };
diff --git a/Engine/Source/Runtime/Engine/Public/HLSLTree/HLSLTreeCommon.h b/Engine/Source/Runtime/Engine/Public/HLSLTree/HLSLTreeCommon.h
--- a/Engine/Source/Runtime/Engine/Public/HLSLTree/HLSLTreeCommon.h
+++ b/Engine/Source/Runtime/Engine/Public/HLSLTree/HLSLTreeCommon.h
@@ -21,7 +21,7 @@ public:
 	Shader::FValue Value;
 
 	virtual void ComputeAnalyticDerivatives(FTree& Tree, FExpressionDerivatives& OutResult) const override;
-	virtual void PrepareValue(FEmitContext& Context, const FRequestedType& RequestedType, FPrepareValueResult& OutResult) const override;
+	virtual bool PrepareValue(FEmitContext& Context, const FRequestedType& RequestedType, FPrepareValueResult& OutResult) const override;
 	virtual void EmitValuePreshader(FEmitContext& Context, const FRequestedType& RequestedType, Shader::FPreshaderData& OutPreshader) const override;
 };
 
@@ -37,7 +37,7 @@ public:
 	EMaterialParameterType ParameterType;
 
 	virtual void ComputeAnalyticDerivatives(FTree& Tree, FExpressionDerivatives& OutResult) const override;
-	virtual void PrepareValue(FEmitContext& Context, const FRequestedType& RequestedType, FPrepareValueResult& OutResult) const override;
+	virtual bool PrepareValue(FEmitContext& Context, const FRequestedType& RequestedType, FPrepareValueResult& OutResult) const override;
 	virtual void EmitValuePreshader(FEmitContext& Context, const FRequestedType& RequestedType, Shader::FPreshaderData& OutPreshader) const override;
 };
 
@@ -122,7 +122,7 @@ public:
 	EExternalInput InputType;
 
 	virtual void ComputeAnalyticDerivatives(FTree& Tree, FExpressionDerivatives& OutResult) const override;
-	virtual void PrepareValue(FEmitContext& Context, const FRequestedType& RequestedType, FPrepareValueResult& OutResult) const override;
+	virtual bool PrepareValue(FEmitContext& Context, const FRequestedType& RequestedType, FPrepareValueResult& OutResult) const override;
 	virtual void EmitValueShader(FEmitContext& Context, const FRequestedType& RequestedType, FEmitValueShaderResult& OutResult) const override;
 };
 
@@ -142,7 +142,7 @@ public:
 	ESamplerSourceMode SamplerSource;
 	ETextureMipValueMode MipValueMode;
 
-	virtual void PrepareValue(FEmitContext& Context, const FRequestedType& RequestedType, FPrepareValueResult& OutResult) const override;
+	virtual bool PrepareValue(FEmitContext& Context, const FRequestedType& RequestedType, FPrepareValueResult& OutResult) const override;
 	virtual void EmitValueShader(FEmitContext& Context, const FRequestedType& RequestedType, FEmitValueShaderResult& OutResult) const override;
 };
 
@@ -162,7 +162,7 @@ public:
 	FExpression* StructExpression;
 
 	virtual void ComputeAnalyticDerivatives(FTree& Tree, FExpressionDerivatives& OutResult) const override;
-	virtual void PrepareValue(FEmitContext& Context, const FRequestedType& RequestedType, FPrepareValueResult& OutResult) const override;
+	virtual bool PrepareValue(FEmitContext& Context, const FRequestedType& RequestedType, FPrepareValueResult& OutResult) const override;
 	virtual void EmitValueShader(FEmitContext& Context, const FRequestedType& RequestedType, FEmitValueShaderResult& OutResult) const override;
 	virtual void EmitValuePreshader(FEmitContext& Context, const FRequestedType& RequestedType, Shader::FPreshaderData& OutPreshader) const override;
 };
@@ -195,7 +195,7 @@ public:
 	FExpression* FieldExpression;
 
 	virtual void ComputeAnalyticDerivatives(FTree& Tree, FExpressionDerivatives& OutResult) const override;
-	virtual void PrepareValue(FEmitContext& Context, const FRequestedType& RequestedType, FPrepareValueResult& OutResult) const override;
+	virtual bool PrepareValue(FEmitContext& Context, const FRequestedType& RequestedType, FPrepareValueResult& OutResult) const override;
 	virtual void EmitValueShader(FEmitContext& Context, const FRequestedType& RequestedType, FEmitValueShaderResult& OutResult) const override;
 	virtual void EmitValuePreshader(FEmitContext& Context, const FRequestedType& RequestedType, Shader::FPreshaderData& OutPreshader) const override;
 };
@@ -213,7 +213,7 @@ public:
 	FExpression* TrueExpression;
 	FExpression* FalseExpression;
 
-	virtual void PrepareValue(FEmitContext& Context, const FRequestedType& RequestedType, FPrepareValueResult& OutResult) const override;
+	virtual bool PrepareValue(FEmitContext& Context, const FRequestedType& RequestedType, FPrepareValueResult& OutResult) const override;
 	virtual void EmitValueShader(FEmitContext& Context, const FRequestedType& RequestedType, FEmitValueShaderResult& OutResult) const override;
 	virtual void EmitValuePreshader(FEmitContext& Context, const FRequestedType& RequestedType, Shader::FPreshaderData& OutPreshader) const override;
 };
@@ -230,7 +230,7 @@ public:
 	FExpression* Input;
 
 	virtual void ComputeAnalyticDerivatives(FTree& Tree, FExpressionDerivatives& OutResult) const override;
-	virtual void PrepareValue(FEmitContext& Context, const FRequestedType& RequestedType, FPrepareValueResult& OutResult) const override;
+	virtual bool PrepareValue(FEmitContext& Context, const FRequestedType& RequestedType, FPrepareValueResult& OutResult) const override;
 	virtual void EmitValueShader(FEmitContext& Context, const FRequestedType& RequestedType, FEmitValueShaderResult& OutResult) const override;
 	virtual void EmitValuePreshader(FEmitContext& Context, const FRequestedType& RequestedType, Shader::FPreshaderData& OutPreshader) const override;
 };
@@ -249,7 +249,7 @@ public:
 	FExpression* Rhs;
 
 	virtual void ComputeAnalyticDerivatives(FTree& Tree, FExpressionDerivatives& OutResult) const override;
-	virtual void PrepareValue(FEmitContext& Context, const FRequestedType& RequestedType, FPrepareValueResult& OutResult) const override;
+	virtual bool PrepareValue(FEmitContext& Context, const FRequestedType& RequestedType, FPrepareValueResult& OutResult) const override;
 	virtual void EmitValueShader(FEmitContext& Context, const FRequestedType& RequestedType, FEmitValueShaderResult& OutResult) const override;
 	virtual void EmitValuePreshader(FEmitContext& Context, const FRequestedType& RequestedType, Shader::FPreshaderData& OutPreshader) const override;
 };
@@ -279,7 +279,7 @@ public:
 	FExpression* Input;
 
 	virtual void ComputeAnalyticDerivatives(FTree& Tree, FExpressionDerivatives& OutResult) const override;
-	virtual void PrepareValue(FEmitContext& Context, const FRequestedType& RequestedType, FPrepareValueResult& OutResult) const override;
+	virtual bool PrepareValue(FEmitContext& Context, const FRequestedType& RequestedType, FPrepareValueResult& OutResult) const override;
 	virtual void EmitValueShader(FEmitContext& Context, const FRequestedType& RequestedType, FEmitValueShaderResult& OutResult) const override;
 	virtual void EmitValuePreshader(FEmitContext& Context, const FRequestedType& RequestedType, Shader::FPreshaderData& OutPreshader) const override;
 };
@@ -296,7 +296,7 @@ public:
 	FExpression* Rhs;
 
 	virtual void ComputeAnalyticDerivatives(FTree& Tree, FExpressionDerivatives& OutResult) const override;
-	virtual void PrepareValue(FEmitContext& Context, const FRequestedType& RequestedType, FPrepareValueResult& OutResult) const override;
+	virtual bool PrepareValue(FEmitContext& Context, const FRequestedType& RequestedType, FPrepareValueResult& OutResult) const override;
 	virtual void EmitValueShader(FEmitContext& Context, const FRequestedType& RequestedType, FEmitValueShaderResult& OutResult) const override;
 	virtual void EmitValuePreshader(FEmitContext& Context, const FRequestedType& RequestedType, Shader::FPreshaderData& OutPreshader) const override;
 };
@@ -304,7 +304,7 @@ public:
 class FExpressionReflectionVector : public FExpression
 {
 public:
-	virtual void PrepareValue(FEmitContext& Context, const FRequestedType& RequestedType, FPrepareValueResult& OutResult) const override;
+	virtual bool PrepareValue(FEmitContext& Context, const FRequestedType& RequestedType, FPrepareValueResult& OutResult) const override;
 	virtual void EmitValueShader(FEmitContext& Context, const FRequestedType& RequestedType, FEmitValueShaderResult& OutResult) const override;
 };
 
diff --git a/Engine/Source/Runtime/Engine/Public/MaterialExpressionIO.h b/Engine/Source/Runtime/Engine/Public/MaterialExpressionIO.h
--- a/Engine/Source/Runtime/Engine/Public/MaterialExpressionIO.h
+++ b/Engine/Source/Runtime/Engine/Public/MaterialExpressionIO.h
@@ -77,7 +77,9 @@ struct FExpressionInput
 
 #if WITH_EDITOR
 	ENGINE_API int32 Compile(class FMaterialCompiler* Compiler);
+	ENGINE_API UE::HLSLTree::FExpression* TryAcquireHLSLExpression(FMaterialHLSLGenerator& Generator, UE::HLSLTree::FScope& Scope) const;
 	ENGINE_API UE::HLSLTree::FExpression* AcquireHLSLExpression(FMaterialHLSLGenerator& Generator, UE::HLSLTree::FScope& Scope) const;
+	ENGINE_API UE::HLSLTree::FExpression* AcquireHLSLExpressionOrConstant(FMaterialHLSLGenerator& Generator, UE::HLSLTree::FScope& Scope, const UE::Shader::FValue& ConstantValue) const;
 	ENGINE_API UE::HLSLTree::FTextureParameterDeclaration* AcquireHLSLTexture(FMaterialHLSLGenerator& Generator, UE::HLSLTree::FScope& Scope) const;
 #endif // WITH_EDITOR
 
diff --git a/Engine/Source/Runtime/Engine/Public/MaterialHLSLGenerator.h b/Engine/Source/Runtime/Engine/Public/MaterialHLSLGenerator.h
--- a/Engine/Source/Runtime/Engine/Public/MaterialHLSLGenerator.h
+++ b/Engine/Source/Runtime/Engine/Public/MaterialHLSLGenerator.h
@@ -23,7 +23,6 @@ class UMaterialExpressionFunctionInput;
 class UMaterialExpressionFunctionOutput;
 struct FFunctionExpressionInput;
 class ITargetPlatform;
-enum class EMaterialGenerateHLSLStatus : uint8;
 
 namespace UE
 {
@@ -50,33 +49,34 @@ struct TMaterialHLSLGeneratorType;
 #define DECLARE_MATERIAL_HLSLGENERATOR_DATA(T) \
 	template<> struct TMaterialHLSLGeneratorType<T> { static const FName& GetTypeName() { static const FName Name(TEXT(#T)); return Name; } }
 
+class FMaterialHLSLErrorHandler : public UE::HLSLTree::FErrorHandlerInterface
+{
+public:
+	explicit FMaterialHLSLErrorHandler(FMaterial& InOutMaterial);
+
+	virtual void AddErrorInternal(UObject* InOwner, FStringView InError) override;
+
+private:
+	FMaterial* Material;
+};
+
 /**
  * MaterialHLSLGenerator is a bridge between a material, and HLSLTree.  It facilitates generating HLSL source code for a given material, using HLSLTree
  */
 class FMaterialHLSLGenerator
 {
 public:
-	FMaterialHLSLGenerator(UMaterial* InTargetMaterial,
-		const FMaterialCompileTargetParameters& InCompilerTarget,
+	FMaterialHLSLGenerator(const FMaterialCompileTargetParameters& InCompilerTarget,
+		FMaterial& InOutMaterial,
 		UE::Shader::FStructTypeRegistry& InOutTypeRegistry,
 		UE::HLSLTree::FTree& InOutTree);
 
 	const FMaterialCompileTargetParameters& GetCompileTarget() const { return CompileTarget; }
-
-	bool Finalize();
-
-	/** Retrieve the compile errors from the generator */
-	void AcquireErrors(FMaterial& OutMaterial);
-
-	EMaterialGenerateHLSLStatus Error(const FString& Message);
-
-	template <typename FormatType, typename... ArgTypes>
-	inline EMaterialGenerateHLSLStatus Errorf(const FormatType& Format, ArgTypes... Args)
-	{
-		return Error(FString::Printf(Format, Args...));
-	}
+	
+	bool Generate();
 
 	UE::HLSLTree::FTree& GetTree() const { return *HLSLTree; }
+	FMaterialHLSLErrorHandler& GetErrors() { return Errors; }
 
 	UE::HLSLTree::FExpression* GetResultExpression() { return ResultExpression; }
 	UE::HLSLTree::FStatement* GetResultStatement() { return ResultStatement; }
@@ -185,21 +185,19 @@ private:
 	void* InternalFindExpressionData(const FName& Type, const UMaterialExpression* MaterialExpression);
 
 	const FMaterialCompileTargetParameters& CompileTarget;
-	UMaterial* TargetMaterial;
+	UMaterial* TargetMaterial = nullptr;
+	FMaterialHLSLErrorHandler Errors;
 
-	const UE::Shader::FStructType* MaterialAttributesType;
+	const UE::Shader::FStructType* MaterialAttributesType = nullptr;
 	UE::Shader::FValue MaterialAttributesDefaultValue;
 
 	UE::HLSLTree::FTree* HLSLTree;
-	UE::Shader::FStructTypeRegistry* TypeRegistry;
+	UE::Shader::FStructTypeRegistry* TypeRegistry = nullptr;
 	UE::HLSLTree::FExpression* ResultExpression = nullptr;
 	UE::HLSLTree::FStatement* ResultStatement = nullptr;
 
-	TArray<UMaterialExpression*> ExpressionStack;
 	TArray<FFunctionCallEntry*, TInlineAllocator<8>> FunctionCallStack;
 	TArray<UE::HLSLTree::FScope*> JoinedScopeStack;
-	TArray<FString> CompileErrors;
-	TArray<UMaterialExpression*> ErrorExpressions;
 	TMap<UE::HLSLTree::FTextureDescription, UE::HLSLTree::FTextureParameterDeclaration*> TextureDeclarationMap;
 	TMap<FName, UE::HLSLTree::FTextureParameterDeclaration*> TextureParameterDeclarationMap;
 	TMap<FSHAHash, FFunctionCallEntry*> FunctionCallMap;
diff --git a/Engine/Source/Runtime/Engine/Public/MaterialShared.h b/Engine/Source/Runtime/Engine/Public/MaterialShared.h
--- a/Engine/Source/Runtime/Engine/Public/MaterialShared.h
+++ b/Engine/Source/Runtime/Engine/Public/MaterialShared.h
@@ -2295,7 +2295,7 @@ private:
 	friend class FMaterialShaderMap;
 	friend class FShaderCompilingManager;
 	friend class FHLSLMaterialTranslator;
-	friend class FMaterialHLSLGenerator;
+	friend class FMaterialHLSLErrorHandler;
 };
 
 /**
diff --git a/Engine/Source/Runtime/Engine/Public/Shader/Preshader.h b/Engine/Source/Runtime/Engine/Public/Shader/Preshader.h
--- a/Engine/Source/Runtime/Engine/Public/Shader/Preshader.h
+++ b/Engine/Source/Runtime/Engine/Public/Shader/Preshader.h
@@ -119,7 +119,7 @@ public:
 	template<>
 	FPreshaderData& Write<FHashedMaterialParameterInfo>(const FHashedMaterialParameterInfo& Value) { return Write(Value.Name).Write(Value.Index).Write(Value.Association); }
 
-	inline FPreshaderData& WriteOpcode(EPreshaderOpcode Op) { return Write<uint8>((uint8)Op); }
+	inline FPreshaderData& WriteOpcode(EPreshaderOpcode Op) { ensure(Op != EPreshaderOpcode::Nop); return Write<uint8>((uint8)Op); }
 
 	LAYOUT_FIELD(TMemoryImageArray<FScriptName>, Names);
 	LAYOUT_FIELD(TMemoryImageArray<FPreshaderStructType>, StructTypes);
