Fix bitfield

diff --git a/Engine/Source/Runtime/Engine/Classes/Engine/Level.h b/Engine/Source/Runtime/Engine/Classes/Engine/Level.h
--- a/Engine/Source/Runtime/Engine/Classes/Engine/Level.h
+++ b/Engine/Source/Runtime/Engine/Classes/Engine/Level.h
@@ -660,7 +660,7 @@ public:
 	uint8										bIsPartitioned : 1;
 	/** Whether the level is a world partition runtime cell. */
 	UPROPERTY()
-	uint8										bIsWorldPartitionRuntimeCell;
+	uint8										bIsWorldPartitionRuntimeCell : 1;
 
 	enum class EIncrementalComponentState : uint8
 	{
