Changing route actor initialization to use the granular method by default again now that we had a few playtests run cleanly.

diff --git a/Engine/Source/Runtime/Engine/Private/Level.cpp b/Engine/Source/Runtime/Engine/Private/Level.cpp
--- a/Engine/Source/Runtime/Engine/Private/Level.cpp
+++ b/Engine/Source/Runtime/Engine/Private/Level.cpp
@@ -91,7 +91,7 @@ static FAutoConsoleVariableRef CVarActorClusteringEnabled(
 	ECVF_Default
 );
 
-bool GUseLegacyRouteActorInitialization = true;
+bool GUseLegacyRouteActorInitialization = false;
 static FAutoConsoleVariableRef CVarUseLegacyRouteActorInitialization(
 	TEXT("s.UseLegacyRouteActorInitialization"),
 	GUseLegacyRouteActorInitialization,
