Add missing forward declaration

diff --git a/Engine/Source/Runtime/Engine/Public/MaterialExpressionIO.h b/Engine/Source/Runtime/Engine/Public/MaterialExpressionIO.h
--- a/Engine/Source/Runtime/Engine/Public/MaterialExpressionIO.h
+++ b/Engine/Source/Runtime/Engine/Public/MaterialExpressionIO.h
@@ -19,6 +19,7 @@ class FTextureParameterDeclaration;
 namespace Shader
 {
 enum class EValueType : uint8;
+struct FValue;
 }
 }
 
