Remove the parameterization of Dipole code path for Separable SSS.

diff --git a/Engine/Shaders/Private/BurleyNormalizedSSSCommon.ush b/Engine/Shaders/Private/BurleyNormalizedSSSCommon.ush
--- a/Engine/Shaders/Private/BurleyNormalizedSSSCommon.ush
+++ b/Engine/Shaders/Private/BurleyNormalizedSSSCommon.ush
@@ -190,7 +190,7 @@ float4 GetBurleyTransmissionProfile(FGBufferData GBufferData, float Thickness)
 	// 0..255, which SubSurface profile to pick
 
 	uint SubsurfaceProfileInt = ExtractSubsurfaceProfileInt(GBufferData);
-	float WorldUnitScale = DecodeWorldUnitScale(View.SSProfilesTexture.Load(int3(SSSS_SUBSURFACE_COLOR_OFFSET, SubsurfaceProfileInt,0)).a)*10.0f; //in cm. we do not multiply by 100.0f as we did in the subsurface shader.
+	float WorldUnitScale = DecodeWorldUnitScale(View.SSProfilesTexture.Load(int3(SSSS_TINT_SCALE_OFFSET, SubsurfaceProfileInt,0)).a)*10.0f; //in cm. we do not multiply by 100.0f as we did in the subsurface shader.
 	float4 MeanFreePath = DecodeDiffuseMeanFreePath(View.SSProfilesTexture.Load(int3(BSSS_DMFP_OFFSET, SubsurfaceProfileInt, 0)));
 	float3 SurfaceAlbedo = View.SSProfilesTexture.Load(int3(BSSS_SURFACEALBEDO_OFFSET, SubsurfaceProfileInt, 0)).rgb;
 	float3 ScalingFactor = GetSearchLightDiffuseScalingFactor3D(SurfaceAlbedo);//assuming that the volume albedo is the same to the surface albedo for transmission;
@@ -224,7 +224,7 @@ float3 GetBurleyTransmissionProfile(
 	const float3 TransmissionThroughput = InternalGetBurleyTransmissionProfile(SubsurfaceAlebdo, MeanFreePathInCm * TransmissionMFPScaleFactor, ThicknessInCm).xyz;
 
 	// Replicate SSS profile distance fade out, by fading the tail of the curve to black.
-	//  * SSS Profile stores transmittance into a discrete set of samples: SSSS_TRANSMISSION_PROFILE_SIZE(=32)
+	//  * SSS Profile stores transmittance into a discrete set of samples: BSSS_TRANSMISSION_PROFILE_SIZE(=32)
 	//  * The transmission distance is set to a max of SSSS_MAX_TRANSMISSION_PROFILE_DISTANCE(=5) (in scaled cm, depend on the WorldToUnitScale, cm by default)
 	//
 	// This code fades the last value to black.
@@ -237,7 +237,7 @@ float3 GetBurleyTransmissionProfile(
 	// 0             \ _______
 	// ----------------------->
 	// .. 29  30  31  32  Sample Index / Distance
-	const float FadeDistance = 1.0f - saturate(ThicknessInCm * SSSS_TRANSMISSION_PROFILE_SIZE - SSSS_MAX_TRANSMISSION_PROFILE_DISTANCE * (SSSS_TRANSMISSION_PROFILE_SIZE - 0.5f));
+	const float FadeDistance = 1.0f - saturate(ThicknessInCm * BSSS_TRANSMISSION_PROFILE_SIZE - SSSS_MAX_TRANSMISSION_PROFILE_DISTANCE * (BSSS_TRANSMISSION_PROFILE_SIZE - 0.5f));
 
 	return TransmissionThroughput * FadeDistance;
 }
@@ -263,5 +263,5 @@ float4 GetSubsurfaceProfileDiffuseMeanFreePath(uint SubsurfaceProfileInt)
 
 float GetSubsurfaceProfileWorldUnitScale(uint SubsurfaceProfileInt)
 {
-	return View.SSProfilesTexture.Load(int3(SSSS_SUBSURFACE_COLOR_OFFSET, SubsurfaceProfileInt, 0)).a;
+	return View.SSProfilesTexture.Load(int3(SSSS_TINT_SCALE_OFFSET, SubsurfaceProfileInt, 0)).a;
 }
diff --git a/Engine/Shaders/Private/PathTracing/Material/PathTracingMaterialSampling.ush b/Engine/Shaders/Private/PathTracing/Material/PathTracingMaterialSampling.ush
--- a/Engine/Shaders/Private/PathTracing/Material/PathTracingMaterialSampling.ush
+++ b/Engine/Shaders/Private/PathTracing/Material/PathTracingMaterialSampling.ush
@@ -382,7 +382,7 @@ FPathTracingPayload UnpackPathTracingPayloadMaterial(FPackedPathTracingPayload P
 				// Burley parameterization
 				float3 SurfaceAlbedo = View.SSProfilesTexture.Load(int3(BSSS_SURFACEALBEDO_OFFSET, ProfileId, 0)).xyz;
 				float3 DiffuseMeanFreePath = DecodeDiffuseMeanFreePath(View.SSProfilesTexture.Load(int3(BSSS_DMFP_OFFSET, ProfileId, 0))).xyz;
-				float WorldUnitScale = DecodeWorldUnitScale(View.SSProfilesTexture.Load(int3(SSSS_SUBSURFACE_COLOR_OFFSET, ProfileId, 0)).a) * BURLEY_MM_2_CM;
+				float WorldUnitScale = DecodeWorldUnitScale(View.SSProfilesTexture.Load(int3(SSSS_TINT_SCALE_OFFSET, ProfileId, 0)).a) * BURLEY_MM_2_CM;
 
 				// Opacity acts as a per-pixel radius multiplier
 				// NOTE: this seems backwards? Opacity=0 acts like default-lit while Opacity=1 acts like SSS?
diff --git a/Engine/Shaders/Private/PostProcessSubsurface.usf b/Engine/Shaders/Private/PostProcessSubsurface.usf
--- a/Engine/Shaders/Private/PostProcessSubsurface.usf
+++ b/Engine/Shaders/Private/PostProcessSubsurface.usf
@@ -931,18 +931,18 @@ void SubsurfaceRecombinePS(in FScreenVertexOutput Input, out float4 OutColor : S
 	const float PixelRadiusScale	= SSSData.ProfileRadiusScale;
 
 	// asset specific color
-	float3 SubsurfaceColor = 0.0f;
+	float3 Tint = 0.0f;
 	float ProfileRadiusScale = 1.0f;
 	if (SSSData.bIsProfile)
 	{
-		SubsurfaceColor = GetSubsurfaceProfileKernel(SSSS_SUBSURFACE_COLOR_OFFSET, ProfileId).rgb;
+		Tint = GetSubsurfaceProfileKernel(SSSS_TINT_SCALE_OFFSET, ProfileId).rgb;
 		ProfileRadiusScale = GetSubsurfaceProfileKernel(SSSS_N_KERNELWEIGHTOFFSET + SSSS_N_KERNELWEIGHTCOUNT - 1, ProfileId).a;
 	}
 	else if (ProfileId == SSS_PROFILE_ID_PERPIXEL)
 	{
 		// This color is used to lerping 'incident light' into 'diffused light', i.e, controlling the sharpness of the lighting.
 		// In case of direct DMPF control, this is not needed.
-		SubsurfaceColor = float3(1, 1, 1);
+		Tint = float3(1, 1, 1);
 	}
 #else
 	const uint ProfileId = ExtractSubsurfaceProfileInt(ScreenSpaceData.GBuffer);
@@ -951,7 +951,7 @@ void SubsurfaceRecombinePS(in FScreenVertexOutput Input, out float4 OutColor : S
 	const float StoredSpecular		= ScreenSpaceData.GBuffer.StoredSpecular;
 	const float PixelRadiusScale	= ScreenSpaceData.GBuffer.CustomData.a;
 	// asset specific color
-	const float3 SubsurfaceColor	= GetSubsurfaceProfileKernel(SSSS_SUBSURFACE_COLOR_OFFSET, ProfileId).rgb;
+	const float3 Tint	= GetSubsurfaceProfileKernel(SSSS_TINT_SCALE_OFFSET, ProfileId).rgb;
 	const float ProfileRadiusScale	= GetSubsurfaceProfileKernel(SSSS_N_KERNELWEIGHTOFFSET + SSSS_N_KERNELWEIGHTCOUNT - 1, ProfileId).a;
 #endif
 
@@ -1002,10 +1002,10 @@ void SubsurfaceRecombinePS(in FScreenVertexOutput Input, out float4 OutColor : S
 	
 	float3 ExtractedNonSubsurface = DiffuseAndSpecular.Specular;
 
-	float3 FadedSubsurfaceColor = SubsurfaceColor * LerpFactor;
+	float3 FadedTint = Tint * LerpFactor;
 
 	// combine potentially half res with full res
-	float3 SubsurfaceLighting = lerp(DiffuseAndSpecular.Diffuse, SSSColor, FadedSubsurfaceColor);
+	float3 SubsurfaceLighting = lerp(DiffuseAndSpecular.Diffuse, SSSColor, FadedTint);
 
 	OutColor = float4(SubsurfaceLighting * StoredBaseColor + ExtractedNonSubsurface, 1.0f);
 }
diff --git a/Engine/Shaders/Private/SSProfilePreIntegratedMobile.usf b/Engine/Shaders/Private/SSProfilePreIntegratedMobile.usf
--- a/Engine/Shaders/Private/SSProfilePreIntegratedMobile.usf
+++ b/Engine/Shaders/Private/SSProfilePreIntegratedMobile.usf
@@ -123,21 +123,11 @@ void SSProfilePreIntegratedPS(
 	float UnClampedNoL = UVPos.x * 2.0f - 1.0f;
 	float Curvature = UVPos.y;
 
-	const bool bUseBurley = IsSubsurfaceProfileUseBurley();
-
 	uint SurfaceAlbedoSampleIndex = BSSS_SURFACEALBEDO_OFFSET;
 	uint DiffuseMeanFreePathSampleIndex = BSSS_DMFP_OFFSET;
 
-	if (!bUseBurley)
-	{
-		SurfaceAlbedoSampleIndex = SSSS_APPROXIMATE_SURFACEALBEDO_OFFSET;
-		DiffuseMeanFreePathSampleIndex = SSSS_APPROXIMATE_DMFP_OFFSET;
-	}
-
 	float4 SurfaceAlbedo = GetSourceSubsurfaceProfileTexture(SurfaceAlbedoSampleIndex, SourceSubsurfaceProfileInt);
-
-	float WorldUnitScale = DecodeWorldUnitScale(GetSourceSubsurfaceProfileTexture(SSSS_SUBSURFACE_COLOR_OFFSET, SourceSubsurfaceProfileInt).a) * BURLEY_CM_2_MM;
-
+	float WorldUnitScale = DecodeWorldUnitScale(GetSourceSubsurfaceProfileTexture(SSSS_TINT_SCALE_OFFSET, SourceSubsurfaceProfileInt).a) * BURLEY_CM_2_MM;
 	float4 DiffuseMeanFreePath = max(DecodeDiffuseMeanFreePath(GetSourceSubsurfaceProfileTexture(DiffuseMeanFreePathSampleIndex, SourceSubsurfaceProfileInt)) * WorldUnitScale, 0.01f);;
 
 	float CosTheta = UnClampedNoL;
diff --git a/Engine/Shaders/Private/SubsurfaceProfileCommon.ush b/Engine/Shaders/Private/SubsurfaceProfileCommon.ush
--- a/Engine/Shaders/Private/SubsurfaceProfileCommon.ush
+++ b/Engine/Shaders/Private/SubsurfaceProfileCommon.ush
@@ -8,14 +8,12 @@
 
 // NOTE: Changing offsets below requires updating all instances of #SSSS_CONSTANTS
 // TODO: This needs to be defined in a single place and shared between C++ and shaders!
-#define SSSS_SUBSURFACE_COLOR_OFFSET			0
-#define BSSS_SURFACEALBEDO_OFFSET                (SSSS_SUBSURFACE_COLOR_OFFSET+1)
+#define SSSS_TINT_SCALE_OFFSET					0
+#define BSSS_SURFACEALBEDO_OFFSET               (SSSS_TINT_SCALE_OFFSET+1)
 #define BSSS_DMFP_OFFSET                        (BSSS_SURFACEALBEDO_OFFSET+1)
 #define SSSS_TRANSMISSION_OFFSET				(BSSS_DMFP_OFFSET+1)
 #define SSSS_BOUNDARY_COLOR_BLEED_OFFSET		(SSSS_TRANSMISSION_OFFSET+1)
-#define SSSS_APPROXIMATE_SURFACEALBEDO_OFFSET	(SSSS_BOUNDARY_COLOR_BLEED_OFFSET+1)
-#define SSSS_APPROXIMATE_DMFP_OFFSET			(SSSS_APPROXIMATE_SURFACEALBEDO_OFFSET+1)
-#define SSSS_DUAL_SPECULAR_OFFSET				(SSSS_APPROXIMATE_DMFP_OFFSET+1)
+#define SSSS_DUAL_SPECULAR_OFFSET				(SSSS_BOUNDARY_COLOR_BLEED_OFFSET+1)
 #define SSSS_KERNEL0_OFFSET						(SSSS_DUAL_SPECULAR_OFFSET+1)
 #define SSSS_KERNEL0_SIZE						13
 #define SSSS_KERNEL1_OFFSET						(SSSS_KERNEL0_OFFSET + SSSS_KERNEL0_SIZE)
@@ -23,10 +21,8 @@
 #define SSSS_KERNEL2_OFFSET						(SSSS_KERNEL1_OFFSET + SSSS_KERNEL1_SIZE)
 #define SSSS_KERNEL2_SIZE						6
 #define SSSS_KERNEL_TOTAL_SIZE					(SSSS_KERNEL0_SIZE + SSSS_KERNEL1_SIZE + SSSS_KERNEL2_SIZE)
-#define SSSS_TRANSMISSION_PROFILE_OFFSET		(SSSS_KERNEL0_OFFSET + SSSS_KERNEL_TOTAL_SIZE)
-#define SSSS_TRANSMISSION_PROFILE_SIZE			32
-#define BSSS_TRANSMISSION_PROFILE_OFFSET        (SSSS_TRANSMISSION_PROFILE_OFFSET + SSSS_TRANSMISSION_PROFILE_SIZE)
-#define BSSS_TRANSMISSION_PROFILE_SIZE			SSSS_TRANSMISSION_PROFILE_SIZE
+#define BSSS_TRANSMISSION_PROFILE_OFFSET		(SSSS_KERNEL0_OFFSET + SSSS_KERNEL_TOTAL_SIZE)
+#define BSSS_TRANSMISSION_PROFILE_SIZE			32
 #define	SSSS_MAX_TRANSMISSION_PROFILE_DISTANCE	5.0f // See MaxTransmissionProfileDistance in ComputeTransmissionProfile(), SeparableSSS.cpp
 #define SSSS_MAX_DUAL_SPECULAR_ROUGHNESS		2.0f
 
diff --git a/Engine/Shaders/Private/TransmissionCommon.ush b/Engine/Shaders/Private/TransmissionCommon.ush
--- a/Engine/Shaders/Private/TransmissionCommon.ush
+++ b/Engine/Shaders/Private/TransmissionCommon.ush
@@ -53,9 +53,8 @@ float4 GetTransmissionProfile(uint SubsurfaceProfileInt, float Thickness)
 {
 	// 0..255, which SubSurface profile to pick
 	const float2 TextureSize = View.SSProfilesTextureSizeAndInvSize.xy;
-	bool bUseBurley = GetSubsurfaceTransmittanceProfileUseBurley(SubsurfaceProfileInt);
-	float ProfileOffset = bUseBurley ? BSSS_TRANSMISSION_PROFILE_OFFSET : SSSS_TRANSMISSION_PROFILE_OFFSET;
-	float2 uv = float2((ProfileOffset + Thickness / SSSS_MAX_TRANSMISSION_PROFILE_DISTANCE * (SSSS_TRANSMISSION_PROFILE_SIZE - 1) + 0.5f), SubsurfaceProfileInt + 0.5f) / (TextureSize.xy);
+	float ProfileOffset = BSSS_TRANSMISSION_PROFILE_OFFSET;
+	float2 uv = float2((ProfileOffset + Thickness / SSSS_MAX_TRANSMISSION_PROFILE_DISTANCE * (BSSS_TRANSMISSION_PROFILE_SIZE - 1) + 0.5f), SubsurfaceProfileInt + 0.5f) / (TextureSize.xy);
 	return Texture2DSampleLevel(View.SSProfilesTexture, View.SSProfilesTransmissionSampler, uv, 0);
 }
 
diff --git a/Engine/Source/Runtime/Engine/Private/Rendering/BurleyNormalizedSSS.cpp b/Engine/Source/Runtime/Engine/Private/Rendering/BurleyNormalizedSSS.cpp
--- a/Engine/Source/Runtime/Engine/Private/Rendering/BurleyNormalizedSSS.cpp
+++ b/Engine/Source/Runtime/Engine/Private/Rendering/BurleyNormalizedSSS.cpp
@@ -224,10 +224,10 @@ void ComputeMirroredBSSSKernel(FLinearColor* TargetBuffer, uint32 TargetBufferSi
 
 
 
-void ComputeTransmissionProfileBurley(FLinearColor* TargetBuffer, uint32 TargetBufferSize, FLinearColor SubsurfaceColor, 
-																		FLinearColor FalloffColor, float ExtinctionScale, 
-																		FLinearColor SurfaceAlbedo, FLinearColor DiffuseMeanFreePathInMm,
-																		float InUnitToMm, FLinearColor TransmissionTintColor)
+void ComputeTransmissionProfileBurley(FLinearColor* TargetBuffer, uint32 TargetBufferSize, 
+									FLinearColor FalloffColor, float ExtinctionScale,
+									FLinearColor SurfaceAlbedo, FLinearColor DiffuseMeanFreePathInMm,
+									float InUnitToMm, FLinearColor TransmissionTintColor)
 {
 	check(TargetBuffer);
 	check(TargetBufferSize > 0);
diff --git a/Engine/Source/Runtime/Engine/Private/Rendering/BurleyNormalizedSSS.h b/Engine/Source/Runtime/Engine/Private/Rendering/BurleyNormalizedSSS.h
--- a/Engine/Source/Runtime/Engine/Private/Rendering/BurleyNormalizedSSS.h
+++ b/Engine/Source/Runtime/Engine/Private/Rendering/BurleyNormalizedSSS.h
@@ -10,13 +10,12 @@ BurleyNormalizedSSS.h: Compute the transmission profile and convert parameters
 
 // @param TargetBuffer, needs to be preallocated with TargetBufferSize (RGB is the sample weight, A is the offset), [0] is the center samples, following elements need to be mirrored with A, -A
 // @param TargetBufferSize >0
-// @parma SubsurfaceColor see SubsurfaceProfile.h
 // @parma FalloffColor see SubsurfaceProfile.h
 // @parma SurfaceAlbedo see SubsurfaceProfile.h
 // @parma DiffuseMeanFreePath mean free path color scaled by mean free path distance (in mm).
 // @param WorldUnitScale see SubsurfaceProfile.h
 // @param TransmissionTintColor see SubsurfaceProfile.h
-void ComputeTransmissionProfileBurley(FLinearColor* TargetBuffer, uint32 TargetBufferSize, FLinearColor SubsurfaceColor, 
+void ComputeTransmissionProfileBurley(FLinearColor* TargetBuffer, uint32 TargetBufferSize, 
 									FLinearColor FalloffColor, float ExtinctionScale,
 									FLinearColor SurfaceAlbedo, FLinearColor DiffuseMeanFreePath, 
 									float WorldUnitScale,FLinearColor TransmissionTintColor);
diff --git a/Engine/Source/Runtime/Engine/Private/Rendering/SeparableSSS.cpp b/Engine/Source/Runtime/Engine/Private/Rendering/SeparableSSS.cpp
deleted file mode 100644
diff --git a/Engine/Source/Runtime/Engine/Private/Rendering/SeparableSSS.h b/Engine/Source/Runtime/Engine/Private/Rendering/SeparableSSS.h
deleted file mode 100644
diff --git a/Engine/Source/Runtime/Engine/Private/Rendering/SubsurfaceProfile.cpp b/Engine/Source/Runtime/Engine/Private/Rendering/SubsurfaceProfile.cpp
--- a/Engine/Source/Runtime/Engine/Private/Rendering/SubsurfaceProfile.cpp
+++ b/Engine/Source/Runtime/Engine/Private/Rendering/SubsurfaceProfile.cpp
@@ -3,7 +3,6 @@
 #include "Engine/SubsurfaceProfile.h"
 #include "RenderingThread.h"
 #include "RendererInterface.h"
-#include "Rendering/SeparableSSS.h"
 #include "Rendering/BurleyNormalizedSSS.h"
 #include "EngineModule.h"
 #include "RenderTargetPool.h"
@@ -261,14 +260,12 @@ static float GetNextSmallerPositiveFloat(float x)
 
 // NOTE: Changing offsets below requires updating all instances of #SSSS_CONSTANTS
 // TODO: This needs to be defined in a single place and shared between C++ and shaders!
-#define SSSS_SUBSURFACE_COLOR_OFFSET			0
-#define BSSS_SURFACEALBEDO_OFFSET                (SSSS_SUBSURFACE_COLOR_OFFSET+1)
+#define SSSS_TINT_SCALE_OFFSET					0
+#define BSSS_SURFACEALBEDO_OFFSET               (SSSS_TINT_SCALE_OFFSET+1)
 #define BSSS_DMFP_OFFSET                        (BSSS_SURFACEALBEDO_OFFSET+1)
 #define SSSS_TRANSMISSION_OFFSET				(BSSS_DMFP_OFFSET+1)
 #define SSSS_BOUNDARY_COLOR_BLEED_OFFSET		(SSSS_TRANSMISSION_OFFSET+1)
-#define SSSS_APPROXIMATE_SURFACEALBEDO_OFFSET	(SSSS_BOUNDARY_COLOR_BLEED_OFFSET+1)
-#define SSSS_APPROXIMATE_DMFP_OFFSET			(SSSS_APPROXIMATE_SURFACEALBEDO_OFFSET+1)
-#define SSSS_DUAL_SPECULAR_OFFSET				(SSSS_APPROXIMATE_DMFP_OFFSET+1)
+#define SSSS_DUAL_SPECULAR_OFFSET				(SSSS_BOUNDARY_COLOR_BLEED_OFFSET+1)
 #define SSSS_KERNEL0_OFFSET						(SSSS_DUAL_SPECULAR_OFFSET+1)
 #define SSSS_KERNEL0_SIZE						13
 #define SSSS_KERNEL1_OFFSET						(SSSS_KERNEL0_OFFSET + SSSS_KERNEL0_SIZE)
@@ -276,10 +273,8 @@ static float GetNextSmallerPositiveFloat(float x)
 #define SSSS_KERNEL2_OFFSET						(SSSS_KERNEL1_OFFSET + SSSS_KERNEL1_SIZE)
 #define SSSS_KERNEL2_SIZE						6
 #define SSSS_KERNEL_TOTAL_SIZE					(SSSS_KERNEL0_SIZE + SSSS_KERNEL1_SIZE + SSSS_KERNEL2_SIZE)
-#define SSSS_TRANSMISSION_PROFILE_OFFSET		(SSSS_KERNEL0_OFFSET + SSSS_KERNEL_TOTAL_SIZE)
-#define SSSS_TRANSMISSION_PROFILE_SIZE			32
-#define BSSS_TRANSMISSION_PROFILE_OFFSET        (SSSS_TRANSMISSION_PROFILE_OFFSET + SSSS_TRANSMISSION_PROFILE_SIZE)
-#define BSSS_TRANSMISSION_PROFILE_SIZE			SSSS_TRANSMISSION_PROFILE_SIZE
+#define BSSS_TRANSMISSION_PROFILE_OFFSET		(SSSS_KERNEL0_OFFSET + SSSS_KERNEL_TOTAL_SIZE)
+#define BSSS_TRANSMISSION_PROFILE_SIZE			32
 #define	SSSS_MAX_TRANSMISSION_PROFILE_DISTANCE	5.0f // See MaxTransmissionProfileDistance in ComputeTransmissionProfile(), SeparableSSS.cpp
 #define	SSSS_MAX_DUAL_SPECULAR_ROUGHNESS		2.0f
 
@@ -384,32 +379,16 @@ void FSubsurfaceProfileTexture::CreateTexture(FRHICommandListImmediate& RHICmdLi
 	{
 		FSubsurfaceProfileStruct Data = SubsurfaceProfileEntries[y].Settings;
 
-		// Convert subsurface parameters from separable to burley for mobile platform, and because the current shader platform is unknown at here, the converted parameters need to be saved in different slots(SSSS_APPROXIMATE_SURFACEALBEDO_OFFSET and SSSS_APPROXIMATE_DMFP_OFFSET).
-		if (!Data.bEnableBurley)
-		{
-			FLinearColor ApproximateSurfaceAlbedo, ApproximateMeanFreePathColor;
-			float ApproximateMeanFreePathDistance;
-			ConvertSubsurfaceParametersFromSeparableToBurley(Data, ApproximateSurfaceAlbedo, ApproximateMeanFreePathColor, ApproximateMeanFreePathDistance);
-			
-			ApproximateMeanFreePathColor = ApproximateMeanFreePathColor.GetClamped(Bias);
-			ApproximateSurfaceAlbedo = ApproximateSurfaceAlbedo.GetClamped(Bias);
-			
-			FLinearColor ApproximateDiffuseMeanFreePathInMm = (ApproximateMeanFreePathColor * ApproximateMeanFreePathDistance) * CmToMm;
-			SetupSurfaceAlbedoAndDiffuseMeanFreePath(ApproximateSurfaceAlbedo, ApproximateDiffuseMeanFreePathInMm);
-
-			TextureRow[SSSS_APPROXIMATE_SURFACEALBEDO_OFFSET] = ApproximateSurfaceAlbedo;
-			TextureRow[SSSS_APPROXIMATE_DMFP_OFFSET] = EncodeDiffuseMeanFreePath(ApproximateDiffuseMeanFreePathInMm);
-		}
-
-		Data.SubsurfaceColor = Data.SubsurfaceColor.GetClamped();
+		Data.Tint = Data.Tint.GetClamped();
 		Data.FalloffColor = Data.FalloffColor.GetClamped(Bias);
 		Data.MeanFreePathColor = Data.MeanFreePathColor.GetClamped(Bias); // In Cm
 		Data.TransmissionTintColor = Data.TransmissionTintColor.GetClamped(Bias);
 		Data.SurfaceAlbedo = Data.SurfaceAlbedo.GetClamped(Bias);
 
-		// to allow blending of the Subsurface with fullres in the shader
-		TextureRow[SSSS_SUBSURFACE_COLOR_OFFSET] = Data.SubsurfaceColor;
-		TextureRow[SSSS_SUBSURFACE_COLOR_OFFSET].A = EncodeWorldUnitScale(Data.WorldUnitScale);
+		// to allow blending of the Subsurface with fullres in the shader and
+		// to depricate {scatter radius, falloff} in favor of albedo/MFP, need the `subsurface color` as tint.
+		TextureRow[SSSS_TINT_SCALE_OFFSET] = Data.Tint;
+		TextureRow[SSSS_TINT_SCALE_OFFSET].A = EncodeWorldUnitScale(Data.WorldUnitScale);
 		
 		const float UnitToCm = Data.WorldUnitScale;
 
@@ -466,22 +445,14 @@ void FSubsurfaceProfileTexture::CreateTexture(FRHICommandListImmediate& RHICmdLi
 			// make Burley fallback looks the same to Burley. Set 1.0f when we have enough samples.
 			const float SamplingCountCompensation = 0.707f;
 			Data.ScatterRadius *= (Data.WorldUnitScale * 10.0f)*2.229f* SamplingCountCompensation;
-
-			// To depricate {scatter radius, falloff} in favor of albedo/MFP, need the subsurface color as tint.
-			TextureRow[SSSS_SUBSURFACE_COLOR_OFFSET].R = Data.Tint.R; 
-			TextureRow[SSSS_SUBSURFACE_COLOR_OFFSET].G = Data.Tint.G;
-			TextureRow[SSSS_SUBSURFACE_COLOR_OFFSET].B = Data.Tint.B;
 		}
 		else
 		{
-			ComputeMirroredSSSKernel(&TextureRow[SSSS_KERNEL0_OFFSET], SSSS_KERNEL0_SIZE, Data.SubsurfaceColor, Data.FalloffColor);
-			ComputeMirroredSSSKernel(&TextureRow[SSSS_KERNEL1_OFFSET], SSSS_KERNEL1_SIZE, Data.SubsurfaceColor, Data.FalloffColor);
-			ComputeMirroredSSSKernel(&TextureRow[SSSS_KERNEL2_OFFSET], SSSS_KERNEL2_SIZE, Data.SubsurfaceColor, Data.FalloffColor);
+			UE_LOG(LogSubsurfaceProfile, Warning, TEXT("Dipole model has already been upgraded to Burley. Should not reach here."));
 		}
 
-		ComputeTransmissionProfile(&TextureRow[SSSS_TRANSMISSION_PROFILE_OFFSET], SSSS_TRANSMISSION_PROFILE_SIZE, Data.SubsurfaceColor, Data.FalloffColor, Data.ExtinctionScale);
 		ComputeTransmissionProfileBurley(&TextureRow[BSSS_TRANSMISSION_PROFILE_OFFSET], BSSS_TRANSMISSION_PROFILE_SIZE, 
-			Data.SubsurfaceColor, Data.FalloffColor, Data.ExtinctionScale, Data.SurfaceAlbedo, DifffuseMeanFreePathInMm, UnitToCm * CmToMm, Data.TransmissionTintColor);
+			Data.FalloffColor, Data.ExtinctionScale, Data.SurfaceAlbedo, DifffuseMeanFreePathInMm, UnitToCm * CmToMm, Data.TransmissionTintColor);
 
 		// could be lower than 1 (but higher than 0) to range compress for better quality (for 8 bit)
 		const float TableMaxRGB = 1.0f;
