Fix crash compling an animation Blueprint with all bones hidden

diff --git a/Engine/Source/Runtime/Engine/Private/Components/SkeletalMeshComponent.cpp b/Engine/Source/Runtime/Engine/Private/Components/SkeletalMeshComponent.cpp
--- a/Engine/Source/Runtime/Engine/Private/Components/SkeletalMeshComponent.cpp
+++ b/Engine/Source/Runtime/Engine/Private/Components/SkeletalMeshComponent.cpp
@@ -1554,9 +1554,9 @@ void USkeletalMeshComponent::FillComponentSpaceTransforms(const USkeletalMesh* I
 	const FTransform* LocalTransformsData = InBoneSpaceTransforms.GetData();
 	FTransform* ComponentSpaceData = OutComponentSpaceTransforms.GetData();
 
-	// First bone is always root bone, and it doesn't have a parent.
+	// First bone (if we have one) is always root bone, and it doesn't have a parent.
 	{
-		check(FillComponentSpaceTransformsRequiredBones[0] == 0);
+		check(FillComponentSpaceTransformsRequiredBones.Num() == 0 || FillComponentSpaceTransformsRequiredBones[0] == 0);
 		OutComponentSpaceTransforms[0] = InBoneSpaceTransforms[0];
 
 #if DO_GUARD_SLOW
