Properly initialize FWorldPartitionActorDescView::bInvalidRuntimeGrid member

diff --git a/Engine/Source/Runtime/Engine/Private/WorldPartition/WorldPartitionActorDescView.cpp b/Engine/Source/Runtime/Engine/Private/WorldPartition/WorldPartitionActorDescView.cpp
--- a/Engine/Source/Runtime/Engine/Private/WorldPartition/WorldPartitionActorDescView.cpp
+++ b/Engine/Source/Runtime/Engine/Private/WorldPartition/WorldPartitionActorDescView.cpp
@@ -7,15 +7,14 @@
 #include "WorldPartition/WorldPartitionActorDesc.h"
 
 FWorldPartitionActorDescView::FWorldPartitionActorDescView()
-	: ActorDesc(nullptr)
-	, bIsForcedNonSpatiallyLoaded(false)
-	, bInvalidDataLayers(false)
+	: FWorldPartitionActorDescView(nullptr)
 {}
 
 FWorldPartitionActorDescView::FWorldPartitionActorDescView(const FWorldPartitionActorDesc* InActorDesc)
 	: ActorDesc(InActorDesc)
 	, bIsForcedNonSpatiallyLoaded(false)
 	, bInvalidDataLayers(false)
+	, bInvalidRuntimeGrid(false)
 {}
 
 const FGuid& FWorldPartitionActorDescView::GetGuid() const
