Allow changing the spatially loaded flag on HLOD actors
The actors can't be edited by users from the editor, but it's still possible that the spatially loaded flag needs to be changed (by code) during the HLOD setup process.
This can occur if the "Is Spatially Loaded" property of the associated HLOD layer is modified

diff --git a/Engine/Source/Runtime/Engine/Public/WorldPartition/HLOD/HLODActor.h b/Engine/Source/Runtime/Engine/Public/WorldPartition/HLOD/HLODActor.h
--- a/Engine/Source/Runtime/Engine/Public/WorldPartition/HLOD/HLODActor.h
+++ b/Engine/Source/Runtime/Engine/Public/WorldPartition/HLOD/HLODActor.h
@@ -59,8 +59,6 @@ public:
 
 	void BuildHLOD(bool bForceBuild = false);
 	uint32 GetHLODHash() const;
-
-	virtual bool CanChangeIsSpatiallyLoadedFlag() const override { return false; }
 #endif // WITH_EDITOR
 
 protected:
