Updated FillMirrorBoneIndexes to take a Skeleton as a parameter. This avoids the confusion about where the reference skeleton should come from, since it doesn't work as expected when using a skeletal mesh reference skeleton.

diff --git a/Engine/Source/Runtime/AnimGraphRuntime/Private/AnimNodes/AnimNode_Mirror.cpp b/Engine/Source/Runtime/AnimGraphRuntime/Private/AnimNodes/AnimNode_Mirror.cpp
--- a/Engine/Source/Runtime/AnimGraphRuntime/Private/AnimNodes/AnimNode_Mirror.cpp
+++ b/Engine/Source/Runtime/AnimGraphRuntime/Private/AnimNodes/AnimNode_Mirror.cpp
@@ -116,7 +116,7 @@ void FAnimNode_MirrorBase::FillCompactPoseAndComponentRefRotations(const FBoneCo
 	if (MirrorDataTable)
 	{
 		TCustomBoneIndexArray<FSkeletonPoseBoneIndex, FSkeletonPoseBoneIndex> MirrorBoneIndexes;
-		MirrorDataTable->FillMirrorBoneIndexes(BoneContainer.GetReferenceSkeleton(), MirrorBoneIndexes);
+		MirrorDataTable->FillMirrorBoneIndexes(BoneContainer.GetSkeletonAsset(), MirrorBoneIndexes);
 		MirrorDataTable->FillCompactPoseMirrorBones(BoneContainer, MirrorBoneIndexes, CompactPoseMirrorBones);
 
 		const int32 NumBones = BoneContainer.GetCompactPoseNumBones();
diff --git a/Engine/Source/Runtime/Engine/Classes/Animation/MirrorDataTable.h b/Engine/Source/Runtime/Engine/Classes/Animation/MirrorDataTable.h
--- a/Engine/Source/Runtime/Engine/Classes/Animation/MirrorDataTable.h
+++ b/Engine/Source/Runtime/Engine/Classes/Animation/MirrorDataTable.h
@@ -151,7 +151,7 @@ public:
 	 * @param	ReferenceSkeleton		The ReferenceSkeleton to compute the mirror index against
 	 * @param	OutMirrorBoneIndexes	An array that provides the bone index of the mirror bone, or INDEX_NONE if the bone is not mirrored
 	 */
-	ENGINE_API void FillMirrorBoneIndexes(const FReferenceSkeleton& ReferenceSkeleton, TCustomBoneIndexArray<FSkeletonPoseBoneIndex, FSkeletonPoseBoneIndex>& OutMirrorBoneIndexes) const;
+	ENGINE_API void FillMirrorBoneIndexes(const USkeleton* Skeleton, TCustomBoneIndexArray<FSkeletonPoseBoneIndex, FSkeletonPoseBoneIndex>& OutMirrorBoneIndexes) const;
 
 #if WITH_EDITOR  
 	/**
diff --git a/Engine/Source/Runtime/Engine/Private/Animation/AnimationRuntime.cpp b/Engine/Source/Runtime/Engine/Private/Animation/AnimationRuntime.cpp
--- a/Engine/Source/Runtime/Engine/Private/Animation/AnimationRuntime.cpp
+++ b/Engine/Source/Runtime/Engine/Private/Animation/AnimationRuntime.cpp
@@ -1311,7 +1311,7 @@ void FAnimationRuntime::MirrorPose(FCompactPose& Pose, const UMirrorDataTable& M
 	if (Skeleton)
 	{
 		TCustomBoneIndexArray<FSkeletonPoseBoneIndex, FSkeletonPoseBoneIndex> MirrorBoneIndexes;
-		MirrorDataTable.FillMirrorBoneIndexes(BoneContainer.GetReferenceSkeleton(), MirrorBoneIndexes);
+		MirrorDataTable.FillMirrorBoneIndexes(Skeleton, MirrorBoneIndexes);
 
 		// Compact pose format of Mirror Bone Map
 		TCustomBoneIndexArray<FCompactPoseBoneIndex, FCompactPoseBoneIndex> CompactPoseMirrorBones;
diff --git a/Engine/Source/Runtime/Engine/Private/Animation/MirrorDataTable.cpp b/Engine/Source/Runtime/Engine/Private/Animation/MirrorDataTable.cpp
--- a/Engine/Source/Runtime/Engine/Private/Animation/MirrorDataTable.cpp
+++ b/Engine/Source/Runtime/Engine/Private/Animation/MirrorDataTable.cpp
@@ -308,8 +308,10 @@ void UMirrorDataTable::FillCompactPoseMirrorBones(const FBoneContainer& BoneCont
 	}
 }
 
-void UMirrorDataTable::FillMirrorBoneIndexes(const FReferenceSkeleton& ReferenceSkeleton, TCustomBoneIndexArray<FSkeletonPoseBoneIndex, FSkeletonPoseBoneIndex>& OutMirrorBoneIndexes) const
+void UMirrorDataTable::FillMirrorBoneIndexes(const USkeleton* InSkeleton, TCustomBoneIndexArray<FSkeletonPoseBoneIndex, FSkeletonPoseBoneIndex>& OutMirrorBoneIndexes) const
 {
+	const FReferenceSkeleton& ReferenceSkeleton = InSkeleton->GetReferenceSkeleton();
+
 	// Reset the mirror table to defaults (no mirroring)
 	OutMirrorBoneIndexes.SetNumUninitialized(ReferenceSkeleton.GetNum());
 	FMemory::Memset(OutMirrorBoneIndexes.GetData(), INDEX_NONE, OutMirrorBoneIndexes.Num() * OutMirrorBoneIndexes.GetTypeSize());
@@ -362,7 +364,7 @@ void UMirrorDataTable::FillMirrorArrays()
 		return; 
 	}
 
-	FillMirrorBoneIndexes(Skeleton->GetReferenceSkeleton(), BoneToMirrorBoneIndex);
+	FillMirrorBoneIndexes(Skeleton, BoneToMirrorBoneIndex);
 
 	TMap<FName, FName> CurveToMirrorCurveMap;
 	
