Adding RemoveRow functionality to CurveTable to match DataTable.

diff --git a/Engine/Source/Runtime/Engine/Classes/Engine/CurveTable.h b/Engine/Source/Runtime/Engine/Classes/Engine/CurveTable.h
--- a/Engine/Source/Runtime/Engine/Classes/Engine/CurveTable.h
+++ b/Engine/Source/Runtime/Engine/Classes/Engine/CurveTable.h
@@ -55,6 +55,8 @@ class UCurveTable
 
 	ECurveTableMode GetCurveTableMode() const { return CurveTableMode; }
 
+	/** Removes a single row from the CurveTable by name. Just returns if row is not found. */
+	ENGINE_API virtual void RemoveRow(FName RowName);
 	ENGINE_API FRichCurve& AddRichCurve(FName RowName);
 	ENGINE_API FSimpleCurve& AddSimpleCurve(FName RowName);
 
diff --git a/Engine/Source/Runtime/Engine/Private/CurveTable.cpp b/Engine/Source/Runtime/Engine/Private/CurveTable.cpp
--- a/Engine/Source/Runtime/Engine/Private/CurveTable.cpp
+++ b/Engine/Source/Runtime/Engine/Private/CurveTable.cpp
@@ -539,6 +539,16 @@ void UCurveTable::EmptyTable()
 	UCurveTable::InvalidateAllCachedCurves();
 }
 
+void UCurveTable::RemoveRow(FName RowName)
+{
+	FRealCurve* Curve = nullptr;
+	RowMap.RemoveAndCopyValue(RowName, Curve);
+	if (Curve != nullptr)
+	{
+		delete Curve;
+	}
+}
+
 FRichCurve& UCurveTable::AddRichCurve(FName RowName)
 {
 	check(CurveTableMode != ECurveTableMode::SimpleCurves);
