Fixed missing copyright notice (Issue #130069)

diff --git a/Engine/Source/Runtime/Engine/Private/VisualLogger/VisualLoggerFilterVolume.cpp b/Engine/Source/Runtime/Engine/Private/VisualLogger/VisualLoggerFilterVolume.cpp
--- a/Engine/Source/Runtime/Engine/Private/VisualLogger/VisualLoggerFilterVolume.cpp
+++ b/Engine/Source/Runtime/Engine/Private/VisualLogger/VisualLoggerFilterVolume.cpp
@@ -1,3 +1,5 @@
+// Copyright Epic Games, Inc. All Rights Reserved.
+
 #include "VisualLogger/VisualLoggerFilterVolume.h"
 
 AVisualLoggerFilterVolume::AVisualLoggerFilterVolume(const FObjectInitializer& ObjectInitializer)
diff --git a/Engine/Source/Runtime/Engine/Public/VisualLogger/VisualLoggerFilterVolume.h b/Engine/Source/Runtime/Engine/Public/VisualLogger/VisualLoggerFilterVolume.h
--- a/Engine/Source/Runtime/Engine/Public/VisualLogger/VisualLoggerFilterVolume.h
+++ b/Engine/Source/Runtime/Engine/Public/VisualLogger/VisualLoggerFilterVolume.h
@@ -1,3 +1,5 @@
+// Copyright Epic Games, Inc. All Rights Reserved.
+
 #pragma once
 
 #include "GameFramework/Volume.h"
