Fixed missing copyright notice (Issue #130069)

diff --git a/Engine/Source/Runtime/Engine/Private/VisualLogger/VisualLoggerFilterVolume.cpp b/Engine/Source/Runtime/Engine/Private/VisualLogger/VisualLoggerFilterVolume.cpp
new file mode 100644
--- /dev/null
+++ b/Engine/Source/Runtime/Engine/Private/VisualLogger/VisualLoggerFilterVolume.cpp
@@ -0,0 +1,10 @@
+// Copyright Epic Games, Inc. All Rights Reserved.
+
+#include "VisualLogger/VisualLoggerFilterVolume.h"
+
+AVisualLoggerFilterVolume::AVisualLoggerFilterVolume(const FObjectInitializer& ObjectInitializer)
+{
+	PrimaryActorTick.bCanEverTick = false;
+	PrimaryActorTick.bStartWithTickEnabled = false;
+	bIsEditorOnlyActor = true;
+}
\ No newline at end of file
diff --git a/Engine/Source/Runtime/Engine/Public/VisualLogger/VisualLoggerFilterVolume.h b/Engine/Source/Runtime/Engine/Public/VisualLogger/VisualLoggerFilterVolume.h
new file mode 100644
--- /dev/null
+++ b/Engine/Source/Runtime/Engine/Public/VisualLogger/VisualLoggerFilterVolume.h
@@ -0,0 +1,18 @@
+// Copyright Epic Games, Inc. All Rights Reserved.
+
+#pragma once
+
+#include "GameFramework/Volume.h"
+#include "VisualLoggerFilterVolume.generated.h"
+
+/**
+ * A volume to be placed in the level while browsing visual logger output.
+ */
+UCLASS()
+class ENGINE_API AVisualLoggerFilterVolume : public AVolume
+{
+	GENERATED_BODY()
+public:
+	AVisualLoggerFilterVolume(const FObjectInitializer& ObjectInitializer);
+	
+};
\ No newline at end of file
