LevelInstance: Pivot cleanup (Icon, Label, prevent: copy/move/label change)

diff --git a/Engine/Source/Editor/EditorStyle/Private/StarshipStyle.cpp b/Engine/Source/Editor/EditorStyle/Private/StarshipStyle.cpp
--- a/Engine/Source/Editor/EditorStyle/Private/StarshipStyle.cpp
+++ b/Engine/Source/Editor/EditorStyle/Private/StarshipStyle.cpp
@@ -6081,6 +6081,7 @@ void FStarshipEditorStyle::FStyle::SetupClassIconsAndThumbnails()
 			{TEXT("Landscape")},
 			{TEXT("LevelBounds")},
 			{TEXT("LevelInstance")},
+			{TEXT("LevelInstancePivot")},
 			{TEXT("PackedLevelActor")},
 			{TEXT("LevelScriptActor")},
 			{TEXT("LevelSequenceActor")},
diff --git a/Engine/Source/Runtime/Engine/Private/LevelInstance/LevelInstanceEditorPivotActor.cpp b/Engine/Source/Runtime/Engine/Private/LevelInstance/LevelInstanceEditorPivotActor.cpp
--- a/Engine/Source/Runtime/Engine/Private/LevelInstance/LevelInstanceEditorPivotActor.cpp
+++ b/Engine/Source/Runtime/Engine/Private/LevelInstance/LevelInstanceEditorPivotActor.cpp
@@ -16,6 +16,9 @@ ALevelInstancePivot::ALevelInstancePivot(const FObjectInitializer& ObjectInitial
 {
 	RootComponent = CreateDefaultSubobject<USceneComponent>(TEXT("RootComponent"));
 	RootComponent->Mobility = EComponentMobility::Static;
+#if WITH_EDITORONLY_DATA
+	bActorLabelEditable = false;
+#endif // WITH_EDITORONLY_DATA
 }
 
 #if WITH_EDITOR
@@ -38,6 +41,8 @@ ALevelInstancePivot* ALevelInstancePivot::Create(ALevelInstance* LevelInstanceAc
 	PivotActor->SpawnTransform = PivotActor->GetActorTransform();
 	PivotActor->OriginalPivotOffset = WorldSettings->LevelInstancePivotOffset;
 	PivotActor->SetLevelInstanceID(LevelInstanceActor->GetLevelInstanceID());
+	// Set Label last as this will call PostEditChangeProperty and update the offset so other fields need to be setup first.
+	PivotActor->SetActorLabel(TEXT("Pivot"));
 
 	return PivotActor;
 }
diff --git a/Engine/Source/Runtime/Engine/Private/LevelInstance/LevelInstanceSubsystem.cpp b/Engine/Source/Runtime/Engine/Private/LevelInstance/LevelInstanceSubsystem.cpp
--- a/Engine/Source/Runtime/Engine/Private/LevelInstance/LevelInstanceSubsystem.cpp
+++ b/Engine/Source/Runtime/Engine/Private/LevelInstance/LevelInstanceSubsystem.cpp
@@ -20,6 +20,7 @@
 #include "LevelInstance/ILevelInstanceEditorModule.h"
 #include "LevelInstance/LevelInstanceEditorInstanceActor.h"
 #include "LevelInstance/LevelInstanceEditorObject.h"
+#include "LevelInstance/LevelInstanceEditorPivotActor.h"
 #include "Misc/ScopedSlowTask.h"
 #include "Misc/ITransaction.h"
 #include "Misc/Paths.h"
@@ -1152,6 +1153,11 @@ ULevelInstanceSubsystem::FLevelsToRemoveScope::~FLevelsToRemoveScope()
 
 bool ULevelInstanceSubsystem::CanMoveActorToLevel(const AActor* Actor, FText* OutReason) const
 {
+	if (Actor->IsA<ALevelInstancePivot>())
+	{
+		return false;
+	}
+
 	if (Actor->GetWorld() == GetWorld())
 	{
 		if (const ALevelInstance* LevelInstanceActor = Cast<ALevelInstance>(Actor))
diff --git a/Engine/Source/Runtime/Engine/Public/LevelInstance/LevelInstanceEditorPivotActor.h b/Engine/Source/Runtime/Engine/Public/LevelInstance/LevelInstanceEditorPivotActor.h
--- a/Engine/Source/Runtime/Engine/Public/LevelInstance/LevelInstanceEditorPivotActor.h
+++ b/Engine/Source/Runtime/Engine/Public/LevelInstance/LevelInstanceEditorPivotActor.h
@@ -27,7 +27,7 @@ public:
 	virtual void PostEditMove(bool bFinished) override;
 	virtual void PostEditChangeProperty(FPropertyChangedEvent& PropertyChangedEvent) override;
 	virtual void PostEditUndo() override;
-
+	virtual bool ShouldExport() override { return false; }
 	void SetPivot(ELevelInstancePivotType PivotType, AActor* PivotActor = nullptr);
 private:
 	void SetLevelInstanceID(const FLevelInstanceID& InLevelInstanceID) { LevelInstanceID = InLevelInstanceID; }
