PR #8735: GPU particle sorting: fix assert when using SM6 (Contributed by [REDACTED_PERSON])

#rnx

diff --git a/Engine/Source/Runtime/Engine/Private/Particles/ParticleSortingGPU.cpp b/Engine/Source/Runtime/Engine/Private/Particles/ParticleSortingGPU.cpp
--- a/Engine/Source/Runtime/Engine/Private/Particles/ParticleSortingGPU.cpp
+++ b/Engine/Source/Runtime/Engine/Private/Particles/ParticleSortingGPU.cpp
@@ -172,7 +172,7 @@ int32 GenerateParticleSortKeys(
 	int32 BatchId
 	)
 {
-	check(FeatureLevel == ERHIFeatureLevel::SM5 || FeatureLevel == ERHIFeatureLevel::ES3_1);
+	check(FeatureLevel >= ERHIFeatureLevel::SM5 || FeatureLevel == ERHIFeatureLevel::ES3_1);
 
 	FParticleKeyGenParameters KeyGenParameters;
 	FParticleKeyGenUniformBufferRef KeyGenUniformBuffer;
