Avoid using ForEachObjectOfClass for invoking EndOfFrameUpdates to preserve flexibility in what can be done during those updates (like creation of UObjects, currently done in Niagara)

diff --git a/Engine/Source/Runtime/Engine/Classes/Engine/Engine.h b/Engine/Source/Runtime/Engine/Classes/Engine/Engine.h
--- a/Engine/Source/Runtime/Engine/Classes/Engine/Engine.h
+++ b/Engine/Source/Runtime/Engine/Classes/Engine/Engine.h
@@ -2493,6 +2493,11 @@ protected:
 	 */
 	 bool PerformError(const TCHAR* Cmd, FOutputDevice& Out = *GLog);
 
+	 /**
+	  * Dispatches EndOfFrameUpdates for all UWorlds
+	  */
+	 static void SendWorldEndOfFrameUpdates();
+
 public:
 	/** @return the GIsEditor flag setting */
 	bool IsEditor();
diff --git a/Engine/Source/Runtime/Engine/Private/UnrealEngine.cpp b/Engine/Source/Runtime/Engine/Private/UnrealEngine.cpp
--- a/Engine/Source/Runtime/Engine/Private/UnrealEngine.cpp
+++ b/Engine/Source/Runtime/Engine/Private/UnrealEngine.cpp
@@ -1527,20 +1527,35 @@ static FAutoConsoleVariableRef CVarIncrementalGCTimePerFrame(
 	ECVF_Default
 );
 
-void UEngine::PreGarbageCollect()
+void UEngine::SendWorldEndOfFrameUpdates()
 {
-	TRACE_CPUPROFILER_EVENT_SCOPE(UEngine::PreGarbageCollect);
-	ForEachObjectOfClass(UWorld::StaticClass(), [](UObject* WorldObj)
+	// Gather worlds that need EOF updates
+	// This must be done in two steps as the object hash table is locked during ForEachObjectOfClass so any NewObject calls would fail
+	TArray<UWorld*, TInlineAllocator<16>> Worlds;
+	ForEachObjectOfClass(UWorld::StaticClass(), [&](UObject* WorldObj)
 	{
-		UWorld* World = CastChecked<UWorld>(WorldObj);
-
-		if (World->HasEndOfFrameUpdates())
+		if (UWorld* World = CastChecked<UWorld>(WorldObj))
 		{
-			// Make sure deferred component updates have been sent to the rendering thread before deleting any UObjects which the rendering thread may be referencing
-			// This fixes rendering thread crashes in the following order of operations 1) UMeshComponent::SetMaterial 2) GC 3) Rendering command that dereferences the UMaterial
-			World->SendAllEndOfFrameUpdates();
+			if (World->HasEndOfFrameUpdates())
+			{
+				Worlds.Add(World);
+			}
 		}
 	});
+
+	for (UWorld* World : Worlds)
+	{
+		// Make sure deferred component updates have been sent to the rendering thread before deleting any UObjects which the rendering thread may be referencing
+		// This fixes rendering thread crashes in the following order of operations 1) UMeshComponent::SetMaterial 2) GC 3) Rendering command that dereferences the UMaterial
+		World->SendAllEndOfFrameUpdates();
+	}
+}
+
+void UEngine::PreGarbageCollect()
+{
+	TRACE_CPUPROFILER_EVENT_SCOPE(UEngine::PreGarbageCollect);
+
+	SendWorldEndOfFrameUpdates();
 }
 
 float UEngine::GetTimeBetweenGarbageCollectionPasses() const
