Avoid chance of double-santizing texture cache keys.

diff --git a/Engine/Source/Runtime/Engine/Private/TextureDerivedData.cpp b/Engine/Source/Runtime/Engine/Private/TextureDerivedData.cpp
--- a/Engine/Source/Runtime/Engine/Private/TextureDerivedData.cpp
+++ b/Engine/Source/Runtime/Engine/Private/TextureDerivedData.cpp
@@ -2094,7 +2094,7 @@ void FTexturePlatformData::Serialize(FArchive& Ar, UTexture* Owner)
 FString FTexturePlatformData::GetDerivedDataMipKeyString(int32 MipIndex, const FTexture2DMipMap& Mip) const
 {
 	const FString& KeyString = DerivedDataKey.Get<FString>();
-	return FDerivedDataCacheInterface::SanitizeCacheKey(*FString::Printf(TEXT("%s_MIP%u_%dx%d"), *KeyString, MipIndex, Mip.SizeX, Mip.SizeY));
+	return FString::Printf(TEXT("%s_MIP%u_%dx%d"), *KeyString, MipIndex, Mip.SizeX, Mip.SizeY);
 }
 
 UE::DerivedData::FPayloadId FTexturePlatformData::MakeMipId(int32 MipIndex)
