World Partition: append suffix to generation log to determine source (PIE, Game, Cook, HLODs).

diff --git a/Engine/Source/Runtime/Engine/Private/WorldPartition/WorldPartitionStreamingGeneration.cpp b/Engine/Source/Runtime/Engine/Private/WorldPartition/WorldPartitionStreamingGeneration.cpp
--- a/Engine/Source/Runtime/Engine/Private/WorldPartition/WorldPartitionStreamingGeneration.cpp
+++ b/Engine/Source/Runtime/Engine/Private/WorldPartition/WorldPartitionStreamingGeneration.cpp
@@ -424,10 +424,9 @@ public:
 		return FActorClusterContext(MoveTemp(ContainerInstances), InFilterActorDescViewFunc);
 	}
 
-	void DumpStateLog()
+	void DumpStateLog(const TCHAR* Suffix)
 	{
-		FString StateLogOutputPath = FPaths::Combine(FPaths::ProjectDir(), TEXT("Saved"), TEXT("WorldPartition"));
-		FString StateLogOutputFilename = FPaths::Combine(StateLogOutputPath, *FString::Printf(TEXT("StreamingGeneration-%s.log"), *FDateTime::Now().ToString()));
+		const FString StateLogOutputFilename = FPaths::ProjectSavedDir() / TEXT("WorldPartition") / FString::Printf(TEXT("StreamingGeneration-%s-%s.log"), Suffix, *FDateTime::Now().ToString());
 		
 		if (FArchive* LogFile = IFileManager::Get().CreateFileWriter(*StateLogOutputFilename))
 		{
@@ -556,7 +555,10 @@ bool UWorldPartition::GenerateStreaming(TArray<FString>* OutPackagesToGenerate)
 
 		// Preparation Phase
 		StreamingGenerator.PreparationPhase(this);
-		StreamingGenerator.DumpStateLog();
+
+		// Dump state log
+		const TCHAR* StateLogSuffix = bIsPIE ? TEXT("PIE") : (IsRunningGame() ? TEXT("Game") : (IsRunningCookCommandlet() ? TEXT("Cook") : TEXT("Unknown")));
+		StreamingGenerator.DumpStateLog(StateLogSuffix);
 
 		// Preparation Phase :: Actor Clusters Creation
 		ActorClusterContext = StreamingGenerator.CreateActorClusters();
@@ -581,6 +583,7 @@ void UWorldPartition::GenerateHLOD(ISourceControlHelper* SourceControlHelper, bo
 	FStreamingGenerationLogErrorHandler LogErrorHandler;
 	FWorldPartitionStreamingGenerator StreamingGenerator(nullptr, &LogErrorHandler);
 	StreamingGenerator.PreparationPhase(this);
+	StreamingGenerator.DumpStateLog(TEXT("HLOD"));
 
 	// Preparation Phase :: Actor Clusters Creation
 	FActorClusterContext ActorClusterContext = StreamingGenerator.CreateActorClusters([](const FWorldPartitionActorDescView& ActorDescView)
