Engine changes to expose velocity functionality to location based impulses

diff --git a/Engine/Source/Runtime/Engine/Classes/Components/PrimitiveComponent.h b/Engine/Source/Runtime/Engine/Classes/Components/PrimitiveComponent.h
--- a/Engine/Source/Runtime/Engine/Classes/Components/PrimitiveComponent.h
+++ b/Engine/Source/Runtime/Engine/Classes/Components/PrimitiveComponent.h
@@ -1384,6 +1384,17 @@ public:
 	UFUNCTION(BlueprintCallable, Category="Physics", meta=(UnsafeDuringActorConstruction="true"))
 	virtual void AddImpulseAtLocation(FVector Impulse, FVector Location, FName BoneName = NAME_None);
 
+	/**
+	 *	Add an impulse to a single rigid body at a specific location. The Strength is taken as a change in angular velocity instead of an impulse (ie. mass will have no effect).
+	 *
+	 *	@param	Impulse		Magnitude and direction of impulse to apply.
+	 *	@param	Location	Point in world space to apply impulse at.
+	 *	@param	BoneName	If a SkeletalMeshComponent, name of bone to apply impulse to. 'None' indicates root body.
+	 */
+	UFUNCTION(BlueprintCallable, Category = "Physics", meta = (UnsafeDuringActorConstruction = "true"))
+	virtual void AddVelocityChangeImpulseAtLocation(FVector Impulse, FVector Location, FName BoneName = NAME_None);
+
+
 	/**
 	 * Add an impulse to all rigid bodies in this component, radiating out from the specified position.
 	 *
diff --git a/Engine/Source/Runtime/Engine/Classes/PhysicsEngine/BodyInstance.h b/Engine/Source/Runtime/Engine/Classes/PhysicsEngine/BodyInstance.h
--- a/Engine/Source/Runtime/Engine/Classes/PhysicsEngine/BodyInstance.h
+++ b/Engine/Source/Runtime/Engine/Classes/PhysicsEngine/BodyInstance.h
@@ -755,8 +755,13 @@ public:
 
 	/** Add an impulse to this body */
 	void AddImpulse(const FVector& Impulse, bool bVelChange);
+	
 	/** Add an impulse to this body and a particular world position */
 	void AddImpulseAtPosition(const FVector& Impulse, const FVector& Position);
+
+	/** Add a velocity change impulse to this body and a particular world position */
+	void AddVelocityChangeImpulseAtLocation(const FVector& Impulse, const FVector& Position);
+
 	/** Set the linear velocity of this body */
 	void SetLinearVelocity(const FVector& NewVel, bool bAddToCurrent, bool bAutoWake = true);
 
diff --git a/Engine/Source/Runtime/Engine/Private/PhysicsEngine/BodyInstance.cpp b/Engine/Source/Runtime/Engine/Private/PhysicsEngine/BodyInstance.cpp
--- a/Engine/Source/Runtime/Engine/Private/PhysicsEngine/BodyInstance.cpp
+++ b/Engine/Source/Runtime/Engine/Private/PhysicsEngine/BodyInstance.cpp
@@ -3597,6 +3597,17 @@ void FBodyInstance::AddImpulseAtPosition(const FVector& Impulse, const FVector&
 	});
 }
 
+void FBodyInstance::AddVelocityChangeImpulseAtLocation(const FVector& Impulse, const FVector& Position)
+{
+	FPhysicsCommand::ExecuteWrite(ActorHandle, [&](const FPhysicsActorHandle& Actor)
+	{
+		if (FPhysicsInterface::IsRigidBody(Actor) && FPhysicsInterface::IsInScene(Actor) && !IsRigidBodyKinematic_AssumesLocked(Actor))
+		{
+			FPhysicsInterface::AddVelocityChangeImpulseAtLocation_AssumesLocked(Actor, Impulse, Position);
+		}
+	});
+}
+
 void FBodyInstance::SetInstanceNotifyRBCollision(bool bNewNotifyCollision)
 {
 
diff --git a/Engine/Source/Runtime/Engine/Private/PrimitiveComponentPhysics.cpp b/Engine/Source/Runtime/Engine/Private/PrimitiveComponentPhysics.cpp
--- a/Engine/Source/Runtime/Engine/Private/PrimitiveComponentPhysics.cpp
+++ b/Engine/Source/Runtime/Engine/Private/PrimitiveComponentPhysics.cpp
@@ -143,6 +143,15 @@ void UPrimitiveComponent::AddImpulseAtLocation(FVector Impulse, FVector Location
 	}
 }
 
+void UPrimitiveComponent::AddVelocityChangeImpulseAtLocation(FVector Impulse, FVector Location, FName BoneName)
+{
+	if (FBodyInstance* BI = GetBodyInstance(BoneName))
+	{
+		WarnInvalidPhysicsOperations(LOCTEXT("AddImpulseAtLocation", "AddImpulseAtLocation"), BI, BoneName);
+		BI->AddVelocityChangeImpulseAtLocation(Impulse, Location);
+	}
+}
+
 void UPrimitiveComponent::AddRadialImpulse(FVector Origin, float Radius, float Strength, ERadialImpulseFalloff Falloff, bool bVelChange)
 {
 	if(bIgnoreRadialImpulse)
diff --git a/Engine/Source/Runtime/PhysicsCore/Private/ChaosEngineInterface.cpp b/Engine/Source/Runtime/PhysicsCore/Private/ChaosEngineInterface.cpp
--- a/Engine/Source/Runtime/PhysicsCore/Private/ChaosEngineInterface.cpp
+++ b/Engine/Source/Runtime/PhysicsCore/Private/ChaosEngineInterface.cpp
@@ -774,6 +774,14 @@ void FChaosEngineInterface::AddImpulseAtLocation_AssumesLocked(const FPhysicsAct
 	}
 }
 
+void FChaosEngineInterface::AddVelocityChangeImpulseAtLocation_AssumesLocked(const FPhysicsActorHandle& InActorReference, const FVector& InVelocityDelta, const FVector& InLocation)
+{
+	if (ensure(FChaosEngineInterface::IsValid(InActorReference)))
+	{
+		AddImpulseAtLocation_AssumesLocked(InActorReference, InActorReference->GetGameThreadAPI().M() * InVelocityDelta, InLocation);
+	}
+}
+
 void FChaosEngineInterface::AddRadialImpulse_AssumesLocked(const FPhysicsActorHandle& InActorReference,const FVector& InOrigin,float InRadius,float InStrength,ERadialImpulseFalloff InFalloff,bool bInVelChange)
 {
 	using namespace Chaos;
diff --git a/Engine/Source/Runtime/PhysicsCore/Public/Chaos/ChaosEngineInterface.h b/Engine/Source/Runtime/PhysicsCore/Public/Chaos/ChaosEngineInterface.h
--- a/Engine/Source/Runtime/PhysicsCore/Public/Chaos/ChaosEngineInterface.h
+++ b/Engine/Source/Runtime/PhysicsCore/Public/Chaos/ChaosEngineInterface.h
@@ -450,6 +450,7 @@ public:
 	static void AddVelocity_AssumesLocked(const FPhysicsActorHandle& InActorReference,const FVector& InForce);
 	static void AddAngularVelocityInRadians_AssumesLocked(const FPhysicsActorHandle& InActorReference,const FVector& InTorque);
 	static void AddImpulseAtLocation_AssumesLocked(const FPhysicsActorHandle& InActorReference,const FVector& InImpulse,const FVector& InLocation);
+	static void AddVelocityChangeImpulseAtLocation_AssumesLocked(const FPhysicsActorHandle& InActorReference, const FVector& InVelocityDelta, const FVector& InLocation);
 	static void AddRadialImpulse_AssumesLocked(const FPhysicsActorHandle& InActorReference,const FVector& InOrigin,float InRadius,float InStrength,ERadialImpulseFalloff InFalloff,bool bInVelChange);
 
 	static bool IsGravityEnabled_AssumesLocked(const FPhysicsActorHandle& InActorReference);
