StringBuilder: Add some missing conversions from deprecated Append to AppendChar.
#rnx

diff --git a/Engine/Source/Runtime/Engine/Private/HLSLTree/HLSLTree.cpp b/Engine/Source/Runtime/Engine/Private/HLSLTree/HLSLTree.cpp
--- a/Engine/Source/Runtime/Engine/Private/HLSLTree/HLSLTree.cpp
+++ b/Engine/Source/Runtime/Engine/Private/HLSLTree/HLSLTree.cpp
@@ -124,7 +124,7 @@ FEmitShaderCode* FEmitContext::EmitFormatCodeInternal(const Shader::FType& Type,
 		}
 		else
 		{
-			FormattedCode.Append(Char);
+			FormattedCode.AppendChar(Char);
 		}
 	}
 	checkf(ArgIndex == ArgList.Num(), TEXT("%d args were provided, but %d were used"), ArgList.Num(), ArgIndex);
@@ -201,10 +201,10 @@ void WriteMaterialUniformAccess(Shader::EValueComponentType ComponentType, uint3
 		// Can skip writing mask if we're taking all 4 components from the register
 		if (NumComponentsInRegister < 4u)
 		{
-			OutResult.Append(TCHAR('.'));
+			OutResult.AppendChar(TCHAR('.'));
 			for (uint32 i = 0u; i < NumComponentsInRegister; ++i)
 			{
-				OutResult.Append(IndexToMask[RegisterOffset + i]);
+				OutResult.AppendChar(IndexToMask[RegisterOffset + i]);
 			}
 		}
 
