Kismet: Update incorrect deprecation version

diff --git a/Engine/Source/Runtime/Engine/Classes/Kismet/KismetMathLibrary.h b/Engine/Source/Runtime/Engine/Classes/Kismet/KismetMathLibrary.h
--- a/Engine/Source/Runtime/Engine/Classes/Kismet/KismetMathLibrary.h
+++ b/Engine/Source/Runtime/Engine/Classes/Kismet/KismetMathLibrary.h
@@ -4078,7 +4078,7 @@ class ENGINE_API UKismetMathLibrary : public UBlueprintFunctionLibrary
 	 * @outparam	OutRectSideA - Vector oriented and sized to represent one edge of the enclosing rectangle, orthogonal to OutRectSideB
 	 * @outparam	OutRectSideB - Vector oriented and sized to represent one edge of the enclosing rectangle, orthogonal to OutRectSideA
 	*/
-	UE_DEPRECATED(5.2, "Use MinAreaRectangle instead. This deprecated version incorrectly returns the average of all input points as the rectangle center.")
+	UE_DEPRECATED(5.0, "Use MinAreaRectangle instead. This deprecated version incorrectly returns the average of all input points as the rectangle center.")
 	UFUNCTION(BlueprintCallable, Category="Math|Geometry", meta=(WorldContext="WorldContextObject", CallableWithoutWorldContext,
 		DeprecatedFunction, DeprecationMessage = "Use 'Min Area Rectangle' instead; this deprecated version incorrectly returns the average of all input points as the rectangle center."))
 	static void MinimumAreaRectangle(UObject* WorldContextObject, const TArray<FVector>& InVerts, const FVector& SampleSurfaceNormal, FVector& OutRectCenter, FRotator& OutRectRotation, float& OutSideLengthX, float& OutSideLengthY, bool bDebugDraw = false);
