Fix crash reimporting UDIM textures, related to UDIM blocks now being sorted

diff --git a/Engine/Source/Runtime/Engine/Private/Texture.cpp b/Engine/Source/Runtime/Engine/Private/Texture.cpp
--- a/Engine/Source/Runtime/Engine/Private/Texture.cpp
+++ b/Engine/Source/Runtime/Engine/Private/Texture.cpp
@@ -1214,7 +1214,7 @@ void FTextureSource::InitBlocked(const ETextureSourceFormat* InLayerFormats,
 	{
 		for (int i = 0; i < InNumBlocks; ++i)
 		{
-			const int64 BlockSize = CalcBlockSize(i);
+			const int64 BlockSize = CalcBlockSize(InBlocks[i]);
 			if (InDataPerBlock[i])
 			{
 				FMemory::Memcpy(DataPtr, InDataPerBlock[i], BlockSize);
