Destroy collisions when destroying an actor in RBAN

diff --git a/Engine/Source/Runtime/Engine/Private/PhysicsEngine/ImmediatePhysics/ImmediatePhysicsChaos/ImmediatePhysicsSimulation_Chaos.cpp b/Engine/Source/Runtime/Engine/Private/PhysicsEngine/ImmediatePhysics/ImmediatePhysicsChaos/ImmediatePhysicsSimulation_Chaos.cpp
--- a/Engine/Source/Runtime/Engine/Private/PhysicsEngine/ImmediatePhysics/ImmediatePhysicsChaos/ImmediatePhysicsSimulation_Chaos.cpp
+++ b/Engine/Source/Runtime/Engine/Private/PhysicsEngine/ImmediatePhysics/ImmediatePhysicsChaos/ImmediatePhysicsSimulation_Chaos.cpp
@@ -465,6 +465,7 @@ namespace ImmediatePhysics_Chaos
 		// @todo(ccaulfield): FActorHandle could remember its index to optimize this
 
 		RemoveFromCollidingPairs(ActorHandle);
+        DestroyActorCollisions(ActorHandle);
 
 		// If any joints reference the particle, we must destroy them
 		TArray<FJointHandle*> ActorJointHandles;
@@ -491,6 +492,11 @@ namespace ImmediatePhysics_Chaos
 		Implementation->bActorsDirty = true;
 	}
 
+    void FSimulation::DestroyActorCollisions(FActorHandle* ActorHandle)
+    {
+        Implementation->Collisions.GetConstraintAllocator().RemoveParticle(ActorHandle->GetParticle());
+    }
+
 	FJointHandle* FSimulation::CreateJoint(FConstraintInstance* ConstraintInstance, FActorHandle* Body1, FActorHandle* Body2)
 	{
 		FJointHandle* JointHandle = new FJointHandle(&Implementation->Joints, ConstraintInstance, Body1, Body2);
diff --git a/Engine/Source/Runtime/Engine/Public/Physics/ImmediatePhysics/ImmediatePhysicsChaos/ImmediatePhysicsSimulation_Chaos.h b/Engine/Source/Runtime/Engine/Public/Physics/ImmediatePhysics/ImmediatePhysicsChaos/ImmediatePhysicsSimulation_Chaos.h
--- a/Engine/Source/Runtime/Engine/Public/Physics/ImmediatePhysics/ImmediatePhysicsChaos/ImmediatePhysicsSimulation_Chaos.h
+++ b/Engine/Source/Runtime/Engine/Public/Physics/ImmediatePhysics/ImmediatePhysicsChaos/ImmediatePhysicsSimulation_Chaos.h
@@ -33,6 +33,8 @@ namespace ImmediatePhysics_Chaos
 		FActorHandle* CreateActor(EActorType ActorType, FBodyInstance* BodyInstance, const FTransform& Transform);
 		void DestroyActor(FActorHandle* ActorHandle);
 
+		void DestroyActorCollisions(FActorHandle* ActorHandle);
+
 		/** Create a physical joint and add it to the simulation */
 		FJointHandle* CreateJoint(FConstraintInstance* ConstraintInstance, FActorHandle* Body1, FActorHandle* Body2);
 		void DestroyJoint(FJointHandle* JointHandle);
