GameplayStatics sound ownership cleanup
- Use WorldContext if OwningActor is null
- Remove duplicate calls to ActorSoundParameterInterface::Fill
- Const correctness around OwningActor when used in PlaySound functions
- Set bStopWhenOwnerDestroyed to false to preserve backward compat in SpawnSoundAtLocation/SpawnSound2D cases

diff --git a/Engine/Source/Runtime/Engine/Classes/Kismet/GameplayStatics.h b/Engine/Source/Runtime/Engine/Classes/Kismet/GameplayStatics.h
--- a/Engine/Source/Runtime/Engine/Classes/Kismet/GameplayStatics.h
+++ b/Engine/Source/Runtime/Engine/Classes/Kismet/GameplayStatics.h
@@ -626,7 +626,7 @@ public:
 	 * @param bIsUISound - True if sound is UI related, else false
 	 */
 	UFUNCTION(BlueprintCallable, BlueprintCosmetic, Category="Audio", meta=( WorldContext="WorldContextObject", AdvancedDisplay = "2", UnsafeDuringActorConstruction = "true" ))
-	static void PlaySound2D(const UObject* WorldContextObject, USoundBase* Sound, float VolumeMultiplier = 1.f, float PitchMultiplier = 1.f, float StartTime = 0.f, USoundConcurrency* ConcurrencySettings = nullptr, AActor* OwningActor = nullptr, bool bIsUISound = true);
+	static void PlaySound2D(const UObject* WorldContextObject, USoundBase* Sound, float VolumeMultiplier = 1.f, float PitchMultiplier = 1.f, float StartTime = 0.f, USoundConcurrency* ConcurrencySettings = nullptr, const AActor* OwningActor = nullptr, bool bIsUISound = true);
 
 	/**
 	 * This function allows users to create Audio Components with settings specifically for non-spatialized,
@@ -678,7 +678,7 @@ public:
 	 *						to do a concurrency limit per owner.
 	 */
 	UFUNCTION(BlueprintCallable, Category="Audio", meta=(WorldContext="WorldContextObject", AdvancedDisplay = "3", UnsafeDuringActorConstruction = "true", Keywords = "play"))
-	static void PlaySoundAtLocation(const UObject* WorldContextObject, USoundBase* Sound, FVector Location, FRotator Rotation, float VolumeMultiplier = 1.f, float PitchMultiplier = 1.f, float StartTime = 0.f, class USoundAttenuation* AttenuationSettings = nullptr, USoundConcurrency* ConcurrencySettings = nullptr, AActor* OwningActor = nullptr, UInitialActiveSoundParams* InitialParams = nullptr);
+	static void PlaySoundAtLocation(const UObject* WorldContextObject, USoundBase* Sound, FVector Location, FRotator Rotation, float VolumeMultiplier = 1.f, float PitchMultiplier = 1.f, float StartTime = 0.f, class USoundAttenuation* AttenuationSettings = nullptr, USoundConcurrency* ConcurrencySettings = nullptr, const AActor* OwningActor = nullptr, UInitialActiveSoundParams* InitialParams = nullptr);
 
 	static void PlaySoundAtLocation(const UObject* WorldContextObject, USoundBase* Sound, FVector Location, float VolumeMultiplier = 1.f, float PitchMultiplier = 1.f, float StartTime = 0.f, class USoundAttenuation* AttenuationSettings = nullptr, USoundConcurrency* ConcurrencySettings = nullptr, UInitialActiveSoundParams* InitialParams = nullptr)
 	{
diff --git a/Engine/Source/Runtime/Engine/Private/ActiveSound.cpp b/Engine/Source/Runtime/Engine/Private/ActiveSound.cpp
--- a/Engine/Source/Runtime/Engine/Private/ActiveSound.cpp
+++ b/Engine/Source/Runtime/Engine/Private/ActiveSound.cpp
@@ -329,7 +329,7 @@ void FActiveSound::SetAudioComponent(const UAudioComponent& Component)
 	SetOwner(Owner);
 }
 
-void FActiveSound::SetOwner(AActor* Actor)
+void FActiveSound::SetOwner(const AActor* Actor)
 {
 	if (Actor)
 	{
diff --git a/Engine/Source/Runtime/Engine/Private/AudioDevice.cpp b/Engine/Source/Runtime/Engine/Private/AudioDevice.cpp
--- a/Engine/Source/Runtime/Engine/Private/AudioDevice.cpp
+++ b/Engine/Source/Runtime/Engine/Private/AudioDevice.cpp
@@ -6105,7 +6105,7 @@ UAudioComponent* FAudioDevice::CreateComponent(USoundBase* Sound, const FCreateC
 	return AudioComponent;
 }
 
-void FAudioDevice::PlaySoundAtLocation(USoundBase* Sound, UWorld* World, float VolumeMultiplier, float PitchMultiplier, float StartTime, const FVector& Location, const FRotator& Rotation, USoundAttenuation* AttenuationSettings, USoundConcurrency* Concurrency, const TArray<FAudioParameter>* Params, AActor* OwningActor)
+void FAudioDevice::PlaySoundAtLocation(USoundBase* Sound, UWorld* World, float VolumeMultiplier, float PitchMultiplier, float StartTime, const FVector& Location, const FRotator& Rotation, USoundAttenuation* AttenuationSettings, USoundConcurrency* Concurrency, const TArray<FAudioParameter>* Params, const AActor* OwningActor)
 {
 	check(IsInGameThread());
 
diff --git a/Engine/Source/Runtime/Engine/Private/Components/AudioComponent.cpp b/Engine/Source/Runtime/Engine/Private/Components/AudioComponent.cpp
--- a/Engine/Source/Runtime/Engine/Private/Components/AudioComponent.cpp
+++ b/Engine/Source/Runtime/Engine/Private/Components/AudioComponent.cpp
@@ -808,6 +808,10 @@ void UAudioComponent::PlayInternal(const PlayInternalRequestData& InPlayRequestD
 		UActorSoundParameterInterface::Fill(Owner, ActorParams);
 		FAudioParameter::Merge(MoveTemp(ActorParams), SoundParams);
 	}
+	else
+	{
+		UE_LOG(LogAudio, VeryVerbose, TEXT("AudioComponent : '%s' playing sound '%s' has no owner"), *GetFullName(), *Sound->GetName());
+	}
 
 	TArray<FAudioParameter> DefaultParamsCopy = DefaultParameters;
 	FAudioParameter::Merge(MoveTemp(DefaultParamsCopy), SoundParams);
diff --git a/Engine/Source/Runtime/Engine/Private/GameplayStatics.cpp b/Engine/Source/Runtime/Engine/Private/GameplayStatics.cpp
--- a/Engine/Source/Runtime/Engine/Private/GameplayStatics.cpp
+++ b/Engine/Source/Runtime/Engine/Private/GameplayStatics.cpp
@@ -1516,7 +1516,7 @@ void UGameplayStatics::SetGlobalListenerFocusParameters(const UObject* WorldCont
 	}
 }
 
-void UGameplayStatics::PlaySound2D(const UObject* WorldContextObject, USoundBase* Sound, float VolumeMultiplier, float PitchMultiplier, float StartTime, USoundConcurrency* ConcurrencySettings, AActor* OwningActor, bool bIsUISound)
+void UGameplayStatics::PlaySound2D(const UObject* WorldContextObject, USoundBase* Sound, float VolumeMultiplier, float PitchMultiplier, float StartTime, USoundConcurrency* ConcurrencySettings, const AActor* OwningActor, bool bIsUISound)
 {
 	if (!Sound || !GEngine || !GEngine->UseSound())
 	{
@@ -1551,10 +1551,14 @@ void UGameplayStatics::PlaySound2D(const UObject* WorldContextObject, USoundBase
 		NewActiveSound.Priority = Sound->Priority;
 		NewActiveSound.SubtitlePriority = Sound->GetSubtitlePriority();
 
-		NewActiveSound.SetOwner(OwningActor);
+		// If OwningActor isn't supplied to this function, derive an owner from the WorldContextObject
+		const AActor* WorldContextOwner = Cast<const AActor>(WorldContextObject);
+		const AActor* ActiveSoundOwner = OwningActor ? OwningActor : WorldContextOwner;
+
+		NewActiveSound.SetOwner(ActiveSoundOwner);
 
 		TArray<FAudioParameter> Params;
-		UActorSoundParameterInterface::Fill(OwningActor, Params);
+		UActorSoundParameterInterface::Fill(ActiveSoundOwner, Params);
 
 		AudioDevice->AddNewActiveSound(NewActiveSound, &Params);
 	}
@@ -1573,9 +1577,13 @@ UAudioComponent* UGameplayStatics::CreateSound2D(const UObject* WorldContextObje
 		return nullptr;
 	}
 
+	// Derive an owner from the WorldContextObject
+	UObject* MutableWorldContext = const_cast<UObject*>(WorldContextObject);
+	AActor* WorldContextOwner = Cast<AActor>(MutableWorldContext);
+
 	FAudioDevice::FCreateComponentParams Params = bPersistAcrossLevelTransition
 		? FAudioDevice::FCreateComponentParams(ThisWorld->GetAudioDeviceRaw())
-		: FAudioDevice::FCreateComponentParams(ThisWorld);
+		: FAudioDevice::FCreateComponentParams(ThisWorld, WorldContextOwner);
 
 	if (ConcurrencySettings)
 	{
@@ -1585,12 +1593,6 @@ UAudioComponent* UGameplayStatics::CreateSound2D(const UObject* WorldContextObje
 	UAudioComponent* AudioComponent = FAudioDevice::CreateComponent(Sound, Params);
 	if (AudioComponent)
 	{
-		TArray<FAudioParameter> ActorParams;
-		const AActor* ActorFromContext = Cast<AActor>(WorldContextObject);
-
-		UActorSoundParameterInterface::Fill(ActorFromContext, ActorParams);
-
-		AudioComponent->SetParameters(MoveTemp(ActorParams));
 		AudioComponent->SetVolumeMultiplier(VolumeMultiplier);
 		AudioComponent->SetPitchMultiplier(PitchMultiplier);
 		AudioComponent->bAllowSpatialization = false;
@@ -1598,6 +1600,7 @@ UAudioComponent* UGameplayStatics::CreateSound2D(const UObject* WorldContextObje
 		AudioComponent->bAutoDestroy = bAutoDestroy;
 		AudioComponent->bIgnoreForFlushing = bPersistAcrossLevelTransition;
 		AudioComponent->SubtitlePriority = Sound->GetSubtitlePriority();
+		AudioComponent->bStopWhenOwnerDestroyed = false;
 	}
 	return AudioComponent;
 }
@@ -1612,7 +1615,7 @@ UAudioComponent* UGameplayStatics::SpawnSound2D(const UObject* WorldContextObjec
 	return AudioComponent;
 }
 
-void UGameplayStatics::PlaySoundAtLocation(const UObject* WorldContextObject, class USoundBase* Sound, FVector Location, FRotator Rotation, float VolumeMultiplier, float PitchMultiplier, float StartTime, class USoundAttenuation* AttenuationSettings, class USoundConcurrency* ConcurrencySettings, AActor* OwningActor, UInitialActiveSoundParams* InitialParams)
+void UGameplayStatics::PlaySoundAtLocation(const UObject* WorldContextObject, class USoundBase* Sound, FVector Location, FRotator Rotation, float VolumeMultiplier, float PitchMultiplier, float StartTime, class USoundAttenuation* AttenuationSettings, class USoundConcurrency* ConcurrencySettings, const AActor* OwningActor, UInitialActiveSoundParams* InitialParams)
 {
 	if (!Sound || !GEngine || !GEngine->UseSound())
 	{
@@ -1633,9 +1636,13 @@ void UGameplayStatics::PlaySoundAtLocation(const UObject* WorldContextObject, cl
 			Params.Append(MoveTemp(InitialParams->AudioParams));
 		}
 
-		UActorSoundParameterInterface::Fill(OwningActor, Params);
+		// If OwningActor isn't supplied to this function, derive an owner from the WorldContextObject
+		const AActor* WorldContextOwner = Cast<const AActor>(WorldContextObject);
+		const AActor* ActiveSoundOwner = OwningActor ? OwningActor : WorldContextOwner;
 
-		AudioDevice->PlaySoundAtLocation(Sound, ThisWorld, VolumeMultiplier, PitchMultiplier, StartTime, Location, Rotation, AttenuationSettings, ConcurrencySettings, &Params, OwningActor);
+		UActorSoundParameterInterface::Fill(ActiveSoundOwner, Params);
+
+		AudioDevice->PlaySoundAtLocation(Sound, ThisWorld, VolumeMultiplier, PitchMultiplier, StartTime, Location, Rotation, AttenuationSettings, ConcurrencySettings, &Params, ActiveSoundOwner);
 	}
 }
 
@@ -1654,7 +1661,11 @@ UAudioComponent* UGameplayStatics::SpawnSoundAtLocation(const UObject* WorldCont
 
 	const bool bIsInGameWorld = ThisWorld->IsGameWorld();
 
-	FAudioDevice::FCreateComponentParams Params(ThisWorld);
+	// Derive an owner from the WorldContextObject
+	UObject* MutableWorldContext = const_cast<UObject*>(WorldContextObject);
+	AActor* WorldContextOwner = Cast<AActor>(MutableWorldContext);
+
+	FAudioDevice::FCreateComponentParams Params(ThisWorld, WorldContextOwner);
 	Params.SetLocation(Location);
 	Params.AttenuationSettings = AttenuationSettings;
 	
@@ -1667,11 +1678,6 @@ UAudioComponent* UGameplayStatics::SpawnSoundAtLocation(const UObject* WorldCont
 
 	if (AudioComponent)
 	{
-		TArray<FAudioParameter> ActorParams;
-		const AActor* ActorFromContext = Cast<AActor>(WorldContextObject);
-		UActorSoundParameterInterface::Fill(ActorFromContext, ActorParams);
-
-		AudioComponent->SetParameters(MoveTemp(ActorParams));
 		AudioComponent->SetWorldLocationAndRotation(Location, Rotation);
 		AudioComponent->SetVolumeMultiplier(VolumeMultiplier);
 		AudioComponent->SetPitchMultiplier(PitchMultiplier);
@@ -1679,6 +1685,7 @@ UAudioComponent* UGameplayStatics::SpawnSoundAtLocation(const UObject* WorldCont
 		AudioComponent->bIsUISound				= !bIsInGameWorld;
 		AudioComponent->bAutoDestroy			= bAutoDestroy;
 		AudioComponent->SubtitlePriority		= Sound->GetSubtitlePriority();
+		AudioComponent->bStopWhenOwnerDestroyed = false;
 		AudioComponent->Play(StartTime);
 	}
 
@@ -1758,10 +1765,6 @@ UAudioComponent* UGameplayStatics::SpawnSoundAttached(USoundBase* Sound, USceneC
 				}
 			}
 
-			TArray<FAudioParameter> ActorParams;
-			UActorSoundParameterInterface::Fill(AttachToComponent->GetOwner(), ActorParams);
-
-			AudioComponent->SetParameters(MoveTemp(ActorParams));
 			AudioComponent->SetVolumeMultiplier(VolumeMultiplier);
 			AudioComponent->SetPitchMultiplier(PitchMultiplier);
 			AudioComponent->bAllowSpatialization = Params.ShouldUseAttenuation();
diff --git a/Engine/Source/Runtime/Engine/Public/ActiveSound.h b/Engine/Source/Runtime/Engine/Public/ActiveSound.h
--- a/Engine/Source/Runtime/Engine/Public/ActiveSound.h
+++ b/Engine/Source/Runtime/Engine/Public/ActiveSound.h
@@ -60,7 +60,6 @@ struct FAttenuationFocusData
 		, FocusFactor(1.0f)
 		, DistanceScale(1.0f)
 		, PriorityScale(1.0f)
-		, PriorityHighest(1.0f)
 		, VolumeScale(1.0f)
 	{
 	}
@@ -308,7 +307,7 @@ public:
 	void ClearAudioComponent();
 	void SetAudioComponent(const FActiveSound& ActiveSound);
 	void SetAudioComponent(const UAudioComponent& Component);
-	void SetOwner(AActor* Owner);
+	void SetOwner(const AActor* Owner);
 	FString GetAudioComponentName() const;
 	FString GetOwnerName() const;
 
diff --git a/Engine/Source/Runtime/Engine/Public/AudioDevice.h b/Engine/Source/Runtime/Engine/Public/AudioDevice.h
--- a/Engine/Source/Runtime/Engine/Public/AudioDevice.h
+++ b/Engine/Source/Runtime/Engine/Public/AudioDevice.h
@@ -764,7 +764,7 @@ public:
 	 * @param	USoundConcurrency	The sound's sound concurrency settings to use (optional). Will use the USoundBase's USoundConcurrency if not specified.
 	 * @param	Params				An optional list of audio component params to immediately apply to a sound.
 	 */
-	void PlaySoundAtLocation(USoundBase* Sound, UWorld* World, float VolumeMultiplier, float PitchMultiplier, float StartTime, const FVector& Location, const FRotator& Rotation, USoundAttenuation* AttenuationSettings = nullptr, USoundConcurrency* ConcurrencySettings = nullptr, const TArray<FAudioParameter>* Params = nullptr, AActor* OwningActor = nullptr);
+	void PlaySoundAtLocation(USoundBase* Sound, UWorld* World, float VolumeMultiplier, float PitchMultiplier, float StartTime, const FVector& Location, const FRotator& Rotation, USoundAttenuation* AttenuationSettings = nullptr, USoundConcurrency* ConcurrencySettings = nullptr, const TArray<FAudioParameter>* Params = nullptr, const AActor* OwningActor = nullptr);
 
 	/**
 	 * Adds an active sound to the audio device
