GameEngine: Fixed SwitchGameWindowToUseGameViewport not taking the window DPI scale into account.
#rnx

diff --git a/Engine/Source/Runtime/Engine/Private/GameEngine.cpp b/Engine/Source/Runtime/Engine/Private/GameEngine.cpp
--- a/Engine/Source/Runtime/Engine/Private/GameEngine.cpp
+++ b/Engine/Source/Runtime/Engine/Private/GameEngine.cpp
@@ -660,7 +660,7 @@ void UGameEngine::SwitchGameWindowToUseGameViewport()
 		TSharedPtr<SWindow> GameViewportWindowPtr = GameViewportWindow.Pin();
 		
 		GameViewportWindowPtr->SetContent(GameViewportWidgetRef);
-		GameViewportWindowPtr->SlatePrepass();
+		GameViewportWindowPtr->SlatePrepass(FSlateApplication::Get().GetApplicationScale() * GameViewportWindowPtr->GetNativeWindow()->GetDPIScaleFactor());
 		
 		if ( SceneViewport.IsValid() )
 		{
