Added Analytics Support for Virtual Assets / Virtualization Module

diff --git a/Engine/Source/Runtime/Engine/Engine.Build.cs b/Engine/Source/Runtime/Engine/Engine.Build.cs
--- a/Engine/Source/Runtime/Engine/Engine.Build.cs
+++ b/Engine/Source/Runtime/Engine/Engine.Build.cs
@@ -8,7 +8,7 @@ public class Engine : ModuleRules
 	public Engine(ReadOnlyTargetRules Target) : base(Target)
 	{
 		PrivateIncludePaths.Add("../Shaders/Shared");
-
+		
 		PrivatePCHHeaderFile = "Private/EnginePrivatePCH.h";
 
 		SharedPCHHeaderFile = "Public/EngineSharedPCH.h";
@@ -19,7 +19,8 @@ public class Engine : ModuleRules
 			new string[] {
 				"Runtime/SynthBenchmark/Public",
 				"Runtime/Engine/Private",
-				"Runtime/Net/Core/Private/Net/Core/PushModel/Types"
+				"Runtime/Net/Core/Private/Net/Core/PushModel/Types",
+				"Developer/Virtualization/Private"
 			}
 		);
 
diff --git a/Engine/Source/Runtime/Engine/Private/Analytics/StudioAnalytics.cpp b/Engine/Source/Runtime/Engine/Private/Analytics/StudioAnalytics.cpp
--- a/Engine/Source/Runtime/Engine/Private/Analytics/StudioAnalytics.cpp
+++ b/Engine/Source/Runtime/Engine/Private/Analytics/StudioAnalytics.cpp
@@ -20,6 +20,7 @@
 #include "HAL/PlatformProcess.h"
 #include "DerivedDataCacheInterface.h"
 #include "DerivedDataCacheUsageStats.h"
+#include "Virtualization/VirtualizationSystem.h"
 
 bool FStudioAnalytics::bInitialized = false;
 volatile double FStudioAnalytics::TimeEstimation = 0;
@@ -27,6 +28,8 @@ FThread FStudioAnalytics::TimerThread;
 TSharedPtr<IAnalyticsProviderET> FStudioAnalytics::Analytics;
 TArray<FAnalyticsEventAttribute> FStudioAnalytics::DefaultAttributes;
 
+using namespace UE::Virtualization;
+
 void FStudioAnalytics::SetProvider(TSharedRef<IAnalyticsProviderET> InAnalytics)
 {
 	checkf(!Analytics.IsValid(), TEXT("FStudioAnalytics::SetProvider called more than once."));
@@ -218,7 +221,7 @@ void FStudioAnalytics::FireEvent_Loading(const FString& LoadingName, double Seco
 			}
 		}
 
-		// Grab the summary stats
+		// Grab the DDC summary stats
 		FDerivedDataCacheSummaryStats DDCSummaryStats;
 		GetDerivedDataCacheRef().GatherSummaryStats(DDCSummaryStats);
 		
@@ -227,7 +230,65 @@ void FStudioAnalytics::FireEvent_Loading(const FString& LoadingName, double Seco
 			FString FormattedAttrName = "DDC.Summary." + Stat.Key;
 			Attributes.Emplace(FormattedAttrName, Stat.Value);
 		}
-		
+
+		// Grab the Virtualization stats
+		IVirtualizationSystem& System = IVirtualizationSystem::Get();
+
+		FPayloadActivityInfo PayloadActivityInfo = System.GetAccumualtedPayloadActivityInfo();
+
+		const FString BaseName = TEXT("Virtualization");
+
+		{
+			FString AttrName = BaseName + TEXT(".Enabled");
+			Attributes.Emplace(MoveTemp(AttrName), System.IsEnabled());
+		}
+
+		{
+			FString AttrName = BaseName + TEXT(".Cache.TimeSpent");
+			Attributes.Emplace(MoveTemp(AttrName), (double)PayloadActivityInfo.Cache.CyclesSpent * FPlatformTime::GetSecondsPerCycle());
+		}
+
+		{
+			FString AttrName = BaseName + TEXT(".Cache.PayloadCount");
+			Attributes.Emplace(MoveTemp(AttrName), (double)PayloadActivityInfo.Cache.PayloadCount);
+		}
+
+		{
+			FString AttrName = BaseName + TEXT(".Cache.TotalBytes");
+			Attributes.Emplace(MoveTemp(AttrName), (double)PayloadActivityInfo.Cache.TotalBytes);
+		}
+
+		{
+			FString AttrName = BaseName + TEXT(".Push.TimeSpent");
+			Attributes.Emplace(MoveTemp(AttrName), (double)PayloadActivityInfo.Push.CyclesSpent * FPlatformTime::GetSecondsPerCycle());
+		}
+
+		{
+			FString AttrName = BaseName + TEXT(".Push.PayloadCount");
+			Attributes.Emplace(MoveTemp(AttrName), (double)PayloadActivityInfo.Push.PayloadCount);
+		}
+
+		{
+			FString AttrName = BaseName + TEXT(".Push.TotalBytes");
+			Attributes.Emplace(MoveTemp(AttrName), (double)PayloadActivityInfo.Push.TotalBytes);
+		}
+
+		{
+			FString AttrName = BaseName + TEXT(".Pull.TimeSpent");
+			Attributes.Emplace(MoveTemp(AttrName), (double)PayloadActivityInfo.Pull.CyclesSpent * FPlatformTime::GetSecondsPerCycle());
+		}
+
+		{
+			FString AttrName = BaseName + TEXT(".Pull.PayloadCount");
+			Attributes.Emplace(MoveTemp(AttrName), (double)PayloadActivityInfo.Pull.PayloadCount);
+		}
+
+		{
+			FString AttrName = BaseName + TEXT(".Pull.TotalBytes");
+			Attributes.Emplace(MoveTemp(AttrName), (double)PayloadActivityInfo.Pull.TotalBytes);
+		}
+
+		// Store it all in the loading event
 		FStudioAnalytics::GetProvider().RecordEvent(TEXT("Core.Loading"), Attributes);
 #endif
 	}
