World Partition: fix crash when cooking.

diff --git a/Engine/Source/Runtime/Engine/Private/WorldPartition/WorldPartition.cpp b/Engine/Source/Runtime/Engine/Private/WorldPartition/WorldPartition.cpp
--- a/Engine/Source/Runtime/Engine/Private/WorldPartition/WorldPartition.cpp
+++ b/Engine/Source/Runtime/Engine/Private/WorldPartition/WorldPartition.cpp
@@ -568,11 +568,14 @@ void UWorldPartition::Uninitialize()
 			if (!IsEngineExitRequested())
 			{
 				// Unload all Editor cells
-				// @todo_ow: Once Metadata is removed from external actor's package, this won't be necessary anymore.
-				EditorHash->ForEachCell([this](UWorldPartitionEditorCell* Cell)
+				if (EditorHash)
 				{
-					UpdateLoadingEditorCell(Cell, /*bShouldBeLoaded*/false, /*bIsFromUserChange*/false);
-				});
+					// @todo_ow: Once Metadata is removed from external actor's package, this won't be necessary anymore.
+					EditorHash->ForEachCell([this](UWorldPartitionEditorCell* Cell)
+					{
+						UpdateLoadingEditorCell(Cell, /*bShouldBeLoaded*/false, /*bIsFromUserChange*/false);
+					});
+				}
 			}
 		}
 
