Moved the PIEStartupTime LoadEvent into the root GameInstanceclass so that it can be recorded from all projects.

diff --git a/Engine/Source/Runtime/Engine/Classes/Engine/GameInstance.h b/Engine/Source/Runtime/Engine/Classes/Engine/GameInstance.h
--- a/Engine/Source/Runtime/Engine/Classes/Engine/GameInstance.h
+++ b/Engine/Source/Runtime/Engine/Classes/Engine/GameInstance.h
@@ -198,6 +198,7 @@ public:
 	FString PIEMapName;
 #if WITH_EDITOR
 	double PIEStartTime = 0;
+	bool bReportedPIEStartupTime = false;
 #endif
 
 	//~ Begin FExec Interface
@@ -249,6 +250,8 @@ public:
 
 	/** Called as soon as the game mode is spawned, to allow additional PIE setting validation prior to creating the local players / etc... (called on pure clients too, in which case the game mode is nullptr) */
 	virtual FGameInstancePIEResult PostCreateGameModeForPIE(const FGameInstancePIEParameters& Params, AGameModeBase* GameMode);
+
+	void ReportPIEStartupTime();
 #endif
 
 	class UEngine* GetEngine() const;
diff --git a/Engine/Source/Runtime/Engine/Private/GameInstance.cpp b/Engine/Source/Runtime/Engine/Private/GameInstance.cpp
--- a/Engine/Source/Runtime/Engine/Private/GameInstance.cpp
+++ b/Engine/Source/Runtime/Engine/Private/GameInstance.cpp
@@ -33,6 +33,8 @@
 #if WITH_EDITOR
 #include "Settings/LevelEditorPlaySettings.h"
 #include "Editor/EditorEngine.h"
+#include "EngineAnalytics.h"
+#include "StudioAnalytics.h"
 #endif
 
 UGameInstance::UGameInstance(const FObjectInitializer& ObjectInitializer)
@@ -320,6 +322,32 @@ FGameInstancePIEResult UGameInstance::InitializeForPlayInEditor(int32 PIEInstanc
 	return FGameInstancePIEResult::Success();
 }
 
+
+void UGameInstance::ReportPIEStartupTime()
+{
+#if WITH_EDITOR
+	if (!bReportedPIEStartupTime)
+	{
+		static bool bHasRunPIEThisSession = false;
+
+		bool bReportFirstTime = false;
+		if (!bHasRunPIEThisSession)
+		{
+			bReportFirstTime = true;
+			bHasRunPIEThisSession = true;
+		}
+
+		bReportedPIEStartupTime = true;
+		const double TimeToStartPIE = FStudioAnalytics::GetAnalyticSeconds() - PIEStartTime;
+		FStudioAnalytics::FireEvent_Loading(TEXT("PIE.TotalStartupTime"), TimeToStartPIE, {
+			FAnalyticsEventAttribute(TEXT("MapName"), UWorld::RemovePIEPrefix(FPackageName::GetShortName(PIEMapName))),
+			FAnalyticsEventAttribute(TEXT("FirstTime"), bReportFirstTime)
+			});
+	}
+#endif
+}
+
+
 FGameInstancePIEResult UGameInstance::StartPlayInEditorGameInstance(ULocalPlayer* LocalPlayer, const FGameInstancePIEParameters& Params)
 {
 	if (!Params.EditorPlaySettings)
@@ -425,7 +453,7 @@ FGameInstancePIEResult UGameInstance::StartPlayInEditorGameInstance(ULocalPlayer
 		{
 			return PostCreateGameModeResult;
 		}
-		
+
 		// Make sure "always loaded" sub-levels are fully loaded
 		PlayWorld->FlushLevelStreaming(EFlushLevelStreamingType::Visibility);
 
@@ -451,7 +479,7 @@ FGameInstancePIEResult UGameInstance::StartPlayInEditorGameInstance(ULocalPlayer
 			// Stream any levels now that need to be loaded before the game starts
 			GEngine->BlockTillLevelStreamingCompleted(PlayWorld);
 		}
-		
+
 		if (Params.NetMode == PIE_ListenServer)
 		{
 			// Add port
@@ -467,6 +495,13 @@ FGameInstancePIEResult UGameInstance::StartPlayInEditorGameInstance(ULocalPlayer
 		}
 
 		PlayWorld->BeginPlay();
+
+#if WITH_EDITOR
+		if (PlayWorld->WorldType == EWorldType::PIE)
+		{
+			ReportPIEStartupTime();
+		}
+#endif
 	}
 
 	// Games can override this to return failure if PIE is not allowed for some reason
