Fix crash undoing delete of virtual bone

diff --git a/Engine/Source/Runtime/Engine/Classes/Animation/SmartName.h b/Engine/Source/Runtime/Engine/Classes/Animation/SmartName.h
--- a/Engine/Source/Runtime/Engine/Classes/Animation/SmartName.h
+++ b/Engine/Source/Runtime/Engine/Classes/Animation/SmartName.h
@@ -62,10 +62,7 @@ struct ENGINE_API FSmartNameMapping
 	GENERATED_USTRUCT_BODY();
 
 	FSmartNameMapping();
-	FSmartNameMapping(FRWLock *Lock);
 	
-	void SetLock(FRWLock *Lock) { RWLock = Lock; }
-
 	// Add a name to the mapping, fails if it already exists
 	// @param InName - The name to add
 	// @return FSmartName - populated smart name
@@ -165,8 +162,6 @@ private:
 #endif
 
 	TMap<FName, FCurveMetaData> CurveMetaDataMap;
-
-	FRWLock* RWLock;
 };
 
 USTRUCT()
@@ -201,8 +196,6 @@ private:
 	// Editor copy of the data we loaded, used to preserve determinism during cooking
 	TMap<FName, FSmartNameMapping> LoadedNameMappings;
 #endif
-	
-	mutable FRWLock RWLock;
 };
 
 template<>
diff --git a/Engine/Source/Runtime/Engine/Private/Animation/SmartName.cpp b/Engine/Source/Runtime/Engine/Private/Animation/SmartName.cpp
--- a/Engine/Source/Runtime/Engine/Private/Animation/SmartName.cpp
+++ b/Engine/Source/Runtime/Engine/Private/Animation/SmartName.cpp
@@ -6,24 +6,23 @@
 #include "UObject/AnimPhysObjectVersion.h"
 #include "Misc/ScopeRWLock.h"
 
+namespace
+{
+	FRWLock SmartNameRWLock;
+}
+
 ////////////////////////////////////////////////////////////////////////
 //
 // FSmartNameMapping
 //
 ///////////////////////////////////////////////////////////////////////
 FSmartNameMapping::FSmartNameMapping()
-	: RWLock(nullptr)
-{
-}
-
-FSmartNameMapping::FSmartNameMapping(FRWLock* Lock)
-	: RWLock(Lock)
 {
 }
 
 void FSmartNameMapping::Iterate(TFunction<void(const FSmartNameMapping* Mapping, SmartName::UID_Type ID)> Callback) const
 {
-	FReadScopeLock Lock(*RWLock);
+	FReadScopeLock Lock(SmartNameRWLock);
 
 	for (int32 NameIndex = 0; NameIndex < CurveNameList.Num(); ++NameIndex)
 	{
@@ -33,7 +32,7 @@ void FSmartNameMapping::Iterate(TFunction<void(const FSmartNameMapping* Mapping,
 
 FSmartName FSmartNameMapping::AddName(FName Name)
 {
-	FWriteScopeLock Lock(*RWLock);
+	FWriteScopeLock Lock(SmartNameRWLock);
 	
 	// Make sure we are not trying to do an invalid add
 	check(Name.IsValid() && !CurveNameList.Contains(Name));
@@ -52,28 +51,28 @@ FSmartName FSmartNameMapping::AddName(FName Name)
 FCurveMetaData* FSmartNameMapping::GetCurveMetaData(FName CurveName)
 {
 	checkSlow(Exists(CurveName));
-	FWriteScopeLock Lock(*RWLock);
+	FWriteScopeLock Lock(SmartNameRWLock);
 	return &CurveMetaDataMap.FindOrAdd(CurveName);
 }
 
 const FCurveMetaData* FSmartNameMapping::GetCurveMetaData(FName CurveName) const
 {
 	checkSlow(Exists(CurveName));
-	FReadScopeLock Lock(*RWLock);
+	FReadScopeLock Lock(SmartNameRWLock);
 	return CurveMetaDataMap.Find(CurveName);
 }
 
 #if !WITH_EDITOR
 const FCurveMetaData& FSmartNameMapping::GetCurveMetaData(SmartName::UID_Type CurveUID) const
 {
-	FReadScopeLock Lock(*RWLock);
+	FReadScopeLock Lock(SmartNameRWLock);
 	return CurveMetaDataList[CurveUID];
 }
 #endif
 
 bool FSmartNameMapping::GetName(const SmartName::UID_Type& Uid, FName& OutName) const
 {
-	FReadScopeLock Lock(*RWLock);
+	FReadScopeLock Lock(SmartNameRWLock);
 	
 	if (CurveNameList.IsValidIndex(Uid))
 	{
@@ -86,7 +85,7 @@ bool FSmartNameMapping::GetName(const SmartName::UID_Type& Uid, FName& OutName)
 #if WITH_EDITOR
 bool FSmartNameMapping::Rename(const SmartName::UID_Type& Uid, FName NewName)
 {
-	FWriteScopeLock Lock(*RWLock);
+	FWriteScopeLock Lock(SmartNameRWLock);
 	
 	FName ExistingName;
 	if(GetName(Uid, ExistingName))
@@ -110,7 +109,7 @@ bool FSmartNameMapping::Rename(const SmartName::UID_Type& Uid, FName NewName)
 
 bool FSmartNameMapping::Remove(const SmartName::UID_Type& Uid)
 {
-	FWriteScopeLock Lock(*RWLock);
+	FWriteScopeLock Lock(SmartNameRWLock);
 	
 	FName ExistingName;
 	if (GetName(Uid, ExistingName))
@@ -125,7 +124,7 @@ bool FSmartNameMapping::Remove(const SmartName::UID_Type& Uid)
 
 bool FSmartNameMapping::Remove(const FName& Name)
 {
-	FWriteScopeLock Lock(*RWLock);
+	FWriteScopeLock Lock(SmartNameRWLock);
 	
 	const SmartName::UID_Type Uid = FindUID(Name);
 	if (Uid != SmartName::MaxUID)
@@ -175,7 +174,7 @@ void FSmartNameMapping::Serialize(FArchive& Ar)
 
 void FSmartNameMapping::FillUidArray(TArray<SmartName::UID_Type>& Array) const
 {
-	FReadScopeLock Lock(*RWLock);
+	FReadScopeLock Lock(SmartNameRWLock);
 	
 	Array.Reset(CurveNameList.Num());
 	
@@ -195,7 +194,7 @@ void FSmartNameMapping::FillUidArray(TArray<SmartName::UID_Type>& Array) const
 
 void FSmartNameMapping::FillNameArray(TArray<FName>& Array) const
 {
-	FReadScopeLock Lock(*RWLock);
+	FReadScopeLock Lock(SmartNameRWLock);
 	
 	//In editor names can be removed and so have to deal with empty slots
 #if WITH_EDITOR
@@ -219,7 +218,7 @@ void FSmartNameMapping::FillUIDToNameArray(TArray<FName>& Array) const
 
 void FSmartNameMapping::FillCurveTypeArray(TArray<FAnimCurveType>& Array) const
 {
-	FReadScopeLock Lock(*RWLock);
+	FReadScopeLock Lock(SmartNameRWLock);
 	Array.Reset(CurveNameList.Num());
 
 	for (const FName& Name : CurveNameList)
@@ -241,7 +240,7 @@ void FSmartNameMapping::FillCurveTypeArray(TArray<FAnimCurveType>& Array) const
 
 void FSmartNameMapping::FillUIDToCurveTypeArray(TArray<FAnimCurveType>& Array) const
 {
-	FReadScopeLock Lock(*RWLock);
+	FReadScopeLock Lock(SmartNameRWLock);
 	
 	Array.Reset(CurveNameList.Num());
 
@@ -269,19 +268,19 @@ void FSmartNameMapping::FillUIDToCurveTypeArray(TArray<FAnimCurveType>& Array) c
 
 bool FSmartNameMapping::Exists(const SmartName::UID_Type& Uid) const
 {
-	FReadScopeLock Lock(*RWLock);
+	FReadScopeLock Lock(SmartNameRWLock);
 	return CurveNameList.IsValidIndex(Uid) && CurveNameList[Uid] != NAME_None;
 }
 
 bool FSmartNameMapping::Exists(const FName& Name) const
 {
-	FReadScopeLock Lock(*RWLock);
+	FReadScopeLock Lock(SmartNameRWLock);
 	return CurveNameList.Contains(Name);
 }
 
 SmartName::UID_Type FSmartNameMapping::FindUID(const FName& Name) const
 {
-	FReadScopeLock Lock(*RWLock);
+	FReadScopeLock Lock(SmartNameRWLock);
 	return CurveNameList.IndexOfByKey(Name);
 }
 
@@ -294,7 +293,7 @@ FArchive& operator<<(FArchive& Ar, FSmartNameMapping& Elem)
 
 bool FSmartNameMapping::FindSmartName(FName Name, FSmartName& OutName) const
 {
-	FReadScopeLock Lock(*RWLock);
+	FReadScopeLock Lock(SmartNameRWLock);
 	SmartName::UID_Type ExistingUID = FindUID(Name);
 	if (ExistingUID != SmartName::MaxUID)
 	{
@@ -307,7 +306,7 @@ bool FSmartNameMapping::FindSmartName(FName Name, FSmartName& OutName) const
 
 bool FSmartNameMapping::FindSmartNameByUID(SmartName::UID_Type UID, FSmartName& OutName) const
 {
-	FReadScopeLock Lock(*RWLock);
+	FReadScopeLock Lock(SmartNameRWLock);
 	FName ExistingName;
 	if (GetName(UID, ExistingName))
 	{
@@ -322,7 +321,7 @@ bool FSmartNameMapping::FindSmartNameByUID(SmartName::UID_Type UID, FSmartName&
 /* initialize curve meta data for the container */
 void FSmartNameMapping::InitializeCurveMetaData(class USkeleton* Skeleton)
 {
-	FWriteScopeLock Lock(*RWLock);
+	FWriteScopeLock Lock(SmartNameRWLock);
 	
 	// initialize bone indices for skeleton
 	for (TPair<FName, FCurveMetaData>& Iter : CurveMetaDataMap)
@@ -341,12 +340,12 @@ void FSmartNameMapping::InitializeCurveMetaData(class USkeleton* Skeleton)
 //////////////////////////////////////////////////////////////////////
 FSmartNameMapping* FSmartNameContainer::AddContainer(FName NewContainerName)
 {
-	FWriteScopeLock Lock(RWLock);
+	FWriteScopeLock Lock(SmartNameRWLock);
 	
 	FSmartNameMapping* ExistingMapping = NameMappings.Find(NewContainerName);
 	if(ExistingMapping == nullptr)
 	{
-		return &NameMappings.Add(NewContainerName, {&RWLock});
+		return &NameMappings.Add(NewContainerName);
 	}
 
 	return ExistingMapping;
@@ -378,11 +377,6 @@ void FSmartNameContainer::PostLoad()
 	// Preserve Load state for deterministic cooking
 	LoadedNameMappings = NameMappings;
 #endif
-	
-	for(auto& Mapping : NameMappings)
-	{
-		Mapping.Value.SetLock(&RWLock);
-	}
 }
 
 FSmartNameMapping* FSmartNameContainer::GetContainerInternal(const FName& ContainerName)
