Extend push model dirty tracking to fast arrays
Some minor deprecation cleanup

diff --git a/Engine/Source/Runtime/CoreUObject/Private/UObject/CoreNet.cpp b/Engine/Source/Runtime/CoreUObject/Private/UObject/CoreNet.cpp
--- a/Engine/Source/Runtime/CoreUObject/Private/UObject/CoreNet.cpp
+++ b/Engine/Source/Runtime/CoreUObject/Private/UObject/CoreNet.cpp
@@ -555,29 +555,4 @@ const TCHAR* LexToString(const EChannelCloseReason Value)
 	}
 
 	return TEXT("Unknown");
-}
-
-void INetSerializeCB::NetSerializeStruct(
-	class UScriptStruct* Struct,
-	class FBitArchive& Ar,
-	class UPackageMap* Map,
-	void* Data,
-	bool& bHasUnmapped)
-{
-	FNetDeltaSerializeInfo Params;
-	Params.Struct = Struct;
-	Params.Map = Map;
-	Params.Data = Data;
-
-	if (Ar.IsSaving())
-	{
-		Params.Writer = static_cast<FBitWriter*>(&Ar);
-	}
-	else
-	{
-		Params.Reader = static_cast<FBitReader*>(&Ar);
-	}
-
-	NetSerializeStruct(Params);
-	bHasUnmapped = Params.bOutHasMoreUnmapped;
 }
\ No newline at end of file
diff --git a/Engine/Source/Runtime/CoreUObject/Public/UObject/CoreNet.h b/Engine/Source/Runtime/CoreUObject/Public/UObject/CoreNet.h
--- a/Engine/Source/Runtime/CoreUObject/Public/UObject/CoreNet.h
+++ b/Engine/Source/Runtime/CoreUObject/Public/UObject/CoreNet.h
@@ -201,10 +201,6 @@ class COREUOBJECT_API UPackageMap : public UObject
 	virtual void Serialize(FArchive& Ar) override;
 
 protected:
-
-	UE_DEPRECATED(4.25, "bSuppressLogs will be removed in a future release.")
-	bool					bSuppressLogs;
-
 	bool					bShouldTrackUnmappedGuids;
 	TSet< FNetworkGUID >	TrackedUnmappedNetGuids;
 	TSet< FNetworkGUID >	TrackedMappedDynamicNetGuids;
@@ -485,14 +481,6 @@ public:
 	 */
 	virtual void NetSerializeStruct(FNetDeltaSerializeInfo& Params) = 0;
 
-	UE_DEPRECATED(4.23, "Please use the version of NetSerializeStruct that accepts an FNetDeltaSerializeInfo reference")
-	virtual void NetSerializeStruct(
-		class UScriptStruct* Struct,
-		class FBitArchive& Ar,
-		class UPackageMap* Map,
-		void* Data,
-		bool& bHasUnmapped);
-
 	/**
 	 * Gathers any guid references for a FastArraySerializer.
 	 * @see GuidReferences.h for more info.
diff --git a/Engine/Source/Runtime/Engine/Private/DataReplication.cpp b/Engine/Source/Runtime/Engine/Private/DataReplication.cpp
--- a/Engine/Source/Runtime/Engine/Private/DataReplication.cpp
+++ b/Engine/Source/Runtime/Engine/Private/DataReplication.cpp
@@ -68,6 +68,12 @@ FAutoConsoleVariableRef CVarPushModelSkipUndirtiedReplicators(
 	GbPushModelSkipUndirtiedReplicators,
 	TEXT("When true, skip replicating any objects that we can safely see aren't dirty."));
 
+bool GbPushModelSkipUndirtiedFastArrays = false;
+FAutoConsoleVariableRef CVarPushModelSkipUndirtiedFastArrays(
+	TEXT("net.PushModelSkipUndirtiedFastArrays"),
+	GbPushModelSkipUndirtiedFastArrays,
+	TEXT("When true, include fast arrays when skipping objects that we can safely see aren't dirty."));
+
 extern int32 GNumSkippedObjectEmptyUpdates;
 
 extern NETCORE_API TAutoConsoleVariable<int32> CVarNetEnableDetailedScopeCounters;
@@ -322,7 +328,11 @@ FObjectReplicator::~FObjectReplicator()
 
 bool FObjectReplicator::SendCustomDeltaProperty(UObject* InObject, FProperty* Property, uint32 ArrayIndex, FNetBitWriter& OutBunch, TSharedPtr<INetDeltaBaseState>& NewFullState, TSharedPtr<INetDeltaBaseState>& OldState)
 {
-	return SendCustomDeltaProperty(InObject, Property->RepIndex + ArrayIndex, OutBunch, NewFullState, OldState);
+	check(RepLayout);
+
+	const uint16 CustomDeltaIndex = RepLayout->GetCustomDeltaIndexFromPropertyRepIndex(Property->RepIndex + ArrayIndex);
+
+	return SendCustomDeltaProperty(InObject, CustomDeltaIndex, OutBunch, NewFullState, OldState);
 }
 
 bool FObjectReplicator::SendCustomDeltaProperty(UObject* InObject, uint16 CustomDeltaIndex, FNetBitWriter& OutBunch, TSharedPtr<INetDeltaBaseState>& NewFullState, TSharedPtr<INetDeltaBaseState>& OldState)
@@ -468,20 +478,20 @@ bool FObjectReplicator::ValidateAgainstState( const UObject* ObjectState )
 	return true;
 }
 
-void FObjectReplicator::InitWithObject( UObject* InObject, UNetConnection * InConnection, bool bUseDefaultState )
+void FObjectReplicator::InitWithObject(UObject* InObject, UNetConnection* InConnection, bool bUseDefaultState)
 {
-	check( GetObject() == nullptr );
-	check( ObjectClass == nullptr );
-	check( bLastUpdateEmpty == false );
-	check( Connection == nullptr );
-	check( OwningChannel == nullptr );
-	check( !RepState.IsValid() );
-	check( RemoteFunctions == nullptr );
-	check( !RepLayout.IsValid() );
+	check(GetObject() == nullptr);
+	check(ObjectClass == nullptr);
+	check(bLastUpdateEmpty == false);
+	check(Connection == nullptr);
+	check(OwningChannel == nullptr);
+	check(!RepState.IsValid());
+	check(RemoteFunctions == nullptr);
+	check(!RepLayout.IsValid());
 
-	SetObject( InObject );
+	SetObject(InObject);
 
-	if ( GetObject() == nullptr )
+	if (GetObject() == nullptr)
 	{
 		// This may seem weird that we're checking for nullptr, but the SetObject above will wrap this object with TWeakObjectPtr
 		// If the object is pending kill, it will switch to nullptr, we're just making sure we handle this invalid edge case
@@ -489,16 +499,16 @@ void FObjectReplicator::InitWithObject( UObject* InObject, UNetConnection * InCo
 		return;
 	}
 
-	ObjectClass					= InObject->GetClass();
-	Connection					= InConnection;
-	RemoteFunctions				= nullptr;
-	bHasReplicatedProperties	= false;
-	bOpenAckCalled				= false;
-	RepState					= nullptr;
-	OwningChannel				= nullptr;		// Initially nullptr until StartReplicating is called
-	TrackedGuidMemoryBytes		= 0;
+	ObjectClass = InObject->GetClass();
+	Connection = InConnection;
+	RemoteFunctions = nullptr;
+	bHasReplicatedProperties = false;
+	bOpenAckCalled = false;
+	RepState = nullptr;
+	OwningChannel = nullptr;		// Initially nullptr until StartReplicating is called
+	TrackedGuidMemoryBytes = 0;
 
-	RepLayout = Connection->Driver->GetObjectClassRepLayout( ObjectClass );
+	RepLayout = Connection->Driver->GetObjectClassRepLayout(ObjectClass);
 
 	// Make a copy of the net properties
 	uint8* Source = bUseDefaultState ? (uint8*)GetObject()->GetArchetype() : (uint8*)InObject;
@@ -517,15 +527,24 @@ void FObjectReplicator::InitWithObject( UObject* InObject, UNetConnection * InCo
 		}
 	}
 
-	InitRecentProperties( Source );
+	InitRecentProperties(Source);
 
 	Connection->Driver->AllOwnedReplicators.Add(this);
 
-	bCanUseNonDirtyOptimization =
-		GbPushModelSkipUndirtiedReplicators &&
-		(RepLayout->IsEmpty() || EnumHasAnyFlags(RepLayout->GetFlags(), ERepLayoutFlags::FullPushSupport)) &&
+	if (GbPushModelSkipUndirtiedFastArrays)
+	{
+		bCanUseNonDirtyOptimization =
+			GbPushModelSkipUndirtiedReplicators &&
+			(RepLayout->IsEmpty() || EnumHasAnyFlags(RepLayout->GetFlags(), ERepLayoutFlags::FullPushSupport));
+	}
+	else
+	{
+		bCanUseNonDirtyOptimization =
+			GbPushModelSkipUndirtiedReplicators &&
+			(RepLayout->IsEmpty() || EnumHasAnyFlags(RepLayout->GetFlags(), ERepLayoutFlags::FullPushProperties)) &&
 			RepState->GetSendingRepState() &&
 			RepState->GetSendingRepState()->RecentCustomDeltaState.Num() == 0;
+	}
 }
 
 void FObjectReplicator::CleanUp()
@@ -2018,10 +2037,6 @@ void FObjectReplicator::QueueRemoteFunctionBunch( UFunction* Func, FOutBunch &Bu
 	
 	RemoteFuncInfo[InfoIdx].LastCallTimestamp = OwningChannel->Connection->Driver->GetElapsedTime();
 
-	PRAGMA_DISABLE_DEPRECATION_WARNINGS
-	RemoteFuncInfo[InfoIdx].LastCallTime = RemoteFuncInfo[InfoIdx].LastCallTimestamp;
-	PRAGMA_ENABLE_DEPRECATION_WARNINGS
-
 	if (RemoteFunctions == nullptr)
 	{
 		RemoteFunctions = new FOutBunch(OwningChannel, 0);
diff --git a/Engine/Source/Runtime/Engine/Private/NetDriver.cpp b/Engine/Source/Runtime/Engine/Private/NetDriver.cpp
--- a/Engine/Source/Runtime/Engine/Private/NetDriver.cpp
+++ b/Engine/Source/Runtime/Engine/Private/NetDriver.cpp
@@ -3136,7 +3136,7 @@ bool UNetDriver::HandleDumpRepLayoutFlagsCommand(const TCHAR* Cmd, FOutputDevice
 
 				for (int32 i = 0; i < ParentCount; ++i)
 				{
-					if ((RepLayout->GetParentCondition(i) != COND_Never) && !RepLayout->IsPushModelProperty(i) && !RepLayout->IsCustomDeltaProperty(i) && (RepLayout->GetParentArrayIndex(i) == 0))
+					if ((RepLayout->GetParentCondition(i) != COND_Never) && !RepLayout->IsPushModelProperty(i) && (RepLayout->GetParentArrayIndex(i) == 0))
 					{
 						if (const FProperty* Property = RepLayout->GetParentProperty(i))
 						{
@@ -4994,7 +4994,7 @@ struct FScopedNetDriverStats
 			CSV_CUSTOM_STAT(Replication, OutKBytes, ((float)FrameOutBytes) / 1024.f, ECsvCustomStatOp::Set );
 			CSV_CUSTOM_STAT(Replication, OutNetGUIDKBytesSec, ((float)NetDriver->NetGUIDOutBytes / 1024.f), ECsvCustomStatOp::Set );
 			CSV_CUSTOM_STAT(Replication, NumClientUpdateLevelVisibility, ((float)GNumClientUpdateLevelVisibility), ECsvCustomStatOp::Set );
-				
+			CSV_CUSTOM_STAT(Replication, NumSkippedObjectEmptyUpdates, (float)GNumSkippedObjectEmptyUpdates, ECsvCustomStatOp::Set);
 
 			SET_DWORD_STAT(STAT_SharedSerializationPropertyHit, GNumSharedSerializationHit);
 			SET_DWORD_STAT(STAT_SharedSerializationPropertyMiss, GNumSharedSerializationMiss);
diff --git a/Engine/Source/Runtime/Engine/Private/RepLayout.cpp b/Engine/Source/Runtime/Engine/Private/RepLayout.cpp
--- a/Engine/Source/Runtime/Engine/Private/RepLayout.cpp
+++ b/Engine/Source/Runtime/Engine/Private/RepLayout.cpp
@@ -1502,8 +1502,8 @@ static void CompareParentProperties(
 		}
 #endif // WITH_PUSH_VALIDATION_SUPPORT
 
-		// If we have full push model support, then we only need to check properties that are actually dirty.
-		else if (EnumHasAnyFlags(SharedParams.Flags, ERepLayoutFlags::FullPushSupport) && !bRecentlyCollectedGarbage)
+		// If we have full push model property support, then we only need to check properties that are actually dirty.
+		else if (EnumHasAnyFlags(SharedParams.Flags, ERepLayoutFlags::FullPushProperties) && !bRecentlyCollectedGarbage)
 		{
 			for (TConstSetBitIterator<> It = SharedParams.PushModelState->GetDirtyProperties(); It; ++It)
 			{
@@ -4371,6 +4371,13 @@ bool FRepLayout::SendCustomDeltaProperty(FNetDeltaSerializeInfo& Params, uint16
 		Params.Writer->SerializeIntPacked(StaticArrayIndex);
 	}
 
+#if WITH_PUSH_MODEL
+	if (EnumHasAnyFlags(Parent.Flags, ERepParentFlags::IsFastArray) && IsPushModelProperty(CustomDeltaProperty.PropertyRepIndex))
+	{
+		static_cast<FFastArraySerializer*>(Params.Data)->CachePushModelState(Params.Object, CustomDeltaProperty.PropertyRepIndex);
+	}
+#endif
+
 	TGuardValue<bool> SupportFastArrayDeltaGuard(Params.bSupportsFastArrayDeltaStructSerialization, bSupportsFastArrayDelta);
 	return CppStructOps->NetDeltaSerialize(Params, Params.Data);
 }
@@ -6016,7 +6023,9 @@ void FRepLayout::InitFromClass(
 	// Tracks the number of (non-delta) lifetime properties so we can check that against our
 	// Push Model Enabled properties.
 	int32 NumberOfLifetimeProperties = 0;
-	int32 NumberOfPushModelProperties = 0;
+	int32 NumberOfLifetimePushModelProperties = 0;
+	int32 NumberOfFastArrayProperties = 0;
+	int32 NumberOfFastArrayPushModelProperties = 0;
 
 	// Setup lifetime replicated properties
 	for (int32 i = 0; i < LifetimeProps.Num(); i++)
@@ -6062,7 +6071,7 @@ void FRepLayout::InitFromClass(
 #if WITH_PUSH_MODEL
 			if (bIsPushModelEnabled && LifetimeProps[i].bIsPushBased)
 			{
-				++NumberOfPushModelProperties;
+				++NumberOfLifetimePushModelProperties;
 				PushModelProperties[ParentIndex] = true;
 			}
 #endif
@@ -6131,6 +6140,14 @@ void FRepLayout::InitFromClass(
 									/*FastArrayItemReplicationIdOffset=*/MaybeFastArrayItem->FindPropertyByName(FastArrayItemReplicationIDName)->GetOffset_ForGC()
 								));
 
+								++NumberOfFastArrayProperties;
+#if WITH_PUSH_MODEL
+								if (bIsPushModelEnabled && LifetimeProps[i].bIsPushBased)
+								{
+									++NumberOfFastArrayPushModelProperties;
+									PushModelProperties[ParentIndex] = true;
+								}
+#endif
 								bAddedFastArray = true;
 								break;
 							}
@@ -6141,7 +6158,7 @@ void FRepLayout::InitFromClass(
 				}
 
 				if (!bAddedFastArray)
-		{
+				{
 					UE_LOG(LogRep, Warning, TEXT("FRepLayout::InitFromClass: Unable to find Fast Array Item array in Fast Array Serializer: %s"), *Parents[ParentIndex].CachedPropertyName.ToString());
 				}
 			}
@@ -6171,9 +6188,16 @@ void FRepLayout::InitFromClass(
 	}	
 
 #if WITH_PUSH_MODEL
-	if (bIsPushModelEnabled && NumberOfPushModelProperties > 0)
+	if (bIsPushModelEnabled && ((NumberOfLifetimePushModelProperties > 0) || (NumberOfFastArrayPushModelProperties > 0)))
 	{
-		Flags |= (NumberOfLifetimeProperties == NumberOfPushModelProperties) ?
+		const bool bFullPushProperties = (NumberOfLifetimeProperties == NumberOfLifetimePushModelProperties);
+
+		if (bFullPushProperties)
+		{
+			Flags |= ERepLayoutFlags::FullPushProperties;
+		}
+
+		Flags |= (bFullPushProperties && (NumberOfFastArrayProperties == NumberOfFastArrayPushModelProperties)) ?
 			ERepLayoutFlags::FullPushSupport :
 			ERepLayoutFlags::PartialPushSupport;
 	}
@@ -8035,6 +8059,11 @@ FProperty* FRepLayout::GetLifetimeCustomDeltaProperty(const uint16 CustomDeltaPr
 	return Parents[CustomDeltaProperty.PropertyRepIndex].Property;
 }
 
+const uint16 FRepLayout::GetCustomDeltaIndexFromPropertyRepIndex(const uint16 PropertyRepIndex) const
+{
+	return LifetimeCustomPropertyState->GetCustomDeltaIndexFromPropertyRepIndex(PropertyRepIndex);
+}
+
 const ELifetimeCondition FRepLayout::GetLifetimeCustomDeltaPropertyCondition(const uint16 CustomDeltaPropertyIndex) const
 {
 	const FLifetimeCustomDeltaProperty& CustomDeltaProperty = LifetimeCustomPropertyState->GetCustomDeltaProperty(CustomDeltaPropertyIndex);
@@ -8169,6 +8198,8 @@ const TCHAR* LexToString(ERepLayoutFlags Flag)
 		return TEXT("HasObjectOrNetSerializeProperties");
 	case ERepLayoutFlags::NoReplicatedProperties:
 		return TEXT("NoReplicatedProperties");
+	case ERepLayoutFlags::FullPushProperties:
+		return TEXT("FullPushProperties");
 	default:
 		check(false);
 		return TEXT("Unknown");
@@ -8192,4 +8223,4 @@ template void FRepLayout::ValidateWithChecksum(TConstRepDataBuffer<DataType> Dat
 REPDATATYPE_SPECIALIZATION(ERepDataBufferType::ShadowBuffer);
 REPDATATYPE_SPECIALIZATION(ERepDataBufferType::ObjectBuffer);
 
-#undef REPDATATYPE_SPECIALIZATION
\ No newline at end of file
+#undef REPDATATYPE_SPECIALIZATION
diff --git a/Engine/Source/Runtime/Engine/Public/Net/DataReplication.h b/Engine/Source/Runtime/Engine/Public/Net/DataReplication.h
--- a/Engine/Source/Runtime/Engine/Public/Net/DataReplication.h
+++ b/Engine/Source/Runtime/Engine/Public/Net/DataReplication.h
@@ -78,10 +78,6 @@ public:
 	{
 		FName FuncName;
 		int32 Calls;
-
-		UE_DEPRECATED(4.25, "Please use LastCallTimestamp instead")
-		float LastCallTime;
-
 		double LastCallTimestamp;
 	};
 
@@ -136,6 +132,7 @@ public:
 	//~ Both of these should be private, IMO, but we'll leave them public for now for back compat
 	//~ in case anyone was using SerializeCustomDeltaProperty already.
 
+	UE_DEPRECATED(5.0, "No longer used.")
 	bool SendCustomDeltaProperty(
 		UObject* InObject,
 		FProperty* Property,
diff --git a/Engine/Source/Runtime/Engine/Public/Net/RepLayout.h b/Engine/Source/Runtime/Engine/Public/Net/RepLayout.h
--- a/Engine/Source/Runtime/Engine/Public/Net/RepLayout.h
+++ b/Engine/Source/Runtime/Engine/Public/Net/RepLayout.h
@@ -1100,9 +1100,10 @@ enum class ERepLayoutFlags : uint8
 	None								= 0,
 	IsActor 							= (1 << 0),	//! This RepLayout is for AActor or a subclass of AActor.
 	PartialPushSupport					= (1 << 1),	//! This RepLayout has some properties that use Push Model and some that don't.
-	FullPushSupport						= (1 << 2),	//! All properties in this RepLayout use Push Model.
+	FullPushSupport						= (1 << 2),	//! All properties and fast arrays in this RepLayout use Push Model.
 	HasObjectOrNetSerializeProperties	= (1 << 3),	//! Will be set for any RepLayout that contains Object or Net Serialize property commands.
 	NoReplicatedProperties				= (1 << 4), //! Will be set if the RepLayout has no lifetime properties, or they are all disabled.
+	FullPushProperties					= (1 << 5), //! All properties in this RepLayout use Push Model.
 };
 ENUM_CLASS_FLAGS(ERepLayoutFlags);
 
@@ -1568,6 +1569,8 @@ public:
 	}
 #endif
 
+	const uint16 GetCustomDeltaIndexFromPropertyRepIndex(const uint16 PropertyRepIndex) const;
+
 	void CountBytes(FArchive& Ar) const;
 
 private:
diff --git a/Engine/Source/Runtime/Net/Core/Classes/Net/Serialization/FastArraySerializer.h b/Engine/Source/Runtime/Net/Core/Classes/Net/Serialization/FastArraySerializer.h
--- a/Engine/Source/Runtime/Net/Core/Classes/Net/Serialization/FastArraySerializer.h
+++ b/Engine/Source/Runtime/Net/Core/Classes/Net/Serialization/FastArraySerializer.h
@@ -16,6 +16,7 @@
 #include "UObject/CoreNet.h"
 #include "Containers/ArrayView.h"
 #include "Net/Core/Misc/GuidReferences.h"
+#include "Net/Core/PushModel/PushModel.h"
 #include "HAL/IConsoleManager.h"
 #include "Templates/EnableIf.h"
 #include "FastArraySerializer.generated.h"
@@ -359,6 +360,10 @@ struct FFastArraySerializer
 	FFastArraySerializer()
 		: IDCounter(0)
 		, ArrayReplicationKey(0)
+#if WITH_PUSH_MODEL
+		, OwningObject(nullptr)
+		, RepIndex(INDEX_NONE)
+#endif // WITH_PUSH_MODEL
 		, CachedNumItems(INDEX_NONE)
 		, CachedNumItemsToConsiderForWriting(INDEX_NONE)
 		, DeltaFlags(EFastArraySerializerDeltaFlags::None)
@@ -383,6 +388,14 @@ struct FFastArraySerializer
 
 	/** List of items that need to be re-serialized when the referenced objects are mapped.*/
 	TMap<int32, FGuidReferencesMap> GuidReferencesMap_StructDelta;
+
+#if WITH_PUSH_MODEL
+	// Object that is replicating this fast array
+	UObject* OwningObject;
+
+	// Property index of this array in the owning object's replication layout
+	int32 RepIndex;
+#endif // WITH_PUSH_MODEL
 	
 	/** This must be called if you add or change an item in the array */
 	void MarkItemDirty(FFastArraySerializerItem & Item)
@@ -418,6 +431,13 @@ struct FFastArraySerializer
 		{
 			ArrayReplicationKey++;
 		}
+
+#if WITH_PUSH_MODEL
+		if (OwningObject != nullptr && RepIndex != INDEX_NONE)
+		{
+			MARK_PROPERTY_DIRTY_UNSAFE(OwningObject, RepIndex);
+		}
+#endif // WITH_PUSH_MODEL
 	}
 
 	/**
@@ -517,6 +537,17 @@ struct FFastArraySerializer
 		return MaxNumberOfAllowedDeletionsPerUpdate;
 	}
 
+#if WITH_PUSH_MODEL
+	void CachePushModelState(UObject* Object, const uint16 PropertyRepIndex)
+	{
+		if (OwningObject == nullptr && RepIndex == INDEX_NONE)
+		{
+			OwningObject = Object;
+			RepIndex = PropertyRepIndex;
+		}
+	}
+#endif // WITH_PUSH_MODEL
+
 private:
 
 	NETCORE_API static int32 MaxNumberOfAllowedChangesPerUpdate;
