Changed how levels listen to FCoreUObjectDelegates::OnAssetLoaded : Fixes crash calling IsTemplate in ~ULevel (and registers only one delegate).

diff --git a/Engine/Source/Runtime/Engine/Classes/Engine/Level.h b/Engine/Source/Runtime/Engine/Classes/Engine/Level.h
--- a/Engine/Source/Runtime/Engine/Classes/Engine/Level.h
+++ b/Engine/Source/Runtime/Engine/Classes/Engine/Level.h
@@ -1314,8 +1314,6 @@ private:
 	/** Sets the level to use or not the actor folder objects feature. */
 	void SetUseActorFoldersInternal(bool bInEnabled);
 
-	void OnAssetLoaded(UObject* Asset);
-
 	friend struct FLevelActorFoldersHelper;
 
 	/** Replace the existing LSA (if set) by spawning a new one based on this level's script blueprint */
diff --git a/Engine/Source/Runtime/Engine/Private/Level.cpp b/Engine/Source/Runtime/Engine/Private/Level.cpp
--- a/Engine/Source/Runtime/Engine/Private/Level.cpp
+++ b/Engine/Source/Runtime/Engine/Private/Level.cpp
@@ -353,6 +353,30 @@ void FLevelSimplificationDetails::PostLoadDeprecated()
 
 TMap<FName, TWeakObjectPtr<UWorld> > ULevel::StreamedLevelsOwningWorld;
 
+#if WITH_EDITOR
+struct FOnWorldLoadedHandler
+{
+	static void Init()
+	{
+		if (!bIsInitialized)
+		{
+			bIsInitialized = true;
+			FCoreUObjectDelegates::OnAssetLoaded.AddLambda([](UObject* Asset)
+			{
+				UWorld* World = Cast<UWorld>(Asset);
+				if (World && World->PersistentLevel)
+				{
+					World->PersistentLevel->OnLevelLoaded();
+				}
+			});
+		}
+	}
+private:
+	static bool bIsInitialized;
+};
+bool FOnWorldLoadedHandler::bIsInitialized = false;
+#endif
+
 ULevel::ULevel( const FObjectInitializer& ObjectInitializer )
 	:	UObject( ObjectInitializer )
 	,	Actors()
@@ -381,10 +405,7 @@ ULevel::ULevel( const FObjectInitializer& ObjectInitializer )
 	bStaticComponentsRegisteredInStreamingManager = false;
 	IncrementalComponentState = EIncrementalComponentState::Init;
 #if WITH_EDITOR
-	if (!IsTemplate())
-	{
-		FCoreUObjectDelegates::OnAssetLoaded.AddUObject(this, &ULevel::OnAssetLoaded);
-	}
+	FOnWorldLoadedHandler::Init();
 #endif
 }
 
@@ -395,17 +416,10 @@ void ULevel::Initialize(const FURL& InURL)
 
 ULevel::~ULevel()
 {
-#if WITH_EDITOR
-	if (!IsTemplate())
-	{
-		FCoreUObjectDelegates::OnAssetLoaded.RemoveAll(this);
-	}
-#endif
 	FTickTaskManagerInterface::Get().FreeTickTaskLevel(TickTaskLevel);
 	TickTaskLevel = NULL;
 }
 
-
 void ULevel::AddReferencedObjects(UObject* InThis, FReferenceCollector& Collector)
 {
 	ULevel* This = CastChecked<ULevel>(InThis);
@@ -2254,16 +2268,6 @@ void ULevel::CommitModelSurfaces()
 }
 
 #if WITH_EDITOR
-
-void ULevel::OnAssetLoaded(UObject* Asset)
-{
-	UWorld* World = Cast<UWorld>(Asset);
-	if (World && World->PersistentLevel == this)
-	{
-		FixupActorFolders();
-	}
-}
-
 void ULevel::FixupActorFolders()
 {
 	if (!IsUsingActorFolders())
