Fix RBAN debug draw jittering due to RBAN being deferred one frame. The simulation can cross the frame/render boundary, so it was missing some frames and potentially doubling up on the next one. Added a new Chaos-specific method to control when we render the debug visualization. Updated the ConsoleCommandUI to set the new values as well.

diff --git a/Engine/Source/Runtime/AnimGraphRuntime/Private/BoneControllers/AnimNode_RigidBody.cpp b/Engine/Source/Runtime/AnimGraphRuntime/Private/BoneControllers/AnimNode_RigidBody.cpp
--- a/Engine/Source/Runtime/AnimGraphRuntime/Private/BoneControllers/AnimNode_RigidBody.cpp
+++ b/Engine/Source/Runtime/AnimGraphRuntime/Private/BoneControllers/AnimNode_RigidBody.cpp
@@ -72,7 +72,9 @@ FAutoConsoleVariableRef CVarRigidBodyNodeSpaceExternalLinearVelocityZ(TEXT("p.Ri
 #endif
 
 bool bRBAN_UseDeferredTask = false;
+bool bRBAN_DebugDraw = false;
 FAutoConsoleVariableRef CVarRigidBodyNodeDeferredSimulation(TEXT("p.RigidBodyNode.UseDeferredTask"), bRBAN_UseDeferredTask, TEXT("Whether to defer the simulation results by one frame so that they can run in a task"), ECVF_Default);
+FAutoConsoleVariableRef CVarRigidBodyNodeDebugDraw(TEXT("p.RigidBodyNode.DebugDraw"), bRBAN_DebugDraw, TEXT("Whether to debug draw the rigid body simulation state. Requires p.Chaos.DebugDraw.Enabled 1 to function as well."), ECVF_Default);
 
 // Array of priorities that can be indexed into with CVars, since task priorities cannot be set from scalability .ini
 static FAutoConsoleTaskPriority GRigidBodyNodeTaskPriorities[] =
@@ -905,6 +907,14 @@ void FAnimNode_RigidBody::EvaluateSkeletalControl_AnyThread(FComponentSpacePoseC
 			{
 				RunPhysicsSimulation(DeltaSeconds, SimSpaceGravity);
 			}
+
+#if WITH_CHAOS
+			// Draw here even if the simulation is deferred since we want the shapes drawn relative to the current transform
+			if (bRBAN_DebugDraw)
+			{
+				PhysicsSimulation->DebugDraw();
+			}
+#endif
 		}
 		
 		//write back to animation system
diff --git a/Engine/Source/Runtime/Engine/Private/PhysicsEngine/ImmediatePhysics/ImmediatePhysicsChaos/ImmediatePhysicsSimulation_Chaos.cpp b/Engine/Source/Runtime/Engine/Private/PhysicsEngine/ImmediatePhysics/ImmediatePhysicsChaos/ImmediatePhysicsSimulation_Chaos.cpp
--- a/Engine/Source/Runtime/Engine/Private/PhysicsEngine/ImmediatePhysics/ImmediatePhysicsChaos/ImmediatePhysicsSimulation_Chaos.cpp
+++ b/Engine/Source/Runtime/Engine/Private/PhysicsEngine/ImmediatePhysics/ImmediatePhysicsChaos/ImmediatePhysicsSimulation_Chaos.cpp
@@ -143,24 +143,27 @@ FAutoConsoleVariableRef CVarChaosImmPhysSolverType(TEXT("p.Chaos.ImmPhys.SolverT
 
 // DebugDraw CVars
 #if CHAOS_DEBUG_DRAW
-int32 ChaosImmediate_DebugDrawParticles = 0;
-int32 ChaosImmediate_DebugDrawShapes = 0;
-int32 ChaosImmediate_DebugDrawShowStatics = 1;
-int32 ChaosImmediate_DebugDrawShowKinematics = 1;
-int32 ChaosImmediate_DebugDrawShowDynamics = 1;
-int32 ChaosImmediate_DebugDrawBounds = 0;
-int32 ChaosImmediate_DebugDrawCollisions = 0;
-int32 ChaosImmediate_DebugDrawJoints = 0;
+bool bChaosImmediate_DebugDrawOnSimulate = false;
+bool bChaosImmediate_DebugDrawParticles = false;
+bool bChaosImmediate_DebugDrawShapes = false;
+bool bChaosImmediate_DebugDrawShowStatics = true;
+bool bChaosImmediate_DebugDrawShowKinematics = true;
+bool bChaosImmediate_DebugDrawShowDynamics = true;
+bool bChaosImmediate_DebugDrawBounds = false;
+bool bChaosImmediate_DebugDrawCollisions = false;
+bool bChaosImmediate_DebugDrawJoints = false;
+bool bChaosImmediate_DebugDrawSimulationSpace = 0;
 Chaos::DebugDraw::FChaosDebugDrawJointFeatures ChaosImmediate_DebugDrawJointFeatures = Chaos::DebugDraw::FChaosDebugDrawJointFeatures::MakeDefault();
-int32 ChaosImmediate_DebugDrawSimulationSpace = 0;
-FAutoConsoleVariableRef CVarChaosImmPhysDebugDrawParticles(TEXT("p.Chaos.ImmPhys.DebugDrawParticles"), ChaosImmediate_DebugDrawParticles, TEXT("Draw Particle Transforms (0 = never; 1 = end of frame; 2 = begin and end of frame; 3 = post-integate, post-apply and post-applypushout;)."));
-FAutoConsoleVariableRef CVarChaosImmPhysDebugDrawShapes(TEXT("p.Chaos.ImmPhys.DebugDrawShapes"), ChaosImmediate_DebugDrawShapes, TEXT("Draw Shapes (0 = never; 1 = end of frame; 2 = begin and end of frame; 3 = post-integate, post-apply and post-applypushout;"));
-FAutoConsoleVariableRef CVarChaosImmPhysDebugDrawShowStatics(TEXT("p.Chaos.ImmPhys.DebugDrawShowStatics"), ChaosImmediate_DebugDrawShowStatics, TEXT("Show statics if shape debug draw is enabled"));
-FAutoConsoleVariableRef CVarChaosImmPhysDebugDrawShowKinematics(TEXT("p.Chaos.ImmPhys.DebugDrawShowKinematics"), ChaosImmediate_DebugDrawShowKinematics, TEXT("Show kinematics if shape debug draw is enabled"));
-FAutoConsoleVariableRef CVarChaosImmPhysDebugDrawShowDynamics(TEXT("p.Chaos.ImmPhys.DebugDrawShowDynamics"), ChaosImmediate_DebugDrawShowDynamics, TEXT("Show dynamics if shape debug draw is enabled"));
-FAutoConsoleVariableRef CVarChaosImmPhysDebugDrawBounds(TEXT("p.Chaos.ImmPhys.DebugDrawBounds"), ChaosImmediate_DebugDrawBounds, TEXT("Draw Particle Bounds (0 = never; 1 = end of frame; 2 = begin and end of frame; 3 = post-integate, post-apply and post-applypushout;)."));
-FAutoConsoleVariableRef CVarChaosImmPhysDebugDrawCollisions(TEXT("p.Chaos.ImmPhys.DebugDrawCollisions"), ChaosImmediate_DebugDrawCollisions, TEXT("Draw Collisions (0 = never; 1 = end of frame; 2 = begin and end of frame; 3 = post-integate, post-apply and post-applypushout;)"));
-FAutoConsoleVariableRef CVarChaosImmPhysDebugDrawJoints(TEXT("p.Chaos.ImmPhys.DebugDrawJoints"), ChaosImmediate_DebugDrawJoints, TEXT("Draw Joints. (0 = never; 1 = end of frame; 2 = begin and end of frame; 3 = post-integate, post-apply and post-applypushout; 4 = each Apply step)."));
+FAutoConsoleVariableRef CVarChaosImmPhysDebugDrawEnabled(TEXT("p.Chaos.ImmPhys.DebugDrawOnSimulate"), bChaosImmediate_DebugDrawOnSimulate, TEXT("Enables debug drawing after the simulation completes."), ECVF_Default);
+FAutoConsoleVariableRef CVarChaosImmPhysDebugDrawParticles(TEXT("p.Chaos.ImmPhys.DebugDrawParticles"), bChaosImmediate_DebugDrawParticles, TEXT("Whether to draw particles when debug drawing."), ECVF_Default);
+FAutoConsoleVariableRef CVarChaosImmPhysDebugDrawShapes(TEXT("p.Chaos.ImmPhys.DebugDrawShapes"), bChaosImmediate_DebugDrawShapes, TEXT("Whether to draw shapes when debug drawing."), ECVF_Default);
+FAutoConsoleVariableRef CVarChaosImmPhysDebugDrawBounds(TEXT("p.Chaos.ImmPhys.DebugDrawBounds"), bChaosImmediate_DebugDrawBounds, TEXT("Whether to draw bounds when debug drawing."), ECVF_Default);
+FAutoConsoleVariableRef CVarChaosImmPhysDebugDrawCollisions(TEXT("p.Chaos.ImmPhys.DebugDrawCollisions"), bChaosImmediate_DebugDrawCollisions, TEXT("Whether to draw collisions when debug drawing."), ECVF_Default);
+FAutoConsoleVariableRef CVarChaosImmPhysDebugDrawJoints(TEXT("p.Chaos.ImmPhys.DebugDrawJoints"), bChaosImmediate_DebugDrawJoints, TEXT("Whether to draw joints when debug drawing."), ECVF_Default);
+FAutoConsoleVariableRef CVarChaosImmPhysDebugDrawSimulationSpace(TEXT("p.Chaos.ImmPhys.DebugDrawSimulationSpace"), bChaosImmediate_DebugDrawSimulationSpace, TEXT("Whether to draw the simulation frame of reference, acceleration and velocity when debug drawing."), ECVF_Default);
+FAutoConsoleVariableRef CVarChaosImmPhysDebugDrawShowStatics(TEXT("p.Chaos.ImmPhys.DebugDrawShowStatics"), bChaosImmediate_DebugDrawShowStatics, TEXT("Show statics if shape debug draw is enabled"), ECVF_Default);
+FAutoConsoleVariableRef CVarChaosImmPhysDebugDrawShowKinematics(TEXT("p.Chaos.ImmPhys.DebugDrawShowKinematics"), bChaosImmediate_DebugDrawShowKinematics, TEXT("Show kinematics if shape debug draw is enabled"), ECVF_Default);
+FAutoConsoleVariableRef CVarChaosImmPhysDebugDrawShowDynamics(TEXT("p.Chaos.ImmPhys.DebugDrawShowDynamics"), bChaosImmediate_DebugDrawShowDynamics, TEXT("Show dynamics if shape debug draw is enabled"), ECVF_Default);
 FAutoConsoleVariableRef CVarChaosImmPhysDebugDrawJointFeaturesCoMConnector(TEXT("p.Chaos.ImmPhys.DebugDraw.JointFeatures.CoMConnector"), ChaosImmediate_DebugDrawJointFeatures.bCoMConnector, TEXT("Joint features mask (see FDebugDrawJointFeatures)."));
 FAutoConsoleVariableRef CVarChaosImmPhysDebugDrawJointFeaturesActorConnector(TEXT("p.Chaos.ImmPhys.DebugDraw.JointFeatures.ActorConnector"), ChaosImmediate_DebugDrawJointFeatures.bActorConnector, TEXT("Joint features mask (see FDebugDrawJointFeatures)."));
 FAutoConsoleVariableRef CVarChaosImmPhysDebugDrawJointFeaturesStretch(TEXT("p.Chaos.ImmPhys.DebugDraw.JointFeatures.Stretch"), ChaosImmediate_DebugDrawJointFeatures.bStretch, TEXT("Joint features mask (see FDebugDrawJointFeatures)."));
@@ -169,7 +172,6 @@ FAutoConsoleVariableRef CVarChaosImmPhysDebugDrawJointFeaturesLevel(TEXT("p.Chao
 FAutoConsoleVariableRef CVarChaosImmPhysDebugDrawJointFeaturesIndex(TEXT("p.Chaos.ImmPhys.DebugDraw.JointFeatures.Index"), ChaosImmediate_DebugDrawJointFeatures.bIndex, TEXT("Joint features mask (see FDebugDrawJointFeatures)."));
 FAutoConsoleVariableRef CVarChaosImmPhysDebugDrawJointFeaturesColor(TEXT("p.Chaos.ImmPhys.DebugDraw.JointFeatures.Color"), ChaosImmediate_DebugDrawJointFeatures.bColor, TEXT("Joint features mask (see FDebugDrawJointFeatures)."));
 FAutoConsoleVariableRef CVarChaosImmPhysDebugDrawJointFeaturesIsland(TEXT("p.Chaos.ImmPhys.DebugDraw.JointFeatures.Island"), ChaosImmediate_DebugDrawJointFeatures.bIsland, TEXT("Joint features mask (see FDebugDrawJointFeatures)."));
-FAutoConsoleVariableRef CVarChaosImmPhysDebugDrawSimulationSpace(TEXT("p.Chaos.ImmPhys.DebugDrawSimulationSpace"), ChaosImmediate_DebugDrawSimulationSpace, TEXT("Draw the simulation frame of reference, acceleration and velocity."));
 
 Chaos::DebugDraw::FChaosDebugDrawSettings ChaosImmPhysDebugDebugDrawSettings(
 	/* ArrowSize =					*/ 1.5f,
@@ -347,28 +349,6 @@ namespace ImmediatePhysics_Chaos
 		Implementation->NarrowPhase.GetContext().bFilteringEnabled = false;
 		Implementation->NarrowPhase.GetContext().bDeferUpdate = true;
 		Implementation->NarrowPhase.GetContext().bAllowManifolds = false;
-
-
-#if CHAOS_DEBUG_DRAW
-		Implementation->Evolution.SetPostIntegrateCallback(
-			[this]()
-			{
-				DebugDrawDynamicParticles(3, 4, 0.25f);
-				DebugDrawConstraints(3, 4, 0.25f);
-			});
-		Implementation->Evolution.SetPostApplyCallback(
-			[this]()
-			{
-				DebugDrawDynamicParticles(3, 3, 0.5f);
-				DebugDrawConstraints(3, 3, 0.5f);
-			});
-		Implementation->Evolution.SetPostApplyPushOutCallback(
-			[this]()
-			{
-				DebugDrawDynamicParticles(3, 3, 1.0f);
-				DebugDrawConstraints(3, 3, 1.0f);
-			});
-#endif
 	}
 
 	FSimulation::~FSimulation()
@@ -758,6 +738,15 @@ namespace ImmediatePhysics_Chaos
 		}
 	}
 
+	void FSimulation::DebugDraw()
+	{
+		DebugDrawStaticParticles();
+		DebugDrawKinematicParticles();
+		DebugDrawDynamicParticles();
+		DebugDrawConstraints();
+		DebugDrawSimulationSpace();
+	}
+
 	FReal FSimulation::UpdateStepTime(const FReal DeltaTime, const FReal MaxStepTime)
 	{
 		Implementation->RollingAverageStepTime = Implementation->RollingAverageStepTime + (DeltaTime - Implementation->RollingAverageStepTime) / Implementation->NumRollingAverageStepTimes;
@@ -841,10 +830,6 @@ namespace ImmediatePhysics_Chaos
 				ChaosImmediate_Collision_PushOutPairIterations);
 		}
 
-		DebugDrawKinematicParticles(2, 2, 0.75f);
-		DebugDrawDynamicParticles(2, 2, 0.75f);
-		DebugDrawConstraints(2, 2, 0.75f);
-
 		// Fixed timestep mode DT (Only used if > 0.0f)
 		FReal FixedStepTime = Implementation->FixedStepTime;
 
@@ -883,32 +868,29 @@ namespace ImmediatePhysics_Chaos
 		Implementation->Evolution.SetSimulationSpace(Implementation->SimulationSpace);
 		Implementation->Evolution.Advance(StepTime, NumSteps, RewindTime);
 
-		DebugDrawStaticParticles(1, 4, 1.0f);
-		DebugDrawKinematicParticles(1, 4, 1.0f);
-		DebugDrawDynamicParticles(1, 3, 1.0f);
-		DebugDrawConstraints(1, 2, 1.0f);
-		DebugDrawSimulationSpace();
+#if CHAOS_DEBUG_DRAW
+		if (bChaosImmediate_DebugDrawOnSimulate)
+		{
+			DebugDraw();
+		}
+#endif
 	}
 
-	void FSimulation::DebugDrawStaticParticles(const int32 MinDebugLevel, const int32 MaxDebugLevel, const float ColorScale)
+	void FSimulation::DebugDrawStaticParticles()
 	{
 #if CHAOS_DEBUG_DRAW
 		using namespace Chaos;
-		if (FDebugDrawQueue::IsDebugDrawingEnabled())
+		if (FDebugDrawQueue::IsDebugDrawingEnabled() && bChaosImmediate_DebugDrawShowStatics)
 		{
-			if (!ChaosImmediate_DebugDrawShowStatics)
-			{
-				return;
-			}
-			if ((ChaosImmediate_DebugDrawParticles >= MinDebugLevel) && (ChaosImmediate_DebugDrawParticles <= MaxDebugLevel))
+			if (bChaosImmediate_DebugDrawParticles)
 			{
 				DebugDraw::DrawParticleTransforms(Implementation->SimulationSpace.Transform, Implementation->Particles.GetActiveStaticParticlesView(), &ChaosImmPhysDebugDebugDrawSettings);
 			}
-			if ((ChaosImmediate_DebugDrawShapes >= MinDebugLevel) && (ChaosImmediate_DebugDrawShapes <= MaxDebugLevel))
+			if (bChaosImmediate_DebugDrawShapes)
 			{
-				DebugDraw::DrawParticleShapes(Implementation->SimulationSpace.Transform, Implementation->Particles.GetActiveStaticParticlesView(), ColorScale, &ChaosImmPhysDebugDebugDrawSettings);
+				DebugDraw::DrawParticleShapes(Implementation->SimulationSpace.Transform, Implementation->Particles.GetActiveStaticParticlesView(), 1.0f, &ChaosImmPhysDebugDebugDrawSettings);
 			}
-			if ((ChaosImmediate_DebugDrawBounds >= MinDebugLevel) && (ChaosImmediate_DebugDrawBounds <= MaxDebugLevel))
+			if (bChaosImmediate_DebugDrawBounds)
 			{
 				DebugDraw::DrawParticleBounds(Implementation->SimulationSpace.Transform, Implementation->Particles.GetActiveStaticParticlesView(), 0.0f, &ChaosImmPhysDebugDebugDrawSettings);
 			}
@@ -916,25 +898,21 @@ namespace ImmediatePhysics_Chaos
 #endif
 	}
 
-	void FSimulation::DebugDrawKinematicParticles(const int32 MinDebugLevel, const int32 MaxDebugLevel, const float ColorScale)
+	void FSimulation::DebugDrawKinematicParticles()
 	{
 #if CHAOS_DEBUG_DRAW
 		using namespace Chaos;
-		if (FDebugDrawQueue::IsDebugDrawingEnabled())
+		if (FDebugDrawQueue::IsDebugDrawingEnabled() && bChaosImmediate_DebugDrawShowKinematics)
 		{
-			if (!ChaosImmediate_DebugDrawShowKinematics)
-			{
-				return;
-			}
-			if ((ChaosImmediate_DebugDrawParticles >= MinDebugLevel) && (ChaosImmediate_DebugDrawParticles <= MaxDebugLevel))
+			if (bChaosImmediate_DebugDrawParticles)
 			{
 				DebugDraw::DrawParticleTransforms(Implementation->SimulationSpace.Transform, Implementation->Particles.GetActiveKinematicParticlesView(), &ChaosImmPhysDebugDebugDrawSettings);
 			}
-			if ((ChaosImmediate_DebugDrawShapes >= MinDebugLevel) && (ChaosImmediate_DebugDrawShapes <= MaxDebugLevel))
+			if (bChaosImmediate_DebugDrawShapes)
 			{
-				DebugDraw::DrawParticleShapes(Implementation->SimulationSpace.Transform, Implementation->Particles.GetActiveKinematicParticlesView(), ColorScale, &ChaosImmPhysDebugDebugDrawSettings);
+				DebugDraw::DrawParticleShapes(Implementation->SimulationSpace.Transform, Implementation->Particles.GetActiveKinematicParticlesView(), 1.0f, &ChaosImmPhysDebugDebugDrawSettings);
 			}
-			if ((ChaosImmediate_DebugDrawBounds >= MinDebugLevel) && (ChaosImmediate_DebugDrawBounds <= MaxDebugLevel))
+			if (bChaosImmediate_DebugDrawBounds)
 			{
 				DebugDraw::DrawParticleBounds(Implementation->SimulationSpace.Transform, Implementation->Particles.GetActiveKinematicParticlesView(), 0.0f, &ChaosImmPhysDebugDebugDrawSettings);
 			}
@@ -942,25 +920,21 @@ namespace ImmediatePhysics_Chaos
 #endif
 	}
 
-	void FSimulation::DebugDrawDynamicParticles(const int32 MinDebugLevel, const int32 MaxDebugLevel, const float ColorScale)
+	void FSimulation::DebugDrawDynamicParticles()
 	{
 #if CHAOS_DEBUG_DRAW
 		using namespace Chaos;
-		if (FDebugDrawQueue::IsDebugDrawingEnabled())
+		if (FDebugDrawQueue::IsDebugDrawingEnabled() && bChaosImmediate_DebugDrawShowDynamics)
 		{
-			if (!ChaosImmediate_DebugDrawShowDynamics)
-			{
-				return;
-			}
-			if ((ChaosImmediate_DebugDrawParticles >= MinDebugLevel) && (ChaosImmediate_DebugDrawParticles <= MaxDebugLevel))
+			if (bChaosImmediate_DebugDrawParticles)
 			{
 				DebugDraw::DrawParticleTransforms(Implementation->SimulationSpace.Transform, Implementation->Particles.GetActiveParticlesView(), &ChaosImmPhysDebugDebugDrawSettings);
 			}
-			if ((ChaosImmediate_DebugDrawShapes >= MinDebugLevel) && (ChaosImmediate_DebugDrawShapes <= MaxDebugLevel))
+			if (bChaosImmediate_DebugDrawShapes)
 			{
-				DebugDraw::DrawParticleShapes(Implementation->SimulationSpace.Transform, Implementation->Particles.GetActiveParticlesView(), ColorScale, &ChaosImmPhysDebugDebugDrawSettings);
+				DebugDraw::DrawParticleShapes(Implementation->SimulationSpace.Transform, Implementation->Particles.GetActiveParticlesView(), 1.0f, &ChaosImmPhysDebugDebugDrawSettings);
 			}
-			if ((ChaosImmediate_DebugDrawBounds >= MinDebugLevel) && (ChaosImmediate_DebugDrawBounds <= MaxDebugLevel))
+			if (bChaosImmediate_DebugDrawBounds)
 			{
 				DebugDraw::DrawParticleBounds(Implementation->SimulationSpace.Transform, Implementation->Particles.GetActiveParticlesView(), 0.0f, &ChaosImmPhysDebugDebugDrawSettings);
 			}
@@ -968,19 +942,19 @@ namespace ImmediatePhysics_Chaos
 #endif
 	}
 
-	void FSimulation::DebugDrawConstraints(const int32 MinDebugLevel, const int32 MaxDebugLevel, const FRealSingle ColorScale)
+	void FSimulation::DebugDrawConstraints()
 	{
 #if CHAOS_DEBUG_DRAW
 		using namespace Chaos;
-		if (FDebugDrawQueue::IsDebugDrawingEnabled())
+		if (FDebugDrawQueue::IsDebugDrawingEnabled() && bChaosImmediate_DebugDrawJoints)
 		{
-			if ((ChaosImmediate_DebugDrawCollisions >= MinDebugLevel) && (ChaosImmediate_DebugDrawCollisions <= MaxDebugLevel))
+			if (bChaosImmediate_DebugDrawCollisions)
 			{
-				DebugDraw::DrawCollisions(Implementation->SimulationSpace.Transform, Implementation->Collisions.GetConstraintAllocator(), ColorScale, &ChaosImmPhysDebugDebugDrawSettings);
+				DebugDraw::DrawCollisions(Implementation->SimulationSpace.Transform, Implementation->Collisions.GetConstraintAllocator(), 1.0f, &ChaosImmPhysDebugDebugDrawSettings);
 			}
-			if ((ChaosImmediate_DebugDrawJoints >= MinDebugLevel) && (ChaosImmediate_DebugDrawJoints <= MaxDebugLevel))
+			if (bChaosImmediate_DebugDrawJoints)
 			{
-				DebugDraw::DrawJointConstraints(Implementation->SimulationSpace.Transform, Implementation->Joints, ColorScale, ChaosImmediate_DebugDrawJointFeatures, &ChaosImmPhysDebugDebugDrawSettings);
+				DebugDraw::DrawJointConstraints(Implementation->SimulationSpace.Transform, Implementation->Joints, 1.0f, ChaosImmediate_DebugDrawJointFeatures, &ChaosImmPhysDebugDebugDrawSettings);
 			}
 		}
 #endif
@@ -990,12 +964,9 @@ namespace ImmediatePhysics_Chaos
 	{
 #if CHAOS_DEBUG_DRAW
 		using namespace Chaos;
-		if (FDebugDrawQueue::IsDebugDrawingEnabled())
+		if (FDebugDrawQueue::IsDebugDrawingEnabled() && bChaosImmediate_DebugDrawSimulationSpace)
 		{
-			if (ChaosImmediate_DebugDrawSimulationSpace)
-			{
-				DebugDraw::DrawSimulationSpace(Implementation->SimulationSpace, &ChaosImmPhysDebugDebugDrawSettings);
-			}
+			DebugDraw::DrawSimulationSpace(Implementation->SimulationSpace, &ChaosImmPhysDebugDebugDrawSettings);
 		}
 #endif
 	}
diff --git a/Engine/Source/Runtime/Engine/Public/Physics/ImmediatePhysics/ImmediatePhysicsChaos/ImmediatePhysicsSimulation_Chaos.h b/Engine/Source/Runtime/Engine/Public/Physics/ImmediatePhysics/ImmediatePhysicsChaos/ImmediatePhysicsSimulation_Chaos.h
--- a/Engine/Source/Runtime/Engine/Public/Physics/ImmediatePhysics/ImmediatePhysicsChaos/ImmediatePhysicsSimulation_Chaos.h
+++ b/Engine/Source/Runtime/Engine/Public/Physics/ImmediatePhysics/ImmediatePhysicsChaos/ImmediatePhysicsSimulation_Chaos.h
@@ -87,16 +87,19 @@ namespace ImmediatePhysics_Chaos
 			const int32 JointPushOutIts,
 			const int32 CollisionPushOutIts);
 
+		/** Explicit debug draw path if the use case needs it to happen at a point outside of the simulation **/
+		void DebugDraw();
+
 	private:
 		void RemoveFromCollidingPairs(FActorHandle* ActorHandle);
 		void PackCollidingPairs();
 		void UpdateActivePotentiallyCollidingPairs();
 		FReal UpdateStepTime(const FReal DeltaTime, const FReal MaxStepTime);
 
-		void DebugDrawStaticParticles(const int32 MinDebugLevel, const int32 MaxDebugLevel, const float ColorScale);
-		void DebugDrawKinematicParticles(const int32 MinDebugLevel, const int32 MaxDebugLevel, const float ColorScale);
-		void DebugDrawDynamicParticles(const int32 MinDebugLevel, const int32 MaxDebugLevel, const float ColorScale);
-		void DebugDrawConstraints(const int32 MinDebugLevel, const int32 MaxDebugLevel, const FRealSingle ColorScale);
+		void DebugDrawStaticParticles();
+		void DebugDrawKinematicParticles();
+		void DebugDrawDynamicParticles();
+		void DebugDrawConstraints();
 		void DebugDrawSimulationSpace();
 
 		struct FImplementation;
