Fixed crash in mobile preview related to instanced static mesh vertex input layout

diff --git a/Engine/Source/Runtime/Engine/Private/InstancedStaticMesh.cpp b/Engine/Source/Runtime/Engine/Private/InstancedStaticMesh.cpp
--- a/Engine/Source/Runtime/Engine/Private/InstancedStaticMesh.cpp
+++ b/Engine/Source/Runtime/Engine/Private/InstancedStaticMesh.cpp
@@ -1000,7 +1000,7 @@ void FInstancedStaticMeshRenderData::BindBuffersToVertexFactories()
 
 	check(IsInRenderingThread());
 
-	const bool bCanUseGPUScene = UseGPUScene(GMaxRHIShaderPlatform, GMaxRHIFeatureLevel);
+	const bool bCanUseGPUScene = UseGPUScene(GMaxRHIShaderPlatform, FeatureLevel);
 	if (!bCanUseGPUScene)
 	{
 		PerInstanceRenderData->InstanceBuffer.FlushGPUUpload();
@@ -1478,19 +1478,20 @@ void FInstancedStaticMeshSceneProxy::SetupProxy(UInstancedStaticMeshComponent* I
 void FInstancedStaticMeshSceneProxy::CreateRenderThreadResources()
 {
 	FStaticMeshSceneProxy::CreateRenderThreadResources();
+
+	const bool bCanUseGPUScene = UseGPUScene(GetScene().GetShaderPlatform(), GetScene().GetFeatureLevel());
 	
 	// Flush upload of GPU data for ISM/HISM
 	if (ensure(InstancedRenderData.PerInstanceRenderData.IsValid()))
 	{
 		FStaticMeshInstanceBuffer& InstanceBuffer = InstancedRenderData.PerInstanceRenderData->InstanceBuffer;
-		const bool bCanUseGPUScene = UseGPUScene(GMaxRHIShaderPlatform, GMaxRHIFeatureLevel);
 		if (!bCanUseGPUScene)
 		{
 			InstanceBuffer.FlushGPUUpload();
 		}
 	}
 
-	if (UseGPUScene(GetScene().GetShaderPlatform(), GetScene().GetFeatureLevel()))
+	if (bCanUseGPUScene)
 	{
 		bSupportsInstanceDataBuffer = true;
 		// TODO: can the PerInstanceRenderData ever not be valid here?
