Add AccountIdHandle to FUniqueNetIdWrapper

diff --git a/Engine/Plugins/Online/Experimental/OnlineServices/Source/OnlineServicesCommon/OnlineServicesCommon.Build.cs b/Engine/Plugins/Online/Experimental/OnlineServices/Source/OnlineServicesCommon/OnlineServicesCommon.Build.cs
--- a/Engine/Plugins/Online/Experimental/OnlineServices/Source/OnlineServicesCommon/OnlineServicesCommon.Build.cs
+++ b/Engine/Plugins/Online/Experimental/OnlineServices/Source/OnlineServicesCommon/OnlineServicesCommon.Build.cs
@@ -12,6 +12,7 @@ public class OnlineServicesCommon : ModuleRules
 
 		PublicDependencyModuleNames.AddRange(
 			new string[] {
+				"CoreOnline",
 				"OnlineServicesInterface"
 			}
 		);
diff --git a/Engine/Plugins/Online/Experimental/OnlineServices/Source/OnlineServicesCommon/Private/Online/OnlineSyntaxTests.cpp b/Engine/Plugins/Online/Experimental/OnlineServices/Source/OnlineServicesCommon/Private/Online/OnlineSyntaxTests.cpp
--- a/Engine/Plugins/Online/Experimental/OnlineServices/Source/OnlineServicesCommon/Private/Online/OnlineSyntaxTests.cpp
+++ b/Engine/Plugins/Online/Experimental/OnlineServices/Source/OnlineServicesCommon/Private/Online/OnlineSyntaxTests.cpp
@@ -1,9 +1,9 @@
 // Copyright Epic Games, Inc. All Rights Reserved.
 
+#include "Online/CoreOnline.h"
 #include "Online/OnlineServicesCommon.h"
 #include "Online/OnlineConfig.h"
 #include "Online/OnlineAsyncOp.h"
-#include "Online/OnlineId.h"
 #include "Online/OnlineComponent.h"
 
 #define TEST_ASYNC_OP_CONTINUATION_SYNTAX 1
diff --git a/Engine/Plugins/Online/Experimental/OnlineServices/Source/OnlineServicesCommon/Public/Online/OnlineAsyncOpCache.h b/Engine/Plugins/Online/Experimental/OnlineServices/Source/OnlineServicesCommon/Public/Online/OnlineAsyncOpCache.h
--- a/Engine/Plugins/Online/Experimental/OnlineServices/Source/OnlineServicesCommon/Public/Online/OnlineAsyncOpCache.h
+++ b/Engine/Plugins/Online/Experimental/OnlineServices/Source/OnlineServicesCommon/Public/Online/OnlineAsyncOpCache.h
@@ -2,8 +2,8 @@
 
 #pragma once
 
+#include "Online/CoreOnline.h"
 #include "Online/OnlineAsyncOp.h"
-#include "Online/OnlineId.h"
 #include "Containers/Map.h"
 #include "Templates/SharedPointer.h"
 
diff --git a/Engine/Plugins/Online/Experimental/OnlineServices/Source/OnlineServicesCommon/Public/Online/OnlineExecHandler.h b/Engine/Plugins/Online/Experimental/OnlineServices/Source/OnlineServicesCommon/Public/Online/OnlineExecHandler.h
--- a/Engine/Plugins/Online/Experimental/OnlineServices/Source/OnlineServicesCommon/Public/Online/OnlineExecHandler.h
+++ b/Engine/Plugins/Online/Experimental/OnlineServices/Source/OnlineServicesCommon/Public/Online/OnlineExecHandler.h
@@ -190,6 +190,13 @@ inline bool ParseOnlineExecParams(const TCHAR*& Cmd, TVariant<Ts...>& Variant)
 	return true;
 }
 
+template<typename IdType>
+inline bool ParseOnlineExecParams(const TCHAR*& Cmd, TOnlineIdHandle<IdType>& Id)
+{
+	// TODO
+	return true;
+}
+
 template <typename T>
 std::enable_if_t<!TModels<Meta::COnlineMetadataAvailable, T>::Value, bool> ParseOnlineExecParams(const TCHAR*& Cmd, T& Value)
 {
diff --git a/Engine/Plugins/Online/Experimental/OnlineServices/Source/OnlineServicesInterface/OnlineServicesInterface.Build.cs b/Engine/Plugins/Online/Experimental/OnlineServices/Source/OnlineServicesInterface/OnlineServicesInterface.Build.cs
--- a/Engine/Plugins/Online/Experimental/OnlineServices/Source/OnlineServicesInterface/OnlineServicesInterface.Build.cs
+++ b/Engine/Plugins/Online/Experimental/OnlineServices/Source/OnlineServicesInterface/OnlineServicesInterface.Build.cs
@@ -12,6 +12,7 @@ public class OnlineServicesInterface : ModuleRules
 
 		PublicDependencyModuleNames.AddRange(
 			new string[] {
+				"CoreOnline"
 			}
 		);
 
diff --git a/Engine/Plugins/Online/Experimental/OnlineServices/Source/OnlineServicesInterface/Private/OnlineId.cpp b/Engine/Plugins/Online/Experimental/OnlineServices/Source/OnlineServicesInterface/Private/OnlineId.cpp
deleted file mode 100644
diff --git a/Engine/Plugins/Online/Experimental/OnlineServices/Source/OnlineServicesInterface/Public/Online/Auth.h b/Engine/Plugins/Online/Experimental/OnlineServices/Source/OnlineServicesInterface/Public/Online/Auth.h
--- a/Engine/Plugins/Online/Experimental/OnlineServices/Source/OnlineServicesInterface/Public/Online/Auth.h
+++ b/Engine/Plugins/Online/Experimental/OnlineServices/Source/OnlineServicesInterface/Public/Online/Auth.h
@@ -2,8 +2,8 @@
 
 #pragma once
 
+#include "Online/CoreOnline.h"
 #include "Online/OnlineAsyncOpHandle.h"
-#include "Online/OnlineId.h"
 #include "Online/OnlineMeta.h"
 
 namespace UE::Online {
diff --git a/Engine/Plugins/Online/Experimental/OnlineServices/Source/OnlineServicesInterface/Public/Online/ExternalUI.h b/Engine/Plugins/Online/Experimental/OnlineServices/Source/OnlineServicesInterface/Public/Online/ExternalUI.h
--- a/Engine/Plugins/Online/Experimental/OnlineServices/Source/OnlineServicesInterface/Public/Online/ExternalUI.h
+++ b/Engine/Plugins/Online/Experimental/OnlineServices/Source/OnlineServicesInterface/Public/Online/ExternalUI.h
@@ -2,8 +2,8 @@
 
 #pragma once
 
+#include "Online/CoreOnline.h"
 #include "Online/OnlineAsyncOpHandle.h"
-#include "Online/OnlineId.h"
 #include "Online/OnlineMeta.h"
 
 namespace UE::Online {
diff --git a/Engine/Plugins/Online/Experimental/OnlineServices/Source/OnlineServicesInterface/Public/Online/Friends.h b/Engine/Plugins/Online/Experimental/OnlineServices/Source/OnlineServicesInterface/Public/Online/Friends.h
--- a/Engine/Plugins/Online/Experimental/OnlineServices/Source/OnlineServicesInterface/Public/Online/Friends.h
+++ b/Engine/Plugins/Online/Experimental/OnlineServices/Source/OnlineServicesInterface/Public/Online/Friends.h
@@ -2,8 +2,8 @@
 
 #pragma once
 
+#include "Online/CoreOnline.h"
 #include "Online/OnlineAsyncOpHandle.h"
-#include "Online/OnlineId.h"
 #include "Online/OnlineMeta.h"
 
 namespace UE::Online {
diff --git a/Engine/Plugins/Online/Experimental/OnlineServices/Source/OnlineServicesInterface/Public/Online/OnlineId.h b/Engine/Plugins/Online/Experimental/OnlineServices/Source/OnlineServicesInterface/Public/Online/OnlineId.h
deleted file mode 100644
diff --git a/Engine/Plugins/Online/Experimental/OnlineServices/Source/OnlineServicesInterface/Public/Online/OnlineServices.h b/Engine/Plugins/Online/Experimental/OnlineServices/Source/OnlineServicesInterface/Public/Online/OnlineServices.h
--- a/Engine/Plugins/Online/Experimental/OnlineServices/Source/OnlineServicesInterface/Public/Online/OnlineServices.h
+++ b/Engine/Plugins/Online/Experimental/OnlineServices/Source/OnlineServicesInterface/Public/Online/OnlineServices.h
@@ -2,7 +2,7 @@
 
 #pragma once
 
-#include "Online/OnlineId.h"
+#include "Online/CoreOnline.h"
 #include "Templates/SharedPointer.h"
 
 class FString;
diff --git a/Engine/Plugins/Online/Experimental/OnlineServices/Source/OnlineServicesInterface/Public/Online/OnlineServicesTypes.h b/Engine/Plugins/Online/Experimental/OnlineServices/Source/OnlineServicesInterface/Public/Online/OnlineServicesTypes.h
deleted file mode 100644
diff --git a/Engine/Plugins/Online/Experimental/OnlineServices/Source/OnlineServicesInterface/Public/Online/Presence.h b/Engine/Plugins/Online/Experimental/OnlineServices/Source/OnlineServicesInterface/Public/Online/Presence.h
--- a/Engine/Plugins/Online/Experimental/OnlineServices/Source/OnlineServicesInterface/Public/Online/Presence.h
+++ b/Engine/Plugins/Online/Experimental/OnlineServices/Source/OnlineServicesInterface/Public/Online/Presence.h
@@ -2,8 +2,8 @@
 
 #pragma once
 
+#include "Online/CoreOnline.h"
 #include "Online/OnlineAsyncOpHandle.h"
-#include "Online/OnlineId.h"
 #include "Online/OnlineMeta.h"
 
 namespace UE::Online {
diff --git a/Engine/Plugins/Online/Experimental/OnlineServicesEOS/Source/OnlineServicesEOS.Build.cs b/Engine/Plugins/Online/Experimental/OnlineServicesEOS/Source/OnlineServicesEOS.Build.cs
--- a/Engine/Plugins/Online/Experimental/OnlineServicesEOS/Source/OnlineServicesEOS.Build.cs
+++ b/Engine/Plugins/Online/Experimental/OnlineServicesEOS/Source/OnlineServicesEOS.Build.cs
@@ -17,6 +17,7 @@ public class OnlineServicesEOS : ModuleRules
 
 		PrivateDependencyModuleNames.AddRange(
 			new string[] {
+				"CoreOnline",
 				"EOSSDK",
 				"EOSShared",
 			}
diff --git a/Engine/Plugins/Online/Experimental/OnlineServicesEOS/Source/Private/OnlineIdEOS.cpp b/Engine/Plugins/Online/Experimental/OnlineServicesEOS/Source/Private/OnlineIdEOS.cpp
--- a/Engine/Plugins/Online/Experimental/OnlineServicesEOS/Source/Private/OnlineIdEOS.cpp
+++ b/Engine/Plugins/Online/Experimental/OnlineServicesEOS/Source/Private/OnlineIdEOS.cpp
@@ -2,8 +2,10 @@
 
 #include "OnlineIdEOS.h"
 
-#include "OnlineServicesEOSTypes.h"
 #include "Online/OnlineAsyncOp.h"
+#include "String/HexToBytes.h"
+#include "String/BytesToHex.h"
+#include "OnlineServicesEOSTypes.h"
 
 #include "eos_connect_types.h"
 #include "eos_connect.h"
@@ -162,6 +164,93 @@ FString FOnlineAccountIdRegistryEOS::ToLogString(const FOnlineAccountIdHandle& H
 	return Result;
 }
 
+enum class EEOSAccountIdElements : uint8
+{	
+	None = 0,
+	EAS = 1 << 0,
+	EOS = 1 << 1
+};
+ENUM_CLASS_FLAGS(EEOSAccountIdElements);
+const uint8 OnlineIdEOSUtf8BufferLength = 32;
+const uint8 OnlineIdEOSHexBufferLength = 16;
+
+TArray<uint8> FOnlineAccountIdRegistryEOS::ToReplicationData(const FOnlineAccountIdHandle& Handle) const
+{
+	TArray<uint8> ReplicationData;
+	if (ValidateOnlineId(Handle))
+	{
+		const FOnlineAccountIdDataEOS& IdData = GetIdData(Handle);
+		EEOSAccountIdElements Elements = EEOSAccountIdElements::None;
+
+		char EasBuffer[EOS_EPICACCOUNTID_MAX_LENGTH + 1] = { 0 };
+		if (EOS_EpicAccountId_IsValid(IdData.EpicAccountId))
+		{
+			int32 EasBufferLength = sizeof(EasBuffer);
+			const EOS_EResult EosResult = EOS_EpicAccountId_ToString(IdData.EpicAccountId, EasBuffer, &EasBufferLength);
+			if (ensure(EosResult == EOS_EResult::EOS_Success))
+			{
+				check(EasBufferLength-1 == OnlineIdEOSUtf8BufferLength);
+				Elements |= EEOSAccountIdElements::EAS;
+			}
+		}
+
+		char EosBuffer[EOS_PRODUCTUSERID_MAX_LENGTH + 1] = { 0 };
+		if (EOS_ProductUserId_IsValid(IdData.ProductUserId))
+		{
+			int32 EosBufferLength = sizeof(EosBuffer);
+			const EOS_EResult EosResult = EOS_ProductUserId_ToString(IdData.ProductUserId, EosBuffer, &EosBufferLength);
+			if (ensure(EosResult == EOS_EResult::EOS_Success))
+			{
+				check(EosBufferLength-1 == OnlineIdEOSUtf8BufferLength);
+				Elements |= EEOSAccountIdElements::EOS;
+			}
+		}
+
+		const int EasHexBufferLength = EnumHasAnyFlags(Elements, EEOSAccountIdElements::EAS) ? OnlineIdEOSHexBufferLength : 0;
+		const int EosHexBufferLength = EnumHasAnyFlags(Elements, EEOSAccountIdElements::EOS) ? OnlineIdEOSHexBufferLength : 0;
+		ReplicationData.SetNumUninitialized(1 + EasHexBufferLength + EosHexBufferLength);
+		ReplicationData[0] = uint8(Elements);
+		if (EasHexBufferLength > 0)
+		{
+			UE::String::HexToBytes(FUtf8StringView(EasBuffer, OnlineIdEOSUtf8BufferLength), ReplicationData.GetData() + 1);
+		}
+		if (EosHexBufferLength > 0)
+		{
+			UE::String::HexToBytes(FUtf8StringView(EosBuffer, OnlineIdEOSUtf8BufferLength), ReplicationData.GetData() + 1 + EasHexBufferLength);
+		}
+	}
+	else
+	{
+		check(!Handle.IsValid()); // Check we haven't been passed a valid handle for a different EOnlineServices.
+	}
+	return ReplicationData;
+}
+
+FOnlineAccountIdHandle FOnlineAccountIdRegistryEOS::FromReplicationData(const TArrayView<uint8> ReplicationData)
+{
+	const EEOSAccountIdElements Elements = EEOSAccountIdElements(ReplicationData[0]);
+	const int EasHexBufferLength = EnumHasAnyFlags(Elements, EEOSAccountIdElements::EAS) ? OnlineIdEOSHexBufferLength : 0;
+	const int EosHexBufferLength = EnumHasAnyFlags(Elements, EEOSAccountIdElements::EOS) ? OnlineIdEOSHexBufferLength : 0;
+
+	EOS_EpicAccountId EpicAccountId = nullptr;
+	if (EasHexBufferLength > 0)
+	{
+		char EasBuffer[EOS_EPICACCOUNTID_MAX_LENGTH + 1] = { 0 };
+		UE::String::BytesToHex(TConstArrayView<uint8>(ReplicationData.GetData() + 1, EasHexBufferLength), EasBuffer);
+		EpicAccountId = EOS_EpicAccountId_FromString(EasBuffer);
+	}
+
+	EOS_ProductUserId ProductUserId = nullptr;
+	if (EosHexBufferLength > 0)
+	{
+		char EosBuffer[EOS_EPICACCOUNTID_MAX_LENGTH + 1] = { 0 };
+		UE::String::BytesToHex(TConstArrayView<uint8>(ReplicationData.GetData() + 1 + EasHexBufferLength, EosHexBufferLength), EosBuffer);
+		ProductUserId = EOS_ProductUserId_FromString(EosBuffer);
+	}
+	
+	return Create(EpicAccountId, ProductUserId);
+}
+
 EOS_EpicAccountId GetEpicAccountId(const FOnlineAccountIdHandle& Handle)
 {
 	return FOnlineAccountIdRegistryEOS::Get().GetIdData(Handle).EpicAccountId;
diff --git a/Engine/Plugins/Online/Experimental/OnlineServicesEOS/Source/Private/OnlineIdEOS.h b/Engine/Plugins/Online/Experimental/OnlineServicesEOS/Source/Private/OnlineIdEOS.h
--- a/Engine/Plugins/Online/Experimental/OnlineServicesEOS/Source/Private/OnlineIdEOS.h
+++ b/Engine/Plugins/Online/Experimental/OnlineServicesEOS/Source/Private/OnlineIdEOS.h
@@ -4,7 +4,7 @@
 
 #include "Async/Future.h"
 #include "Delegates/Delegate.h"
-#include "Online/OnlineId.h"
+#include "Online/CoreOnline.h"
 #include "Online/OnlineIdCommon.h"
 
 #if defined(EOS_PLATFORM_BASE_FILE_NAME)
@@ -39,8 +39,12 @@ public:
 
 	// Begin IOnlineAccountIdRegistry
 	virtual FString ToLogString(const FOnlineAccountIdHandle& Handle) const override;
+	virtual TArray<uint8> ToReplicationData(const FOnlineAccountIdHandle& Handle) const override;
+	virtual FOnlineAccountIdHandle FromReplicationData(const TArrayView<uint8> ReplicationString) override;
 	// End IOnlineAccountIdRegistry
 
+	virtual ~FOnlineAccountIdRegistryEOS() = default;
+
 private:
 	mutable FRWLock Lock;
 
@@ -49,7 +53,6 @@ private:
 	TMap<EOS_EpicAccountId, FOnlineAccountIdHandle> EpicAccountIdToHandle; // Map of EOS_EpicAccountId to the associated handle.
 	TMap<EOS_ProductUserId, FOnlineAccountIdHandle> ProductUserIdToHandle; // Map of EOS_ProductUserId to the associated handle.
 
-	virtual ~FOnlineAccountIdRegistryEOS() = default;
 };
 
 EOS_EpicAccountId GetEpicAccountId(const FOnlineAccountIdHandle& Handle);
diff --git a/Engine/Plugins/Online/Experimental/OnlineServicesEOS/Source/Private/Tests/OnlineIdEOSTests.cpp b/Engine/Plugins/Online/Experimental/OnlineServicesEOS/Source/Private/Tests/OnlineIdEOSTests.cpp
new file mode 100644
--- /dev/null
+++ b/Engine/Plugins/Online/Experimental/OnlineServicesEOS/Source/Private/Tests/OnlineIdEOSTests.cpp
@@ -0,0 +1,40 @@
+// Copyright Epic Games, Inc. All Rights Reserved.
+
+#include "Misc/AutomationTest.h"
+#include "OnlineIdEOS.h"
+
+IMPLEMENT_SIMPLE_AUTOMATION_TEST(FAccountIdReplicationTest, 
+	"System.Engine.Online.EosAccountIdReplicationTest", 
+	EAutomationTestFlags::ApplicationContextMask | EAutomationTestFlags::SmokeFilter)
+bool FAccountIdReplicationTest::RunTest(const FString& Parameters)
+{
+	using namespace UE::Online;
+	const EOS_EpicAccountId EasId = EOS_EpicAccountId_FromString("d03352b601a44af7b02ed1faab668e0f");
+	const EOS_ProductUserId EosId = EOS_ProductUserId_FromString("31a0b16a1f0743978458e3733bb3be14");
+
+	{
+		FOnlineAccountIdRegistryEOS Registry;
+		const FOnlineAccountIdHandle Handle = Registry.Create(EasId, EosId);
+		TArray<uint8> RepData = Registry.ToReplicationData(Handle);
+		const FOnlineAccountIdHandle Handle2 = Registry.FromReplicationData(RepData);
+		UTEST_EQUAL(TEXT(""), Handle, Handle2);
+	}
+
+	{
+		FOnlineAccountIdRegistryEOS Registry;
+		const FOnlineAccountIdHandle Handle = Registry.Create(EasId, nullptr);
+		TArray<uint8> RepData = Registry.ToReplicationData(Handle);
+		const FOnlineAccountIdHandle Handle2 = Registry.FromReplicationData(RepData);
+		UTEST_EQUAL(TEXT(""), Handle, Handle2);
+	}
+
+	{
+		FOnlineAccountIdRegistryEOS Registry;
+		const FOnlineAccountIdHandle Handle = Registry.Create(nullptr, EosId);
+		TArray<uint8> RepData = Registry.ToReplicationData(Handle);
+		const FOnlineAccountIdHandle Handle2 = Registry.FromReplicationData(RepData);
+		UTEST_EQUAL(TEXT(""), Handle, Handle2);
+	}
+
+	return true;
+}
\ No newline at end of file
diff --git a/Engine/Plugins/Online/OnlineFramework/Source/Party/Party.Build.cs b/Engine/Plugins/Online/OnlineFramework/Source/Party/Party.Build.cs
--- a/Engine/Plugins/Online/OnlineFramework/Source/Party/Party.Build.cs
+++ b/Engine/Plugins/Online/OnlineFramework/Source/Party/Party.Build.cs
@@ -16,6 +16,7 @@ public class Party : ModuleRules
 
 		PublicDependencyModuleNames.AddRange(
 			new string[] {
+				"CoreOnline",
 				"OnlineSubsystem",
 				"OnlineSubsystemUtils",
 			}
diff --git a/Engine/Plugins/Online/OnlineSubsystemUtils/Source/OnlineSubsystemUtils/OnlineSubsystemUtils.Build.cs b/Engine/Plugins/Online/OnlineSubsystemUtils/Source/OnlineSubsystemUtils/OnlineSubsystemUtils.Build.cs
--- a/Engine/Plugins/Online/OnlineSubsystemUtils/Source/OnlineSubsystemUtils/OnlineSubsystemUtils.Build.cs
+++ b/Engine/Plugins/Online/OnlineSubsystemUtils/Source/OnlineSubsystemUtils/OnlineSubsystemUtils.Build.cs
@@ -54,6 +54,7 @@ public class OnlineSubsystemUtils : ModuleRules
 			new string[]
 			{
 				"Core",
+				"CoreOnline",
 				"CoreUObject",
 				"NetCore"
 			}
diff --git a/Engine/Plugins/Online/OnlineSubsystemUtils/Source/OnlineSubsystemUtils/Private/OnlineSubsystemUtilsModule.cpp b/Engine/Plugins/Online/OnlineSubsystemUtils/Source/OnlineSubsystemUtils/Private/OnlineSubsystemUtilsModule.cpp
--- a/Engine/Plugins/Online/OnlineSubsystemUtils/Source/OnlineSubsystemUtils/Private/OnlineSubsystemUtilsModule.cpp
+++ b/Engine/Plugins/Online/OnlineSubsystemUtils/Source/OnlineSubsystemUtils/Private/OnlineSubsystemUtilsModule.cpp
@@ -226,8 +226,10 @@ PRAGMA_ENABLE_DEPRECATION_WARNINGS
 		ensure(SubsystemNameToHash.Num() == (HashId - 1));
 		ensure(HashToSubsystemName.Num() == (HashId - 1));
 
-		// FUniqueNetIdRepl uses 5 bits to transmit the HashId and 31 is used for OnlineSubsystems not included in this list
-		check(HashId < 31);
+		// FUniqueNetIdRepl uses 5 bits to transmit the HashId
+		// 31 is used for OnlineSubsystems not included in this list
+		// 30 is used for OnlineServices (v2) net ids
+		check(HashId < 30);
 	}
 
 	/** If false it will not try to do online PIE at all */
diff --git a/Engine/Source/Runtime/CoreOnline/Private/Online/CoreOnline.cpp b/Engine/Source/Runtime/CoreOnline/Private/Online/CoreOnline.cpp
new file mode 100644
--- /dev/null
+++ b/Engine/Source/Runtime/CoreOnline/Private/Online/CoreOnline.cpp
@@ -0,0 +1,84 @@
+// Copyright Epic Games, Inc. All Rights Reserved.
+
+#include "Online/CoreOnline.h"
+
+#include "Misc/LazySingleton.h"
+
+namespace UE::Online {
+
+FOnlineIdRegistryRegistry& FOnlineIdRegistryRegistry::Get()
+{
+	return TLazySingleton<FOnlineIdRegistryRegistry>::Get();
+}
+
+void FOnlineIdRegistryRegistry::TearDown()
+{
+	return TLazySingleton<FOnlineIdRegistryRegistry>::TearDown();
+}
+
+void FOnlineIdRegistryRegistry::RegisterAccountIdRegistry(EOnlineServices OnlineServices, IOnlineAccountIdRegistry* Registry, int32 Priority)
+{
+	FAccountIdRegistryAndPriority* Found = AccountIdRegistries.Find(OnlineServices);
+	if (Found == nullptr || Found->Priority < Priority)
+	{
+		AccountIdRegistries.Emplace(OnlineServices, FAccountIdRegistryAndPriority(Registry, Priority));
+	}
+}
+
+void FOnlineIdRegistryRegistry::UnregisterAccountIdRegistry(EOnlineServices OnlineServices, int32 Priority)
+{
+	FAccountIdRegistryAndPriority* Found = AccountIdRegistries.Find(OnlineServices);
+	if (Found != nullptr && Found->Priority == Priority)
+	{
+		AccountIdRegistries.Remove(OnlineServices);
+	}
+}
+
+IOnlineAccountIdRegistry* FOnlineIdRegistryRegistry::GetAccountIdRegistry(EOnlineServices OnlineServices)
+{
+	IOnlineAccountIdRegistry* Registry = nullptr;
+	if (FAccountIdRegistryAndPriority* Found = AccountIdRegistries.Find(OnlineServices))
+	{
+		Registry = Found->Registry;
+	}
+	return Registry;
+}
+
+FString ToLogString(const FOnlineAccountIdHandle& Id)
+{
+	FString Result;
+	if (IOnlineAccountIdRegistry* Registry = FOnlineIdRegistryRegistry::Get().GetAccountIdRegistry(Id.GetOnlineServicesType()))
+	{
+		Result = Registry->ToLogString(Id);
+	}
+	return Result;
+}
+
+}	/* UE::Online */
+
+FString FUniqueNetIdWrapper::ToDebugString() const
+{
+	FString Result;
+	if (IsValid())
+	{
+		if (Variant.IsType<FUniqueNetIdPtr>())
+		{
+			const FUniqueNetIdPtr& Ptr = Variant.Get<FUniqueNetIdPtr>();
+			Result = FString::Printf(TEXT("%s:%s"), *Ptr->GetType().ToString(), *Ptr->ToDebugString());
+		}
+		else if (Variant.IsType<UE::Online::FOnlineAccountIdHandle>())
+		{
+			const UE::Online::FOnlineAccountIdHandle& Handle = Variant.Get<UE::Online::FOnlineAccountIdHandle>();
+			Result = ToLogString(Handle);
+		}
+		else
+		{
+			checkNoEntry();
+		}
+	}
+	else
+	{
+		Result = TEXT("INVALID");
+	}
+	return Result;
+}
diff --git a/Engine/Source/Runtime/CoreOnline/Public/Online/CoreOnline.h b/Engine/Source/Runtime/CoreOnline/Public/Online/CoreOnline.h
--- a/Engine/Source/Runtime/CoreOnline/Public/Online/CoreOnline.h
+++ b/Engine/Source/Runtime/CoreOnline/Public/Online/CoreOnline.h
@@ -5,8 +5,11 @@
 #include "CoreMinimal.h"
 #include "Hash/CityHash.h"
 #include "Online/CoreOnlineFwd.h"
+#include "Misc/TVariant.h"
 #include "UObject/ObjectMacros.h"
 
+class FLazySingleton;
+
 UE_DEPRECATED(5.0, "Use NAME_GameSession.")
 inline constexpr EName GameSessionName = NAME_GameSession;
 UE_DEPRECATED(5.0, "Use NAME_PartySession.")
@@ -172,10 +175,187 @@ public:
 	}
 };
 
+namespace UE::Online {
+
+/** Tags used as template argument to TOnlineIdHandle to make it a compile error to assign between id's of different types */
+namespace OnlineIdHandleTags
+{
+	struct FAccount {};
+	struct FSession {};
+	struct FParty {};
+}
+
+enum class EOnlineServices : uint8
+{
+	// Null, Providing minimal functionality when no backend services are required
+	Null,
+	// Epic Online Services
+	Epic,
+	// Xbox services
+	Xbox,
+	// PlayStation Network
+	PSN,
+	// Nintendo
+	Nintendo,
+	// Stadia,
+	Stadia,
+	// Steam
+	Steam,
+	// Google
+	Google,
+	// GooglePlay
+	GooglePlay,
+	// Apple
+	Apple,
+	// GameKit
+	AppleGameKit,
+	// Samsung
+	Samsung,
+	// Oculus
+	Oculus,
+	// Tencent
+	Tencent,
+	// Reserved for future use/platform extensions
+	Reserved_14,
+	Reserved_15,
+	Reserved_16,
+	Reserved_17,
+	Reserved_18,
+	Reserved_19,
+	Reserved_20,
+	Reserved_21,
+	Reserved_22,
+	Reserved_23,
+	Reserved_24,
+	Reserved_25,
+	Reserved_26,
+	Reserved_27,
+	// For game specific Online Services
+	GameDefined_0 = 28,
+	GameDefined_1,
+	GameDefined_2,
+	GameDefined_3,
+	// Platform native, may not exist for all platforms
+	Platform = 254,
+	// Default, configured via ini, TODO: List specific ini section/key
+	Default = 255
+};
+
+/**
+ * A handle to an id which uniquely identifies a persistent or transient online resource, i.e. account/session/party etc, within a given Online Services provider.
+ * At most one id, and therefore one handle, exists for any given resource. The id and handle persist until the OnlineServices module is unloaded.
+ * Passed to and returned from OnlineServices APIs.
+ */
+template<typename IdType>
+class TOnlineIdHandle
+{
+public:
+	TOnlineIdHandle() = default;
+	TOnlineIdHandle(EOnlineServices Type, uint32 Handle)
+	{
+		check(Handle < 0xFF000000);
+		Value = (Handle & 0x00FFFFFF) | (uint32(Type) << 24);
+	}
+
+	inline bool IsValid() const { return GetHandle() != 0; }
+
+	EOnlineServices GetOnlineServicesType() const { return EOnlineServices(Value >> 24); }
+	uint32 GetHandle() const { return Value & 0x00FFFFFF; }
+
+	bool operator==(const TOnlineIdHandle& Other) const { return Value == Other.Value; }
+	bool operator!=(const TOnlineIdHandle& Other) const { return Value != Other.Value; }
+
+private:
+	uint32 Value = uint32(EOnlineServices::Null) << 24;
+};
+
+using FOnlineAccountIdHandle = TOnlineIdHandle<OnlineIdHandleTags::FAccount>;
+
+COREONLINE_API FString ToLogString(const FOnlineAccountIdHandle& Id);
+
+template<typename IdType>
+inline uint32 GetTypeHash(const TOnlineIdHandle<IdType>& Handle)
+{
+	using ::GetTypeHash;
+	return HashCombine(GetTypeHash(Handle.GetOnlineServicesType()), GetTypeHash(Handle.GetHandle()));
+}
+
+template<typename IdType>
+class IOnlineIdRegistry
+{
+public:
+	virtual FString ToLogString(const TOnlineIdHandle<IdType>& Handle) const = 0;
+	virtual TArray<uint8> ToReplicationData(const TOnlineIdHandle<IdType>& Handle) const = 0;
+	virtual TOnlineIdHandle<IdType> FromReplicationData(const TArrayView<uint8> Handle) = 0;
+};
+
+using IOnlineAccountIdRegistry = IOnlineIdRegistry<OnlineIdHandleTags::FAccount>;
+
+class FOnlineIdRegistryRegistry
+{
+public:
+	/**
+	 * Get the FOnlineIdRegistryRegistry singleton
+	 *
+	 * @return The FOnlineIdRegistryRegistry singleton instance
+	 */
+	COREONLINE_API static FOnlineIdRegistryRegistry& Get();
+
+	/**
+	 * Tear down the singleton instance
+	 */
+	COREONLINE_API static void TearDown();
+	/**
+	 * Register a registry for a given OnlineServices implementation and TOnlineIdHandle type
+	 *
+	 * @param OnlineServices Services that the registry is for
+	 * @param Registry the registry of online ids
+	 * @param Priority Integer priority, allows an existing registry to be extended and registered with a higher priority so it is used instead
+	 */
+	COREONLINE_API void RegisterAccountIdRegistry(EOnlineServices OnlineServices, IOnlineAccountIdRegistry* Registry, int32 Priority = 0);
+
+	/**
+	 * Unregister a previously registered Id registry
+	 *
+	 * @param OnlineServices Services that the registry is for
+	 * @param Priority Integer priority, will be unregistered only if the priority matches the one that is registered
+	 */
+	COREONLINE_API void UnregisterAccountIdRegistry(EOnlineServices OnlineServices, int32 Priority = 0);
+
+	/**
+	 * Get the account id registry
+	 *
+	 * @param OnlineServices Type of online services for the IOnlineServices instance
+	 *
+	 * @return The account id registry, or an invalid pointer if it is not registered
+	 */
+	COREONLINE_API IOnlineAccountIdRegistry* GetAccountIdRegistry(EOnlineServices OnlineServices);
+
+private:
+	struct FAccountIdRegistryAndPriority
+	{
+		FAccountIdRegistryAndPriority(IOnlineAccountIdRegistry* InRegistry, int32 InPriority)
+			: Registry(InRegistry), Priority(InPriority) {}
+
+		IOnlineAccountIdRegistry* Registry;
+		int32 Priority;
+	};
+
+	TMap<EOnlineServices, FAccountIdRegistryAndPriority> AccountIdRegistries;
+
+	FOnlineIdRegistryRegistry() = default;
+	~FOnlineIdRegistryRegistry() = default;
+	friend FLazySingleton;
+};
+
+}	/* UE::Online */
+
 USTRUCT(noexport)
 struct FUniqueNetIdWrapper
 {
 	//GENERATED_BODY()
+	
+	using FVariantType = TVariant<FUniqueNetIdPtr, UE::Online::FOnlineAccountIdHandle>;
 
 	FUniqueNetIdWrapper() = default;
 	virtual ~FUniqueNetIdWrapper() = default;
@@ -183,58 +363,114 @@ struct FUniqueNetIdWrapper
 	// copy operators generated by compiler
 
 	FUniqueNetIdWrapper(const FUniqueNetIdRef& InUniqueNetId)
-		: UniqueNetId(InUniqueNetId)
 	{
+		Variant.Emplace<FUniqueNetIdPtr>(InUniqueNetId);
 	}
 
 	FUniqueNetIdWrapper(const FUniqueNetIdPtr& InUniqueNetId)
-		: UniqueNetId(InUniqueNetId)
+	{
+		Variant.Emplace<FUniqueNetIdPtr>(InUniqueNetId);
+	}
+
+	FUniqueNetIdWrapper(const FVariantType& InVariant)
+		: Variant(InVariant)
 	{
 	}
 
 	// temporarily restored implicit conversion from FUniqueNetId
 	FUniqueNetIdWrapper(const FUniqueNetId& InUniqueNetId)
-		: UniqueNetId(InUniqueNetId.AsShared())
 	{
+		Variant.Emplace<FUniqueNetIdPtr>(InUniqueNetId.AsShared());
+	}
+
+	bool IsV1() const
+	{
+		return Variant.IsType<FUniqueNetIdPtr>();
+	}
+
+	FUniqueNetIdPtr GetV1() const
+	{
+		FUniqueNetIdPtr Result;
+		if (ensure(IsV1()))
+		{
+			Result = Variant.Get<FUniqueNetIdPtr>();
+		}
+		return Result;
+	}
+
+	bool IsV2() const
+	{
+		return Variant.IsType<UE::Online::FOnlineAccountIdHandle>();
+	}
+
+	UE::Online::FOnlineAccountIdHandle GetV2() const
+	{
+		UE::Online::FOnlineAccountIdHandle Result;
+		if (ensure(IsV2()))
+		{
+			Result = Variant.Get<UE::Online::FOnlineAccountIdHandle>();
+		}
+		return Result;
 	}
 
 	FName GetType() const
 	{
-		return IsValid() ? UniqueNetId->GetType() : NAME_None;
+		FName Result = NAME_None;
+		if (IsValid() && ensure(IsV1()))
+		{
+			Result = GetV1()->GetType();
+		}
+		return Result;
 	}
-	
+
 	/** Convert this value to a string */
 	FString ToString() const
 	{
-		return IsValid() ? UniqueNetId->ToString() : TEXT("INVALID");
+		FString Result;
+		if(IsValid() && ensure(IsV1()))
+		{
+			Result = GetV1()->ToString();
+		}
+		return Result;
 	}
 
 	/** Convert this value to a string with additional information */
-	FString ToDebugString() const
-	{
-		return IsValid() ? FString::Printf(TEXT("%s:%s"), *UniqueNetId->GetType().ToString(), *UniqueNetId->ToDebugString()) : TEXT("INVALID");
-	}
+	COREONLINE_API FString ToDebugString() const;
 
-	/** Is the FUniqueNetId wrapped in this object valid */
+	/** Is the Variant wrapped in this object valid */
 	bool IsValid() const
 	{
-		return UniqueNetId.IsValid() && UniqueNetId->IsValid();
+		if (IsV1())
+		{
+			const FUniqueNetIdPtr& Ptr = GetV1();
+			return Ptr.IsValid() && Ptr->IsValid();
+		}
+		else
+		{
+			const UE::Online::FOnlineAccountIdHandle& Handle = GetV2();
+			return Handle.IsValid();
+		}
 	}
-	
-	/** 
+
+	/**
 	 * Assign a unique id to this wrapper object
 	 *
 	 * @param InUniqueNetId id to associate
 	 */
 	virtual void SetUniqueNetId(const FUniqueNetIdPtr& InUniqueNetId)
 	{
-		UniqueNetId = InUniqueNetId;
+		Variant.Emplace<FUniqueNetIdPtr>(InUniqueNetId);
+	}
+
+	virtual void SetAccountId(const UE::Online::FOnlineAccountIdHandle& Handle)
+	{
+		Variant.Emplace<UE::Online::FOnlineAccountIdHandle>(Handle);
 	}
 
 	/** @return unique id associated with this wrapper object */
-	const FUniqueNetIdPtr& GetUniqueNetId() const
+	FUniqueNetIdPtr GetUniqueNetId() const
 	{
-		return UniqueNetId;
+		return GetV1();
 	}
 
 	/**
@@ -242,7 +478,7 @@ struct FUniqueNetIdWrapper
 	 */
 	const FUniqueNetId& operator*() const
 	{
-		return *UniqueNetId;
+		return *GetV1();
 	}
 
 	/**
@@ -250,23 +486,26 @@ struct FUniqueNetIdWrapper
 	 */
 	const FUniqueNetId* operator->() const
 	{
-		return UniqueNetId.Get();
+		return GetV1().Get();
 	}
 
 	/**
 	* Friend function for using FUniqueNetIdWrapper as a hashable key
 	*/
-	friend inline uint32 GetTypeHash(FUniqueNetIdWrapper const& Value)
+	friend inline uint32 GetTypeHash(const FUniqueNetIdWrapper& Value)
 	{
 		if (Value.IsValid())
 		{
-			return GetTypeHash(*Value);
-		}
-		else
-		{
-			// If we hit this, something went wrong and we have received an unhashable wrapper.
-			return INDEX_NONE;
+			if (Value.IsV1())
+			{
+				return GetTypeHash(*Value.GetV1());
+			}
+			else
+			{
+				return GetTypeHash(Value.GetV2());
+			}
 		}
+		return INDEX_NONE;
 	}
 
 	static FUniqueNetIdWrapper Invalid()
@@ -277,8 +516,32 @@ struct FUniqueNetIdWrapper
 
 	friend bool operator==(const FUniqueNetIdWrapper& Lhs, const FUniqueNetIdWrapper& Rhs)
 	{
-		bool bLhsValid = Lhs.IsValid();
-		return bLhsValid == Rhs.IsValid() && (!bLhsValid || *Lhs == *Rhs);
+		const bool bLhsValid = Lhs.IsValid();
+		if (bLhsValid != Rhs.IsValid())
+		{
+			// Different validity
+			return false;
+		}
+		if (!bLhsValid)
+		{
+			// Both invalid
+			return true;
+		}
+		if (Lhs.Variant.GetIndex() != Rhs.Variant.GetIndex())
+		{
+			// Different variant
+			return false;
+		}
+
+		if (Lhs.IsV1())
+		{
+			// Pointers can point to equivalent objects
+			return *Lhs.GetV1() == *Rhs.GetV1();
+		}
+		else
+		{
+			return Lhs.GetV2() == Rhs.GetV2();
+		}
 	}
 
 	friend bool operator!=(const FUniqueNetIdWrapper& Lhs, const FUniqueNetIdWrapper& Rhs)
@@ -288,8 +551,23 @@ struct FUniqueNetIdWrapper
 
 	friend bool operator==(const FUniqueNetIdWrapper& Lhs, const FUniqueNetId& Rhs)
 	{
-		bool bLhsValid = Lhs.IsValid();
-		return bLhsValid == Rhs.IsValid() && (!bLhsValid || *Lhs == Rhs);
+		const bool bLhsValid = Lhs.IsValid();
+		if (bLhsValid != Rhs.IsValid())
+		{
+			// Different validity
+			return false;
+		}
+		if (!bLhsValid)
+		{
+			// Both invalid
+			return true;
+		}
+		if (!Lhs.IsV1())
+		{
+			// Different variant
+			return false;
+		}
+		return *Lhs.GetV1() == Rhs;
 	}
 
 	friend bool operator!=(const FUniqueNetIdWrapper& Lhs, const FUniqueNetId& Rhs)
@@ -331,7 +609,7 @@ struct FUniqueNetIdWrapper
 protected:
 
 	// Actual unique id
-	FUniqueNetIdPtr UniqueNetId;
+	FVariantType Variant;
 };
 
 template <typename ValueType>
diff --git a/Engine/Source/Runtime/Engine/Classes/GameFramework/OnlineReplStructs.h b/Engine/Source/Runtime/Engine/Classes/GameFramework/OnlineReplStructs.h
--- a/Engine/Source/Runtime/Engine/Classes/GameFramework/OnlineReplStructs.h
+++ b/Engine/Source/Runtime/Engine/Classes/GameFramework/OnlineReplStructs.h
@@ -13,6 +13,7 @@
 #include "OnlineReplStructs.generated.h"
 
 class FJsonValue;
+enum class EUniqueIdEncodingFlags : uint8;
 
 /**
  * Wrapper for opaque type FUniqueNetId
@@ -30,7 +31,7 @@ struct FUniqueNetIdRepl : public FUniqueNetIdWrapper
 	}
 
 	FUniqueNetIdRepl(const FUniqueNetIdRepl& InWrapper)
-		: FUniqueNetIdWrapper(InWrapper.UniqueNetId)
+		: FUniqueNetIdWrapper(InWrapper.Variant)
 	{
 	}
 
@@ -51,10 +52,16 @@ struct FUniqueNetIdRepl : public FUniqueNetIdWrapper
 
 	virtual ~FUniqueNetIdRepl() {}
 
-	virtual void SetUniqueNetId(const FUniqueNetIdPtr& InUniqueNetId) override
+	virtual void SetUniqueNetId(const FUniqueNetIdPtr& UniqueNetId) override
 	{
 		ReplicationBytes.Empty();
-		FUniqueNetIdWrapper::SetUniqueNetId(InUniqueNetId);
+		FUniqueNetIdWrapper::SetUniqueNetId(UniqueNetId);
+	}
+
+	virtual void SetAccountId(const UE::Online::FOnlineAccountIdHandle& Handle) override
+	{
+		ReplicationBytes.Empty();
+		FUniqueNetIdWrapper::SetAccountId(Handle);
 	}
 
 	/** Export contents of this struct as a string */
@@ -66,6 +73,10 @@ struct FUniqueNetIdRepl : public FUniqueNetIdWrapper
 	/** Network serialization */
 	ENGINE_API bool NetSerialize(FArchive& Ar, class UPackageMap* Map, bool& bOutSuccess);
 
+	void NetSerializeLoadV1Encoded(FArchive& Ar, const EUniqueIdEncodingFlags EncodingFlags, bool& bOutSuccess);
+	void NetSerializeLoadV1Unencoded(FArchive& Ar, const EUniqueIdEncodingFlags EncodingFlags, bool& bOutSuccess);
+	void NetSerializeLoadV2(FArchive& Ar, const EUniqueIdEncodingFlags EncodingFlags, bool& bOutSuccess);
+
 	/** Serialization to any FArchive */
 	ENGINE_API friend FArchive& operator<<( FArchive& Ar, FUniqueNetIdRepl& UniqueNetId);
 
@@ -99,6 +110,8 @@ protected:
 	void UniqueIdFromString(FName Type, const FString& Contents);
 	/** Helper to make network serializable representation */
 	void MakeReplicationData();
+	void MakeReplicationDataV1();
+	void MakeReplicationDataV2();
 	/** Network serialized data cache */
 	UPROPERTY(Transient)
 	TArray<uint8> ReplicationBytes;
diff --git a/Engine/Source/Runtime/Engine/Private/OnlineReplStructs.cpp b/Engine/Source/Runtime/Engine/Private/OnlineReplStructs.cpp
--- a/Engine/Source/Runtime/Engine/Private/OnlineReplStructs.cpp
+++ b/Engine/Source/Runtime/Engine/Private/OnlineReplStructs.cpp
@@ -44,6 +44,8 @@ ENUM_CLASS_FLAGS(EUniqueIdEncodingFlags);
 
 /** Use highest value for type for other (out of engine) oss type */
 const uint8 TypeHash_Other = 31;
+/** Use next highest value for V2 net id */
+const uint8 TypeHash_V2 = 30;
 
 FArchive& operator<<(FArchive& Ar, FUniqueNetIdRepl& UniqueNetId)
 {
@@ -108,17 +110,32 @@ inline uint8 GetTypeHashFromEncoding(EUniqueIdEncodingFlags inFlags)
  */
 void FUniqueNetIdRepl::MakeReplicationData()
 {
-	//LOG_SCOPE_VERBOSITY_OVERRIDE(LogNet, ELogVerbosity::VeryVerbose);
-	//UE_LOG(LogNet, VeryVerbose, TEXT("MakeReplicationData %s"), *ToString());
-
-	FString Contents;
 	if (IsValid())
 	{
-		Contents = UniqueNetId->ToString();
+		if (IsV1())
+		{
+			MakeReplicationDataV1();
+		}
+		else
+		{
+			MakeReplicationDataV2();
+		}
 	}
+	else
+	{
+		EUniqueIdEncodingFlags EncodingFlags = (EUniqueIdEncodingFlags::IsEncoded | EUniqueIdEncodingFlags::IsEmpty);
 
+		ReplicationBytes.Empty();
+		FMemoryWriter Writer(ReplicationBytes);
+		Writer << EncodingFlags;
+	}
+}
+
+void FUniqueNetIdRepl::MakeReplicationDataV1()
+{
+	FString Contents = GetUniqueNetId()->ToString();
 	const int32 Length = Contents.Len();
-	if (Length > 0)
+	if (ensure(Length > 0))
 	{
 		// For now don't allow odd chars (HexToBytes adds a 0)
 		const bool bEvenChars = (Length % 2) == 0;
@@ -202,10 +219,32 @@ void FUniqueNetIdRepl::MakeReplicationData()
 		ReplicationBytes.Empty();
 		FMemoryWriter Writer(ReplicationBytes);
 		Writer << EncodingFlags;
-		//UE_LOG(LogNet, VeryVerbose, TEXT("Empty/Invalid UniqueId, serializing %d bytes"), ReplicationBytes.Num());
 	}
 }
 
+void FUniqueNetIdRepl::MakeReplicationDataV2()
+{
+	EUniqueIdEncodingFlags EncodingFlags = EUniqueIdEncodingFlags::IsEncoded;
+	EncodingFlags |= static_cast<EUniqueIdEncodingFlags>(TypeHash_V2 << 3);
+
+	UE::Online::FOnlineAccountIdHandle Handle = GetV2();
+	UE::Online::EOnlineServices OnlineServicesType = Handle.GetOnlineServicesType();
+
+	TArray<uint8> ReplicationData;
+	const UE::Online::IOnlineAccountIdRegistry* Registry = UE::Online::FOnlineIdRegistryRegistry::Get().GetAccountIdRegistry(OnlineServicesType);
+	if (ensure(Registry))
+	{
+		ReplicationData = Registry->ToReplicationData(Handle);
+		check(!ReplicationData.IsEmpty());
+	}
+
+	ReplicationBytes.Empty();
+	FMemoryWriter Writer(ReplicationBytes);
+	Writer << EncodingFlags;
+	Writer << OnlineServicesType;
+	Writer << ReplicationData;
+}
+
 void FUniqueNetIdRepl::UniqueIdFromString(FName Type, const FString& Contents)
 {
 	// Don't need to distinguish OSS interfaces here with world because we just want the create function below
@@ -231,92 +270,25 @@ bool FUniqueNetIdRepl::NetSerialize(FArchive& Ar, UPackageMap* Map, bool& bOutSu
 	else if (Ar.IsLoading())
 	{
 		// @note: start by assuming a replicated nullptr unique id
-		UniqueNetId.Reset();
+		SetUniqueNetId(nullptr);
 
 		EUniqueIdEncodingFlags EncodingFlags = EUniqueIdEncodingFlags::NotEncoded;
 		Ar << EncodingFlags;
 		if (!Ar.IsError())
 		{
+			
 			if (EnumHasAllFlags(EncodingFlags, EUniqueIdEncodingFlags::IsEncoded))
 			{
 				if (!EnumHasAllFlags(EncodingFlags, EUniqueIdEncodingFlags::IsEmpty))
 				{
-					// Non empty and hex encoded
-					uint8 TypeHash = GetTypeHashFromEncoding(EncodingFlags);
-					if (TypeHash == 0)
-					{
-						// If no type was encoded, assume default
-						TypeHash = UOnlineEngineInterface::Get()->GetReplicationHashForSubsystem(UOnlineEngineInterface::Get()->GetDefaultOnlineSubsystemName());
-					}
-					FName Type;
-					bool bValidTypeHash = TypeHash != 0;
-					if (TypeHash == TypeHash_Other)
-					{
-						FString TypeString;
-						Ar << TypeString;
-						Type = FName(*TypeString);
-						if (Ar.IsError() || Type == NAME_None)
-						{
-							bValidTypeHash = false;
-						}
-					}
-					else
-					{
-						Type = UOnlineEngineInterface::Get()->GetSubsystemFromReplicationHash(TypeHash);
-					}
-
-					if (bValidTypeHash)
+					const uint8 TypeHash = GetTypeHashFromEncoding(EncodingFlags);
+					if (TypeHash == TypeHash_V2)
 					{
-						// Get the size
-						uint8 EncodedSize = 0;
-						Ar << EncodedSize;
-						if (!Ar.IsError())
-						{
-							if (EncodedSize > 0)
-							{
-								uint8* TempBytes = (uint8*)FMemory_Alloca(EncodedSize);
-								Ar.Serialize(TempBytes, EncodedSize);
-								if (!Ar.IsError())
-								{
-									FString Contents = BytesToHex(TempBytes, EncodedSize);
-									if (Contents.Len() > 0)
-									{
-										if (Type != NAME_None)
-										{
-											// BytesToHex loses case
-											Contents.ToLowerInline();
-											if (EnumHasAllFlags(EncodingFlags, EUniqueIdEncodingFlags::IsPadded))
-											{
-												Contents.RightChopInline(1); // remove padded character
-											}
-											UniqueIdFromString(Type, Contents);
-										}
-										else
-										{
-											UE_LOG(LogNet, Warning, TEXT("Error with unique id type"));
-										}
-									}
-								}
-								else
-								{
-									UE_LOG(LogNet, Warning, TEXT("Error with encoded unique id contents"));
-								}
-							}
-							else
-							{
-								UE_LOG(LogNet, Warning, TEXT("Empty Encoding!"));
-							}
-
-							bOutSuccess = (EncodedSize == 0) || IsValid();
-						}
-						else
-						{
-							UE_LOG(LogNet, Warning, TEXT("Error with encoded unique id size"));
-						}
+						NetSerializeLoadV2(Ar, EncodingFlags, bOutSuccess);
 					}
 					else
 					{
-						UE_LOG(LogNet, Warning, TEXT("Error with encoded type hash"));
+						NetSerializeLoadV1Encoded(Ar, EncodingFlags, bOutSuccess);
 					}
 				}
 				else
@@ -327,64 +299,165 @@ bool FUniqueNetIdRepl::NetSerialize(FArchive& Ar, UPackageMap* Map, bool& bOutSu
 			}
 			else
 			{
-				// Original FString serialization goes here
-				uint8 TypeHash = GetTypeHashFromEncoding(EncodingFlags);
-				if (TypeHash == 0)
-				{
-					// If no type was encoded, assume default
-					TypeHash = UOnlineEngineInterface::Get()->GetReplicationHashForSubsystem(UOnlineEngineInterface::Get()->GetDefaultOnlineSubsystemName());
-				}
-				FName Type;
-				bool bValidTypeHash = TypeHash != 0;
-				if (TypeHash == TypeHash_Other)
-				{
-					FString TypeString;
-					Ar << TypeString;
-					Type = FName(*TypeString);
-					if (Ar.IsError() || Type == NAME_None)
-					{
-						bValidTypeHash = false;
-					}
-				}
-				else
-				{
-					Type = UOnlineEngineInterface::Get()->GetSubsystemFromReplicationHash(TypeHash);
-				}
+				NetSerializeLoadV1Unencoded(Ar, EncodingFlags, bOutSuccess);
+			}
+		}
+		else
+		{
+			UE_LOG(LogNet, Warning, TEXT("Error serializing unique id"));
+		}
+	}
+
+	return true;
+}
+
+void FUniqueNetIdRepl::NetSerializeLoadV1Encoded(FArchive& Ar, const EUniqueIdEncodingFlags EncodingFlags, bool& bOutSuccess)
+{
+	// Non empty and hex encoded
+	uint8 TypeHash = GetTypeHashFromEncoding(EncodingFlags);
+	if (TypeHash == 0)
+	{
+		// If no type was encoded, assume default
+		TypeHash = UOnlineEngineInterface::Get()->GetReplicationHashForSubsystem(UOnlineEngineInterface::Get()->GetDefaultOnlineSubsystemName());
+	}
+	FName Type;
+	bool bValidTypeHash = TypeHash != 0;
+	if (TypeHash == TypeHash_Other)
+	{
+		FString TypeString;
+		Ar << TypeString;
+		Type = FName(*TypeString);
+		if (Ar.IsError() || Type == NAME_None)
+		{
+			bValidTypeHash = false;
+		}
+	}
+	else
+	{
+		Type = UOnlineEngineInterface::Get()->GetSubsystemFromReplicationHash(TypeHash);
+	}
 
-				if (bValidTypeHash)
+	if (bValidTypeHash)
+	{
+		// Get the size
+		uint8 EncodedSize = 0;
+		Ar << EncodedSize;
+		if (!Ar.IsError())
+		{
+			if (EncodedSize > 0)
+			{
+				uint8* TempBytes = (uint8*)FMemory_Alloca(EncodedSize);
+				Ar.Serialize(TempBytes, EncodedSize);
+				if (!Ar.IsError())
 				{
-					FString Contents;
-					Ar << Contents;
-					if (!Ar.IsError())
+					FString Contents = BytesToHex(TempBytes, EncodedSize);
+					if (Contents.Len() > 0)
 					{
 						if (Type != NAME_None)
 						{
+							// BytesToHex loses case
+							Contents.ToLowerInline();
+							if (EnumHasAllFlags(EncodingFlags, EUniqueIdEncodingFlags::IsPadded))
+							{
+								Contents.RightChopInline(1); // remove padded character
+							}
 							UniqueIdFromString(Type, Contents);
-							bOutSuccess = !Contents.IsEmpty();
 						}
 						else
 						{
 							UE_LOG(LogNet, Warning, TEXT("Error with unique id type"));
 						}
 					}
-					else
-					{
-						UE_LOG(LogNet, Warning, TEXT("Error with unencoded unique id"));
-					}
 				}
 				else
 				{
-					UE_LOG(LogNet, Warning, TEXT("Error with encoded type hash"));
+					UE_LOG(LogNet, Warning, TEXT("Error with encoded unique id contents"));
 				}
 			}
+			else
+			{
+				UE_LOG(LogNet, Warning, TEXT("Empty Encoding!"));
+			}
+
+			bOutSuccess = (EncodedSize == 0) || IsValid();
 		}
 		else
 		{
-			UE_LOG(LogNet, Warning, TEXT("Error serializing unique id"));
+			UE_LOG(LogNet, Warning, TEXT("Error with encoded unique id size"));
 		}
 	}
+	else
+	{
+		UE_LOG(LogNet, Warning, TEXT("Error with encoded type hash"));
+	}
+}
 
-	return true;
+void FUniqueNetIdRepl::NetSerializeLoadV1Unencoded(FArchive& Ar, const EUniqueIdEncodingFlags EncodingFlags, bool& bOutSuccess)
+{
+	uint8 TypeHash = GetTypeHashFromEncoding(EncodingFlags);
+	if (TypeHash == 0)
+	{
+		// If no type was encoded, assume default
+		TypeHash = UOnlineEngineInterface::Get()->GetReplicationHashForSubsystem(UOnlineEngineInterface::Get()->GetDefaultOnlineSubsystemName());
+	}
+	FName Type;
+	bool bValidTypeHash = TypeHash != 0;
+	if (TypeHash == TypeHash_Other)
+	{
+		FString TypeString;
+		Ar << TypeString;
+		Type = FName(*TypeString);
+		if (Ar.IsError() || Type == NAME_None)
+		{
+			bValidTypeHash = false;
+		}
+	}
+	else
+	{
+		Type = UOnlineEngineInterface::Get()->GetSubsystemFromReplicationHash(TypeHash);
+	}
+
+	if (bValidTypeHash)
+	{
+		FString Contents;
+		Ar << Contents;
+		if (!Ar.IsError())
+		{
+			if (Type != NAME_None)
+			{
+				UniqueIdFromString(Type, Contents);
+				bOutSuccess = !Contents.IsEmpty();
+			}
+			else
+			{
+				UE_LOG(LogNet, Warning, TEXT("Error with unique id type"));
+			}
+		}
+		else
+		{
+			UE_LOG(LogNet, Warning, TEXT("Error with unencoded unique id"));
+		}
+	}
+	else
+	{
+		UE_LOG(LogNet, Warning, TEXT("Error with encoded type hash"));
+	}
+}
+
+void FUniqueNetIdRepl::NetSerializeLoadV2(FArchive& Ar, const EUniqueIdEncodingFlags EncodingFlags, bool& bOutSuccess)
+{
+	UE::Online::EOnlineServices OnlineServicesType;
+	Ar << OnlineServicesType;
+	TArray<uint8> ReplicationData;
+	Ar << ReplicationData;
+
+	UE::Online::IOnlineAccountIdRegistry* Registry = UE::Online::FOnlineIdRegistryRegistry::Get().GetAccountIdRegistry(OnlineServicesType);
+	if (ensure(Registry))
+	{
+		const UE::Online::FOnlineAccountIdHandle Handle = Registry->FromReplicationData(ReplicationData);
+		check(Handle.IsValid());
+		SetAccountId(Handle);
+	}
 }
 
 bool FUniqueNetIdRepl::Serialize(FArchive& Ar)
@@ -402,14 +475,21 @@ bool FUniqueNetIdRepl::ExportTextItem(FString& ValueStr, FUniqueNetIdRepl const&
 
 	if (IsValid())
 	{
-		FName Type = UniqueNetId->GetType();
-		if (Type == UOnlineEngineInterface::Get()->GetDefaultOnlineSubsystemName())
+		if (IsV1())
 		{
-			ValueStr += FString::Printf(TEXT("%s"), *UniqueNetId->ToString());
+			FName Type = GetType();
+			if (Type == UOnlineEngineInterface::Get()->GetDefaultOnlineSubsystemName())
+			{
+				ValueStr += ToString();
+			}
+			else
+			{
+				ValueStr += FString::Printf(TEXT("%s:%s"), *Type.ToString(), *ToString());
+			}
 		}
 		else
 		{
-			ValueStr += FString::Printf(TEXT("%s:%s"), *Type.ToString(), *UniqueNetId->ToString());
+			ValueStr += ToDebugString();
 		}
 	}
 	else
@@ -484,7 +564,7 @@ TSharedRef<FJsonValue> FUniqueNetIdRepl::ToJson() const
 {
 	if (IsValid())
 	{
-		const FString JsonString = FString::Printf(TEXT("%s:%s"), *UniqueNetId->GetType().ToString(), *ToString());
+		const FString JsonString = FString::Printf(TEXT("%s:%s"), *GetType().ToString(), *ToString());
 		return MakeShareable(new FJsonValueString(JsonString));
 	}
 	else
