Fix crash for unassigned RVT sample parameter material expression.
Compilation code changed to need a texture, so make this an error condition now.

diff --git a/Engine/Source/Runtime/Engine/Private/Materials/MaterialExpressions.cpp b/Engine/Source/Runtime/Engine/Private/Materials/MaterialExpressions.cpp
--- a/Engine/Source/Runtime/Engine/Private/Materials/MaterialExpressions.cpp
+++ b/Engine/Source/Runtime/Engine/Private/Materials/MaterialExpressions.cpp
@@ -2387,9 +2387,9 @@ int32 UMaterialExpressionRuntimeVirtualTextureSample::Compile(class FMaterialCom
 	bool bIsVirtualTextureValid = VirtualTexture != nullptr;
 	if (!bIsVirtualTextureValid)
 	{
-		if (!bIsParameter)
+		if (bIsParameter)
 		{
-			Compiler->Error(TEXT("Missing input Virtual Texture"));
+			return Compiler->Error(TEXT("Missing input Virtual Texture"));
 		}
 	}
 	else if (VirtualTexture->GetMaterialType() != MaterialType)
@@ -2450,7 +2450,7 @@ int32 UMaterialExpressionRuntimeVirtualTextureSample::Compile(class FMaterialCom
 	switch (OutputIndex)
 	{
 	case 0: 
-		if ((bIsParameter || bIsVirtualTextureValid) && bIsBaseColorValid)
+		if (bIsVirtualTextureValid && bIsBaseColorValid)
 		{
 			switch (MaterialType)
 			{
@@ -2466,7 +2466,7 @@ int32 UMaterialExpressionRuntimeVirtualTextureSample::Compile(class FMaterialCom
 		}
 		break;
 	case 1:
-		if ((bIsParameter || bIsVirtualTextureValid) && bIsSpecularValid)
+		if (bIsVirtualTextureValid && bIsSpecularValid)
 		{
 			switch (MaterialType)
 			{
@@ -2481,7 +2481,7 @@ int32 UMaterialExpressionRuntimeVirtualTextureSample::Compile(class FMaterialCom
 		}
 		break;
 	case 2:
-		if ((bIsParameter || bIsVirtualTextureValid) && bIsRoughnessValid)
+		if (bIsVirtualTextureValid && bIsRoughnessValid)
 		{
 			switch (MaterialType)
 			{
@@ -2497,7 +2497,7 @@ int32 UMaterialExpressionRuntimeVirtualTextureSample::Compile(class FMaterialCom
 		}
 		break;
 	case 3:
-		if ((bIsParameter || bIsVirtualTextureValid) && bIsNormalValid)
+		if (bIsVirtualTextureValid && bIsNormalValid)
 		{
 			switch (MaterialType)
 			{
@@ -2513,7 +2513,7 @@ int32 UMaterialExpressionRuntimeVirtualTextureSample::Compile(class FMaterialCom
 		}
 		break;
 	case 4:
-		if ((bIsParameter || bIsVirtualTextureValid) && bIsWorldHeightValid)
+		if (bIsVirtualTextureValid && bIsWorldHeightValid)
 		{
 			UnpackType = EVirtualTextureUnpackType::HeightR16;
 		}
@@ -2523,7 +2523,7 @@ int32 UMaterialExpressionRuntimeVirtualTextureSample::Compile(class FMaterialCom
 		}
 		break;
 	case 5:
-		if ((bIsParameter || bIsVirtualTextureValid) && bIsMaskValid)
+		if (bIsVirtualTextureValid && bIsMaskValid)
 		{
 			UnpackTarget = 2; UnpackMask = 0x8; break;
 		}
