- Changed the PreviewAllowlistCvars code to be a prefix instead of full CVar
- CHange PreviewAllowlistCvars (for Android and IOS) from None to sg. so that scalability groups can be previewed
- Fixed the PreviewAllowlistCVars/PreviewDenylistCVars value lookup to use the proper ConfigSystem instead of GConfig
#p4v-preflight-copy 18533519

diff --git a/Engine/Config/BaseDeviceProfiles.ini b/Engine/Config/BaseDeviceProfiles.ini
--- a/Engine/Config/BaseDeviceProfiles.ini
+++ b/Engine/Config/BaseDeviceProfiles.ini
@@ -245,9 +245,9 @@ BaseProfileName=
 +CVars=sg.TextureQuality=3
 +CVars=sg.EffectsQuality=3
 +CVars=sg.FoliageQuality=3
-; PreviewAllowlistCVars and PreviewDenyListCVars are arrays of cvars that are included or excluded from being applied in mobile preview.
+; PreviewAllowlistCVars and PreviewDenyListCVars are arrays of cvars (or prefixes) that are included or excluded from being applied in editor preview
 ; If any PreviewAllowlistCVars is set, cvars are denied by default.
-PreviewAllowlistCVars=None
+PreviewAllowlistCVars=sg.
 
 [iPad2 DeviceProfile]
 DeviceType=IOS
@@ -649,9 +649,9 @@ BaseProfileName=
 +CVars=r.Android.DisableVulkanSM5Support=1
 ; PF_B8G8R8A8
 +CVars=r.DefaultBackBufferPixelFormat=0 
-; PreviewAllowlistCVars and PreviewDenyListCVars are arrays of cvars that are included or excluded from being applied in mobile preview.
+; PreviewAllowlistCVars and PreviewDenyListCVars are arrays of cvars (or prefixes) that are included or excluded from being applied in editor preview.
 ; If any PreviewAllowlistCVars is set, cvars are denied by default.
-PreviewAllowlistCVars=None
+PreviewAllowlistCVars=sg.
 
 [Android_Low DeviceProfile]
 DeviceType=Android
diff --git a/Engine/Source/Runtime/Engine/Private/DeviceProfiles/DeviceProfileManager.cpp b/Engine/Source/Runtime/Engine/Private/DeviceProfiles/DeviceProfileManager.cpp
--- a/Engine/Source/Runtime/Engine/Private/DeviceProfiles/DeviceProfileManager.cpp
+++ b/Engine/Source/Runtime/Engine/Private/DeviceProfiles/DeviceProfileManager.cpp
@@ -245,8 +245,8 @@ void UDeviceProfileManager::ProcessDeviceProfileIniSettings(const FString& Devic
 	FString SectionSuffix = *FString::Printf(TEXT(" %s"), *UDeviceProfile::StaticClass()->GetName());
 
 #if WITH_EDITOR
-	TSet<FString> PreviewAllowlistCVars;
-	TSet<FString> PreviewDenylistCVars;
+	TArray<FString> PreviewAllowlistCVars;
+	TArray<FString> PreviewDenylistCVars;
 	bool bFoundAllowDeny = false;
 	if (Mode == EDeviceProfileMode::DPM_CacheValues)
 	{
@@ -255,19 +255,8 @@ void UDeviceProfileManager::ProcessDeviceProfileIniSettings(const FString& Devic
 			PreviewAllowlistCVars.Num()==0 && PreviewDenylistCVars.Num()==0 && !CurrentProfileName.IsEmpty() && AvailableProfiles.Contains(CurrentSectionName);
 			CurrentProfileName = GConfig->GetStr(*CurrentSectionName, TEXT("BaseProfileName"), GDeviceProfilesIni), CurrentSectionName = CurrentProfileName + SectionSuffix)
 		{
-			TArray<FString> TempAllowlist;
-			GConfig->GetArray(*CurrentSectionName, TEXT("PreviewAllowlistCVars"), TempAllowlist, GDeviceProfilesIni);
-			for( FString& Item : TempAllowlist)
-			{
-				PreviewAllowlistCVars.Add(Item);
-			}
-
-			TArray<FString> TempDenylist;
-			GConfig->GetArray(*CurrentSectionName, TEXT("PreviewDenylistCVars"), TempDenylist, GDeviceProfilesIni);
-			for (FString& Item : TempDenylist)
-			{
-				PreviewDenylistCVars.Add(Item);
-			}
+			ConfigSystem->GetArray(*CurrentSectionName, TEXT("PreviewAllowlistCVars"), PreviewAllowlistCVars, GDeviceProfilesIni);
+			ConfigSystem->GetArray(*CurrentSectionName, TEXT("PreviewDenylistCVars"), PreviewDenylistCVars, GDeviceProfilesIni);
 		}
 	}
 #endif
@@ -406,13 +395,13 @@ void UDeviceProfileManager::ProcessDeviceProfileIniSettings(const FString& Devic
 #if WITH_EDITOR
 							if (Mode == EDeviceProfileMode::DPM_CacheValues)
 							{
-								if (PreviewDenylistCVars.Contains(CVarKey))
+								if (PreviewDenylistCVars.ContainsByPredicate([&CVarKey](const FString& Entry) { return CVarKey.StartsWith(Entry); }))
 								{
 									UE_LOG(LogInit, Log, TEXT("Skipping Device Profile CVar due to PreviewDenylistCVars: [[%s]]"), *CVarKey);
 									continue;
 								}
 
-								if (PreviewAllowlistCVars.Num() > 0 && !PreviewAllowlistCVars.Contains(CVarKey))
+								if (PreviewAllowlistCVars.Num() > 0 && !PreviewAllowlistCVars.ContainsByPredicate([&CVarKey](const FString& Entry) { return CVarKey.StartsWith(Entry); }))
 								{
 									UE_LOG(LogInit, Log, TEXT("Skipping Device Profile CVar due to PreviewAllowlistCVars: [[%s]]"), *CVarKey);
 									continue;
