World Partition handles
- Actor descriptor handles can now detect if their owning container is valid or not.
- Changed unit tests to use the new unit test map /Engine/WorldPartition/WorldPartitionUnitTest.
- Added the destructed container test case in unit tests.
- Removed all the unit test specific cases from actor descriptors code.
- Removed all code clearing actor descriptor handles in BeginDestroy.

diff --git a/Engine/Source/Runtime/Engine/Private/Tests/WorldPartition/WorldPartitionHandleTests.cpp b/Engine/Source/Runtime/Engine/Private/Tests/WorldPartition/WorldPartitionHandleTests.cpp
--- a/Engine/Source/Runtime/Engine/Private/Tests/WorldPartition/WorldPartitionHandleTests.cpp
+++ b/Engine/Source/Runtime/Engine/Private/Tests/WorldPartition/WorldPartitionHandleTests.cpp
@@ -3,6 +3,7 @@
 #include "Misc/AutomationTest.h"
 #include "WorldPartition/WorldPartitionHandle.h"
 #include "WorldPartition/WorldPartitionActorDesc.h"
+#include "WorldPartition/ActorDescContainer.h"
 
 #if WITH_DEV_AUTOMATION_TESTS
 
@@ -16,208 +17,215 @@ namespace WorldPartitionTests
 	bool FWorldPartitionSoftRefTest::RunTest(const FString& Parameters)
 	{
 #if WITH_EDITOR
-		TestTrue(TEXT("SoftRef sizeof"), sizeof(FWorldPartitionHandle) == sizeof(void*));
-		TestTrue(TEXT("HardRef sizeof"), sizeof(FWorldPartitionReference) == sizeof(void*));
+		FWorldPartitionHandle Handle;
+		FWorldPartitionReference Reference;
 
-		TUniquePtr<FWorldPartitionActorDesc> ActorDescPtr1(new FWorldPartitionActorDesc());
-		TUniquePtr<FWorldPartitionActorDesc> ActorDescPtr2(new FWorldPartitionActorDesc());
+		{
+			UActorDescContainer* ActorDescContainer = NewObject<UActorDescContainer>(GetTransientPackage());
+			ActorDescContainer->Initialize(nullptr, TEXT("/Engine/WorldPartition/WorldPartitionUnitTest"));
 
-		FWorldPartitionHandle Handle(&ActorDescPtr1);
-		FWorldPartitionReference Reference(&ActorDescPtr2);
+			Handle = FWorldPartitionHandle(ActorDescContainer, FGuid(TEXT("5D9F93BA407A811AFDDDAAB4F1CECC6A")));
+			Reference = FWorldPartitionReference(ActorDescContainer, FGuid(TEXT("0D2B04D240BE5DE58FE437A8D2DBF5C9")));
 
-		TestTrue(TEXT("Handle soft refcount"), Handle->GetSoftRefCount() == 1);
-		TestTrue(TEXT("Handle hard refcount"), Handle->GetHardRefCount() == 0);
-		TestTrue(TEXT("Reference soft refcount"), Reference->GetSoftRefCount() == 0);
-		TestTrue(TEXT("Reference hard refcount"), Reference->GetHardRefCount() == 1);
+			TestTrue(TEXT("Handle soft refcount"), Handle->GetSoftRefCount() == 1);
+			TestTrue(TEXT("Handle hard refcount"), Handle->GetHardRefCount() == 0);
+			TestTrue(TEXT("Reference soft refcount"), Reference->GetSoftRefCount() == 0);
+			TestTrue(TEXT("Reference hard refcount"), Reference->GetHardRefCount() == 1);
 
-		// Pin handle scope test
-		{
-			FWorldPartitionHandlePinRefScope PinRefScopeHandle(Handle);
-			TestTrue(TEXT("Pin to Handle soft refcount"), Handle->GetSoftRefCount() == 1);
-			TestTrue(TEXT("Pin to Handle hard refcount"), Handle->GetHardRefCount() == 0);
-		}
+			// Pin handle scope test
+			{
+				FWorldPartitionHandlePinRefScope PinRefScopeHandle(Handle);
+				TestTrue(TEXT("Pin to Handle soft refcount"), Handle->GetSoftRefCount() == 1);
+				TestTrue(TEXT("Pin to Handle hard refcount"), Handle->GetHardRefCount() == 0);
+			}
 
-		// Pin reference scope test
-		{
-			FWorldPartitionHandlePinRefScope PinRefScopeReference(Reference);
-			TestTrue(TEXT("Pin to Reference soft refcount"), Reference->GetSoftRefCount() == 0);
-			TestTrue(TEXT("Pin to Reference hard refcount"), Reference->GetHardRefCount() == 2);
-		}
+			// Pin reference scope test
+			{
+				FWorldPartitionHandlePinRefScope PinRefScopeReference(Reference);
+				TestTrue(TEXT("Pin to Reference soft refcount"), Reference->GetSoftRefCount() == 0);
+				TestTrue(TEXT("Pin to Reference hard refcount"), Reference->GetHardRefCount() == 2);
+			}
 
-		TestTrue(TEXT("Handle/Reference equality"), Handle != Reference);
-		TestTrue(TEXT("Reference/Handle equality"), Reference != Handle);
+			TestTrue(TEXT("Handle/Reference equality"), Handle != Reference);
+			TestTrue(TEXT("Reference/Handle equality"), Reference != Handle);
 		
-		TestTrue(TEXT("Handle soft refcount"), Handle->GetSoftRefCount() == 1);
-		TestTrue(TEXT("Handle hard refcount"), Handle->GetHardRefCount() == 0);
-		TestTrue(TEXT("Reference soft refcount"), Reference->GetSoftRefCount() == 0);
-		TestTrue(TEXT("Reference hard refcount"), Reference->GetHardRefCount() == 1);
-
-		// Conversions
-		{
-			FWorldPartitionHandle HandleToReference = Reference;
-			TestTrue(TEXT("Handle/Reference equality"), HandleToReference == Reference);
-			TestTrue(TEXT("Reference/Handle equality"), Reference == HandleToReference);
-			TestTrue(TEXT("Reference soft refcount"), Reference->GetSoftRefCount() == 1);
+			TestTrue(TEXT("Handle soft refcount"), Handle->GetSoftRefCount() == 1);
+			TestTrue(TEXT("Handle hard refcount"), Handle->GetHardRefCount() == 0);
+			TestTrue(TEXT("Reference soft refcount"), Reference->GetSoftRefCount() == 0);
 			TestTrue(TEXT("Reference hard refcount"), Reference->GetHardRefCount() == 1);
 
-			FWorldPartitionReference ReferenceToHandle = Handle;
-			TestTrue(TEXT("Handle/Reference equality"), ReferenceToHandle == Handle);
-			TestTrue(TEXT("Handle/Reference equality"), Handle == ReferenceToHandle);
-			TestTrue(TEXT("Reference soft refcount"), Reference->GetSoftRefCount() == 1);
-			TestTrue(TEXT("Reference hard refcount"), Reference->GetHardRefCount() == 1);
-		}
+			// Conversions
+			{
+				FWorldPartitionHandle HandleToReference = Reference;
+				TestTrue(TEXT("Handle/Reference equality"), HandleToReference == Reference);
+				TestTrue(TEXT("Reference/Handle equality"), Reference == HandleToReference);
+				TestTrue(TEXT("Reference soft refcount"), Reference->GetSoftRefCount() == 1);
+				TestTrue(TEXT("Reference hard refcount"), Reference->GetHardRefCount() == 1);
 
-		TestTrue(TEXT("Handle soft refcount"), Handle->GetSoftRefCount() == 1);
-		TestTrue(TEXT("Handle hard refcount"), Handle->GetHardRefCount() == 0);
-		TestTrue(TEXT("Reference soft refcount"), Reference->GetSoftRefCount() == 0);
-		TestTrue(TEXT("Reference hard refcount"), Reference->GetHardRefCount() == 1);
+				FWorldPartitionReference ReferenceToHandle = Handle;
+				TestTrue(TEXT("Handle/Reference equality"), ReferenceToHandle == Handle);
+				TestTrue(TEXT("Handle/Reference equality"), Handle == ReferenceToHandle);
+				TestTrue(TEXT("Reference soft refcount"), Reference->GetSoftRefCount() == 1);
+				TestTrue(TEXT("Reference hard refcount"), Reference->GetHardRefCount() == 1);
+			}
 
-		// inplace new test
-		{
-			uint8 Buffer[sizeof(FWorldPartitionHandle)];
-			FWorldPartitionHandle* HandlePtr = new (Buffer) FWorldPartitionHandle(Reference);
+			TestTrue(TEXT("Handle soft refcount"), Handle->GetSoftRefCount() == 1);
+			TestTrue(TEXT("Handle hard refcount"), Handle->GetHardRefCount() == 0);
+			TestTrue(TEXT("Reference soft refcount"), Reference->GetSoftRefCount() == 0);
+			TestTrue(TEXT("Reference hard refcount"), Reference->GetHardRefCount() == 1);
 
-			TestTrue(TEXT("Handle array soft refcount"), Reference->GetSoftRefCount() == 1);
-			TestTrue(TEXT("Handle array hard refcount"), Reference->GetHardRefCount() == 1);
+			// inplace new test
+			{
+				uint8 Buffer[sizeof(FWorldPartitionHandle)];
+				FWorldPartitionHandle* HandlePtr = new (Buffer) FWorldPartitionHandle(Reference);
 
-			HandlePtr->~FWorldPartitionHandle();
+				TestTrue(TEXT("Handle array soft refcount"), Reference->GetSoftRefCount() == 1);
+				TestTrue(TEXT("Handle array hard refcount"), Reference->GetHardRefCount() == 1);
 
-			TestTrue(TEXT("Handle array soft refcount"), Reference->GetSoftRefCount() == 0);
-			TestTrue(TEXT("Handle array hard refcount"), Reference->GetHardRefCount() == 1);
-		}
+				HandlePtr->~FWorldPartitionHandle();
 
-		// TArray test
-		{
-			TArray<FWorldPartitionHandle> HandleList;
-			HandleList.Add(Handle);
+				TestTrue(TEXT("Handle array soft refcount"), Reference->GetSoftRefCount() == 0);
+				TestTrue(TEXT("Handle array hard refcount"), Reference->GetHardRefCount() == 1);
+			}
+
+			// TArray test
+			{
+				TArray<FWorldPartitionHandle> HandleList;
+				HandleList.Add(Handle);
 
-			TestTrue(TEXT("Handle array soft refcount"), Handle->GetSoftRefCount() == 2);
-			TestTrue(TEXT("Handle array hard refcount"), Handle->GetHardRefCount() == 0);
+				TestTrue(TEXT("Handle array soft refcount"), Handle->GetSoftRefCount() == 2);
+				TestTrue(TEXT("Handle array hard refcount"), Handle->GetHardRefCount() == 0);
 
-			FWorldPartitionReference ReferenceToHandle = Handle;
-			TestTrue(TEXT("Handle/Reference equality"), ReferenceToHandle == Handle);
-			TestTrue(TEXT("Handle/Reference equality"), Handle == ReferenceToHandle);
-			TestTrue(TEXT("Handle soft refcount"), Handle->GetSoftRefCount() == 2);
-			TestTrue(TEXT("Handle hard refcount"), Handle->GetHardRefCount() == 1);
+				FWorldPartitionReference ReferenceToHandle = Handle;
+				TestTrue(TEXT("Handle/Reference equality"), ReferenceToHandle == Handle);
+				TestTrue(TEXT("Handle/Reference equality"), Handle == ReferenceToHandle);
+				TestTrue(TEXT("Handle soft refcount"), Handle->GetSoftRefCount() == 2);
+				TestTrue(TEXT("Handle hard refcount"), Handle->GetHardRefCount() == 1);
 
-			TestTrue(TEXT("Handle array contains handle"), HandleList.Contains(Handle));
-			TestTrue(TEXT("Handle array contains reference"), HandleList.Contains(ReferenceToHandle));
+				TestTrue(TEXT("Handle array contains handle"), HandleList.Contains(Handle));
+				TestTrue(TEXT("Handle array contains reference"), HandleList.Contains(ReferenceToHandle));
 
-			HandleList.Add(Reference);
+				HandleList.Add(Reference);
 			
-			TestTrue(TEXT("Handle array contains reference"), HandleList.Contains(Reference));
-			TestTrue(TEXT("Handle array soft refcount"), Reference->GetSoftRefCount() == 1);
-			TestTrue(TEXT("Handle array hard refcount"), Reference->GetHardRefCount() == 1);
+				TestTrue(TEXT("Handle array contains reference"), HandleList.Contains(Reference));
+				TestTrue(TEXT("Handle array soft refcount"), Reference->GetSoftRefCount() == 1);
+				TestTrue(TEXT("Handle array hard refcount"), Reference->GetHardRefCount() == 1);
 			
-			HandleList.Remove(Handle);
-			TestTrue(TEXT("Handle array soft refcount"), Handle->GetSoftRefCount() == 1);
-			TestTrue(TEXT("Handle array hard refcount"), Handle->GetHardRefCount() == 1);
-
-			HandleList.Remove(Reference);
-			TestTrue(TEXT("Handle array soft refcount"), Reference->GetSoftRefCount() == 0);
-			TestTrue(TEXT("Handle array hard refcount"), Reference->GetHardRefCount() == 1);
-		}
+				HandleList.Remove(Handle);
+				TestTrue(TEXT("Handle array soft refcount"), Handle->GetSoftRefCount() == 1);
+				TestTrue(TEXT("Handle array hard refcount"), Handle->GetHardRefCount() == 1);
 
-		TestTrue(TEXT("Handle soft refcount"), Handle->GetSoftRefCount() == 1);
-		TestTrue(TEXT("Handle hard refcount"), Handle->GetHardRefCount() == 0);
-		TestTrue(TEXT("Reference soft refcount"), Reference->GetSoftRefCount() == 0);
-		TestTrue(TEXT("Reference hard refcount"), Reference->GetHardRefCount() == 1);
+				HandleList.Remove(Reference);
+				TestTrue(TEXT("Handle array soft refcount"), Reference->GetSoftRefCount() == 0);
+				TestTrue(TEXT("Handle array hard refcount"), Reference->GetHardRefCount() == 1);
+			}
 
-		// TSet test
-		{
-			TSet<FWorldPartitionHandle> HandleSet;
-			HandleSet.Add(Handle);
+			TestTrue(TEXT("Handle soft refcount"), Handle->GetSoftRefCount() == 1);
+			TestTrue(TEXT("Handle hard refcount"), Handle->GetHardRefCount() == 0);
+			TestTrue(TEXT("Reference soft refcount"), Reference->GetSoftRefCount() == 0);
+			TestTrue(TEXT("Reference hard refcount"), Reference->GetHardRefCount() == 1);
 
-			TestTrue(TEXT("Handle set soft refcount"), Handle->GetSoftRefCount() == 2);
-			TestTrue(TEXT("Handle set hard refcount"), Handle->GetHardRefCount() == 0);
+			// TSet test
+			{
+				TSet<FWorldPartitionHandle> HandleSet;
+				HandleSet.Add(Handle);
 
-			FWorldPartitionReference ReferenceToHandle = Handle;
-			TestTrue(TEXT("Handle/Reference equality"), ReferenceToHandle == Handle);
-			TestTrue(TEXT("Handle/Reference equality"), Handle == ReferenceToHandle);
-			TestTrue(TEXT("Handle soft refcount"), Handle->GetSoftRefCount() == 2);
-			TestTrue(TEXT("Handle hard refcount"), Handle->GetHardRefCount() == 1);
+				TestTrue(TEXT("Handle set soft refcount"), Handle->GetSoftRefCount() == 2);
+				TestTrue(TEXT("Handle set hard refcount"), Handle->GetHardRefCount() == 0);
 
-			TestTrue(TEXT("Handle set contains handle"), HandleSet.Contains(Handle));
-			TestTrue(TEXT("Handle set contains reference"), HandleSet.Contains(ReferenceToHandle));
+				FWorldPartitionReference ReferenceToHandle = Handle;
+				TestTrue(TEXT("Handle/Reference equality"), ReferenceToHandle == Handle);
+				TestTrue(TEXT("Handle/Reference equality"), Handle == ReferenceToHandle);
+				TestTrue(TEXT("Handle soft refcount"), Handle->GetSoftRefCount() == 2);
+				TestTrue(TEXT("Handle hard refcount"), Handle->GetHardRefCount() == 1);
 
-			HandleSet.Add(Reference);
-			TestTrue(TEXT("Handle set contains reference"), HandleSet.Contains(Reference));
+				TestTrue(TEXT("Handle set contains handle"), HandleSet.Contains(Handle));
+				TestTrue(TEXT("Handle set contains reference"), HandleSet.Contains(ReferenceToHandle));
 
-			TestTrue(TEXT("Reference soft refcount"), Reference->GetSoftRefCount() == 1);
-			TestTrue(TEXT("Reference hard refcount"), Reference->GetHardRefCount() == 1);
-		}
+				HandleSet.Add(Reference);
+				TestTrue(TEXT("Handle set contains reference"), HandleSet.Contains(Reference));
 
-		TestTrue(TEXT("Handle soft refcount"), Handle->GetSoftRefCount() == 1);
-		TestTrue(TEXT("Handle hard refcount"), Handle->GetHardRefCount() == 0);
-		TestTrue(TEXT("Reference soft refcount"), Reference->GetSoftRefCount() == 0);
-		TestTrue(TEXT("Reference hard refcount"), Reference->GetHardRefCount() == 1);
+				TestTrue(TEXT("Reference soft refcount"), Reference->GetSoftRefCount() == 1);
+				TestTrue(TEXT("Reference hard refcount"), Reference->GetHardRefCount() == 1);
+			}
 
-		// Move tests
-		{
-			// Handle move
 			TestTrue(TEXT("Handle soft refcount"), Handle->GetSoftRefCount() == 1);
 			TestTrue(TEXT("Handle hard refcount"), Handle->GetHardRefCount() == 0);
-			{
-				FWorldPartitionHandle HandleCopy(MoveTemp(Handle));
-				TestTrue(TEXT("Handle move src not valid"), !Handle.IsValid());
-				TestTrue(TEXT("Handle move dst valid"), HandleCopy.IsValid());
-				TestTrue(TEXT("Handle soft refcount"), HandleCopy->GetSoftRefCount() == 1);
-				TestTrue(TEXT("Handle hard refcount"), HandleCopy->GetHardRefCount() == 0);
-
-				Handle = MoveTemp(HandleCopy);
-				TestTrue(TEXT("Handle move src not valid"), !HandleCopy.IsValid());
-				TestTrue(TEXT("Handle move dst valid"), Handle.IsValid());
-				TestTrue(TEXT("Handle soft refcount"), Handle->GetSoftRefCount() == 1);
-				TestTrue(TEXT("Handle hard refcount"), Handle->GetHardRefCount() == 0);
-			}
-
-			// Reference move
 			TestTrue(TEXT("Reference soft refcount"), Reference->GetSoftRefCount() == 0);
 			TestTrue(TEXT("Reference hard refcount"), Reference->GetHardRefCount() == 1);
+
+			// Move tests
 			{
-				FWorldPartitionReference ReferenceCopy(MoveTemp(Reference));
-				TestTrue(TEXT("Reference move src not valid"), !Reference.IsValid());
-				TestTrue(TEXT("Reference move dst valid"), ReferenceCopy.IsValid());
-				TestTrue(TEXT("Reference soft refcount"), ReferenceCopy->GetSoftRefCount() == 0);
-				TestTrue(TEXT("Reference hard refcount"), ReferenceCopy->GetHardRefCount() == 1);
-
-				Reference = MoveTemp(ReferenceCopy);
-				TestTrue(TEXT("Reference move src not valid"), !ReferenceCopy.IsValid());
-				TestTrue(TEXT("Reference move dst valid"), Reference.IsValid());
+				// Handle move
+				TestTrue(TEXT("Handle soft refcount"), Handle->GetSoftRefCount() == 1);
+				TestTrue(TEXT("Handle hard refcount"), Handle->GetHardRefCount() == 0);
+				{
+					FWorldPartitionHandle HandleCopy(MoveTemp(Handle));
+					TestTrue(TEXT("Handle move src not valid"), !Handle.IsValid());
+					TestTrue(TEXT("Handle move dst valid"), HandleCopy.IsValid());
+					TestTrue(TEXT("Handle soft refcount"), HandleCopy->GetSoftRefCount() == 1);
+					TestTrue(TEXT("Handle hard refcount"), HandleCopy->GetHardRefCount() == 0);
+
+					Handle = MoveTemp(HandleCopy);
+					TestTrue(TEXT("Handle move src not valid"), !HandleCopy.IsValid());
+					TestTrue(TEXT("Handle move dst valid"), Handle.IsValid());
+					TestTrue(TEXT("Handle soft refcount"), Handle->GetSoftRefCount() == 1);
+					TestTrue(TEXT("Handle hard refcount"), Handle->GetHardRefCount() == 0);
+				}
+
+				// Reference move
 				TestTrue(TEXT("Reference soft refcount"), Reference->GetSoftRefCount() == 0);
 				TestTrue(TEXT("Reference hard refcount"), Reference->GetHardRefCount() == 1);
+				{
+					FWorldPartitionReference ReferenceCopy(MoveTemp(Reference));
+					TestTrue(TEXT("Reference move src not valid"), !Reference.IsValid());
+					TestTrue(TEXT("Reference move dst valid"), ReferenceCopy.IsValid());
+					TestTrue(TEXT("Reference soft refcount"), ReferenceCopy->GetSoftRefCount() == 0);
+					TestTrue(TEXT("Reference hard refcount"), ReferenceCopy->GetHardRefCount() == 1);
+
+					Reference = MoveTemp(ReferenceCopy);
+					TestTrue(TEXT("Reference move src not valid"), !ReferenceCopy.IsValid());
+					TestTrue(TEXT("Reference move dst valid"), Reference.IsValid());
+					TestTrue(TEXT("Reference soft refcount"), Reference->GetSoftRefCount() == 0);
+					TestTrue(TEXT("Reference hard refcount"), Reference->GetHardRefCount() == 1);
+				}
+
+				// Handle reference move
+				{
+					FWorldPartitionHandle HandleFromReference(MoveTemp(Reference));
+					TestTrue(TEXT("Handle move src not valid"), !Reference.IsValid());
+					TestTrue(TEXT("Handle move dst valid"), HandleFromReference.IsValid());
+					TestTrue(TEXT("Handle soft refcount"), HandleFromReference->GetSoftRefCount() == 1);
+					TestTrue(TEXT("Handle hard refcount"), HandleFromReference->GetHardRefCount() == 0);
+
+					Reference = MoveTemp(HandleFromReference);
+					TestTrue(TEXT("Handle move src not valid"), !HandleFromReference.IsValid());
+					TestTrue(TEXT("Handle move dst valid"), Reference.IsValid());
+					TestTrue(TEXT("Handle soft refcount"), Reference->GetSoftRefCount() == 0);
+					TestTrue(TEXT("Handle hard refcount"), Reference->GetHardRefCount() == 1);
+				}
+
+				// Reference handle move
+				{
+					FWorldPartitionReference ReferenceFromHandle(MoveTemp(Handle));
+					TestTrue(TEXT("Reference move src not valid"), !Handle.IsValid());
+					TestTrue(TEXT("Reference move dst valid"), ReferenceFromHandle.IsValid());
+					TestTrue(TEXT("Reference soft refcount"), ReferenceFromHandle->GetSoftRefCount() == 0);
+					TestTrue(TEXT("Reference hard refcount"), ReferenceFromHandle->GetHardRefCount() == 1);
+
+					Handle = MoveTemp(ReferenceFromHandle);
+					TestTrue(TEXT("Reference move src not valid"), !ReferenceFromHandle.IsValid());
+					TestTrue(TEXT("Reference move dst valid"), Handle.IsValid());
+					TestTrue(TEXT("Reference soft refcount"), Handle->GetSoftRefCount() == 1);
+					TestTrue(TEXT("Reference hard refcount"), Handle->GetHardRefCount() == 0);
+				}
 			}
 
-			// Handle reference move
-			{
-				FWorldPartitionHandle HandleFromReference(MoveTemp(Reference));
-				TestTrue(TEXT("Handle move src not valid"), !Reference.IsValid());
-				TestTrue(TEXT("Handle move dst valid"), HandleFromReference.IsValid());
-				TestTrue(TEXT("Handle soft refcount"), HandleFromReference->GetSoftRefCount() == 1);
-				TestTrue(TEXT("Handle hard refcount"), HandleFromReference->GetHardRefCount() == 0);
-
-				Reference = MoveTemp(HandleFromReference);
-				TestTrue(TEXT("Handle move src not valid"), !HandleFromReference.IsValid());
-				TestTrue(TEXT("Handle move dst valid"), Reference.IsValid());
-				TestTrue(TEXT("Handle soft refcount"), Reference->GetSoftRefCount() == 0);
-				TestTrue(TEXT("Handle hard refcount"), Reference->GetHardRefCount() == 1);
-			}
-
-			// Reference handle move
-			{
-				FWorldPartitionReference ReferenceFromHandle(MoveTemp(Handle));
-				TestTrue(TEXT("Reference move src not valid"), !Handle.IsValid());
-				TestTrue(TEXT("Reference move dst valid"), ReferenceFromHandle.IsValid());
-				TestTrue(TEXT("Reference soft refcount"), ReferenceFromHandle->GetSoftRefCount() == 0);
-				TestTrue(TEXT("Reference hard refcount"), ReferenceFromHandle->GetHardRefCount() == 1);
-
-				Handle = MoveTemp(ReferenceFromHandle);
-				TestTrue(TEXT("Reference move src not valid"), !ReferenceFromHandle.IsValid());
-				TestTrue(TEXT("Reference move dst valid"), Handle.IsValid());
-				TestTrue(TEXT("Reference soft refcount"), Handle->GetSoftRefCount() == 1);
-				TestTrue(TEXT("Reference hard refcount"), Handle->GetHardRefCount() == 0);
-			}
+			// Make sure the container gets destroyed so we can test destrucing dangling handles
+			CollectGarbage(GARBAGE_COLLECTION_KEEPFLAGS);
 		}
+
+		TestTrue(TEXT("Invalid container"), true);
 #endif
 		return true;
 	}
diff --git a/Engine/Source/Runtime/Engine/Private/WorldPartition/WorldPartition.cpp b/Engine/Source/Runtime/Engine/Private/WorldPartition/WorldPartition.cpp
--- a/Engine/Source/Runtime/Engine/Private/WorldPartition/WorldPartition.cpp
+++ b/Engine/Source/Runtime/Engine/Private/WorldPartition/WorldPartition.cpp
@@ -317,8 +317,6 @@ void UWorldPartition::OnPreBeginPIE(bool bStartSimulate)
 	check(!bIsPIE);
 	bIsPIE = true;
 
-	check(IsMainWorldPartition());
-
 	OnBeginPlay();
 }
 
@@ -344,8 +342,6 @@ void UWorldPartition::OnCancelPIE()
 
 void UWorldPartition::OnEndPlay()
 {
-	check(IsMainWorldPartition());
-
 	RuntimeHash->FlushStreaming();
 	RuntimeHash->OnEndPlay();
 
@@ -383,8 +379,6 @@ void UWorldPartition::Initialize(UWorld* InWorld, const FTransform& InTransform)
 	UWorld* OuterWorld = GetTypedOuter<UWorld>();
 	check(OuterWorld);
 
-	check(IsMainWorldPartition());
-
 	RegisterDelegates();
 
 	AWorldPartitionReplay::Initialize(World);
@@ -595,12 +589,6 @@ bool UWorldPartition::IsInitialized() const
 	return InitState == EWorldPartitionInitState::Initialized;
 }
 
-bool UWorldPartition::IsMainWorldPartition() const
-{
-	check(World);
-	return World == GetTypedOuter<UWorld>();
-}
-
 void UWorldPartition::OnPostBugItGoCalled(const FVector& Loc, const FRotator& Rot)
 {
 #if WITH_EDITOR
@@ -1114,15 +1102,6 @@ void UWorldPartition::UnhashActorDesc(FWorldPartitionActorDesc* ActorDesc)
 }
 #endif
 
-void UWorldPartition::BeginDestroy()
-{
-	Super::BeginDestroy();
-
-#if WITH_EDITOR
-	LoadedSubobjects.Empty();
-#endif
-}
-
 bool UWorldPartition::ResolveSubobject(const TCHAR* SubObjectPath, UObject*& OutObject, bool bLoadIfExists)
 {
 	if (GetWorld())
diff --git a/Engine/Source/Runtime/Engine/Private/WorldPartition/WorldPartitionActorDesc.cpp b/Engine/Source/Runtime/Engine/Private/WorldPartition/WorldPartitionActorDesc.cpp
--- a/Engine/Source/Runtime/Engine/Private/WorldPartition/WorldPartitionActorDesc.cpp
+++ b/Engine/Source/Runtime/Engine/Private/WorldPartition/WorldPartitionActorDesc.cpp
@@ -77,11 +77,10 @@ void FWorldPartitionActorDesc::Init(const AActor* InActor)
 
 	ActorLabel = *InActor->GetActorLabel(false);
 
-	// Only set ActorPtr/Container on WorldPartition ActorDescs
-	if (UWorldPartition* WorldPartition = InActor->GetLevel()->GetWorldPartition())
+	if (UActorDescContainer* ActorContainer = InActor->GetLevel()->GetWorldPartition())
 	{
 		ActorPtr = (AActor*)InActor;
-		Container = WorldPartition;
+		Container = ActorContainer;
 	}
 }
 
@@ -103,12 +102,8 @@ void FWorldPartitionActorDesc::Init(UActorDescContainer* InContainer, const FWor
 	// Serialize metadata payload
 	Serialize(MetadataAr);
 
-	// Only set ActorPtr/Container on WorldPartition ActorDescs
-	if (UWorldPartition* WorldPartition = Cast<UWorldPartition>(InContainer))
-	{
-		Container = InContainer;
-		ActorPtr = FindObject<AActor>(nullptr, *ActorPath.ToString());
-	}
+	Container = InContainer;
+	ActorPtr = FindObject<AActor>(nullptr, *ActorPath.ToString());
 }
 
 bool FWorldPartitionActorDesc::Equals(const FWorldPartitionActorDesc* Other) const
@@ -312,13 +307,6 @@ UHLODLayer* FWorldPartitionActorDesc::GetHLODLayer() const
 
 bool FWorldPartitionActorDesc::IsLoaded(bool bEvenIfPendingKill) const
 {
-#if WITH_DEV_AUTOMATION_TESTS
-	if (GIsAutomationTesting)
-	{
-		return HardRefCount > 0;
-	}
-#endif
-
 	return ActorPtr.IsValid(bEvenIfPendingKill);
 }
 
@@ -338,13 +326,9 @@ AActor* FWorldPartitionActorDesc::Load() const
 	// The, if the actor isn't loaded, load it
 	if (ActorPtr.IsExplicitlyNull())
 	{
-		// For now we assume that an ActorDesc that gets loaded is owned by an actual WorldPartition
-		UWorldPartition* WorldPartition = Container ? CastChecked<UWorldPartition>(Container) : nullptr;
-		check(WorldPartition || GIsAutomationTesting);
-
 		const FLinkerInstancingContext* InstancingContext = nullptr;
 		FSoftObjectPathFixupArchive* SoftObjectPathFixupArchive = nullptr;
-		if (WorldPartition && WorldPartition->InstancingContext.IsInstanced())
+		if (UWorldPartition* WorldPartition = Cast<UWorldPartition>(Container); WorldPartition && WorldPartition->InstancingContext.IsInstanced())
 		{
 			InstancingContext = &WorldPartition->InstancingContext;
 			SoftObjectPathFixupArchive = WorldPartition->InstancingSoftObjectPathFixupArchive.Get();
@@ -361,7 +345,6 @@ AActor* FWorldPartitionActorDesc::Load() const
 		}
 
 		Package = LoadPackage(Package, *ActorPackage.ToString(), LOAD_None, nullptr, InstancingContext);
-		check(Package || GIsAutomationTesting);
 
 		if (Package)
 		{
diff --git a/Engine/Source/Runtime/Engine/Private/WorldPartition/WorldPartitionEditorCell.cpp b/Engine/Source/Runtime/Engine/Private/WorldPartition/WorldPartitionEditorCell.cpp
--- a/Engine/Source/Runtime/Engine/Private/WorldPartition/WorldPartitionEditorCell.cpp
+++ b/Engine/Source/Runtime/Engine/Private/WorldPartition/WorldPartitionEditorCell.cpp
@@ -41,14 +41,6 @@ void UWorldPartitionEditorCell::AddReferencedObjects(UObject* InThis, FReference
 	Super::AddReferencedObjects(InThis, Collector);
 }
 
-void UWorldPartitionEditorCell::BeginDestroy()
-{
-	// Release handles/references
-	Actors.Empty();
-	LoadedActors.Empty();
-	Super::BeginDestroy();
-}
-
 void UWorldPartitionEditorCell::AddActor(const FWorldPartitionHandle& ActorHandle)
 {
 	AddActor(ActorHandle->GetGuid(), ActorHandle);
diff --git a/Engine/Source/Runtime/Engine/Private/WorldPartition/WorldPartitionHandle.cpp b/Engine/Source/Runtime/Engine/Private/WorldPartition/WorldPartitionHandle.cpp
--- a/Engine/Source/Runtime/Engine/Private/WorldPartition/WorldPartitionHandle.cpp
+++ b/Engine/Source/Runtime/Engine/Private/WorldPartition/WorldPartitionHandle.cpp
@@ -15,6 +15,11 @@ TUniquePtr<FWorldPartitionActorDesc>* FWorldPartitionHandleUtils::GetActorDesc(U
 	return nullptr;
 }
 
+UActorDescContainer* FWorldPartitionHandleUtils::GetActorDescContainer(TUniquePtr<FWorldPartitionActorDesc>* ActorDesc)
+{
+	return ActorDesc ? ActorDesc->Get()->GetContainer() : nullptr;
+}
+
 bool FWorldPartitionHandleUtils::IsActorDescLoaded(FWorldPartitionActorDesc* ActorDesc)
 {
 	return ActorDesc->IsLoaded();
diff --git a/Engine/Source/Runtime/Engine/Public/WorldPartition/WorldPartition.h b/Engine/Source/Runtime/Engine/Public/WorldPartition/WorldPartition.h
--- a/Engine/Source/Runtime/Engine/Public/WorldPartition/WorldPartition.h
+++ b/Engine/Source/Runtime/Engine/Public/WorldPartition/WorldPartition.h
@@ -106,10 +106,6 @@ private:
 	virtual void OnActorDescUnregistered(const FWorldPartitionActorDesc&) override;
 #endif
 
-	//~ Begin UObject Interface
-	virtual void BeginDestroy() override;
-	//~ End UObject Interface
-
 	static void WorldPartitionOnLevelRemovedFromWorld(class ULevel* Level, UWorld* InWorld);
 
 public:
@@ -217,7 +213,6 @@ private:
 	FWorldPartitionReference WorldDataLayersActor;
 #endif
 
-	bool IsMainWorldPartition() const;
 	void OnWorldMatchStarting();
 
 	void OnPostBugItGoCalled(const FVector& Loc, const FRotator& Rot);
diff --git a/Engine/Source/Runtime/Engine/Public/WorldPartition/WorldPartitionActorDesc.h b/Engine/Source/Runtime/Engine/Public/WorldPartition/WorldPartitionActorDesc.h
--- a/Engine/Source/Runtime/Engine/Public/WorldPartition/WorldPartitionActorDesc.h
+++ b/Engine/Source/Runtime/Engine/Public/WorldPartition/WorldPartitionActorDesc.h
@@ -126,6 +126,11 @@ protected:
 		return HardRefCount;
 	}
 
+	UActorDescContainer* GetContainer() const
+	{
+		return Container;
+	}
+
 public:
 	const TArray<FGuid>& GetReferences() const
 	{
diff --git a/Engine/Source/Runtime/Engine/Public/WorldPartition/WorldPartitionEditorCell.h b/Engine/Source/Runtime/Engine/Public/WorldPartition/WorldPartitionEditorCell.h
--- a/Engine/Source/Runtime/Engine/Public/WorldPartition/WorldPartitionEditorCell.h
+++ b/Engine/Source/Runtime/Engine/Public/WorldPartition/WorldPartitionEditorCell.h
@@ -70,8 +70,6 @@ public:
 	typedef TActorHandle<FWorldPartitionHandle> FActorHandle;
 	typedef TActorHandle<FWorldPartitionReference> FActorReference;
 
-	virtual void BeginDestroy() override;
-
 	void AddActor(const FWorldPartitionHandle& ActorHandle);
 	void AddActor(const FGuid& Source, const FWorldPartitionHandle& ActorHandle);
 	void RemoveActor(const FWorldPartitionHandle& ActorHandle);
diff --git a/Engine/Source/Runtime/Engine/Public/WorldPartition/WorldPartitionHandle.h b/Engine/Source/Runtime/Engine/Public/WorldPartition/WorldPartitionHandle.h
--- a/Engine/Source/Runtime/Engine/Public/WorldPartition/WorldPartitionHandle.h
+++ b/Engine/Source/Runtime/Engine/Public/WorldPartition/WorldPartitionHandle.h
@@ -3,6 +3,8 @@
 
 #include "CoreMinimal.h"
 #include "Templates/UniquePtr.h"
+#include "UObject/WeakObjectPtrTemplates.h"
+#include "UObject/WeakObjectPtr.h"
 
 #if WITH_EDITOR
 class UActorDescContainer;
@@ -11,6 +13,7 @@ class FWorldPartitionActorDesc;
 struct ENGINE_API FWorldPartitionHandleUtils
 {
 	static TUniquePtr<FWorldPartitionActorDesc>* GetActorDesc(UActorDescContainer* Container, const FGuid& ActorGuid);
+	static UActorDescContainer* GetActorDescContainer(TUniquePtr<FWorldPartitionActorDesc>* ActorDesc);
 	static bool IsActorDescLoaded(FWorldPartitionActorDesc* ActorDesc);
 };
 
@@ -23,7 +26,8 @@ public:
 	{}
 
 	FORCEINLINE TWorldPartitionHandle(TUniquePtr<FWorldPartitionActorDesc>* InActorDesc)
-		: ActorDesc(InActorDesc)
+		: Container(FWorldPartitionHandleUtils::GetActorDescContainer(InActorDesc))
+		, ActorDesc(InActorDesc)
 	{
 		if (IsValid())
 		{
@@ -32,7 +36,8 @@ public:
 	}
 
 	FORCEINLINE TWorldPartitionHandle(UActorDescContainer* Container, const FGuid& ActorGuid)
-		: ActorDesc(FWorldPartitionHandleUtils::GetActorDesc(Container, ActorGuid))
+		: Container(Container)
+		, ActorDesc(FWorldPartitionHandleUtils::GetActorDesc(Container, ActorGuid))
 	{
 		if (IsValid())
 		{
@@ -41,13 +46,15 @@ public:
 	}
 
 	FORCEINLINE TWorldPartitionHandle(const TWorldPartitionHandle& Other)
-		: ActorDesc(nullptr)
+		: Container(nullptr)
+		, ActorDesc(nullptr)
 	{
 		*this = Other;
 	}
 
 	FORCEINLINE TWorldPartitionHandle(TWorldPartitionHandle&& Other)
-		: ActorDesc(nullptr)
+		: Container(nullptr)
+		, ActorDesc(nullptr)
 	{
 		*this = MoveTemp(Other);
 	}
@@ -84,6 +91,7 @@ public:
 				DecRefCount();
 			}
 
+			Container = Other.Container;
 			ActorDesc = Other.ActorDesc;
 
 			if (IsValid())
@@ -104,7 +112,10 @@ public:
 				DecRefCount();
 			}
 
+			Container = Other.Container;
 			ActorDesc = Other.ActorDesc;
+			
+			Other.Container = nullptr;
 			Other.ActorDesc = nullptr;
 		}
 
@@ -120,6 +131,7 @@ public:
 			DecRefCount();
 		}
 
+		Container = Other.Container;
 		ActorDesc = Other.ActorDesc;
 
 		if (IsValid())
@@ -138,12 +150,15 @@ public:
 			DecRefCount();
 		}
 
+		Container = Other.Container;
 		ActorDesc = Other.ActorDesc;
 
 		if (IsValid())
 		{
 			IncRefCount();
 			Other.DecRefCount();
+
+			Other.Container = nullptr;
 			Other.ActorDesc = nullptr;
 		}
 
@@ -162,7 +177,7 @@ public:
 
 	FORCEINLINE bool IsValid() const
 	{
-		return ActorDesc && ActorDesc->IsValid();
+		return Container.IsValid() && ActorDesc && ActorDesc->IsValid();
 	}
 
 	FORCEINLINE bool IsLoaded() const
@@ -177,6 +192,7 @@ public:
 
 	FORCEINLINE void Reset()
 	{
+		Container = nullptr;
 		ActorDesc = nullptr;
 	}
 	
@@ -219,6 +235,7 @@ public:
 		Impl::DecRefCount(ActorDesc->Get());
 	}
 
+	TWeakObjectPtr<UActorDescContainer> Container;
 	TUniquePtr<FWorldPartitionActorDesc>* ActorDesc;
 };
 
diff --git a/Engine/Source/Runtime/Landscape/Classes/LandscapeInfo.h b/Engine/Source/Runtime/Landscape/Classes/LandscapeInfo.h
--- a/Engine/Source/Runtime/Landscape/Classes/LandscapeInfo.h
+++ b/Engine/Source/Runtime/Landscape/Classes/LandscapeInfo.h
@@ -169,8 +169,6 @@ public:
 	//~ End UObject Interface
 
 #if WITH_EDITOR
-	virtual void BeginDestroy() override;
-
 	// @todo document 
 	// all below.
 	LANDSCAPE_API bool SupportsLandscapeEditing() const;
diff --git a/Engine/Source/Runtime/Landscape/Private/LandscapeInfo.cpp b/Engine/Source/Runtime/Landscape/Private/LandscapeInfo.cpp
--- a/Engine/Source/Runtime/Landscape/Private/LandscapeInfo.cpp
+++ b/Engine/Source/Runtime/Landscape/Private/LandscapeInfo.cpp
@@ -12,17 +12,4 @@ LANDSCAPE_API FLandscapeInfoLayerSettings::FLandscapeInfoLayerSettings(ULandscap
 	, DebugColorChannel(0)
 	, bValid(false)
 #endif
-{
-}
-
-#if WITH_EDITOR
-
-void ULandscapeInfo::BeginDestroy()
-{
-	// Release WorldPartition Actor Handles
-	ProxyHandles.Empty();
-	SplineHandles.Empty();
-	Super::BeginDestroy();
-}
-
-#endif // WITH_EDITOR
\ No newline at end of file
+{}
\ No newline at end of file
