Animation:
* SkeletalMeshComponents will now run multithreadable tasks on the taskgraph on forked multithreaded servers

diff --git a/Engine/Source/Runtime/Engine/Private/Components/SkeletalMeshComponent.cpp b/Engine/Source/Runtime/Engine/Private/Components/SkeletalMeshComponent.cpp
--- a/Engine/Source/Runtime/Engine/Private/Components/SkeletalMeshComponent.cpp
+++ b/Engine/Source/Runtime/Engine/Private/Components/SkeletalMeshComponent.cpp
@@ -2295,7 +2295,7 @@ void USkeletalMeshComponent::RefreshBoneTransforms(FActorComponentTickFunction*
 
 	const bool bShouldDoParallelInterpolation = bShouldDoInterpolation && CVarUseParallelAnimationInterpolation.GetValueOnGameThread() == 1;
 
-	const bool bDoPAE = !!CVarUseParallelAnimationEvaluation.GetValueOnGameThread() && FApp::ShouldUseThreadingForPerformance();
+	const bool bDoPAE = !!CVarUseParallelAnimationEvaluation.GetValueOnGameThread() && (FApp::ShouldUseThreadingForPerformance() || FForkProcessHelper::SupportsMultithreadingPostFork());
 
 	const bool bMainInstanceValidForParallelWork = AnimScriptInstance == nullptr || AnimScriptInstance->CanRunParallelWork();
 	const bool bPostInstanceValidForParallelWork = PostProcessAnimInstance == nullptr || PostProcessAnimInstance->CanRunParallelWork();
@@ -2500,7 +2500,7 @@ void USkeletalMeshComponent::DispatchParallelTickPose(FActorComponentTickFunctio
 				QUICK_SCOPE_CYCLE_COUNTER(STAT_USkeletalMeshComponent_RefreshBoneTransforms_DispatchParallelTickPose);
 
 				// This duplicates *some* of the logic from RefreshBoneTransforms()
-				const bool bDoPAE = !!CVarUseParallelAnimationEvaluation.GetValueOnGameThread() && FApp::ShouldUseThreadingForPerformance();
+				const bool bDoPAE = !!CVarUseParallelAnimationEvaluation.GetValueOnGameThread() && (FApp::ShouldUseThreadingForPerformance() || FForkProcessHelper::SupportsMultithreadingPostFork());
 
 				const bool bDoParallelUpdate = bDoPAE && TickFunction->IsCompletionHandleValid();
 
