World Partition: better unit tests for invalid handles.

diff --git a/Engine/Source/Runtime/Engine/Private/Tests/WorldPartition/WorldPartitionHandleTests.cpp b/Engine/Source/Runtime/Engine/Private/Tests/WorldPartition/WorldPartitionHandleTests.cpp
--- a/Engine/Source/Runtime/Engine/Private/Tests/WorldPartition/WorldPartitionHandleTests.cpp
+++ b/Engine/Source/Runtime/Engine/Private/Tests/WorldPartition/WorldPartitionHandleTests.cpp
@@ -17,215 +17,216 @@ namespace WorldPartitionTests
 	bool FWorldPartitionSoftRefTest::RunTest(const FString& Parameters)
 	{
 #if WITH_EDITOR
-		FWorldPartitionHandle Handle;
-		FWorldPartitionReference Reference;
+		UActorDescContainer* ActorDescContainer = NewObject<UActorDescContainer>(GetTransientPackage());
+		ActorDescContainer->Initialize(nullptr, TEXT("/Engine/WorldPartition/WorldPartitionUnitTest"));
 
-		{
-			UActorDescContainer* ActorDescContainer = NewObject<UActorDescContainer>(GetTransientPackage());
-			ActorDescContainer->Initialize(nullptr, TEXT("/Engine/WorldPartition/WorldPartitionUnitTest"));
+		FWorldPartitionHandle Handle = FWorldPartitionHandle(ActorDescContainer, FGuid(TEXT("5D9F93BA407A811AFDDDAAB4F1CECC6A")));
+		FWorldPartitionReference Reference = FWorldPartitionReference(ActorDescContainer, FGuid(TEXT("0D2B04D240BE5DE58FE437A8D2DBF5C9")));
 
-			Handle = FWorldPartitionHandle(ActorDescContainer, FGuid(TEXT("5D9F93BA407A811AFDDDAAB4F1CECC6A")));
-			Reference = FWorldPartitionReference(ActorDescContainer, FGuid(TEXT("0D2B04D240BE5DE58FE437A8D2DBF5C9")));
+		TestTrue(TEXT("Handle soft refcount"), Handle->GetSoftRefCount() == 1);
+		TestTrue(TEXT("Handle hard refcount"), Handle->GetHardRefCount() == 0);
+		TestTrue(TEXT("Reference soft refcount"), Reference->GetSoftRefCount() == 0);
+		TestTrue(TEXT("Reference hard refcount"), Reference->GetHardRefCount() == 1);
 
-			TestTrue(TEXT("Handle soft refcount"), Handle->GetSoftRefCount() == 1);
-			TestTrue(TEXT("Handle hard refcount"), Handle->GetHardRefCount() == 0);
-			TestTrue(TEXT("Reference soft refcount"), Reference->GetSoftRefCount() == 0);
-			TestTrue(TEXT("Reference hard refcount"), Reference->GetHardRefCount() == 1);
-
-			// Pin handle scope test
-			{
-				FWorldPartitionHandlePinRefScope PinRefScopeHandle(Handle);
-				TestTrue(TEXT("Pin to Handle soft refcount"), Handle->GetSoftRefCount() == 1);
-				TestTrue(TEXT("Pin to Handle hard refcount"), Handle->GetHardRefCount() == 0);
-			}
+		// Pin handle scope test
+		{
+			FWorldPartitionHandlePinRefScope PinRefScopeHandle(Handle);
+			TestTrue(TEXT("Pin to Handle soft refcount"), Handle->GetSoftRefCount() == 1);
+			TestTrue(TEXT("Pin to Handle hard refcount"), Handle->GetHardRefCount() == 0);
+		}
 
-			// Pin reference scope test
-			{
-				FWorldPartitionHandlePinRefScope PinRefScopeReference(Reference);
-				TestTrue(TEXT("Pin to Reference soft refcount"), Reference->GetSoftRefCount() == 0);
-				TestTrue(TEXT("Pin to Reference hard refcount"), Reference->GetHardRefCount() == 2);
-			}
+		// Pin reference scope test
+		{
+			FWorldPartitionHandlePinRefScope PinRefScopeReference(Reference);
+			TestTrue(TEXT("Pin to Reference soft refcount"), Reference->GetSoftRefCount() == 0);
+			TestTrue(TEXT("Pin to Reference hard refcount"), Reference->GetHardRefCount() == 2);
+		}
 
-			TestTrue(TEXT("Handle/Reference equality"), Handle != Reference);
-			TestTrue(TEXT("Reference/Handle equality"), Reference != Handle);
+		TestTrue(TEXT("Handle/Reference equality"), Handle != Reference);
+		TestTrue(TEXT("Reference/Handle equality"), Reference != Handle);
 		
-			TestTrue(TEXT("Handle soft refcount"), Handle->GetSoftRefCount() == 1);
-			TestTrue(TEXT("Handle hard refcount"), Handle->GetHardRefCount() == 0);
-			TestTrue(TEXT("Reference soft refcount"), Reference->GetSoftRefCount() == 0);
-			TestTrue(TEXT("Reference hard refcount"), Reference->GetHardRefCount() == 1);
+		TestTrue(TEXT("Handle soft refcount"), Handle->GetSoftRefCount() == 1);
+		TestTrue(TEXT("Handle hard refcount"), Handle->GetHardRefCount() == 0);
+		TestTrue(TEXT("Reference soft refcount"), Reference->GetSoftRefCount() == 0);
+		TestTrue(TEXT("Reference hard refcount"), Reference->GetHardRefCount() == 1);
 
-			// Conversions
-			{
-				FWorldPartitionHandle HandleToReference = Reference;
-				TestTrue(TEXT("Handle/Reference equality"), HandleToReference == Reference);
-				TestTrue(TEXT("Reference/Handle equality"), Reference == HandleToReference);
-				TestTrue(TEXT("Reference soft refcount"), Reference->GetSoftRefCount() == 1);
-				TestTrue(TEXT("Reference hard refcount"), Reference->GetHardRefCount() == 1);
-
-				FWorldPartitionReference ReferenceToHandle = Handle;
-				TestTrue(TEXT("Handle/Reference equality"), ReferenceToHandle == Handle);
-				TestTrue(TEXT("Handle/Reference equality"), Handle == ReferenceToHandle);
-				TestTrue(TEXT("Reference soft refcount"), Reference->GetSoftRefCount() == 1);
-				TestTrue(TEXT("Reference hard refcount"), Reference->GetHardRefCount() == 1);
-			}
+		// Conversions
+		{
+			FWorldPartitionHandle HandleToReference = Reference;
+			TestTrue(TEXT("Handle/Reference equality"), HandleToReference == Reference);
+			TestTrue(TEXT("Reference/Handle equality"), Reference == HandleToReference);
+			TestTrue(TEXT("Reference soft refcount"), Reference->GetSoftRefCount() == 1);
+			TestTrue(TEXT("Reference hard refcount"), Reference->GetHardRefCount() == 1);
 
-			TestTrue(TEXT("Handle soft refcount"), Handle->GetSoftRefCount() == 1);
-			TestTrue(TEXT("Handle hard refcount"), Handle->GetHardRefCount() == 0);
-			TestTrue(TEXT("Reference soft refcount"), Reference->GetSoftRefCount() == 0);
+			FWorldPartitionReference ReferenceToHandle = Handle;
+			TestTrue(TEXT("Handle/Reference equality"), ReferenceToHandle == Handle);
+			TestTrue(TEXT("Handle/Reference equality"), Handle == ReferenceToHandle);
+			TestTrue(TEXT("Reference soft refcount"), Reference->GetSoftRefCount() == 1);
 			TestTrue(TEXT("Reference hard refcount"), Reference->GetHardRefCount() == 1);
+		}
 
-			// inplace new test
-			{
-				uint8 Buffer[sizeof(FWorldPartitionHandle)];
-				FWorldPartitionHandle* HandlePtr = new (Buffer) FWorldPartitionHandle(Reference);
+		TestTrue(TEXT("Handle soft refcount"), Handle->GetSoftRefCount() == 1);
+		TestTrue(TEXT("Handle hard refcount"), Handle->GetHardRefCount() == 0);
+		TestTrue(TEXT("Reference soft refcount"), Reference->GetSoftRefCount() == 0);
+		TestTrue(TEXT("Reference hard refcount"), Reference->GetHardRefCount() == 1);
 
-				TestTrue(TEXT("Handle array soft refcount"), Reference->GetSoftRefCount() == 1);
-				TestTrue(TEXT("Handle array hard refcount"), Reference->GetHardRefCount() == 1);
+		// inplace new test
+		{
+			uint8 Buffer[sizeof(FWorldPartitionHandle)];
+			FWorldPartitionHandle* HandlePtr = new (Buffer) FWorldPartitionHandle(Reference);
 
-				HandlePtr->~FWorldPartitionHandle();
+			TestTrue(TEXT("Handle array soft refcount"), Reference->GetSoftRefCount() == 1);
+			TestTrue(TEXT("Handle array hard refcount"), Reference->GetHardRefCount() == 1);
 
-				TestTrue(TEXT("Handle array soft refcount"), Reference->GetSoftRefCount() == 0);
-				TestTrue(TEXT("Handle array hard refcount"), Reference->GetHardRefCount() == 1);
-			}
+			HandlePtr->~FWorldPartitionHandle();
 
-			// TArray test
-			{
-				TArray<FWorldPartitionHandle> HandleList;
-				HandleList.Add(Handle);
+			TestTrue(TEXT("Handle array soft refcount"), Reference->GetSoftRefCount() == 0);
+			TestTrue(TEXT("Handle array hard refcount"), Reference->GetHardRefCount() == 1);
+		}
 
-				TestTrue(TEXT("Handle array soft refcount"), Handle->GetSoftRefCount() == 2);
-				TestTrue(TEXT("Handle array hard refcount"), Handle->GetHardRefCount() == 0);
+		// TArray test
+		{
+			TArray<FWorldPartitionHandle> HandleList;
+			HandleList.Add(Handle);
 
-				FWorldPartitionReference ReferenceToHandle = Handle;
-				TestTrue(TEXT("Handle/Reference equality"), ReferenceToHandle == Handle);
-				TestTrue(TEXT("Handle/Reference equality"), Handle == ReferenceToHandle);
-				TestTrue(TEXT("Handle soft refcount"), Handle->GetSoftRefCount() == 2);
-				TestTrue(TEXT("Handle hard refcount"), Handle->GetHardRefCount() == 1);
+			TestTrue(TEXT("Handle array soft refcount"), Handle->GetSoftRefCount() == 2);
+			TestTrue(TEXT("Handle array hard refcount"), Handle->GetHardRefCount() == 0);
 
-				TestTrue(TEXT("Handle array contains handle"), HandleList.Contains(Handle));
-				TestTrue(TEXT("Handle array contains reference"), HandleList.Contains(ReferenceToHandle));
+			FWorldPartitionReference ReferenceToHandle = Handle;
+			TestTrue(TEXT("Handle/Reference equality"), ReferenceToHandle == Handle);
+			TestTrue(TEXT("Handle/Reference equality"), Handle == ReferenceToHandle);
+			TestTrue(TEXT("Handle soft refcount"), Handle->GetSoftRefCount() == 2);
+			TestTrue(TEXT("Handle hard refcount"), Handle->GetHardRefCount() == 1);
 
-				HandleList.Add(Reference);
+			TestTrue(TEXT("Handle array contains handle"), HandleList.Contains(Handle));
+			TestTrue(TEXT("Handle array contains reference"), HandleList.Contains(ReferenceToHandle));
+
+			HandleList.Add(Reference);
 			
-				TestTrue(TEXT("Handle array contains reference"), HandleList.Contains(Reference));
-				TestTrue(TEXT("Handle array soft refcount"), Reference->GetSoftRefCount() == 1);
-				TestTrue(TEXT("Handle array hard refcount"), Reference->GetHardRefCount() == 1);
+			TestTrue(TEXT("Handle array contains reference"), HandleList.Contains(Reference));
+			TestTrue(TEXT("Handle array soft refcount"), Reference->GetSoftRefCount() == 1);
+			TestTrue(TEXT("Handle array hard refcount"), Reference->GetHardRefCount() == 1);
 			
-				HandleList.Remove(Handle);
-				TestTrue(TEXT("Handle array soft refcount"), Handle->GetSoftRefCount() == 1);
-				TestTrue(TEXT("Handle array hard refcount"), Handle->GetHardRefCount() == 1);
+			HandleList.Remove(Handle);
+			TestTrue(TEXT("Handle array soft refcount"), Handle->GetSoftRefCount() == 1);
+			TestTrue(TEXT("Handle array hard refcount"), Handle->GetHardRefCount() == 1);
 
-				HandleList.Remove(Reference);
-				TestTrue(TEXT("Handle array soft refcount"), Reference->GetSoftRefCount() == 0);
-				TestTrue(TEXT("Handle array hard refcount"), Reference->GetHardRefCount() == 1);
-			}
+			HandleList.Remove(Reference);
+			TestTrue(TEXT("Handle array soft refcount"), Reference->GetSoftRefCount() == 0);
+			TestTrue(TEXT("Handle array hard refcount"), Reference->GetHardRefCount() == 1);
+		}
 
-			TestTrue(TEXT("Handle soft refcount"), Handle->GetSoftRefCount() == 1);
-			TestTrue(TEXT("Handle hard refcount"), Handle->GetHardRefCount() == 0);
-			TestTrue(TEXT("Reference soft refcount"), Reference->GetSoftRefCount() == 0);
-			TestTrue(TEXT("Reference hard refcount"), Reference->GetHardRefCount() == 1);
+		TestTrue(TEXT("Handle soft refcount"), Handle->GetSoftRefCount() == 1);
+		TestTrue(TEXT("Handle hard refcount"), Handle->GetHardRefCount() == 0);
+		TestTrue(TEXT("Reference soft refcount"), Reference->GetSoftRefCount() == 0);
+		TestTrue(TEXT("Reference hard refcount"), Reference->GetHardRefCount() == 1);
 
-			// TSet test
-			{
-				TSet<FWorldPartitionHandle> HandleSet;
-				HandleSet.Add(Handle);
+		// TSet test
+		{
+			TSet<FWorldPartitionHandle> HandleSet;
+			HandleSet.Add(Handle);
 
-				TestTrue(TEXT("Handle set soft refcount"), Handle->GetSoftRefCount() == 2);
-				TestTrue(TEXT("Handle set hard refcount"), Handle->GetHardRefCount() == 0);
+			TestTrue(TEXT("Handle set soft refcount"), Handle->GetSoftRefCount() == 2);
+			TestTrue(TEXT("Handle set hard refcount"), Handle->GetHardRefCount() == 0);
 
-				FWorldPartitionReference ReferenceToHandle = Handle;
-				TestTrue(TEXT("Handle/Reference equality"), ReferenceToHandle == Handle);
-				TestTrue(TEXT("Handle/Reference equality"), Handle == ReferenceToHandle);
-				TestTrue(TEXT("Handle soft refcount"), Handle->GetSoftRefCount() == 2);
-				TestTrue(TEXT("Handle hard refcount"), Handle->GetHardRefCount() == 1);
+			FWorldPartitionReference ReferenceToHandle = Handle;
+			TestTrue(TEXT("Handle/Reference equality"), ReferenceToHandle == Handle);
+			TestTrue(TEXT("Handle/Reference equality"), Handle == ReferenceToHandle);
+			TestTrue(TEXT("Handle soft refcount"), Handle->GetSoftRefCount() == 2);
+			TestTrue(TEXT("Handle hard refcount"), Handle->GetHardRefCount() == 1);
 
-				TestTrue(TEXT("Handle set contains handle"), HandleSet.Contains(Handle));
-				TestTrue(TEXT("Handle set contains reference"), HandleSet.Contains(ReferenceToHandle));
+			TestTrue(TEXT("Handle set contains handle"), HandleSet.Contains(Handle));
+			TestTrue(TEXT("Handle set contains reference"), HandleSet.Contains(ReferenceToHandle));
 
-				HandleSet.Add(Reference);
-				TestTrue(TEXT("Handle set contains reference"), HandleSet.Contains(Reference));
+			HandleSet.Add(Reference);
+			TestTrue(TEXT("Handle set contains reference"), HandleSet.Contains(Reference));
 
-				TestTrue(TEXT("Reference soft refcount"), Reference->GetSoftRefCount() == 1);
-				TestTrue(TEXT("Reference hard refcount"), Reference->GetHardRefCount() == 1);
-			}
+			TestTrue(TEXT("Reference soft refcount"), Reference->GetSoftRefCount() == 1);
+			TestTrue(TEXT("Reference hard refcount"), Reference->GetHardRefCount() == 1);
+		}
 
+		TestTrue(TEXT("Handle soft refcount"), Handle->GetSoftRefCount() == 1);
+		TestTrue(TEXT("Handle hard refcount"), Handle->GetHardRefCount() == 0);
+		TestTrue(TEXT("Reference soft refcount"), Reference->GetSoftRefCount() == 0);
+		TestTrue(TEXT("Reference hard refcount"), Reference->GetHardRefCount() == 1);
+
+		// Move tests
+		{
+			// Handle move
 			TestTrue(TEXT("Handle soft refcount"), Handle->GetSoftRefCount() == 1);
 			TestTrue(TEXT("Handle hard refcount"), Handle->GetHardRefCount() == 0);
-			TestTrue(TEXT("Reference soft refcount"), Reference->GetSoftRefCount() == 0);
-			TestTrue(TEXT("Reference hard refcount"), Reference->GetHardRefCount() == 1);
-
-			// Move tests
 			{
-				// Handle move
+				FWorldPartitionHandle HandleCopy(MoveTemp(Handle));
+				TestTrue(TEXT("Handle move src not valid"), !Handle.IsValid());
+				TestTrue(TEXT("Handle move dst valid"), HandleCopy.IsValid());
+				TestTrue(TEXT("Handle soft refcount"), HandleCopy->GetSoftRefCount() == 1);
+				TestTrue(TEXT("Handle hard refcount"), HandleCopy->GetHardRefCount() == 0);
+
+				Handle = MoveTemp(HandleCopy);
+				TestTrue(TEXT("Handle move src not valid"), !HandleCopy.IsValid());
+				TestTrue(TEXT("Handle move dst valid"), Handle.IsValid());
 				TestTrue(TEXT("Handle soft refcount"), Handle->GetSoftRefCount() == 1);
 				TestTrue(TEXT("Handle hard refcount"), Handle->GetHardRefCount() == 0);
-				{
-					FWorldPartitionHandle HandleCopy(MoveTemp(Handle));
-					TestTrue(TEXT("Handle move src not valid"), !Handle.IsValid());
-					TestTrue(TEXT("Handle move dst valid"), HandleCopy.IsValid());
-					TestTrue(TEXT("Handle soft refcount"), HandleCopy->GetSoftRefCount() == 1);
-					TestTrue(TEXT("Handle hard refcount"), HandleCopy->GetHardRefCount() == 0);
-
-					Handle = MoveTemp(HandleCopy);
-					TestTrue(TEXT("Handle move src not valid"), !HandleCopy.IsValid());
-					TestTrue(TEXT("Handle move dst valid"), Handle.IsValid());
-					TestTrue(TEXT("Handle soft refcount"), Handle->GetSoftRefCount() == 1);
-					TestTrue(TEXT("Handle hard refcount"), Handle->GetHardRefCount() == 0);
-				}
-
-				// Reference move
+			}
+
+			// Reference move
+			TestTrue(TEXT("Reference soft refcount"), Reference->GetSoftRefCount() == 0);
+			TestTrue(TEXT("Reference hard refcount"), Reference->GetHardRefCount() == 1);
+			{
+				FWorldPartitionReference ReferenceCopy(MoveTemp(Reference));
+				TestTrue(TEXT("Reference move src not valid"), !Reference.IsValid());
+				TestTrue(TEXT("Reference move dst valid"), ReferenceCopy.IsValid());
+				TestTrue(TEXT("Reference soft refcount"), ReferenceCopy->GetSoftRefCount() == 0);
+				TestTrue(TEXT("Reference hard refcount"), ReferenceCopy->GetHardRefCount() == 1);
+
+				Reference = MoveTemp(ReferenceCopy);
+				TestTrue(TEXT("Reference move src not valid"), !ReferenceCopy.IsValid());
+				TestTrue(TEXT("Reference move dst valid"), Reference.IsValid());
 				TestTrue(TEXT("Reference soft refcount"), Reference->GetSoftRefCount() == 0);
 				TestTrue(TEXT("Reference hard refcount"), Reference->GetHardRefCount() == 1);
-				{
-					FWorldPartitionReference ReferenceCopy(MoveTemp(Reference));
-					TestTrue(TEXT("Reference move src not valid"), !Reference.IsValid());
-					TestTrue(TEXT("Reference move dst valid"), ReferenceCopy.IsValid());
-					TestTrue(TEXT("Reference soft refcount"), ReferenceCopy->GetSoftRefCount() == 0);
-					TestTrue(TEXT("Reference hard refcount"), ReferenceCopy->GetHardRefCount() == 1);
-
-					Reference = MoveTemp(ReferenceCopy);
-					TestTrue(TEXT("Reference move src not valid"), !ReferenceCopy.IsValid());
-					TestTrue(TEXT("Reference move dst valid"), Reference.IsValid());
-					TestTrue(TEXT("Reference soft refcount"), Reference->GetSoftRefCount() == 0);
-					TestTrue(TEXT("Reference hard refcount"), Reference->GetHardRefCount() == 1);
-				}
-
-				// Handle reference move
-				{
-					FWorldPartitionHandle HandleFromReference(MoveTemp(Reference));
-					TestTrue(TEXT("Handle move src not valid"), !Reference.IsValid());
-					TestTrue(TEXT("Handle move dst valid"), HandleFromReference.IsValid());
-					TestTrue(TEXT("Handle soft refcount"), HandleFromReference->GetSoftRefCount() == 1);
-					TestTrue(TEXT("Handle hard refcount"), HandleFromReference->GetHardRefCount() == 0);
-
-					Reference = MoveTemp(HandleFromReference);
-					TestTrue(TEXT("Handle move src not valid"), !HandleFromReference.IsValid());
-					TestTrue(TEXT("Handle move dst valid"), Reference.IsValid());
-					TestTrue(TEXT("Handle soft refcount"), Reference->GetSoftRefCount() == 0);
-					TestTrue(TEXT("Handle hard refcount"), Reference->GetHardRefCount() == 1);
-				}
-
-				// Reference handle move
-				{
-					FWorldPartitionReference ReferenceFromHandle(MoveTemp(Handle));
-					TestTrue(TEXT("Reference move src not valid"), !Handle.IsValid());
-					TestTrue(TEXT("Reference move dst valid"), ReferenceFromHandle.IsValid());
-					TestTrue(TEXT("Reference soft refcount"), ReferenceFromHandle->GetSoftRefCount() == 0);
-					TestTrue(TEXT("Reference hard refcount"), ReferenceFromHandle->GetHardRefCount() == 1);
-
-					Handle = MoveTemp(ReferenceFromHandle);
-					TestTrue(TEXT("Reference move src not valid"), !ReferenceFromHandle.IsValid());
-					TestTrue(TEXT("Reference move dst valid"), Handle.IsValid());
-					TestTrue(TEXT("Reference soft refcount"), Handle->GetSoftRefCount() == 1);
-					TestTrue(TEXT("Reference hard refcount"), Handle->GetHardRefCount() == 0);
-				}
 			}
 
-			// Make sure the container gets destroyed so we can test destrucing dangling handles
-			CollectGarbage(GARBAGE_COLLECTION_KEEPFLAGS);
+			// Handle reference move
+			{
+				FWorldPartitionHandle HandleFromReference(MoveTemp(Reference));
+				TestTrue(TEXT("Handle move src not valid"), !Reference.IsValid());
+				TestTrue(TEXT("Handle move dst valid"), HandleFromReference.IsValid());
+				TestTrue(TEXT("Handle soft refcount"), HandleFromReference->GetSoftRefCount() == 1);
+				TestTrue(TEXT("Handle hard refcount"), HandleFromReference->GetHardRefCount() == 0);
+
+				Reference = MoveTemp(HandleFromReference);
+				TestTrue(TEXT("Handle move src not valid"), !HandleFromReference.IsValid());
+				TestTrue(TEXT("Handle move dst valid"), Reference.IsValid());
+				TestTrue(TEXT("Handle soft refcount"), Reference->GetSoftRefCount() == 0);
+				TestTrue(TEXT("Handle hard refcount"), Reference->GetHardRefCount() == 1);
+			}
+
+			// Reference handle move
+			{
+				FWorldPartitionReference ReferenceFromHandle(MoveTemp(Handle));
+				TestTrue(TEXT("Reference move src not valid"), !Handle.IsValid());
+				TestTrue(TEXT("Reference move dst valid"), ReferenceFromHandle.IsValid());
+				TestTrue(TEXT("Reference soft refcount"), ReferenceFromHandle->GetSoftRefCount() == 0);
+				TestTrue(TEXT("Reference hard refcount"), ReferenceFromHandle->GetHardRefCount() == 1);
+
+				Handle = MoveTemp(ReferenceFromHandle);
+				TestTrue(TEXT("Reference move src not valid"), !ReferenceFromHandle.IsValid());
+				TestTrue(TEXT("Reference move dst valid"), Handle.IsValid());
+				TestTrue(TEXT("Reference soft refcount"), Handle->GetSoftRefCount() == 1);
+				TestTrue(TEXT("Reference hard refcount"), Handle->GetHardRefCount() == 0);
+			}
 		}
 
-		TestTrue(TEXT("Invalid container"), true);
+		// Make sure the container gets destroyed so we can test destructing dangling handles
+		{
+			TestTrue(TEXT("Invalid container test"), Handle.IsValid());
+			TestTrue(TEXT("Invalid container test"), Reference.IsValid());
+
+			CollectGarbage(GARBAGE_COLLECTION_KEEPFLAGS);
+
+			TestFalse(TEXT("Invalid container test"), Handle.IsValid());
+			TestFalse(TEXT("Invalid container test"), Reference.IsValid());
+		}
 #endif
 		return true;
 	}
