Add Group/SortPriority for UMaterialExpressionCurveAtlasRowParameter::GetParameterValue

diff --git a/Engine/Source/Runtime/Engine/Private/Materials/MaterialExpressions.cpp b/Engine/Source/Runtime/Engine/Private/Materials/MaterialExpressions.cpp
--- a/Engine/Source/Runtime/Engine/Private/Materials/MaterialExpressions.cpp
+++ b/Engine/Source/Runtime/Engine/Private/Materials/MaterialExpressions.cpp
@@ -19215,6 +19215,8 @@ bool UMaterialExpressionCurveAtlasRowParameter::GetParameterValue(FMaterialParam
 
 	OutMeta.Value = Value;
 	OutMeta.ExpressionGuid = ExpressionGUID;
+	OutMeta.Group = Group;
+	OutMeta.SortPriority = SortPriority;
 	return true;
 }
 
