No longer treat editor only child actors as transient since a mechanism to programmatically flag them as editor only has been added since the transient hack was put in place

diff --git a/Engine/Source/Runtime/Engine/Private/Components/ChildActorComponent.cpp b/Engine/Source/Runtime/Engine/Private/Components/ChildActorComponent.cpp
--- a/Engine/Source/Runtime/Engine/Private/Components/ChildActorComponent.cpp
+++ b/Engine/Source/Runtime/Engine/Private/Components/ChildActorComponent.cpp
@@ -646,7 +646,7 @@ void UChildActorComponent::CreateChildActor()
 
 #if WITH_EDITOR
 				Params.bCreateActorPackage = false;
-				Params.OverridePackage = GetOwner()->GetExternalPackage();
+				Params.OverridePackage = (MyOwner ? MyOwner->GetExternalPackage() : nullptr);
 				Params.OverrideActorGuid = CachedInstanceData ? CachedInstanceData->ChildActorGUID : FGuid();
 #endif
 				if (ChildActorTemplate && ChildActorTemplate->GetClass() == ChildActorClass)
@@ -658,10 +658,9 @@ void UChildActorComponent::CreateChildActor()
 				{
 					Params.ObjectFlags &= ~RF_Transactional;
 				}
-				if (HasAllFlags(RF_Transient) || IsEditorOnly() || (MyOwner && (MyOwner->HasAllFlags(RF_Transient) || MyOwner->IsEditorOnly())))
+				if (HasAllFlags(RF_Transient) || (MyOwner && MyOwner->HasAllFlags(RF_Transient)))
 				{
-					// If this component or its owner are transient or editor only, set our created actor to transient. 
-					// We can't programatically set editor only on an actor so this is the best option
+					// If this component or its owner are transient, set our created actor to transient. 
 					Params.ObjectFlags |= RF_Transient;
 				}
 
@@ -674,6 +673,12 @@ void UChildActorComponent::CreateChildActor()
 				// If spawn was successful, 
 				if(ChildActor != nullptr)
 				{
+					if (IsEditorOnly() || (MyOwner && MyOwner->IsEditorOnly()))
+					{
+						// If this component or its owner are editor only, set our created actor to editor only. 
+						ChildActor->bIsEditorOnlyActor = true;
+					}
+
 					ChildActorName = ChildActor->GetFName();
 
 					// Parts that we deferred from SpawnActor
