Fix crash when destroying landscape actor descriptors during garbage collection.

diff --git a/Engine/Source/Runtime/Engine/Private/WorldPartition/Landscape/LandscapeActorDesc.cpp b/Engine/Source/Runtime/Engine/Private/WorldPartition/Landscape/LandscapeActorDesc.cpp
--- a/Engine/Source/Runtime/Engine/Private/WorldPartition/Landscape/LandscapeActorDesc.cpp
+++ b/Engine/Source/Runtime/Engine/Private/WorldPartition/Landscape/LandscapeActorDesc.cpp
@@ -58,12 +58,13 @@ void FLandscapeActorDesc::OnUnregister()
 {
 	FPartitionActorDesc::OnUnregister();
 
-	if (ULandscapeInfo* LandscapeInfo = ULandscapeInfo::FindOrCreate(Container->GetWorld(), GridGuid))
+	if (ULandscapeInfo* LandscapeInfo = ULandscapeInfo::Find(Container->GetWorld(), GridGuid))
 	{
 		FWorldPartitionHandle Handle(Container, GetGuid());
-		check(Handle.IsValid());
-		LandscapeInfo->ProxyHandles.Remove(Handle);
+		if (Handle.IsValid())
+		{
+			LandscapeInfo->ProxyHandles.Remove(Handle);
+		}
 	}
 }
-
 #endif
diff --git a/Engine/Source/Runtime/Engine/Private/WorldPartition/Landscape/LandscapeSplineActorDesc.cpp b/Engine/Source/Runtime/Engine/Private/WorldPartition/Landscape/LandscapeSplineActorDesc.cpp
--- a/Engine/Source/Runtime/Engine/Private/WorldPartition/Landscape/LandscapeSplineActorDesc.cpp
+++ b/Engine/Source/Runtime/Engine/Private/WorldPartition/Landscape/LandscapeSplineActorDesc.cpp
@@ -57,12 +57,13 @@ void FLandscapeSplineActorDesc::OnUnregister()
 {
 	FWorldPartitionActorDesc::OnUnregister();
 
-	if (ULandscapeInfo* LandscapeInfo = ULandscapeInfo::FindOrCreate(Container->GetWorld(), LandscapeGuid))
+	if (ULandscapeInfo* LandscapeInfo = ULandscapeInfo::Find(Container->GetWorld(), LandscapeGuid))
 	{
 		FWorldPartitionHandle Handle(Container, GetGuid());
-		check(Handle.IsValid());
-		LandscapeInfo->SplineHandles.Remove(Handle);
+		if (Handle.IsValid())
+		{
+			LandscapeInfo->SplineHandles.Remove(Handle);
+		}
 	}
 }
-
 #endif
