World Partition: moved actor descriptors instancing code into a virtual fonction implemented by containers.

diff --git a/Engine/Source/Runtime/Engine/Private/WorldPartition/WorldPartition.cpp b/Engine/Source/Runtime/Engine/Private/WorldPartition/WorldPartition.cpp
--- a/Engine/Source/Runtime/Engine/Private/WorldPartition/WorldPartition.cpp
+++ b/Engine/Source/Runtime/Engine/Private/WorldPartition/WorldPartition.cpp
@@ -1085,6 +1085,17 @@ void UWorldPartition::OnActorDescUnregistered(const FWorldPartitionActorDesc& Ac
 	}
 }
 
+bool UWorldPartition::GetInstancingContext(const FLinkerInstancingContext*& OutInstancingContext, FSoftObjectPathFixupArchive*& OutSoftObjectPathFixupArchive) const
+{
+	if (InstancingContext.IsInstanced())
+	{
+		OutInstancingContext = &InstancingContext;
+		OutSoftObjectPathFixupArchive = InstancingSoftObjectPathFixupArchive.Get();
+		return true;
+	}
+	return false;
+}
+
 void UWorldPartition::HashActorDesc(FWorldPartitionActorDesc* ActorDesc)
 {
 	check(ActorDesc);
diff --git a/Engine/Source/Runtime/Engine/Private/WorldPartition/WorldPartitionActorDesc.cpp b/Engine/Source/Runtime/Engine/Private/WorldPartition/WorldPartitionActorDesc.cpp
--- a/Engine/Source/Runtime/Engine/Private/WorldPartition/WorldPartitionActorDesc.cpp
+++ b/Engine/Source/Runtime/Engine/Private/WorldPartition/WorldPartitionActorDesc.cpp
@@ -328,10 +328,10 @@ AActor* FWorldPartitionActorDesc::Load() const
 	{
 		const FLinkerInstancingContext* InstancingContext = nullptr;
 		FSoftObjectPathFixupArchive* SoftObjectPathFixupArchive = nullptr;
-		if (UWorldPartition* WorldPartition = Cast<UWorldPartition>(Container); WorldPartition && WorldPartition->InstancingContext.IsInstanced())
+
+		if (Container)
 		{
-			InstancingContext = &WorldPartition->InstancingContext;
-			SoftObjectPathFixupArchive = WorldPartition->InstancingSoftObjectPathFixupArchive.Get();
+			Container->GetInstancingContext(InstancingContext, SoftObjectPathFixupArchive);
 		}
 
 		UPackage* Package = nullptr;
diff --git a/Engine/Source/Runtime/Engine/Public/WorldPartition/ActorDescContainer.h b/Engine/Source/Runtime/Engine/Public/WorldPartition/ActorDescContainer.h
--- a/Engine/Source/Runtime/Engine/Public/WorldPartition/ActorDescContainer.h
+++ b/Engine/Source/Runtime/Engine/Public/WorldPartition/ActorDescContainer.h
@@ -45,6 +45,10 @@ public:
 	AActor* PinActor(const FGuid& ActorGuid);
 	void UnpinActor(const FGuid& ActorGuid);
 	bool IsActorPinned(const FGuid& ActorGuid) const { return PinnedActors.Contains(ActorGuid); }
+
+	/** Instancing support */
+	virtual bool GetInstancingContext(const FLinkerInstancingContext*& OutInstancingContext, FSoftObjectPathFixupArchive*& OutSoftObjectPathFixupArchive) const { return false; }
+
 public:
 	DECLARE_EVENT_OneParam(UWorldPartition, FActorDescAddedEvent, FWorldPartitionActorDesc*);
 	FActorDescAddedEvent OnActorDescAddedEvent;
diff --git a/Engine/Source/Runtime/Engine/Public/WorldPartition/WorldPartition.h b/Engine/Source/Runtime/Engine/Public/WorldPartition/WorldPartition.h
--- a/Engine/Source/Runtime/Engine/Public/WorldPartition/WorldPartition.h
+++ b/Engine/Source/Runtime/Engine/Public/WorldPartition/WorldPartition.h
@@ -96,7 +96,7 @@ private:
 	void OnBeginPlay();
 	void OnEndPlay();
 
-	// UActorDescContainer events
+	//~ Begin UActorDescContainer Interface
 	virtual void OnActorDescAdded(FWorldPartitionActorDesc* NewActorDesc) override;
 	virtual void OnActorDescRemoved(FWorldPartitionActorDesc* ActorDesc) override;
 	virtual void OnActorDescUpdating(FWorldPartitionActorDesc* ActorDesc) override;
@@ -104,6 +104,9 @@ private:
 
 	virtual void OnActorDescRegistered(const FWorldPartitionActorDesc&) override;
 	virtual void OnActorDescUnregistered(const FWorldPartitionActorDesc&) override;
+
+	virtual bool GetInstancingContext(const FLinkerInstancingContext*& OutInstancingContext, FSoftObjectPathFixupArchive*& OutSoftObjectPathFixupArchive) const override;
+	//~ End UActorDescContainer Interface
 #endif
 
 	static void WorldPartitionOnLevelRemovedFromWorld(class ULevel* Level, UWorld* InWorld);
