Call IsPackageExternal() instead of manually testing for RF_HasExternalPackage.

diff --git a/Engine/Source/Runtime/Engine/Private/ActorFolder.cpp b/Engine/Source/Runtime/Engine/Private/ActorFolder.cpp
--- a/Engine/Source/Runtime/Engine/Private/ActorFolder.cpp
+++ b/Engine/Source/Runtime/Engine/Private/ActorFolder.cpp
@@ -47,7 +47,7 @@ bool UActorFolder::IsAsset() const
 {
 	// Actor Folders are considered assets to allow using the asset logic for save dialogs, etc.
 	// Also, they return true even if pending kill, in order to show up as deleted in these dialogs.
-	return HasAnyFlags(RF_HasExternalPackage) && !GetPackage()->HasAnyFlags(RF_Transient) && !HasAnyFlags(RF_Transient | RF_ClassDefaultObject);
+	return IsPackageExternal() && !GetPackage()->HasAnyFlags(RF_Transient) && !HasAnyFlags(RF_Transient | RF_ClassDefaultObject);
 }
 
 namespace ActorFolder
