- Actor folder objects feature is no longer active while cooking or while running with -game
- For PIE partitioned worlds, propagate PersistentLevel's used actor folders in each generated runtime level. This allows to remove fallback to persistent level in ULevel::GetActorFolder.
- Modified actors in partitioned worlds are now in the correct folders in PIE

diff --git a/Engine/Source/Runtime/Engine/Classes/Engine/Level.h b/Engine/Source/Runtime/Engine/Classes/Engine/Level.h
--- a/Engine/Source/Runtime/Engine/Classes/Engine/Level.h
+++ b/Engine/Source/Runtime/Engine/Classes/Engine/Level.h
@@ -68,12 +68,12 @@ struct ENGINE_API FLevelActorFoldersHelper
 {
 private:
 	static void SetUseActorFolders(ULevel* InLevel, bool bInEnabled);
-	static void AddActorFolder(ULevel* InLevel, UActorFolder* InActorFolder, bool bInShouldDirtyLevel);
+	static void AddActorFolder(ULevel* InLevel, UActorFolder* InActorFolder, bool bInShouldDirtyLevel, bool bInShouldBroadcast = true);
 
 	friend class UWorld;
 	friend class UActorFolder;
 	friend class UWorldPartitionConvertCommandlet;
-	friend class FWorldPartitionLevelHelper;
+	friend class UWorldPartitionLevelStreamingDynamic;
 };
 
 class ENGINE_API FLevelPartitionOperationScope
diff --git a/Engine/Source/Runtime/Engine/Private/Level.cpp b/Engine/Source/Runtime/Engine/Private/Level.cpp
--- a/Engine/Source/Runtime/Engine/Private/Level.cpp
+++ b/Engine/Source/Runtime/Engine/Private/Level.cpp
@@ -110,13 +110,16 @@ void FLevelActorFoldersHelper::SetUseActorFolders(ULevel* InLevel, bool bInEnabl
 	InLevel->SetUseActorFoldersInternal(bInEnabled);
 }
 
-void FLevelActorFoldersHelper::AddActorFolder(ULevel* InLevel, UActorFolder* InActorFolder, bool bInShouldDirtyLevel)
+void FLevelActorFoldersHelper::AddActorFolder(ULevel* InLevel, UActorFolder* InActorFolder, bool bInShouldDirtyLevel, bool bInShouldBroadcast)
 {
 	InLevel->Modify(bInShouldDirtyLevel);
 	check(InActorFolder->GetGuid().IsValid());
 	InLevel->ActorFolders.Add(InActorFolder->GetGuid(), InActorFolder);
 
-	GEngine->BroadcastActorFolderAdded(InActorFolder);
+	if (bInShouldBroadcast)
+	{
+		GEngine->BroadcastActorFolderAdded(InActorFolder);
+	}
 }
 
 FLevelPartitionOperationScope::FLevelPartitionOperationScope(ULevel* InLevel)
@@ -351,31 +354,16 @@ void FLevelSimplificationDetails::PostLoadDeprecated()
 
 }
 
-TMap<FName, TWeakObjectPtr<UWorld> > ULevel::StreamedLevelsOwningWorld;
-
-#if WITH_EDITOR
-struct FOnWorldLoadedHandler
+static bool IsActorFolderObjectsFeatureAvailable()
 {
-	static void Init()
-	{
-		if (!bIsInitialized)
-		{
-			bIsInitialized = true;
-			FCoreUObjectDelegates::OnAssetLoaded.AddLambda([](UObject* Asset)
-			{
-				UWorld* World = Cast<UWorld>(Asset);
-				if (World && World->PersistentLevel)
-				{
-					World->PersistentLevel->OnLevelLoaded();
-				}
-			});
-		}
-	}
-private:
-	static bool bIsInitialized;
-};
-bool FOnWorldLoadedHandler::bIsInitialized = false;
+#if WITH_EDITOR
+	return !IsRunningCookCommandlet() && !IsRunningGame();
+#else
+	return false;
 #endif
+}
+
+TMap<FName, TWeakObjectPtr<UWorld> > ULevel::StreamedLevelsOwningWorld;
 
 ULevel::ULevel( const FObjectInitializer& ObjectInitializer )
 	:	UObject( ObjectInitializer )
@@ -398,15 +386,12 @@ ULevel::ULevel( const FObjectInitializer& ObjectInitializer )
 	bPromptWhenAddingToLevelBeforeCheckout = true;
 	bPromptWhenAddingToLevelOutsideBounds = true;
 	bUseActorFolders = false;
-	bFixupActorFoldersAtLoad = true;
+	bFixupActorFoldersAtLoad = IsActorFolderObjectsFeatureAvailable();
 #endif	
 	bActorClusterCreated = false;
 	bIsPartitioned = false;
 	bStaticComponentsRegisteredInStreamingManager = false;
 	IncrementalComponentState = EIncrementalComponentState::Init;
-#if WITH_EDITOR
-	FOnWorldLoadedHandler::Init();
-#endif
 }
 
 void ULevel::Initialize(const FURL& InURL)
@@ -887,7 +872,7 @@ void ULevel::PostLoad()
 	Super::PostLoad();
 
 #if WITH_EDITOR
-	if (IsUsingActorFolders() && IsUsingExternalObjects())
+	if (IsUsingActorFolders() && IsUsingExternalObjects() && IsActorFolderObjectsFeatureAvailable())
 	{
 		if (!bWasDuplicated)
 		{
@@ -2270,7 +2255,7 @@ void ULevel::CommitModelSurfaces()
 #if WITH_EDITOR
 void ULevel::FixupActorFolders()
 {
-	if (!IsUsingActorFolders())
+	if (!IsUsingActorFolders() || !IsActorFolderObjectsFeatureAvailable())
 	{
 		return;
 	}
@@ -2956,11 +2941,6 @@ UActorFolder* ULevel::GetActorFolder(const FName& InPath, bool bSkipDeleted) con
 				return ActorFolder;
 			}
 		}
-		// In PIE, fallback on persistent level
-		if (GetOutermost()->HasAnyPackageFlags(PKG_PlayInEditor) && (this != GetWorld()->PersistentLevel))
-		{
-			return GetWorld()->PersistentLevel->GetActorFolder(InPath, bSkipDeleted);
-		}
 	}
 	return nullptr;
 }
@@ -2978,11 +2958,6 @@ UActorFolder* ULevel::GetActorFolder(const FGuid& InGuid, bool bSkipDeleted) con
 			}
 			return ActorFolder;
 		}
-		// In PIE, fallback on persistent level
-		if (GetOutermost()->HasAnyPackageFlags(PKG_PlayInEditor) && (this != GetWorld()->PersistentLevel))
-		{
-			return GetWorld()->PersistentLevel->GetActorFolder(InGuid, bSkipDeleted);
-		}
 	}
 	return nullptr;
 }
diff --git a/Engine/Source/Runtime/Engine/Private/WorldPartition/WorldPartitionActorDescView.cpp b/Engine/Source/Runtime/Engine/Private/WorldPartition/WorldPartitionActorDescView.cpp
--- a/Engine/Source/Runtime/Engine/Private/WorldPartition/WorldPartitionActorDescView.cpp
+++ b/Engine/Source/Runtime/Engine/Private/WorldPartition/WorldPartitionActorDescView.cpp
@@ -124,6 +124,11 @@ const FGuid& FWorldPartitionActorDescView::GetParentActor() const
 	return ActorDesc->GetParentActor();
 }
 
+const FGuid& FWorldPartitionActorDescView::GetFolderGuid() const
+{
+	return ActorDesc->GetFolderGuid();
+}
+
 bool FWorldPartitionActorDescView::GetContainerInstance(const UActorDescContainer*& OutLevelContainer, FTransform& OutLevelTransform, EContainerClusterMode& OutClusterMode) const
 {
 	return ActorDesc->GetContainerInstance(OutLevelContainer, OutLevelTransform, OutClusterMode);
diff --git a/Engine/Source/Runtime/Engine/Private/WorldPartition/WorldPartitionLevelHelper.cpp b/Engine/Source/Runtime/Engine/Private/WorldPartition/WorldPartitionLevelHelper.cpp
--- a/Engine/Source/Runtime/Engine/Private/WorldPartition/WorldPartitionLevelHelper.cpp
+++ b/Engine/Source/Runtime/Engine/Private/WorldPartition/WorldPartitionLevelHelper.cpp
@@ -214,12 +214,6 @@ ULevel* FWorldPartitionLevelHelper::CreateEmptyLevelForRuntimeCell(const UWorld*
 	check(InWorld->PersistentLevel->Model && NewLevel->Model);
 	NewLevel->Model->LightingGuid = InWorld->PersistentLevel->Model->LightingGuid;
 
-	// Propagate ActorFolder flag to new levels
-	if (InWorld->IsPlayInEditor() && InWorld->PersistentLevel->IsUsingActorFolders())
-	{
-		FLevelActorFoldersHelper::SetUseActorFolders(NewLevel, true);
-	}
-
 	return NewLevel;
 }
 
diff --git a/Engine/Source/Runtime/Engine/Private/WorldPartition/WorldPartitionLevelStreamingDynamic.cpp b/Engine/Source/Runtime/Engine/Private/WorldPartition/WorldPartitionLevelStreamingDynamic.cpp
--- a/Engine/Source/Runtime/Engine/Private/WorldPartition/WorldPartitionLevelStreamingDynamic.cpp
+++ b/Engine/Source/Runtime/Engine/Private/WorldPartition/WorldPartitionLevelStreamingDynamic.cpp
@@ -56,6 +56,7 @@ void UWorldPartitionLevelStreamingDynamic::Initialize(const UWorldPartitionRunti
 	StreamingPriority = 0;
 	ChildPackages = InCell.GetPackages();
 	ActorContainer = InCell.ActorContainer;
+	ActorFolders = InCell.GetActorFolders().Array();
 
 	UWorld* OuterWorld = InCell.GetOuterUWorldPartition()->GetTypedOuter<UWorld>();
 	OriginalLevelPackageName = OuterWorld->GetPackage()->GetLoadedPath().GetPackageFName();
@@ -91,6 +92,19 @@ void UWorldPartitionLevelStreamingDynamic::CreateRuntimeLevel()
 	RuntimeLevel = FWorldPartitionLevelHelper::CreateEmptyLevelForRuntimeCell(PlayWorld, WorldAssetPackageName);
 	check(RuntimeLevel);
 
+	// Propagate ActorFolder flag to the runtime level and prepare its ActorFolders list
+	if (PlayWorld->PersistentLevel->IsUsingActorFolders() && ActorFolders.Num())
+	{
+		FLevelActorFoldersHelper::SetUseActorFolders(RuntimeLevel, true);
+		for (const FGuid& ActorFolderGuid : ActorFolders)
+		{
+			if (UActorFolder* ActorFolder = PlayWorld->PersistentLevel->GetActorFolder(ActorFolderGuid))
+			{
+				FLevelActorFoldersHelper::AddActorFolder(RuntimeLevel, ActorFolder, /*bInShouldDirtyLevel*/ false, /*bInShouldBroadcast*/ false);
+			}
+		}
+	}
+
 	// Set flag here as this level isn't async loaded
 	RuntimeLevel->bClientOnlyVisible = bClientOnlyVisible;
 
@@ -191,7 +205,7 @@ bool UWorldPartitionLevelStreamingDynamic::RequestLevel(UWorld* InPersistentWorl
 			check(UWorld::FindWorldInPackage(CellLevelPackage));
 			check(RuntimeLevel->OwningWorld);
 			check(RuntimeLevel->OwningWorld->WorldType == EWorldType::PIE || ((IsRunningGame() || IsRunningDedicatedServer()) && RuntimeLevel->OwningWorld->WorldType == EWorldType::Game));
-	
+
 			if (IssueLoadRequests())
 			{
 				// Editor immediately blocks on load and we also block if background level streaming is disabled.
diff --git a/Engine/Source/Runtime/Engine/Private/WorldPartition/WorldPartitionRuntimeLevelStreamingCell.cpp b/Engine/Source/Runtime/Engine/Private/WorldPartition/WorldPartitionRuntimeLevelStreamingCell.cpp
--- a/Engine/Source/Runtime/Engine/Private/WorldPartition/WorldPartitionRuntimeLevelStreamingCell.cpp
+++ b/Engine/Source/Runtime/Engine/Private/WorldPartition/WorldPartitionRuntimeLevelStreamingCell.cpp
@@ -102,6 +102,10 @@ void UWorldPartitionRuntimeLevelStreamingCell::AddActorToCell(const FWorldPartit
 {
 	check(!ActorDescView.GetActorIsEditorOnly());
 	Packages.Emplace(ActorDescView.GetActorPackage(), ActorDescView.GetActorPath(), InContainerID, InContainerTransform, InContainer->GetContainerPackage());
+	if (ActorDescView.GetFolderGuid().IsValid())
+	{
+		ActorFolders.Add(ActorDescView.GetFolderGuid());
+	}
 }
 
 UWorldPartitionLevelStreamingDynamic* UWorldPartitionRuntimeLevelStreamingCell::CreateLevelStreaming(const FString& InPackageName) const
diff --git a/Engine/Source/Runtime/Engine/Public/WorldPartition/WorldPartitionActorDescView.h b/Engine/Source/Runtime/Engine/Public/WorldPartition/WorldPartitionActorDescView.h
--- a/Engine/Source/Runtime/Engine/Public/WorldPartition/WorldPartitionActorDescView.h
+++ b/Engine/Source/Runtime/Engine/Public/WorldPartition/WorldPartitionActorDescView.h
@@ -42,6 +42,7 @@ public:
 	uint32 GetTag() const;
 	const FGuid& GetParentActor() const;
 	FName GetActorName() const;
+	const FGuid& GetFolderGuid() const;
 
 	bool GetContainerInstance(const UActorDescContainer*& OutLevelContainer, FTransform& OutLevelTransform, EContainerClusterMode& OutClusterMode) const;
 
diff --git a/Engine/Source/Runtime/Engine/Public/WorldPartition/WorldPartitionLevelStreamingDynamic.h b/Engine/Source/Runtime/Engine/Public/WorldPartition/WorldPartitionLevelStreamingDynamic.h
--- a/Engine/Source/Runtime/Engine/Public/WorldPartition/WorldPartitionLevelStreamingDynamic.h
+++ b/Engine/Source/Runtime/Engine/Public/WorldPartition/WorldPartitionLevelStreamingDynamic.h
@@ -53,6 +53,7 @@ private:
 	FName OriginalLevelPackageName;
 	TArray<FWorldPartitionRuntimeCellObjectMapping> ChildPackages;
 	TArray<FWorldPartitionRuntimeCellObjectMapping> ChildPackagesToLoad;
+	TArray<FGuid> ActorFolders;
 
 #if WITH_EDITORONLY_DATA
 	UPROPERTY()
diff --git a/Engine/Source/Runtime/Engine/Public/WorldPartition/WorldPartitionRuntimeLevelStreamingCell.h b/Engine/Source/Runtime/Engine/Public/WorldPartition/WorldPartitionRuntimeLevelStreamingCell.h
--- a/Engine/Source/Runtime/Engine/Public/WorldPartition/WorldPartitionRuntimeLevelStreamingCell.h
+++ b/Engine/Source/Runtime/Engine/Public/WorldPartition/WorldPartitionRuntimeLevelStreamingCell.h
@@ -41,6 +41,7 @@ class UWorldPartitionRuntimeLevelStreamingCell : public UWorldPartitionRuntimeSp
 	//~End UWorldPartitionRuntimeCell Interface
 
 	const TArray<FWorldPartitionRuntimeCellObjectMapping>& GetPackages() const { return Packages; }
+	const TSet<FGuid>& GetActorFolders() const { return ActorFolders; }
 #endif
 
 private:
@@ -61,6 +62,9 @@ private:
 #if WITH_EDITORONLY_DATA
 	UPROPERTY()
 	TArray<FWorldPartitionRuntimeCellObjectMapping> Packages;
+
+	UPROPERTY()
+	TSet<FGuid> ActorFolders;
 #endif
 
 	UPROPERTY()
