Type layout fix for UAnimSequence
Reduces the size of UAnimSequence from 480 bytes to 464 (eliminating 16 bytes of padding)

diff --git a/Engine/Source/Runtime/Engine/Classes/Animation/AnimSequence.h b/Engine/Source/Runtime/Engine/Classes/Animation/AnimSequence.h
--- a/Engine/Source/Runtime/Engine/Classes/Animation/AnimSequence.h
+++ b/Engine/Source/Runtime/Engine/Classes/Animation/AnimSequence.h
@@ -309,14 +309,14 @@ public:
 	UPROPERTY(EditAnywhere, Category=AdditiveSettings, meta=(DisplayName = "Base Pose Type"))
 	TEnumAsByte<enum EAdditiveBasePoseType> RefPoseType;
 
+	/* Additve reference frame if RefPoseType == AnimFrame **/
+	UPROPERTY(EditAnywhere, Category = AdditiveSettings)
+	int32 RefFrameIndex;
+	
 	/* Additive reference animation if it's relevant - i.e. AnimScaled or AnimFrame **/
 	UPROPERTY(EditAnywhere, Category=AdditiveSettings, meta=(DisplayName = "Base Pose Animation"))
 	TObjectPtr<class UAnimSequence> RefPoseSeq;
 
-	/* Additve reference frame if RefPoseType == AnimFrame **/
-	UPROPERTY(EditAnywhere, Category=AdditiveSettings)
-	int32 RefFrameIndex;
-
 	/** Base pose to use when retargeting */
 	UPROPERTY(EditAnywhere, AssetRegistrySearchable, Category=Animation)
 	FName RetargetSource;
@@ -388,12 +388,7 @@ public:
 	int32 MarkerDataUpdateCounter;
 #endif // WITH_EDITORONLY_DATA
 
-	/** Authored Sync markers */
-	UPROPERTY()
-	TArray<FAnimSyncMarker>		AuthoredSyncMarkers;
 
-	/** List of Unique marker names in this animation sequence */
-	TArray<FName>				UniqueMarkerNames;
 
 public:
 	//~ Begin UObject Interface
@@ -1005,6 +1000,15 @@ public:
 	// Should we be always using our raw data (i.e is our compressed data stale)
 	bool bUseRawDataOnly;
 
+public:
+	/** Authored Sync markers */
+	UPROPERTY()
+	TArray<FAnimSyncMarker>		AuthoredSyncMarkers;
+
+	/** List of Unique marker names in this animation sequence */
+	TArray<FName>				UniqueMarkerNames;
+
+private:
 #if WITH_EDITOR
 	// Are we currently compressing this animation
 	bool bCompressionInProgress;
