Deprecating UWorld IsClient and IsServer in favor of GetNetMode/IsNetMode

PR #8581: Fix UWorld::IsClient() in single instance PIE (Contributed by [REDACTED_PERSON])

diff --git a/Engine/Source/Runtime/Engine/Classes/Engine/World.h b/Engine/Source/Runtime/Engine/Classes/Engine/World.h
--- a/Engine/Source/Runtime/Engine/Classes/Engine/World.h
+++ b/Engine/Source/Runtime/Engine/Classes/Engine/World.h
@@ -3299,9 +3299,11 @@ public:
 	bool Listen( FURL& InURL );
 
 	/** @return true if this level is a client */
+	UE_DEPRECATED(5.0, "Use GetNetMode or IsNetMode instead for more accurate results.")
 	bool IsClient() const;
 
 	/** @return true if this level is a server */
+	UE_DEPRECATED(5.0, "Use GetNetMode or IsNetMode instead for more accurate results")
 	bool IsServer() const;
 
 	/** @return true if the world is in the paused state */
diff --git a/Engine/Source/Runtime/Engine/Private/ActorConstruction.cpp b/Engine/Source/Runtime/Engine/Private/ActorConstruction.cpp
--- a/Engine/Source/Runtime/Engine/Private/ActorConstruction.cpp
+++ b/Engine/Source/Runtime/Engine/Private/ActorConstruction.cpp
@@ -515,7 +515,7 @@ void AActor::RerunConstructionScripts()
 
 		// Exchange net roles before running construction scripts
 		UWorld *OwningWorld = GetWorld();
-		if (OwningWorld && !OwningWorld->IsServer())
+		if (OwningWorld && OwningWorld->IsNetMode(NM_Client))
 		{
 			ExchangeNetRoles(true);
 		}
diff --git a/Engine/Source/Runtime/Engine/Private/Level.cpp b/Engine/Source/Runtime/Engine/Private/Level.cpp
--- a/Engine/Source/Runtime/Engine/Private/Level.cpp
+++ b/Engine/Source/Runtime/Engine/Private/Level.cpp
@@ -2462,7 +2462,8 @@ ALevelScriptActor* ULevel::GetLevelScriptActor() const
 void ULevel::InitializeNetworkActors()
 {
 	check( OwningWorld );
-	bool			bIsServer				= OwningWorld->IsServer();
+
+	const bool bIsClient = OwningWorld->IsNetMode(NM_Client);
 
 	// Kill non relevant client actors and set net roles correctly
 	for( int32 ActorIndex=0; ActorIndex<Actors.Num(); ActorIndex++ )
@@ -2487,7 +2488,7 @@ void ULevel::InitializeNetworkActors()
 					}
 				}
 
-				if (!bIsServer)
+				if (bIsClient)
 				{
 					if (!Actor->bNetLoadOnClient)
 					{
diff --git a/Engine/Source/Runtime/Engine/Private/LocalPlayer.cpp b/Engine/Source/Runtime/Engine/Private/LocalPlayer.cpp
--- a/Engine/Source/Runtime/Engine/Private/LocalPlayer.cpp
+++ b/Engine/Source/Runtime/Engine/Private/LocalPlayer.cpp
@@ -281,7 +281,7 @@ void ULocalPlayer::PlayerRemoved()
 bool ULocalPlayer::SpawnPlayActor(const FString& URL,FString& OutError, UWorld* InWorld)
 {
 	check(InWorld);
-	if ( InWorld->IsServer() )
+	if (!InWorld->IsNetMode(NM_Client))
 	{
 		FURL PlayerURL(NULL, *URL, TRAVEL_Absolute);
 
diff --git a/Engine/Source/Runtime/Engine/Private/PlayerController.cpp b/Engine/Source/Runtime/Engine/Private/PlayerController.cpp
--- a/Engine/Source/Runtime/Engine/Private/PlayerController.cpp
+++ b/Engine/Source/Runtime/Engine/Private/PlayerController.cpp
@@ -547,7 +547,7 @@ void APlayerController::ServerNotifyLoadedWorld_Implementation(FName WorldPackag
 
 	// Only valid for calling, for PC's in the process of seamless traveling
 	// NOTE: SeamlessTravelCount tracks client seamless travel, through the serverside gameplay code; this should not be replaced.
-	if (CurWorld != NULL && CurWorld->IsServer() && SeamlessTravelCount > 0 && LastCompletedSeamlessTravelCount < SeamlessTravelCount)
+	if (CurWorld != NULL && !CurWorld->IsNetMode(NM_Client) && SeamlessTravelCount > 0 && LastCompletedSeamlessTravelCount < SeamlessTravelCount)
 	{
 		// Update our info on what world the client is in
 		UNetConnection* const Connection = Cast<UNetConnection>(Player);
diff --git a/Engine/Source/Runtime/Engine/Private/UnrealEngine.cpp b/Engine/Source/Runtime/Engine/Private/UnrealEngine.cpp
--- a/Engine/Source/Runtime/Engine/Private/UnrealEngine.cpp
+++ b/Engine/Source/Runtime/Engine/Private/UnrealEngine.cpp
@@ -13042,7 +13042,7 @@ bool UEngine::HandleStreamMapCommand( const TCHAR* Cmd, FOutputDevice& Ar, UWorl
 #if WITH_SERVER_CODE
 bool UEngine::HandleServerTravelCommand( const TCHAR* Cmd, FOutputDevice& Ar, UWorld* InWorld )
 {
-	if( InWorld->IsServer() )
+	if (!InWorld->IsNetMode(NM_Client))
 	{
 		FString URL(Cmd);
 		FString MapName, Options;
@@ -15283,7 +15283,7 @@ bool UEngine::CommitMapChange( FWorldContext &Context )
 					StreamingLevel->SetShouldBeVisible(true);
 
 #if WITH_SERVER_CODE
-					if (Context.World()->IsServer())
+					if (!Context.World()->IsNetMode(NM_Client))
 					{
 						const bool bShouldBeLoaded = StreamingLevel->ShouldBeLoaded();
 						const bool bShouldBeVisible = StreamingLevel->ShouldBeVisible();
diff --git a/Engine/Source/Runtime/Engine/Private/World.cpp b/Engine/Source/Runtime/Engine/Private/World.cpp
--- a/Engine/Source/Runtime/Engine/Private/World.cpp
+++ b/Engine/Source/Runtime/Engine/Private/World.cpp
@@ -4620,7 +4620,7 @@ void UWorld::ClearNetDriver(UNetDriver* Driver)
 
 bool UWorld::SetGameMode(const FURL& InURL)
 {
-	if( IsServer() && !AuthorityGameMode )
+	if (!IsNetMode(NM_Client) && !AuthorityGameMode)
 	{
 		AuthorityGameMode = GetGameInstance()->CreateGameModeForURL(InURL, this);
 		if( AuthorityGameMode != NULL )
@@ -6089,10 +6089,12 @@ void UWorld::NotifyControlMessage(UNetConnection* Connection, uint8 MessageType,
 				{
 					if (ChildIndex >= 0)
 					{
-						SwapConnection = Connection->Children.IsValidIndex(ChildIndex) ? ToRawPtr(Connection->Children[ChildIndex]) : NULL;
+						SwapConnection = Connection->Children.IsValidIndex(ChildIndex) ? ToRawPtr(Connection->Children[ChildIndex]) : nullptr;
 					}
+
 					bool bSuccess = false;
-					if (SwapConnection != NULL)
+
+					if (SwapConnection != nullptr)
 					{
 						bSuccess = DestroySwappedPC(SwapConnection);
 					}
diff --git a/Engine/Source/Runtime/Experimental/ChaosSolverEngine/Private/Chaos/ChaosDebugDrawComponent.cpp b/Engine/Source/Runtime/Experimental/ChaosSolverEngine/Private/Chaos/ChaosDebugDrawComponent.cpp
--- a/Engine/Source/Runtime/Experimental/ChaosSolverEngine/Private/Chaos/ChaosDebugDrawComponent.cpp
+++ b/Engine/Source/Runtime/Experimental/ChaosSolverEngine/Private/Chaos/ChaosDebugDrawComponent.cpp
@@ -303,7 +303,7 @@ void UChaosDebugDrawComponent::HandlePostWorldInitialization(UWorld* World, cons
 void UChaosDebugDrawComponent::CreateDebugDrawActor(UWorld* World)
 {
 #if CHAOS_DEBUG_DRAW
-	if (!World->IsClient())
+	if (!World->IsNetMode(NM_Client))
 		return;
 
 	static FName NAME_ChaosDebugDrawActor = TEXT("ChaosDebugDrawActor");
