World Partition:
- Automatically set the actor descriptors container when added to a container, instead of directly trying to determine the container type.
- Updated actor descriptors handle unit tests to include owned container tests.

diff --git a/Engine/Source/Runtime/Engine/Private/Tests/WorldPartition/WorldPartitionHandleTests.cpp b/Engine/Source/Runtime/Engine/Private/Tests/WorldPartition/WorldPartitionHandleTests.cpp
--- a/Engine/Source/Runtime/Engine/Private/Tests/WorldPartition/WorldPartitionHandleTests.cpp
+++ b/Engine/Source/Runtime/Engine/Private/Tests/WorldPartition/WorldPartitionHandleTests.cpp
@@ -23,6 +23,9 @@ namespace WorldPartitionTests
 		FWorldPartitionHandle Handle = FWorldPartitionHandle(ActorDescContainer, FGuid(TEXT("5D9F93BA407A811AFDDDAAB4F1CECC6A")));
 		FWorldPartitionReference Reference = FWorldPartitionReference(ActorDescContainer, FGuid(TEXT("0D2B04D240BE5DE58FE437A8D2DBF5C9")));
 
+		TestTrue(TEXT("Handle container"), Handle->GetContainer() == ActorDescContainer);
+		TestTrue(TEXT("Reference container"), Reference->GetContainer() == ActorDescContainer);
+
 		TestTrue(TEXT("Handle soft refcount"), Handle->GetSoftRefCount() == 1);
 		TestTrue(TEXT("Handle hard refcount"), Handle->GetHardRefCount() == 0);
 		TestTrue(TEXT("Reference soft refcount"), Reference->GetSoftRefCount() == 0);
diff --git a/Engine/Source/Runtime/Engine/Private/WorldPartition/ActorDescContainer.cpp b/Engine/Source/Runtime/Engine/Private/WorldPartition/ActorDescContainer.cpp
--- a/Engine/Source/Runtime/Engine/Private/WorldPartition/ActorDescContainer.cpp
+++ b/Engine/Source/Runtime/Engine/Private/WorldPartition/ActorDescContainer.cpp
@@ -84,7 +84,7 @@ void UActorDescContainer::Initialize(UWorld* InWorld, FName InPackageName, TFunc
 
 				TUniquePtr<FWorldPartitionActorDesc> NewActorDesc(AActor::StaticCreateClassActorDesc(ActorDescInitData.NativeClass));
 
-				NewActorDesc->Init(this, ActorDescInitData);
+				NewActorDesc->Init(ActorDescInitData);
 			
 				//check(bIsValidClass || (NewActorDesc->GetActorIsEditorOnly() && IsRunningGame()));
 
@@ -106,6 +106,7 @@ void UActorDescContainer::Initialize(UWorld* InWorld, FName InPackageName, TFunc
 
 		if (ActorDesc.IsValid())
 		{
+			ActorDesc->SetContainer(this);
 			AddActorDescriptor(ActorDesc.Release());
 		}
 	}
@@ -208,6 +209,7 @@ void UActorDescContainer::OnObjectPreSave(UObject* Object, FObjectPreSaveContext
 				else
 				{
 					FWorldPartitionActorDesc* const AddedActorDesc = AddActor(Actor);
+					AddedActorDesc->SetContainer(this);
 					OnActorDescAdded(AddedActorDesc);
 				}
 			}
@@ -249,6 +251,7 @@ void UActorDescContainer::RemoveActor(const FGuid& ActorGuid)
 	{
 		OnActorDescRemoved(ExistingActorDesc->Get());
 		RemoveActorDescriptor(ExistingActorDesc->Get());
+		ExistingActorDesc->Get()->SetContainer(nullptr);
 		ExistingActorDesc->Reset();
 	}
 }
diff --git a/Engine/Source/Runtime/Engine/Private/WorldPartition/LevelInstance/LevelInstanceActorDesc.cpp b/Engine/Source/Runtime/Engine/Private/WorldPartition/LevelInstance/LevelInstanceActorDesc.cpp
--- a/Engine/Source/Runtime/Engine/Private/WorldPartition/LevelInstance/LevelInstanceActorDesc.cpp
+++ b/Engine/Source/Runtime/Engine/Private/WorldPartition/LevelInstance/LevelInstanceActorDesc.cpp
@@ -51,12 +51,12 @@ void FLevelInstanceActorDesc::Init(const AActor* InActor)
 	DesiredRuntimeBehavior = LevelInstanceActor->GetDesiredRuntimeBehavior();
 }
 
-void FLevelInstanceActorDesc::Init(UActorDescContainer* InContainer, const FWorldPartitionActorDescInitData& DescData)
+void FLevelInstanceActorDesc::Init(const FWorldPartitionActorDescInitData& DescData)
 {
 	ALevelInstance* CDO = DescData.NativeClass->GetDefaultObject<ALevelInstance>();
 	DesiredRuntimeBehavior = CDO->GetDefaultRuntimeBehavior();
 
-	FWorldPartitionActorDesc::Init(InContainer, DescData);
+	FWorldPartitionActorDesc::Init(DescData);
 }
 
 bool FLevelInstanceActorDesc::Equals(const FWorldPartitionActorDesc* Other) const
diff --git a/Engine/Source/Runtime/Engine/Private/WorldPartition/WorldPartitionActorDesc.cpp b/Engine/Source/Runtime/Engine/Private/WorldPartition/WorldPartitionActorDesc.cpp
--- a/Engine/Source/Runtime/Engine/Private/WorldPartition/WorldPartitionActorDesc.cpp
+++ b/Engine/Source/Runtime/Engine/Private/WorldPartition/WorldPartitionActorDesc.cpp
@@ -13,7 +13,7 @@
 #include "Editor/GroupActor.h"
 #include "UObject/UE5MainStreamObjectVersion.h"
 #include "UObject/UE5ReleaseStreamObjectVersion.h"
-#include "WorldPartition/WorldPartition.h"
+#include "WorldPartition/WorldPartitionLog.h"
 #include "WorldPartition/ActorDescContainer.h"
 #include "WorldPartition/DataLayer/DataLayer.h"
 #include "WorldPartition/DataLayer/WorldDataLayers.h"
@@ -64,7 +64,7 @@ void FWorldPartitionActorDesc::Init(const AActor* InActor)
 		ParentActor = AttachParentActor->GetActorGuid();
 	}
 	
-	TArray<AActor*> ActorReferences = ActorsReferencesUtils::GetExternalActorReferences((AActor*)InActor);
+	TArray<AActor*> ActorReferences = ActorsReferencesUtils::GetExternalActorReferences(const_cast<AActor*>(InActor));
 
 	if (ActorReferences.Num())
 	{
@@ -77,14 +77,11 @@ void FWorldPartitionActorDesc::Init(const AActor* InActor)
 
 	ActorLabel = *InActor->GetActorLabel(false);
 
-	if (UActorDescContainer* ActorContainer = InActor->GetLevel()->GetWorldPartition())
-	{
-		ActorPtr = (AActor*)InActor;
-		Container = ActorContainer;
-	}
+	Container = nullptr;
+	ActorPtr = const_cast<AActor*>(InActor);
 }
 
-void FWorldPartitionActorDesc::Init(UActorDescContainer* InContainer, const FWorldPartitionActorDescInitData& DescData)
+void FWorldPartitionActorDesc::Init(const FWorldPartitionActorDescInitData& DescData)
 {
 	ActorPackage = DescData.PackageName;
 	ActorPath = DescData.ActorPath;
@@ -102,7 +99,7 @@ void FWorldPartitionActorDesc::Init(UActorDescContainer* InContainer, const FWor
 	// Serialize metadata payload
 	Serialize(MetadataAr);
 
-	Container = InContainer;
+	Container = nullptr;
 	ActorPtr = FindObject<AActor>(nullptr, *ActorPath.ToString());
 }
 
diff --git a/Engine/Source/Runtime/Engine/Public/WorldPartition/LevelInstance/LevelInstanceActorDesc.h b/Engine/Source/Runtime/Engine/Public/WorldPartition/LevelInstance/LevelInstanceActorDesc.h
--- a/Engine/Source/Runtime/Engine/Public/WorldPartition/LevelInstance/LevelInstanceActorDesc.h
+++ b/Engine/Source/Runtime/Engine/Public/WorldPartition/LevelInstance/LevelInstanceActorDesc.h
@@ -31,7 +31,7 @@ public:
 protected:
 	FLevelInstanceActorDesc();
 	virtual void Init(const AActor* InActor) override;
-	virtual void Init(UActorDescContainer* InContainer, const FWorldPartitionActorDescInitData& DescData) override;
+	virtual void Init(const FWorldPartitionActorDescInitData& DescData) override;
 	virtual bool Equals(const FWorldPartitionActorDesc* Other) const override;
 	virtual void TransferFrom(const FWorldPartitionActorDesc* From) override;
 	virtual void Serialize(FArchive& Ar) override;
diff --git a/Engine/Source/Runtime/Engine/Public/WorldPartition/WorldPartitionActorDesc.h b/Engine/Source/Runtime/Engine/Public/WorldPartition/WorldPartitionActorDesc.h
--- a/Engine/Source/Runtime/Engine/Public/WorldPartition/WorldPartitionActorDesc.h
+++ b/Engine/Source/Runtime/Engine/Public/WorldPartition/WorldPartitionActorDesc.h
@@ -131,6 +131,12 @@ protected:
 		return Container;
 	}
 
+	void SetContainer(UActorDescContainer* InContainer)
+	{
+		check(!Container || !InContainer);
+		Container = InContainer;
+	}
+
 public:
 	const TArray<FGuid>& GetReferences() const
 	{
@@ -149,7 +155,7 @@ public:
 	void UnregisterActor();
 
 	virtual void Init(const AActor* InActor);
-	virtual void Init(UActorDescContainer* InContainer, const FWorldPartitionActorDescInitData& DescData);
+	virtual void Init(const FWorldPartitionActorDescInitData& DescData);
 
 	virtual bool Equals(const FWorldPartitionActorDesc* Other) const;
 
@@ -162,6 +168,7 @@ protected:
 
 	virtual void TransferFrom(const FWorldPartitionActorDesc* From)
 	{
+		Container = From->Container;
 		SoftRefCount = From->SoftRefCount;
 		HardRefCount = From->HardRefCount;
 	}
