MaterialParameterCollection Blueprintability
+ Adding GetScalarParameterNames/GetVectorParameterNames blueprint callable functions
+ Adding GetScalar/VectorParameterDefaultValue blueprint callable functions to expose the non-transient values to blueprint
+ Adding code-only GetScalar/VectorParameterIndexByName functions to make it easier to access parameters within the internal arrays
+ Adding editor- and code-only functions to set default values on the collection asset

diff --git a/Engine/Source/Runtime/Engine/Classes/Materials/MaterialParameterCollection.h b/Engine/Source/Runtime/Engine/Classes/Materials/MaterialParameterCollection.h
--- a/Engine/Source/Runtime/Engine/Classes/Materials/MaterialParameterCollection.h
+++ b/Engine/Source/Runtime/Engine/Classes/Materials/MaterialParameterCollection.h
@@ -28,7 +28,7 @@ struct FCollectionParameterBase
 	{
 		FPlatformMisc::CreateGuid(Id);
 	}
-
+	
 	/** The name of the parameter.  Changing this name will break any blueprints that reference the parameter. */
 	UPROPERTY(EditAnywhere, Category=Parameter)
 	FName ParameterName;
@@ -65,7 +65,7 @@ struct FCollectionVectorParameter : public FCollectionParameterBase
 	{
 		ParameterName = FName(TEXT("Vector"));
 	}
-
+	
 	UPROPERTY(EditAnywhere, Category=Parameter)
 	FLinearColor DefaultValue;
 };
@@ -88,7 +88,51 @@ class UMaterialParameterCollection : public UObject
 
 	UPROPERTY(EditAnywhere, Category=Material, Meta = (TitleProperty = "ParameterName"))
 	TArray<FCollectionVectorParameter> VectorParameters;
+	
+#if WITH_EDITOR
+	/** Set the default value of a scalar parameter on the Material Parameter Collection asset itself by struct **/
+	bool SetScalarParameterDefaultValueByInfo(FCollectionScalarParameter ScalarParameter);
 
+	/** Set the default value of a scalar parameter on the Material Parameter Collection asset itself by name **/
+	bool SetScalarParameterDefaultValue(FName ParameterName, const float Value);
+
+	/** Set the default value of a vector parameter on the Material Parameter Collection asset itself by struct **/
+	bool SetVectorParameterDefaultValueByInfo(FCollectionVectorParameter VectorParameter);
+	
+	/** Set the default value of a vector parameter on the Material Parameter Collection asset itself by name **/
+	bool SetVectorParameterDefaultValue(FName ParameterName, const FLinearColor& Value);
+#endif // WITH_EDITOR
+
+	/** Get the index in the ScalarParameters array of the parameter matching the input parameter name, returns -1 if not found **/
+	int32 GetScalarParameterIndexByName(FName ParameterName);
+
+	/** Get the index in the VectorParameters array of the parameter matching the input parameter name, returns -1 if not found **/
+	int32 GetVectorParameterIndexByName(FName ParameterName);
+
+	/** Returns an array of the names of all the scalar parameters in this Material Parameter Collection **/
+	UFUNCTION(BlueprintCallable, Category="Rendering|Material", meta=(Keywords="GetScalarParameterNames"))
+	TArray<FName> GetScalarParameterNames();
+
+	/** Returns an array of the names of all the vector parameters in this Material Parameter Collection **/
+	UFUNCTION(BlueprintCallable, Category="Rendering|Material", meta=(Keywords="GetVectorParameterNames"))
+	TArray<FName> GetVectorParameterNames();
+
+	/** Gets the default value of a scalar parameter from a material collection.
+	 * @param ParameterName - The name of the value to get the value of
+	 * @param bParameterFound - if a parameter with the input name was found
+	 * @returns the value of the parameter
+	 **/
+	UFUNCTION(BlueprintCallable, Category="Rendering|Material", meta=(Keywords="GetFloatParameterDefaultValue"))
+	float GetScalarParameterDefaultValue(FName ParameterName, bool& bParameterFound);
+
+	/** Gets the default value of a scalar parameter from a material collection.
+	 * @param ParameterName - The name of the value to get the value of
+	 * @param bParameterFound - if a parameter with the input name was found
+	 * @returns the value of the parameter
+	 **/
+	UFUNCTION(BlueprintCallable, Category="Rendering|Material", meta=(Keywords="GetFloatParameterDefaultValue"))
+	FLinearColor GetVectorParameterDefaultValue(FName ParameterName, bool& bParameterFound);
+	
 	//~ Begin UObject Interface.
 #if WITH_EDITOR
 	using Super::PreEditChange;
diff --git a/Engine/Source/Runtime/Engine/Private/Materials/ParameterCollection.cpp b/Engine/Source/Runtime/Engine/Private/Materials/ParameterCollection.cpp
--- a/Engine/Source/Runtime/Engine/Private/Materials/ParameterCollection.cpp
+++ b/Engine/Source/Runtime/Engine/Private/Materials/ParameterCollection.cpp
@@ -164,6 +164,128 @@ void SanitizeParameters(TArray<ParameterType>& Parameters)
 int32 PreviousNumScalarParameters = 0;
 int32 PreviousNumVectorParameters = 0;
 
+bool UMaterialParameterCollection::SetScalarParameterDefaultValueByInfo(FCollectionScalarParameter ScalarParameter)
+{
+	if(GetScalarParameterByName(ScalarParameter.ParameterName))
+	{
+		// if the input parameter exists, pass the name and value down to SetScalarParameterDefaultValue
+		// since we want to preserve the Guid of the parameter that's already on the asset
+		return SetScalarParameterDefaultValue(ScalarParameter.ParameterName, ScalarParameter.DefaultValue);
+	}
+	// otherwise, return false
+	return false;
+}
+
+bool UMaterialParameterCollection::SetScalarParameterDefaultValue(FName ParameterName, const float Value)
+{
+	// make sure the input parameter name exists in the array
+	const int32 ParameterIndex = GetScalarParameterIndexByName(ParameterName);
+	if(ParameterIndex == -1)
+	{
+		return false;
+	}
+	// if so, change the value on the parameter in the array itself to maintain its GUID
+	ScalarParameters[ParameterIndex].DefaultValue = Value;
+
+	// and return a positive result
+	return true;
+}
+
+bool UMaterialParameterCollection::SetVectorParameterDefaultValueByInfo(FCollectionVectorParameter VectorParameter)
+{
+	if(GetScalarParameterByName(VectorParameter.ParameterName))
+    {
+    	// if the input parameter exists, pass the name and value down to SetVectorParameterDefaultValue
+    	// since we want to preserve the Guid of the parameter that's already on the asset
+    	return SetVectorParameterDefaultValue(VectorParameter.ParameterName, VectorParameter.DefaultValue);
+    }
+	// otherwise, return false
+    return false;
+}
+
+bool UMaterialParameterCollection::SetVectorParameterDefaultValue(FName ParameterName, const FLinearColor& Value)
+{
+	// make sure the input parameter name exists in the array
+	const int32 ParameterIndex = GetVectorParameterIndexByName(ParameterName);
+	if(ParameterIndex == -1)
+	{
+		return false;
+	}
+	// if so, change the value on the parameter in the array itself to maintain its GUID
+	VectorParameters[ParameterIndex].DefaultValue = Value;
+
+	// and return a positive result
+	return true;
+}
+
+int32 UMaterialParameterCollection::GetScalarParameterIndexByName(FName ParameterName)
+{
+	// loop over all the available scalar parameters and look for a name match
+	for (int32 ParameterIndex = 0; ParameterIndex < ScalarParameters.Num(); ParameterIndex++)
+	{
+		if(ScalarParameters[ParameterIndex].ParameterName == ParameterName)
+		{
+			return ParameterIndex;
+		}
+	}
+	// if not found, return -1
+	return -1;
+}
+
+int32 UMaterialParameterCollection::GetVectorParameterIndexByName(FName ParameterName)
+{
+	// loop over all the available vector parameters and look for a name match
+	for (int32 ParameterIndex = 0; ParameterIndex < VectorParameters.Num(); ParameterIndex++)
+	{
+		if(VectorParameters[ParameterIndex].ParameterName == ParameterName)
+		{
+			return ParameterIndex;
+		}
+	}
+	// if not found, return -1
+	return -1;
+}
+
+TArray<FName> UMaterialParameterCollection::GetScalarParameterNames()
+{
+	TArray<FName> Names;
+	GetParameterNames(Names, false);
+	return Names;
+}
+
+TArray<FName> UMaterialParameterCollection::GetVectorParameterNames()
+{
+	TArray<FName> Names;
+	GetParameterNames(Names, true);
+	return Names;
+}
+
+float UMaterialParameterCollection::GetScalarParameterDefaultValue(FName ParameterName, bool& bParameterFound)
+{
+	const int32 ParameterIndex = GetScalarParameterIndexByName(ParameterName);
+	bParameterFound = true;
+	if(ParameterIndex == -1)
+	{
+		bParameterFound = false;
+		return 0.0;
+	}
+	
+	return ScalarParameters[ParameterIndex].DefaultValue;
+}
+
+FLinearColor UMaterialParameterCollection::GetVectorParameterDefaultValue(FName ParameterName, bool& bParameterFound)
+{
+	const int32 ParameterIndex = GetVectorParameterIndexByName(ParameterName);
+	bParameterFound = true;
+	if(ParameterIndex == -1)
+	{
+		bParameterFound = false;
+		return FLinearColor::Black;
+	}
+	
+	return VectorParameters[ParameterIndex].DefaultValue;
+}
+
 void UMaterialParameterCollection::PreEditChange(FProperty* PropertyThatWillChange)
 {
 	Super::PreEditChange(PropertyThatWillChange);
