New material HLSL translator supports functions using control flow

diff --git a/Engine/Source/Editor/MaterialEditor/Private/MaterialEditor.cpp b/Engine/Source/Editor/MaterialEditor/Private/MaterialEditor.cpp
--- a/Engine/Source/Editor/MaterialEditor/Private/MaterialEditor.cpp
+++ b/Engine/Source/Editor/MaterialEditor/Private/MaterialEditor.cpp
@@ -549,8 +549,15 @@ void FMaterialEditor::InitMaterialEditor( const EToolkitMode::Type Mode, const T
 		// Support undo/redo for the material function if it exists
 		MaterialFunction->SetFlags(RF_Transactional);
 
+		MaterialFunction->MaterialGraph = Material->MaterialGraph;
+		MaterialFunction->EditorMaterial = Material;
+
 		Material->Expressions = MaterialFunction->FunctionExpressions;
 		Material->EditorComments = MaterialFunction->FunctionEditorComments;
+		Material->ExpressionExecBegin = MaterialFunction->ExpressionExecBegin;
+		Material->ExpressionExecEnd = MaterialFunction->ExpressionExecEnd;
+		Material->bEnableExecWire = MaterialFunction->IsUsingControlFlow();
+		Material->bEnableNewHLSLGenerator = MaterialFunction->IsUsingNewHLSLGenerator();
 
 		// Remove NULL entries, so the rest of the material editor can assume all entries of Material->Expressions are valid
 		// This can happen if an expression class was removed
@@ -690,6 +697,8 @@ void FMaterialEditor::InitMaterialEditor( const EToolkitMode::Type Mode, const T
 				SetPreviewExpression(FirstOutput);
 			}
 		}
+
+		//Material->CreateExecutionFlowExpressions();
 	}
 
 	// Store the name of this material (for the tutorial widget meta)
@@ -2375,6 +2384,8 @@ void FMaterialEditor::UpdatePreviewMaterial( bool bForce )
 		// The preview material's expressions array must stay up to date before recompiling 
 		// So that RebuildMaterialFunctionInfo will see all the nested material functions that may need to be updated
 		ExpressionPreviewMaterial->Expressions = Material->Expressions;
+		ExpressionPreviewMaterial->ExpressionExecBegin = Material->ExpressionExecBegin;
+		ExpressionPreviewMaterial->ExpressionExecEnd = Material->ExpressionExecEnd;
 
 		if (MaterialFunction)
 		{
@@ -2478,6 +2489,8 @@ bool FMaterialEditor::UpdateOriginalMaterial()
 		// Copy the expressions back from the preview material
 		MaterialFunction->FunctionExpressions = Material->Expressions;
 		MaterialFunction->FunctionEditorComments = Material->EditorComments;
+		MaterialFunction->ExpressionExecBegin = Material->ExpressionExecBegin;
+		MaterialFunction->ExpressionExecEnd = Material->ExpressionExecEnd;
 
 		// Preserve the thumbnail info
 		UThumbnailInfo* OriginalThumbnailInfo = MaterialFunction->ParentFunction->ThumbnailInfo;
@@ -4840,6 +4853,8 @@ void FMaterialEditor::SetPreviewExpression(UMaterialExpression* NewPreviewExpres
 			// Create the expression preview material if it hasnt already been created
 			ExpressionPreviewMaterial = NewObject<UPreviewMaterial>(GetTransientPackage(), NAME_None, RF_Public);
 			ExpressionPreviewMaterial->bIsPreviewMaterial = true;
+			ExpressionPreviewMaterial->bEnableNewHLSLGenerator = Material->IsUsingNewHLSLGenerator();
+			ExpressionPreviewMaterial->bEnableExecWire = Material->IsUsingControlFlow();
 			if (Material->IsUIMaterial())
 			{
 				ExpressionPreviewMaterial->MaterialDomain = MD_UI;
@@ -4867,6 +4882,8 @@ void FMaterialEditor::SetPreviewExpression(UMaterialExpression* NewPreviewExpres
 		// The expression preview material's expressions array must stay up to date before recompiling 
 		// So that RebuildMaterialFunctionInfo will see all the nested material functions that may need to be updated
 		ExpressionPreviewMaterial->Expressions = Material->Expressions;
+		ExpressionPreviewMaterial->ExpressionExecBegin = Material->ExpressionExecBegin;
+		ExpressionPreviewMaterial->ExpressionExecEnd = Material->ExpressionExecEnd;
 
 		// The preview window should now show the expression preview material
 		SetPreviewMaterial( ExpressionPreviewMaterial );
diff --git a/Engine/Source/Editor/MaterialEditor/Private/MaterialEditorUtilities.cpp b/Engine/Source/Editor/MaterialEditor/Private/MaterialEditorUtilities.cpp
--- a/Engine/Source/Editor/MaterialEditor/Private/MaterialEditorUtilities.cpp
+++ b/Engine/Source/Editor/MaterialEditor/Private/MaterialEditorUtilities.cpp
@@ -240,7 +240,7 @@ void FMaterialEditorUtilities::GetVisibleMaterialParameters(const UMaterial* Mat
 	TArray<FGetVisibleMaterialParametersFunctionState*> FunctionStack;
 	FunctionStack.Push(FunctionState.Get());
 
-	if (Material->IsCompiledWithExecutionFlow())
+	if (Material->IsUsingControlFlow())
 	{
 		GetVisibleMaterialParametersFromExpression(FMaterialExpressionKey(Material->ExpressionExecBegin, INDEX_NONE), MaterialInstance, VisibleExpressions, FunctionStack);
 	}
diff --git a/Engine/Source/Editor/UnrealEd/Private/MaterialGraph.cpp b/Engine/Source/Editor/UnrealEd/Private/MaterialGraph.cpp
--- a/Engine/Source/Editor/UnrealEd/Private/MaterialGraph.cpp
+++ b/Engine/Source/Editor/UnrealEd/Private/MaterialGraph.cpp
@@ -120,12 +120,6 @@ void UMaterialGraph::RebuildGraphInternal(const TMap<UMaterialExpression*, TArra
 		//^^^ New material properties go above here. ^^^^
 		MaterialInputs.Add(FMaterialInputInfo(LOCTEXT("MaterialAttributes", "Material Attributes"), MP_MaterialAttributes, LOCTEXT("MaterialAttributesToolTip", "Material Attributes")));
 
-		if (Material->IsCompiledWithExecutionFlow())
-		{
-			check(Material->ExpressionExecBegin);
-			InitExpressionNewNode<UMaterialGraphNode>(this, Material->ExpressionExecBegin, false);
-		}
-		
 		// Add Root Node
 		{
 			FGraphNodeCreator<UMaterialGraphNode_Root> NodeCreator(*this);
@@ -135,6 +129,24 @@ void UMaterialGraph::RebuildGraphInternal(const TMap<UMaterialExpression*, TArra
 		}
 	}
 
+	if (Material->IsUsingControlFlow())
+	{
+		if (ensure(Material->ExpressionExecBegin))
+		{
+			InitExpressionNewNode<UMaterialGraphNode>(this, Material->ExpressionExecBegin, false);
+		}
+		
+		if (MaterialFunction)
+		{
+			check(MaterialFunction->ExpressionExecBegin == Material->ExpressionExecBegin);
+			check(MaterialFunction->ExpressionExecEnd == Material->ExpressionExecEnd);
+			if (ensure(Material->ExpressionExecEnd))
+			{
+				InitExpressionNewNode<UMaterialGraphNode>(this, Material->ExpressionExecEnd, false);
+			}
+		}
+	}
+
 	TArray<UMaterialExpression*> ChildSubGraphExpressions;
 
 	// Composite's use reroutes under the hood that we don't want to create nodes for, gather their expressions for checking
@@ -380,9 +392,10 @@ void UMaterialGraph::LinkGraphNodesFromMaterial()
 				UMaterialExpression* ConnectedExpression = ExecOutput->GetExpression();
 				if (ConnectedExpression)
 				{
-					if (ConnectedExpression == Material->ExpressionExecEnd)
+					if (ConnectedExpression == Material->ExpressionExecEnd && RootNode)
 					{
-						// Exec end point is the root node
+						// Exec end point is the root node (assuming there is a RootNode)
+						// MaterialFunctions currently do not have a RootNode, and instead have the ExecEnd node as a stand-alone
 						Pin->MakeLinkTo(RootNode->GetExecInputPin());
 					}
 					else if (UMaterialGraphNode* GraphNode = Cast<UMaterialGraphNode>(ConnectedExpression->GraphNode))
diff --git a/Engine/Source/Editor/UnrealEd/Private/MaterialGraphNode_Root.cpp b/Engine/Source/Editor/UnrealEd/Private/MaterialGraphNode_Root.cpp
--- a/Engine/Source/Editor/UnrealEd/Private/MaterialGraphNode_Root.cpp
+++ b/Engine/Source/Editor/UnrealEd/Private/MaterialGraphNode_Root.cpp
@@ -76,7 +76,7 @@ void UMaterialGraphNode_Root::CreateInputPins()
 {
 	UMaterialGraph* MaterialGraph = CastChecked<UMaterialGraph>(GetGraph());
 
-	if (Material->IsCompiledWithExecutionFlow())
+	if (Material->IsUsingControlFlow())
 	{
 		// Create the execution pin
 		UEdGraphPin* NewPin = CreatePin(EGPD_Input, UMaterialGraphSchema::PC_Exec, NAME_None, NAME_None);
diff --git a/Engine/Source/Runtime/Engine/Classes/Materials/Material.h b/Engine/Source/Runtime/Engine/Classes/Materials/Material.h
--- a/Engine/Source/Runtime/Engine/Classes/Materials/Material.h
+++ b/Engine/Source/Runtime/Engine/Classes/Materials/Material.h
@@ -1090,7 +1090,7 @@ public:
 	/** Like IsPropertyActive(), but considers any state overriden by DerivedMaterial */
 	ENGINE_API bool IsPropertyActiveInDerived(EMaterialProperty InProperty, const UMaterialInterface* DerivedMaterial) const;
 
-	ENGINE_API bool IsCompiledWithExecutionFlow() const;
+	ENGINE_API bool IsUsingControlFlow() const;
 
 	ENGINE_API bool IsUsingNewHLSLGenerator() const;
 
diff --git a/Engine/Source/Runtime/Engine/Classes/Materials/MaterialExpressionFunctionOutput.h b/Engine/Source/Runtime/Engine/Classes/Materials/MaterialExpressionFunctionOutput.h
--- a/Engine/Source/Runtime/Engine/Classes/Materials/MaterialExpressionFunctionOutput.h
+++ b/Engine/Source/Runtime/Engine/Classes/Materials/MaterialExpressionFunctionOutput.h
@@ -65,6 +65,7 @@ class UMaterialExpressionFunctionOutput : public UMaterialExpression
 	virtual bool IsResultMaterialAttributes(int32 OutputIndex) override;
 	virtual bool IsResultStrataMaterial(int32 OutputIndex) override;
 	virtual void GatherStrataMaterialInfo(FStrataMaterialInfo& StrataMaterialInfo, int32 OutputIndex) override;
+	bool GenerateHLSLExpression(FMaterialHLSLGenerator& Generator, UE::HLSLTree::FScope& Scope, int32 OutputIndex, UE::HLSLTree::FExpression*& OutExpression);
 #endif // WITH_EDITOR
 	//~ End UMaterialExpression Interface
 
diff --git a/Engine/Source/Runtime/Engine/Classes/Materials/MaterialFunction.h b/Engine/Source/Runtime/Engine/Classes/Materials/MaterialFunction.h
--- a/Engine/Source/Runtime/Engine/Classes/Materials/MaterialFunction.h
+++ b/Engine/Source/Runtime/Engine/Classes/Materials/MaterialFunction.h
@@ -38,6 +38,13 @@ class UMaterialFunction : public UMaterialFunctionInterface
 	/** Array of material expressions, excluding Comments.  Used by the material editor. */
 	UPROPERTY()
 	TArray<TObjectPtr<UMaterialExpression>> FunctionExpressions;
+
+	/** The execution begin expression, if material is using an exec wire */
+	UPROPERTY()
+	TObjectPtr<class UMaterialExpressionExecBegin> ExpressionExecBegin;
+
+	UPROPERTY()
+	TObjectPtr<class UMaterialExpressionExecEnd> ExpressionExecEnd;
 #endif // WITH_EDITORONLY_DATA
 
 	/** Whether to list this function in the material function library, which is a window in the material editor that lists categorized functions. */
@@ -48,6 +55,12 @@ class UMaterialFunction : public UMaterialFunctionInterface
 	UPROPERTY(EditAnywhere, Category=MaterialFunction)
 	uint8 bPrefixParameterNames:1;
 
+	UPROPERTY(EditAnywhere, Category = MaterialFunction)
+	uint8 bEnableExecWire : 1;
+
+	UPROPERTY(EditAnywhere, Category = MaterialFunction)
+	uint8 bEnableNewHLSLGenerator : 1;
+
 #if WITH_EDITORONLY_DATA
 	/** 
 	 * Categories that this function belongs to in the material function library.  
@@ -72,12 +85,17 @@ class UMaterialFunction : public UMaterialFunctionInterface
 	UPROPERTY(transient)
 	TObjectPtr<UMaterial> PreviewMaterial;
 
+	// The UMaterial which represents this function while the function itself is open in the material editor
+	TObjectPtr<UMaterial> EditorMaterial;
+
 	UPROPERTY()
 	TArray<TObjectPtr<class UMaterialExpressionMaterialFunctionCall>> DependentFunctionExpressionCandidates;
 
 	/** Determines the blend mode when previewing a material function. */
 	UPROPERTY(EditAnywhere, Category = Preview, AssetRegistrySearchable)
 	TEnumAsByte<enum EBlendMode> PreviewBlendMode = BLEND_Opaque;
+
+	class UMaterialGraph* MaterialGraph = nullptr;
 private:
 	/** Transient flag used to track re-entrance in recursive functions like IsDependent. */
 	UPROPERTY(transient)
@@ -172,5 +190,10 @@ public:
 	ENGINE_API bool SetFontParameterValueEditorOnly(FName ParameterName, class UFont* InFontValue, int32 InFontPage);
 	ENGINE_API bool SetStaticComponentMaskParameterValueEditorOnly(FName ParameterName, bool R, bool G, bool B, bool A, FGuid OutExpressionGuid);
 	ENGINE_API bool SetStaticSwitchParameterValueEditorOnly(FName ParameterName, bool OutValue, FGuid OutExpressionGuid);
+
+	virtual bool IsUsingControlFlow() const override;
+	virtual bool IsUsingNewHLSLGenerator() const override;
+
+	void CreateExecutionFlowExpressions();
 #endif // WITH_EDITOR
 };
diff --git a/Engine/Source/Runtime/Engine/Classes/Materials/MaterialFunctionInterface.h b/Engine/Source/Runtime/Engine/Classes/Materials/MaterialFunctionInterface.h
--- a/Engine/Source/Runtime/Engine/Classes/Materials/MaterialFunctionInterface.h
+++ b/Engine/Source/Runtime/Engine/Classes/Materials/MaterialFunctionInterface.h
@@ -323,5 +323,8 @@ public:
 	ENGINE_API bool OverrideNamedFontParameter(const FHashedMaterialParameterInfo& ParameterInfo, class UFont*& OutFontValue, int32& OutFontPage);
 	ENGINE_API bool OverrideNamedStaticSwitchParameter(const FHashedMaterialParameterInfo& ParameterInfo, bool& OutValue, FGuid& OutExpressionGuid);
 	ENGINE_API bool OverrideNamedStaticComponentMaskParameter(const FHashedMaterialParameterInfo& ParameterInfo, bool& OutR, bool& OutG, bool& OutB, bool& OutA, FGuid& OutExpressionGuid);
+
+	virtual bool IsUsingControlFlow() const { return false; }
+	virtual bool IsUsingNewHLSLGenerator() const { return false; }
 #endif // WITH_EDITOR
 };
diff --git a/Engine/Source/Runtime/Engine/Private/HLSLTree/HLSLTree.cpp b/Engine/Source/Runtime/Engine/Private/HLSLTree/HLSLTree.cpp
--- a/Engine/Source/Runtime/Engine/Private/HLSLTree/HLSLTree.cpp
+++ b/Engine/Source/Runtime/Engine/Private/HLSLTree/HLSLTree.cpp
@@ -1,490 +1,133 @@
 // Copyright Epic Games, Inc. All Rights Reserved.
 #include "HLSLTree/HLSLTree.h"
+#include "HLSLTree/HLSLTreeEmit.h"
 #include "Misc/StringBuilder.h"
 #include "Misc/MemStack.h"
 #include "Misc/MemStackUtility.h"
 #include "Shader/ShaderTypes.h"
-#include "MaterialShared.h" // TODO - split preshader out into its own module
-#include "Misc/LargeWorldRenderPosition.h"
-
-static FSHAHash HashString(const TCHAR* String, int32 Length)
-{
-	FSHAHash Hash;
-	FSHA1::HashBuffer(String, Length * sizeof(TCHAR), Hash.Hash);
-	return Hash;
-}
-
-static FSHAHash HashString(const FStringBuilderBase& StringBuilder)
-{
-	return HashString(StringBuilder.GetData(), StringBuilder.Len());
-}
-
-static FSHAHash HashString(FStringView String)
-{
-	return HashString(String.GetData(), String.Len());
-}
+#include "Shader/Preshader.h"
 
 namespace UE
 {
 namespace HLSLTree
 {
 
-EExpressionEvaluation CombineEvaluations(EExpressionEvaluation Lhs, EExpressionEvaluation Rhs)
+/**
+ * Represents a phi node (see various topics on single static assignment)
+ * A phi node takes on a value based on the previous scope that was executed.
+ * In practice, this means the generated HLSL code will declare a local variable before all the previous scopes, then assign that variable the proper value from within each scope
+ */
+class FExpressionLocalPHI final : public FExpression
 {
-	if (Lhs == EExpressionEvaluation::None)
-	{
-		// If either is 'None', return the other
-		return Rhs;
-	}
-	else if (Rhs == EExpressionEvaluation::None)
-	{
-		return Lhs;
-	}
-	else if (Lhs == EExpressionEvaluation::Constant && Rhs == EExpressionEvaluation::Constant)
-	{
-		// 2 constants make a constant
-		return EExpressionEvaluation::Constant;
-	}
-	else if (Lhs == EExpressionEvaluation::Shader || Rhs == EExpressionEvaluation::Shader)
+public:
+	FExpressionLocalPHI(const FName& InLocalName, TArrayView<FScope*> InPreviousScopes) : LocalName(InLocalName), NumValues(InPreviousScopes.Num())
 	{
-		// If either requires shader, shader is required
-		return EExpressionEvaluation::Shader;
+		for (int32 i = 0; i < InPreviousScopes.Num(); ++i)
+		{
+			Scopes[i] = InPreviousScopes[i];
+			Values[i] = nullptr;
+		}
 	}
-	// Any combination of constants/preshader can make a preshader
-	return EExpressionEvaluation::Preshader;
-}
 
-FEmitContext::FEmitContext(FMemStackBase& InAllocator, FErrorHandlerInterface& InErrors, const Shader::FStructTypeRegistry& InTypeRegistry)
-	: Allocator(&InAllocator)
-	, Errors(&InErrors)
-	, TypeRegistry(&InTypeRegistry)
-{
-}
+	FExpressionLocalPHI(const FExpressionLocalPHI* Source, EDerivativeCoordinate Coord);
 
-FEmitContext::~FEmitContext()
-{
-}
+	virtual void ComputeAnalyticDerivatives(FTree& Tree, FExpressionDerivatives& OutResult) const override;
+	virtual bool PrepareValue(FEmitContext& Context, const FRequestedType& RequestedType, FPrepareValueResult& OutResult) const override;
+	virtual void EmitValueShader(FEmitContext& Context, FEmitShaderScope& Scope, const FRequestedType& RequestedType, FEmitValueShaderResult& OutResult) const override;
 
-const TCHAR* FEmitContext::AcquireLocalDeclarationCode()
-{
-	return MemStack::AllocateStringf(*Allocator, TEXT("Local%d"), NumExpressionLocals++);
-}
+	TArray<EDerivativeCoordinate, TInlineAllocator<8>> DerivativeChain;
+	FName LocalName;
+	FScope* Scopes[MaxNumPreviousScopes];
+	FExpression* Values[MaxNumPreviousScopes];
+	int32 NumValues = 0;
+};
 
-namespace Private
+/**
+ * Represents a call to a function that includes its own scope/control-flow
+ * Scope for the function will be linked into the generated material
+ */
+class FExpressionFunctionCall final : public FExpression
 {
-void MoveToScope(FEmitShaderCode* ShaderValue, FScope* Scope)
-{
-	if (ShaderValue->Scope != Scope)
-	{
-		FScope* NewScope = FScope::FindSharedParent(ShaderValue->Scope, Scope);
-		check(NewScope);
+public:
+	FExpressionFunctionCall(FFunction* InFunction, TArrayView<FExpression*> InInputExpressions, int32 InOutputIndex) : Function(InFunction), InputExpressions(InInputExpressions), OutputIndex(InOutputIndex) {}
 
-		ShaderValue->Scope = NewScope;
-		for (FEmitShaderCode* Dependency : ShaderValue->Dependencies)
-		{
-			MoveToScope(Dependency, NewScope);
-		}
-	}
-}
+	virtual bool PrepareValue(FEmitContext& Context, const FRequestedType& RequestedType, FPrepareValueResult& OutResult) const override;
+	virtual void EmitValueShader(FEmitContext& Context, FEmitShaderScope& Scope, const FRequestedType& RequestedType, FEmitValueShaderResult& OutResult) const override;
 
-void FormatArg_ShaderValue(FEmitShaderCode* ShaderValue, FEmitShaderValueDependencies& OutDependencies, FStringBuilderBase& OutCode)
-{
-	OutDependencies.Add(ShaderValue);
-	OutCode.Append(ShaderValue->Reference);
-}
+	FFunction* Function;
+	TArray<FExpression*, TInlineAllocator<8>> InputExpressions;
+	int32 OutputIndex;
+};
 
-} // namespace Private
+FUnaryOpDescription::FUnaryOpDescription()
+	: Name(nullptr), Operator(nullptr), PreshaderOpcode(Shader::EPreshaderOpcode::Nop)
+{}
 
-FEmitShaderCode* FEmitContext::EmitFormatCodeInternal(const Shader::FType& Type, const TCHAR* Format, bool bInline, const FFormatArgList& ArgList)
-{
-	check(!Type.IsVoid());
+FUnaryOpDescription::FUnaryOpDescription(const TCHAR* InName, const TCHAR* InOperator, Shader::EPreshaderOpcode InOpcode)
+	: Name(InName), Operator(InOperator), PreshaderOpcode(InOpcode)
+{}
 
-	TStringBuilder<1024> FormattedCode;
-	FEmitShaderValueDependencies Dependencies;
-	int32 ArgIndex = 0;
-	while (true)
-	{
-		const TCHAR Char = *Format++;
-		if (Char == 0)
-		{
-			break;
-		}
-		else if (Char == TEXT('%'))
-		{
-			const FFormatArgVariant& Arg = ArgList[ArgIndex++];
-			switch (Arg.Type)
-			{
-			case EFormatArgType::ShaderValue: Private::FormatArg_ShaderValue(Arg.ShaderValue, Dependencies, FormattedCode); break;
-			case EFormatArgType::String: FormattedCode.Append(Arg.String); break;
-			case EFormatArgType::Int: FormattedCode.Appendf(TEXT("%d"), Arg.Int); break;
-			default:
-				checkNoEntry();
-				break;
-			}
-		}
-		else
-		{
-			FormattedCode.AppendChar(Char);
-		}
-	}
-	checkf(ArgIndex == ArgList.Num(), TEXT("%d args were provided, but %d were used"), ArgList.Num(), ArgIndex);
+FBinaryOpDescription::FBinaryOpDescription()
+	: Name(nullptr), Operator(nullptr), PreshaderOpcode(Shader::EPreshaderOpcode::Nop)
+{}
 
-	return EmitCodeInternal(Type, FormattedCode.ToView(), bInline, Dependencies);
-}
+FBinaryOpDescription::FBinaryOpDescription(const TCHAR* InName, const TCHAR* InOperator, Shader::EPreshaderOpcode InOpcode)
+	: Name(InName), Operator(InOperator), PreshaderOpcode(InOpcode)
+{}
 
-FEmitShaderCode* FEmitContext::EmitCodeInternal(const Shader::FType& Type, FStringView Code, bool bInline, TArrayView<FEmitShaderCode*> Dependencies)
+FUnaryOpDescription GetUnaryOpDesription(EUnaryOp Op)
 {
-	FScope* CurrentScope = ScopeStack.Last();
-	check(IsScopeLive(CurrentScope));
-
-	FEmitShaderCode* ShaderValue = nullptr;
-
-	// Check to see if we've already generated code for an equivalent expression
-	const FSHAHash ShaderHash = HashString(Code);
-	FEmitShaderCode** const PrevShaderValue = ShaderValueMap.Find(ShaderHash);
-	if (PrevShaderValue)
+	switch (Op)
 	{
-		ShaderValue = *PrevShaderValue;
-		check(ShaderValue->Type == Type);
-		Private::MoveToScope(ShaderValue, CurrentScope);
+	case EUnaryOp::None: return FUnaryOpDescription(TEXT("None"), TEXT(""), Shader::EPreshaderOpcode::Nop); break;
+	case EUnaryOp::Neg: return FUnaryOpDescription(TEXT("Neg"), TEXT("-"), Shader::EPreshaderOpcode::Neg); break;
+	case EUnaryOp::Rcp: return FUnaryOpDescription(TEXT("Rcp"), TEXT("/"), Shader::EPreshaderOpcode::Rcp); break;
+	default: checkNoEntry(); return FUnaryOpDescription();
 	}
-	else
-	{
-		ShaderValue = new(*Allocator) FEmitShaderCode(CurrentScope, Type);
-		ShaderValue->Hash = ShaderHash;
-		ShaderValue->Dependencies = MemStack::AllocateArrayView(*Allocator, Dependencies);
-		if (bInline)
-		{
-			ShaderValue->Reference = MemStack::AllocateString(*Allocator, Code);
-		}
-		else
-		{
-			ShaderValue->Reference = AcquireLocalDeclarationCode();
-			ShaderValue->Value = MemStack::AllocateString(*Allocator, Code);
-			ShaderValueMap.Add(ShaderHash, ShaderValue);
-		}
-	}
-
-	return ShaderValue;
 }
 
-namespace Private
+FBinaryOpDescription GetBinaryOpDesription(EBinaryOp Op)
 {
-void WriteMaterialUniformAccess(Shader::EValueComponentType ComponentType, uint32 NumComponents, uint32 UniformOffset, FStringBuilderBase& OutResult)
-{
-	static const TCHAR IndexToMask[] = TEXT("xyzw");
-	uint32 RegisterIndex = UniformOffset / 4;
-	uint32 RegisterOffset = UniformOffset % 4;
-	uint32 NumComponentsToWrite = NumComponents;
-	bool bConstructor = false;
-
-	check(ComponentType == Shader::EValueComponentType::Float || ComponentType == Shader::EValueComponentType::Int);
-	const bool bIsInt = (ComponentType == Shader::EValueComponentType::Int);
-
-	while (NumComponentsToWrite > 0u)
-	{
-		const uint32 NumComponentsInRegister = FMath::Min(NumComponentsToWrite, 4u - RegisterOffset);
-		if (NumComponentsInRegister < NumComponents && !bConstructor)
-		{
-			// Uniform will be split across multiple registers, so add the constructor to concat them together
-			OutResult.Appendf(TEXT("%s%d("), Shader::GetComponentTypeName(ComponentType), NumComponents);
-			bConstructor = true;
-		}
-
-		if (bIsInt)
-		{
-			// PreshaderBuffer is typed as float4, so reinterpret as 'int' if needed
-			OutResult.Append(TEXT("asint("));
-		}
-
-		OutResult.Appendf(TEXT("Material.PreshaderBuffer[%u]"), RegisterIndex);
-		// Can skip writing mask if we're taking all 4 components from the register
-		if (NumComponentsInRegister < 4u)
-		{
-			OutResult.AppendChar(TCHAR('.'));
-			for (uint32 i = 0u; i < NumComponentsInRegister; ++i)
-			{
-				OutResult.AppendChar(IndexToMask[RegisterOffset + i]);
-			}
-		}
-
-		if (bIsInt)
-		{
-			OutResult.Append(TEXT(")"));
-		}
-
-		NumComponentsToWrite -= NumComponentsInRegister;
-		RegisterIndex++;
-		RegisterOffset = 0u;
-		if (NumComponentsToWrite > 0u)
-		{
-			OutResult.Append(TEXT(", "));
-		}
-	}
-	if (bConstructor)
+	switch (Op)
 	{
-		OutResult.Append(TEXT(")"));
+	case EBinaryOp::None: return FBinaryOpDescription(TEXT("None"), TEXT(""), Shader::EPreshaderOpcode::Nop); break;
+	case EBinaryOp::Add: return FBinaryOpDescription(TEXT("Add"), TEXT("+"), Shader::EPreshaderOpcode::Add); break;
+	case EBinaryOp::Sub: return FBinaryOpDescription(TEXT("Subtract"), TEXT("-"), Shader::EPreshaderOpcode::Sub); break;
+	case EBinaryOp::Mul: return FBinaryOpDescription(TEXT("Multiply"), TEXT("*"), Shader::EPreshaderOpcode::Mul); break;
+	case EBinaryOp::Div: return FBinaryOpDescription(TEXT("Divide"), TEXT("/"), Shader::EPreshaderOpcode::Div); break;
+	case EBinaryOp::Less: return FBinaryOpDescription(TEXT("Less"), TEXT("<"), Shader::EPreshaderOpcode::Nop); break;
+	default: checkNoEntry(); return FBinaryOpDescription();
 	}
 }
-} // namespace Private
 
-FEmitShaderCode* FEmitContext::EmitPreshaderOrConstant(const FRequestedType& RequestedType, FExpression* Expression)
+EExpressionEvaluation CombineEvaluations(EExpressionEvaluation Lhs, EExpressionEvaluation Rhs)
 {
-	Shader::FPreshaderData LocalPreshader;
-	Expression->EmitValuePreshader(*this, RequestedType, LocalPreshader);
-
-	const Shader::FType Type = RequestedType.GetType();
-
-	FSHA1 Hasher;
-	Hasher.Update((uint8*)&Type, sizeof(Type));
-	LocalPreshader.AppendHash(Hasher);
-	const FSHAHash Hash = Hasher.Finalize();
-	FEmitShaderCode* const* PrevShaderValue = PreshaderValueMap.Find(Hash);
-	if (PrevShaderValue)
-	{
-		FScope* CurrentScope = ScopeStack.Last();
-		check(IsScopeLive(CurrentScope));
-
-		FEmitShaderCode* ShaderValue = *PrevShaderValue;
-		check(ShaderValue->Type == Type);
-		Private::MoveToScope(ShaderValue, CurrentScope);
-		return ShaderValue;
-	}
-
-	Shader::FPreshaderStack Stack;
-	const Shader::FPreshaderValue ConstantValue = LocalPreshader.EvaluateConstant(*Material, Stack);
-
-	TStringBuilder<1024> FormattedCode;
-	if (Type.IsStruct())
-	{
-		FormattedCode.Append(TEXT("{ "));
-	}
-
-	FMaterialUniformPreshaderHeader* PreshaderHeader = nullptr;
-	uint32 CurrentBoolUniformOffset = ~0u;
-	uint32 CurrentNumBoolComponents = 32u;
-
-	int32 ComponentIndex = 0;
-	for (int32 FieldIndex = 0; FieldIndex < Type.GetNumFlatFields(); ++FieldIndex)
+	if (Lhs == EExpressionEvaluation::None)
 	{
-		if (FieldIndex > 0)
-		{
-			FormattedCode.Append(TEXT(", "));
-		}
-
-		const Shader::EValueType FieldType = Type.GetFlatFieldType(FieldIndex);
-		const Shader::FValueTypeDescription TypeDesc = Shader::GetValueTypeDescription(FieldType);
-		const int32 NumFieldComponents = TypeDesc.NumComponents;
-		const EExpressionEvaluation FieldEvaluation = Expression->GetPreparedType().GetFieldEvaluation(ComponentIndex, NumFieldComponents);
-
-		if (FieldEvaluation == EExpressionEvaluation::Preshader)
-		{
-			// Only need to allocate uniform buffer for non-constant components
-			// Constant components can have their value inlined into the shader directly
-			FUniformExpressionSet& UniformExpressionSet = MaterialCompilationOutput->UniformExpressionSet;
-			if (!PreshaderHeader)
-			{
-				// Allocate a preshader header the first time we hit a non-constant field
-				PreshaderHeader = &UniformExpressionSet.UniformPreshaders.AddDefaulted_GetRef();
-				PreshaderHeader->FieldIndex = UniformExpressionSet.UniformPreshaderFields.Num();
-				PreshaderHeader->NumFields = 0u;
-				PreshaderHeader->OpcodeOffset = UniformExpressionSet.UniformPreshaderData.Num();
-				Expression->EmitValuePreshader(*this, RequestedType, UniformExpressionSet.UniformPreshaderData);
-				PreshaderHeader->OpcodeSize = UniformExpressionSet.UniformPreshaderData.Num() - PreshaderHeader->OpcodeOffset;
-			}
-
-			FMaterialUniformPreshaderField& PreshaderField = UniformExpressionSet.UniformPreshaderFields.AddDefaulted_GetRef();
-			PreshaderField.ComponentIndex = ComponentIndex;
-			PreshaderField.Type = FieldType;
-			PreshaderHeader->NumFields++;
-
-			if (TypeDesc.ComponentType == Shader::EValueComponentType::Bool)
-			{
-				// 'Bool' uniforms are packed into bits
-				if (CurrentNumBoolComponents + NumFieldComponents > 32u)
-				{
-					CurrentBoolUniformOffset = UniformPreshaderOffset++;
-					CurrentNumBoolComponents = 0u;
-				}
-
-				const uint32 RegisterIndex = CurrentBoolUniformOffset / 4;
-				const uint32 RegisterOffset = CurrentBoolUniformOffset % 4;
-				FormattedCode.Appendf(TEXT("UnpackUniform_%s(asuint(Material.PreshaderBuffer[%u][%u]), %u)"),
-					TypeDesc.Name,
-					RegisterIndex,
-					RegisterOffset,
-					CurrentNumBoolComponents);
-
-				PreshaderField.BufferOffset = CurrentBoolUniformOffset * 32u + CurrentNumBoolComponents;
-				CurrentNumBoolComponents += NumFieldComponents;
-			}
-			else if (TypeDesc.ComponentType == Shader::EValueComponentType::Double)
-			{
-				// Double uniforms are split into Tile/Offset components to make FLWCScalar/FLWCVectors
-				PreshaderField.BufferOffset = UniformPreshaderOffset;
-
-				if (NumFieldComponents > 1)
-				{
-					FormattedCode.Appendf(TEXT("MakeLWCVector%d("), NumFieldComponents);
-				}
-				else
-				{
-					FormattedCode.Append(TEXT("MakeLWCScalar("));
-				}
-
-				// Write the tile uniform
-				Private::WriteMaterialUniformAccess(Shader::EValueComponentType::Float, NumFieldComponents, UniformPreshaderOffset, FormattedCode);
-				UniformPreshaderOffset += NumFieldComponents;
-				FormattedCode.Append(TEXT(", "));
-
-				// Write the offset uniform
-				Private::WriteMaterialUniformAccess(Shader::EValueComponentType::Float, NumFieldComponents, UniformPreshaderOffset, FormattedCode);
-				UniformPreshaderOffset += NumFieldComponents;
-				FormattedCode.Append(TEXT(")"));
-			}
-			else
-			{
-				// Float/Int uniforms are written directly to the uniform buffer
-				const uint32 RegisterOffset = UniformPreshaderOffset % 4;
-				if (RegisterOffset + NumFieldComponents > 4u)
-				{
-					// If this uniform would span multiple registers, align offset to the next register to avoid this
-					// TODO - we could keep track of this empty padding space, and pack other smaller uniform types here
-					UniformPreshaderOffset = Align(UniformPreshaderOffset, 4u);
-				}
-
-				PreshaderField.BufferOffset = UniformPreshaderOffset;
-				Private::WriteMaterialUniformAccess(TypeDesc.ComponentType, NumFieldComponents, UniformPreshaderOffset, FormattedCode);
-				UniformPreshaderOffset += NumFieldComponents;
-			}
-		}
-		else
-		{
-			// We allow FieldEvaluation to be 'None', since in that case we still need to fill in a value for the HLSL initializer
-			check(FieldEvaluation == EExpressionEvaluation::Constant || FieldEvaluation == EExpressionEvaluation::None);
-
-			// The type generated by the preshader might not match the expected type
-			// In the future, with new HLSLTree, preshader could potentially include explicit cast opcodes, and avoid implicit conversions
-			Shader::FValue FieldConstantValue(ConstantValue.Type.GetComponentType(ComponentIndex), NumFieldComponents);
-			for (int32 i = 0; i < NumFieldComponents; ++i)
-			{
-				// Allow replicating scalar values
-				FieldConstantValue.Component[i] = ConstantValue.Component.Num() == 1 ? ConstantValue.Component[0] : ConstantValue.Component[ComponentIndex + i];
-			}
-
-			if (TypeDesc.ComponentType == Shader::EValueComponentType::Double)
-			{
-				const Shader::FDoubleValue DoubleValue = FieldConstantValue.AsDouble();
-				TStringBuilder<256> TileValue;
-				TStringBuilder<256> OffsetValue;
-				for (int32 Index = 0; Index < NumFieldComponents; ++Index)
-				{
-					if (Index > 0)
-					{
-						TileValue.Append(TEXT(", "));
-						OffsetValue.Append(TEXT(", "));
-					}
-
-					const FLargeWorldRenderScalar Value(DoubleValue[Index]);
-					TileValue.Appendf(TEXT("%#.9gf"), Value.GetTile());
-					OffsetValue.Appendf(TEXT("%#.9gf"), Value.GetOffset());
-				}
-
-				if (NumFieldComponents > 1)
-				{
-					FormattedCode.Appendf(TEXT("MakeLWCVector%d(float%d(%s), float%d(%s))"), NumFieldComponents, NumFieldComponents, TileValue.ToString(), NumFieldComponents, OffsetValue.ToString());
-				}
-				else
-				{
-					FormattedCode.Appendf(TEXT("MakeLWCScalar(%s, %s)"), TileValue.ToString(), OffsetValue.ToString());
-				}
-			}
-			else
-			{
-				const Shader::FValue CastFieldConstantValue = Shader::Cast(FieldConstantValue, FieldType);
-				if (NumFieldComponents > 1)
-				{
-					FormattedCode.Appendf(TEXT("%s("), TypeDesc.Name);
-				}
-				for (int32 Index = 0; Index < NumFieldComponents; ++Index)
-				{
-					if (Index > 0)
-					{
-						FormattedCode.Append(TEXT(", "));
-					}
-					CastFieldConstantValue.Component[Index].ToString(TypeDesc.ComponentType, FormattedCode);
-				}
-				if (NumFieldComponents > 1)
-				{
-					FormattedCode.Append(TEXT(")"));
-				}
-			}
-		}
-		ComponentIndex += NumFieldComponents;
+		// If either is 'None', return the other
+		return Rhs;
 	}
-	check(ComponentIndex == Type.GetNumComponents());
-
-	if (Type.IsStruct())
+	else if (Rhs == EExpressionEvaluation::None)
 	{
-		FormattedCode.Append(TEXT(" }"));
+		return Lhs;
 	}
-
-	const bool bInline = !Type.IsStruct(); // struct declarations can't be inline, due to HLSL syntax
-	FEmitShaderCode* ShaderValue = EmitCodeInternal(Type, FormattedCode.ToView(), bInline, TArrayView<FEmitShaderCode*>());
-	PreshaderValueMap.Add(Hash, ShaderValue);
-
-	return ShaderValue;
-}
-
-FEmitShaderCode* FEmitContext::EmitConstantZero(const Shader::FType& Type)
-{
-	return EmitInlineCode(Type, TEXT("((%)0)"), Type.GetName());
-}
-
-namespace Private
-{
-void EmitShaderValue(FEmitContext& Context, FEmitShaderCode* ShaderValue)
-{
-	if (ShaderValue->Scope)
+	else if (Lhs == EExpressionEvaluation::Constant && Rhs == EExpressionEvaluation::Constant)
 	{
-		// Emit dependencies first
-		for (FEmitShaderCode* Dependency : ShaderValue->Dependencies)
-		{
-			EmitShaderValue(Context, Dependency);
-		}
-		// Don't need a declaration for inline values
-		if (!ShaderValue->IsInline())
-		{
-			ShaderValue->Scope->EmitDeclarationf(Context, TEXT("const %s %s = %s;"),
-				ShaderValue->Type.GetName(),
-				ShaderValue->Reference,
-				ShaderValue->Value);
-		}
-		ShaderValue->Scope = nullptr; // Don't emit again
+		// 2 constants make a constant
+		return EExpressionEvaluation::Constant;
 	}
-}
-}
-
-void FEmitContext::Finalize()
-{
-	check(ScopeStack.Num() == 0);
-	for (const auto& It : ShaderValueMap)
+	else if (Lhs == EExpressionEvaluation::Shader || Rhs == EExpressionEvaluation::Shader)
 	{
-		Private::EmitShaderValue(*this, It.Value);
+		// If either requires shader, shader is required
+		return EExpressionEvaluation::Shader;
 	}
-
-	ShaderValueMap.Reset();
-	PreshaderValueMap.Reset();
-	LocalPHIs.Reset();
-
-	MaterialCompilationOutput->UniformExpressionSet.UniformPreshaderBufferSize = (UniformPreshaderOffset + 3u) / 4u;
+	// Any combination of constants/preshader can make a preshader
+	return EExpressionEvaluation::Preshader;
 }
 
 void FScope::Reset()
 {
 	State = EScopeState::Uninitialized;
-	Declarations = FCodeList();
-	Statements = FCodeList();
 }
 
 FScope* FScope::FindSharedParent(FScope* Lhs, FScope* Rhs)
@@ -526,12 +169,8 @@ void FExpressionLocalPHI::ComputeAnalyticDerivatives(FTree& Tree, FExpressionDer
 {
 	// We don't have values assigned at the time analytic derivatives are computed
 	// It's possible the derivatives will be end up being invalid, but that case will need to be detected later, during PrepareValue
-	FExpressionLocalPHI* ExpressionDdx = Tree.NewExpression<FExpressionLocalPHI>(this, EDerivativeCoordinate::Ddx);
-	FExpressionLocalPHI* ExpressionDdy = Tree.NewExpression<FExpressionLocalPHI>(this, EDerivativeCoordinate::Ddy);
-	Tree.PHIExpressions.Add(ExpressionDdx);
-	Tree.PHIExpressions.Add(ExpressionDdy);
-	OutResult.ExpressionDdx = ExpressionDdx;
-	OutResult.ExpressionDdy = ExpressionDdy;
+	OutResult.ExpressionDdx = Tree.NewExpression<FExpressionLocalPHI>(this, EDerivativeCoordinate::Ddx);
+	OutResult.ExpressionDdy = Tree.NewExpression<FExpressionLocalPHI>(this, EDerivativeCoordinate::Ddy);
 }
 
 bool FExpressionLocalPHI::PrepareValue(FEmitContext& Context, const FRequestedType& RequestedType, FPrepareValueResult& OutResult) const
@@ -636,71 +275,104 @@ bool FExpressionLocalPHI::PrepareValue(FEmitContext& Context, const FRequestedTy
 	return true;
 }
 
-void FExpressionLocalPHI::EmitValueShader(FEmitContext& Context, const FRequestedType& RequestedType, FEmitValueShaderResult& OutResult) const
+void FExpressionLocalPHI::EmitValueShader(FEmitContext& Context, FEmitShaderScope& Scope, const FRequestedType& RequestedType, FEmitValueShaderResult& OutResult) const
 {
-	int32 LocalPHIIndex = Context.LocalPHIs.Find(this);
-	if (LocalPHIIndex == INDEX_NONE)
+	FEmitShaderExpression* const* PrevEmitExpression = Context.EmitLocalPHIMap.Find(this);
+	FEmitShaderExpression* EmitExpression = PrevEmitExpression ? *PrevEmitExpression : nullptr;
+	if (!EmitExpression)
 	{
+		const int32 LocalPHIIndex = Context.NumExpressionLocalPHIs++;
+		const FRequestedType LocalType = GetRequestedType();
+
 		// This is the first time we've emitted shader code for this PHI
-		// First add it to the list, so if this is called recursively, this path will only be taken the first time
-		LocalPHIIndex = Context.LocalPHIs.Add(this);
+		// Create an expression and add it to the map first, so if this is called recursively this path will only be taken the first time
+		EmitExpression = OutResult.Code = Context.EmitInlineExpression(Scope,
+			LocalType.GetType(),
+			TEXT("LocalPHI%"), LocalPHIIndex);
+		Context.EmitLocalPHIMap.Add(this, EmitExpression);
 
 		// Find the outermost scope to declare our local variable
-		FScope* DeclarationScope = Context.ScopeStack.Last();
+		FEmitShaderScope* EmitDeclarationScope = &Scope;
+		FEmitShaderScope* EmitValueScopes[MaxNumPreviousScopes] = { nullptr };
 		for (int32 i = 0; i < NumValues; ++i)
 		{
-			DeclarationScope = FScope::FindSharedParent(DeclarationScope, Scopes[i]);
-			if (!DeclarationScope)
+			EmitValueScopes[i] = Context.AcquireEmitScope(Scopes[i]);
+			EmitDeclarationScope = FEmitShaderScope::FindSharedParent(EmitDeclarationScope, EmitValueScopes[i]);
+			if (!EmitDeclarationScope)
 			{
 				Context.Errors->AddError(TEXT("Invalid LocalPHI"));
 				return;
 			}
 		}
 
-		const FRequestedType LocalType = GetRequestedType();
-
-		bool bNeedToAddDeclaration = true;
+		FEmitShaderStatement* EmitDeclaration = nullptr;
 		for (int32 i = 0; i < NumValues; ++i)
 		{
-			FScope* ValueScope = Scopes[i];
-			check(IsScopeLive(ValueScope));
-
-			Context.ScopeStack.Add(ValueScope);
-			FEmitShaderCode* ShaderValue = Values[i]->GetValueShader(Context, LocalType);
-			Context.ScopeStack.Pop();
-
-			if (ValueScope == DeclarationScope)
+			FEmitShaderScope* EmitValueScope = EmitValueScopes[i];
+			if (EmitValueScope == EmitDeclarationScope)
 			{
-				ValueScope->EmitDeclarationf(Context, TEXT("%s LocalPHI%d = %s;"),
+				FEmitShaderExpression* ShaderValue = Values[i]->GetValueShader(Context, *EmitValueScope, LocalType);
+				EmitDeclaration = Context.EmitStatement(*EmitValueScope, TEXT("% LocalPHI% = %;"),
 					LocalType.GetName(),
 					LocalPHIIndex,
-					ShaderValue->Reference);
-				bNeedToAddDeclaration = false;
-			}
-			else
-			{
-				ValueScope->EmitStatementf(Context, TEXT("LocalPHI%d = %s;"),
-					LocalPHIIndex,
-					ShaderValue->Reference);
+					ShaderValue);
+				break;
 			}
 		}
-
-		if (bNeedToAddDeclaration)
+		if (!EmitDeclaration)
 		{
-			check(IsScopeLive(DeclarationScope));
-			DeclarationScope->EmitDeclarationf(Context, TEXT("%s LocalPHI%d;"),
+			EmitDeclaration = Context.EmitStatement(*EmitDeclarationScope, TEXT("% LocalPHI%;"),
 				LocalType.GetName(),
 				LocalPHIIndex);
 		}
+
+		FEmitShaderNode* Dependencies[MaxNumPreviousScopes] = { nullptr };
+		int32 NumDependencies = 0;
+		for (int32 i = 0; i < NumValues; ++i)
+		{
+			FEmitShaderScope* EmitValueScope = EmitValueScopes[i];
+			if (EmitValueScope != EmitDeclarationScope)
+			{
+				FEmitShaderExpression* ShaderValue = Values[i]->GetValueShader(Context, *EmitValueScope, LocalType);
+				FEmitShaderStatement* EmitAssignment = Context.EmitStatementWithDependency(*EmitValueScope, EmitDeclaration, TEXT("LocalPHI% = %;"),
+					LocalPHIIndex,
+					ShaderValue);
+				Dependencies[NumDependencies++] = EmitAssignment;
+			}
+		}
+
+		// Fill in the expression's dependencies
+		EmitExpression->Dependencies = MemStack::AllocateArrayView(*Context.Allocator, MakeArrayView(Dependencies, NumDependencies));
 	}
 
-	const Shader::FType LocalType = GetType();
-	OutResult.Code = Context.EmitInlineCode(LocalType, TEXT("LocalPHI%"), LocalPHIIndex);
+	OutResult.Code = EmitExpression;
 }
 
-void FStatement::Reset()
+bool FExpressionFunctionCall::PrepareValue(FEmitContext& Context, const FRequestedType& RequestedType, FPrepareValueResult& OutResult) const
 {
-	bEmitShader = false;
+	if (!PrepareScope(Context, Function->RootScope))
+	{
+		return false;
+	}
+
+	FPreparedType OutputType = PrepareExpressionValue(Context, Function->OutputExpressions[OutputIndex], RequestedType);
+	return OutResult.SetType(Context, RequestedType, OutputType);
+}
+
+void FExpressionFunctionCall::EmitValueShader(FEmitContext& Context, FEmitShaderScope& Scope, const FRequestedType& RequestedType, FEmitValueShaderResult& OutResult) const
+{
+	FEmitShaderNode* const* PrevDependency = Context.EmitFunctionMap.Find(Function);
+	FEmitShaderNode* Dependency = PrevDependency ? *PrevDependency : nullptr;
+	if (!Dependency)
+	{
+		// Inject the function's root scope at scope where it's called
+		FEmitShaderScope* EmitCalledScope = Context.AcquireEmitScope(Function->CalledScope);
+		Dependency = Context.EmitNextScope(*EmitCalledScope, Function->RootScope);
+		Context.EmitFunctionMap.Add(Function, Dependency);
+	}
+
+	FEmitShaderExpression* EmitFunctionOutput = Function->OutputExpressions[OutputIndex]->GetValueShader(Context, Scope, RequestedType);
+	OutResult.Code = Context.EmitInlineExpressionWithDependencies(Scope, MakeArrayView(&Dependency, 1), EmitFunctionOutput->Type, TEXT("%"), EmitFunctionOutput);
 }
 
 void FExpression::Reset()
@@ -1165,13 +837,12 @@ bool FPrepareValueResult::SetType(FEmitContext& Context, const FRequestedType& R
 
 bool FPrepareValueResult::SetForwardValue(FEmitContext& Context, const FRequestedType& RequestedType, FExpression* InForwardValue)
 {
-	check(InForwardValue);
 	if (InForwardValue != ForwardValue)
 	{
 		PreparedType = PrepareExpressionValue(Context, InForwardValue, RequestedType);
 		ForwardValue = InForwardValue;
 	}
-	return true;
+	return (bool)InForwardValue;
 }
 
 void FExpression::ComputeAnalyticDerivatives(FTree& Tree, FExpressionDerivatives& OutResult) const
@@ -1179,7 +850,7 @@ void FExpression::ComputeAnalyticDerivatives(FTree& Tree, FExpressionDerivatives
 	// nop
 }
 
-void FExpression::EmitValueShader(FEmitContext& Context, const FRequestedType& RequestedType, FEmitValueShaderResult& OutResult) const
+void FExpression::EmitValueShader(FEmitContext& Context, FEmitShaderScope& Scope, const FRequestedType& RequestedType, FEmitValueShaderResult& OutResult) const
 {
 	check(false);
 }
@@ -1189,42 +860,42 @@ void FExpression::EmitValuePreshader(FEmitContext& Context, const FRequestedType
 	check(false);
 }
 
-FEmitShaderCode* FExpression::GetValueShader(FEmitContext& Context, const FRequestedType& RequestedType, const Shader::FType& ResultType)
+FEmitShaderExpression* FExpression::GetValueShader(FEmitContext& Context, FEmitShaderScope& Scope, const FRequestedType& RequestedType, const Shader::FType& ResultType)
 {
 	FOwnerScope OwnerScope(*Context.Errors, GetOwner());
 	if (PrepareValueResult.ForwardValue)
 	{
-		return PrepareValueResult.ForwardValue->GetValueShader(Context, RequestedType, ResultType);
+		return PrepareValueResult.ForwardValue->GetValueShader(Context, Scope, RequestedType, ResultType);
 	}
 
 	const EExpressionEvaluation Evaluation = PrepareValueResult.PreparedType.GetEvaluation(RequestedType);
 	check(Evaluation != EExpressionEvaluation::None);
 
-	FEmitShaderCode* Value = nullptr;
+	FEmitShaderExpression* Value = nullptr;
 	if (Evaluation == EExpressionEvaluation::Constant || Evaluation == EExpressionEvaluation::Preshader)
 	{
-		Value = Context.EmitPreshaderOrConstant(RequestedType, this);
+		Value = Context.EmitPreshaderOrConstant(Scope, RequestedType, this);
 	}
 	else
 	{
 		check(Evaluation == EExpressionEvaluation::Shader);
 		FEmitValueShaderResult Result;
-		EmitValueShader(Context, RequestedType, Result);
+		EmitValueShader(Context, Scope, RequestedType, Result);
 		Value = Result.Code;
 	}
 
-	Value = Context.EmitCast(Value, ResultType);
+	Value = Context.EmitCast(Scope, Value, ResultType);
 	return Value;
 }
 
-FEmitShaderCode* FExpression::GetValueShader(FEmitContext& Context, const FRequestedType& RequestedType)
+FEmitShaderExpression* FExpression::GetValueShader(FEmitContext& Context, FEmitShaderScope& Scope, const FRequestedType& RequestedType)
 {
-	return GetValueShader(Context, RequestedType, RequestedType.GetType());
+	return GetValueShader(Context, Scope, RequestedType, RequestedType.GetType());
 }
 
-FEmitShaderCode* FExpression::GetValueShader(FEmitContext& Context)
+FEmitShaderExpression* FExpression::GetValueShader(FEmitContext& Context, FEmitShaderScope& Scope)
 {
-	return GetValueShader(Context, GetRequestedType());
+	return GetValueShader(Context, Scope, GetRequestedType());
 }
 
 void FExpression::GetValuePreshader(FEmitContext& Context, const FRequestedType& RequestedType, Shader::FPreshaderData& OutPreshader)
@@ -1243,12 +914,17 @@ void FExpression::GetValuePreshader(FEmitContext& Context, const FRequestedType&
 		EmitValuePreshader(Context, RequestedType, OutPreshader);
 		bReentryFlag = false;
 	}
-	else
+	else if (Evaluation == EExpressionEvaluation::Constant)
 	{
-		check(Evaluation == EExpressionEvaluation::Constant);
 		const Shader::FValue ConstantValue = GetValueConstant(Context, RequestedType);
 		OutPreshader.WriteOpcode(Shader::EPreshaderOpcode::Constant).Write(ConstantValue);
 	}
+	else
+	{
+		// Value is not used, write a dummy value
+		check(Evaluation == EExpressionEvaluation::None);
+		OutPreshader.WriteOpcode(Shader::EPreshaderOpcode::ConstantZero).Write(Shader::EValueType::Float1);
+	}
 }
 
 Shader::FValue FExpression::GetValueConstant(FEmitContext& Context, const FRequestedType& RequestedType)
@@ -1302,39 +978,6 @@ void FScope::AddPreviousScope(FScope& Scope)
 	PreviousScope[NumPreviousScopes++] = &Scope;
 }
 
-void FScope::InternalEmitCode(FEmitContext& Context, FCodeList& List, ENextScopeFormat ScopeFormat, FScope* Scope, const TCHAR* String, int32 Length)
-{
-	if (Scope && Scope->ContainedStatement && !Scope->ContainedStatement->bEmitShader)
-	{
-		Scope->ContainedStatement->bEmitShader = true;
-		Context.ScopeStack.Add(Scope);
-		Scope->ContainedStatement->EmitShader(Context);
-		Context.ScopeStack.Pop();
-	}
-
-	const int32 SizeofString = sizeof(TCHAR) * Length;
-	void* Memory = Context.Allocator->Alloc(sizeof(FCodeEntry) + SizeofString, alignof(FCodeEntry));
-	FCodeEntry* CodeEntry = new(Memory) FCodeEntry();
-	FMemory::Memcpy(CodeEntry->String, String, SizeofString);
-	CodeEntry->String[Length] = 0;
-	CodeEntry->Length = Length;
-	CodeEntry->Scope = Scope;
-	CodeEntry->ScopeFormat = ScopeFormat;
-	CodeEntry->Next = nullptr;
-
-	if (!List.First)
-	{
-		List.First = CodeEntry;
-		List.Last = CodeEntry;
-	}
-	else
-	{
-		List.Last->Next = CodeEntry;
-		List.Last = CodeEntry;
-	}
-	List.Num++;
-}
-
 bool PrepareScope(FEmitContext& Context, FScope* InScope)
 {
 	if (InScope && InScope->State == EScopeState::Uninitialized)
@@ -1384,17 +1027,6 @@ void FTree::Destroy(FTree* Tree)
 	}
 }
 
-//
-static void WriteIndent(int32 IndentLevel, FStringBuilderBase& InOutString)
-{
-	const int32 Offset = InOutString.AddUninitialized(IndentLevel);
-	TCHAR* Result = InOutString.GetData() + Offset;
-	for (int32 i = 0; i < IndentLevel; ++i)
-	{
-		*Result++ = TEXT('\t');
-	}
-}
-
 void FScope::MarkLive()
 {
 	if (State == EScopeState::Uninitialized)
@@ -1421,54 +1053,6 @@ void FScope::MarkDead()
 	State = EScopeState::Dead;
 }
 
-void FScope::WriteHLSL(int32 Indent, FStringBuilderBase& OutString) const
-{
-	{
-		const FCodeEntry* CodeDeclaration = Declarations.First;
-		while (CodeDeclaration)
-		{
-			check(!CodeDeclaration->Scope);
-			WriteIndent(Indent, OutString);
-			OutString.Append(CodeDeclaration->String, CodeDeclaration->Length);
-			OutString.AppendChar(TEXT('\n'));
-			CodeDeclaration = CodeDeclaration->Next;
-		}
-	}
-
-	{
-		const FCodeEntry* CodeStatement = Statements.First;
-		while (CodeStatement)
-		{
-			if (CodeStatement->Length > 0)
-			{
-				WriteIndent(Indent, OutString);
-				OutString.Append(CodeStatement->String, CodeStatement->Length);
-				OutString.AppendChar(TEXT('\n'));
-			}
-			if (CodeStatement->Scope)
-			{
-				int32 NextIndent = Indent;
-				bool bNeedToCloseScope = false;
-				if (CodeStatement->ScopeFormat == ENextScopeFormat::Scoped)
-				{
-					WriteIndent(Indent, OutString);
-					OutString.Append(TEXT("{\n"));
-					NextIndent++;
-					bNeedToCloseScope = true;
-				}
-
-				CodeStatement->Scope->WriteHLSL(NextIndent, OutString);
-				if (bNeedToCloseScope)
-				{
-					WriteIndent(Indent, OutString);
-					OutString.Append(TEXT("}\n"));
-				}
-			}
-			CodeStatement = CodeStatement->Next;
-		}
-	}
-}
-
 void FOwnerContext::PushOwner(UObject* Owner)
 {
 	OwnerStack.Add(Owner);
@@ -1484,6 +1068,12 @@ UObject* FOwnerContext::GetCurrentOwner() const
 	return (OwnerStack.Num() > 0) ? OwnerStack.Last() : nullptr;
 }
 
+uint64 Private::GetNextTypeHash()
+{
+	static uint64 Hash = 1;
+	return Hash++;
+}
+
 void FTree::ResetNodes()
 {
 	FNode* Node = Nodes;
@@ -1526,16 +1116,30 @@ bool FTree::Finalize()
 
 bool FTree::EmitShader(FEmitContext& Context, FStringBuilderBase& OutCode) const
 {
-	if (RootScope->ContainedStatement)
+	FEmitShaderScope* EmitRootScope = Context.AcquireEmitScope(RootScope);
+	
+	// Link all nodes to their proper scope
+	for (FEmitShaderNode* EmitNode : Context.EmitNodes)
 	{
-		RootScope->ContainedStatement->bEmitShader = true;
-		Context.ScopeStack.Add(RootScope);
-		RootScope->ContainedStatement->EmitShader(Context);
-		Context.ScopeStack.Pop(false);
+		FEmitShaderScope* EmitScope = EmitNode->Scope;
+		if (EmitScope)
+		{
+			EmitNode->NextScopedNode = EmitScope->FirstNode;
+			EmitScope->FirstNode = EmitNode;
+		}
+	}
+
+	{
+		FEmitShaderScopeStack Stack;
+		TStringBuilder<1024> ScopeCode;
+		Stack.Emplace(EmitRootScope, 1, ScopeCode);
+		EmitRootScope->EmitShaderCode(Stack);
+		check(Stack.Num() == 1);
+		OutCode.Append(ScopeCode.ToView());
 	}
 
 	Context.Finalize();
-	RootScope->WriteHLSL(1, OutCode);
+
 	return true;
 }
 
@@ -1546,8 +1150,25 @@ void FTree::RegisterNode(FNode* Node)
 	Nodes = Node;
 }
 
-void FTree::RegisterExpression(FExpression* Expression)
+FExpression* FTree::FindExpression(FXxHash64 Hash) const
+{
+	FExpression* const* FoundExpression = ExpressionMap.Find(Hash);
+	if (FoundExpression)
+	{
+		return *FoundExpression;
+	}
+	return nullptr;
+}
+
+void FTree::RegisterExpression(FExpression* Expression, FXxHash64 Hash)
 {
+	ExpressionMap.Add(Hash, Expression);
+}
+
+void FTree::RegisterExpression(FExpressionLocalPHI* Expression, FXxHash64 Hash)
+{
+	PHIExpressions.Add(Expression);
+	RegisterExpression(static_cast<FExpression*>(Expression), Hash);
 }
 
 void FTree::RegisterStatement(FScope& Scope, FStatement* Statement)
@@ -1574,14 +1195,7 @@ FExpression* FTree::AcquireLocal(FScope& Scope, const FName& LocalName)
 	const TArrayView<FScope*> PreviousScopes = Scope.GetPreviousScopes();
 	if (PreviousScopes.Num() > 1)
 	{
-		FExpressionLocalPHI* Expression = NewExpression<FExpressionLocalPHI>();
-		Expression->LocalName = LocalName;
-		Expression->NumValues = PreviousScopes.Num();
-		for (int32 i = 0; i < PreviousScopes.Num(); ++i)
-		{
-			Expression->Scopes[i] = PreviousScopes[i];
-		}
-		PHIExpressions.Add(Expression);
+		FExpression* Expression = NewExpression<FExpressionLocalPHI>(LocalName, PreviousScopes);
 		Scope.LocalMap.Add(LocalName, Expression);
 		return Expression;
 	}
@@ -1594,6 +1208,18 @@ FExpression* FTree::AcquireLocal(FScope& Scope, const FName& LocalName)
 	return nullptr;
 }
 
+FExpression* FTree::NewFunctionCall(FScope& Scope, FFunction* Function, TArrayView<FExpression*> InputExpressions, int32 OutputIndex)
+{
+	FScope* CalledScope = &Scope;
+	if (Function->CalledScope)
+	{
+		CalledScope = FScope::FindSharedParent(CalledScope, Function->CalledScope);
+		check(CalledScope);
+	}
+	Function->CalledScope = CalledScope;
+	return NewExpression<FExpressionFunctionCall>(Function, InputExpressions, OutputIndex);
+}
+
 const FExpressionDerivatives& FTree::GetAnalyticDerivatives(FExpression* InExpression)
 {
 	if (!InExpression)
@@ -1632,6 +1258,13 @@ FScope* FTree::NewOwnedScope(FStatement& Owner)
 	return NewScope;
 }
 
+FFunction* FTree::NewFunction()
+{
+	FFunction* NewFunction = NewNode<FFunction>();
+	NewFunction->RootScope = NewNode<FScope>();
+	return NewFunction;
+}
+
 FTextureParameterDeclaration* FTree::NewTextureParameterDeclaration(const FName& Name, const FTextureDescription& DefaultValue)
 {
 	FTextureParameterDeclaration* Declaration = NewNode<FTextureParameterDeclaration>(Name, DefaultValue);
diff --git a/Engine/Source/Runtime/Engine/Private/HLSLTree/HLSLTreeCommon.cpp b/Engine/Source/Runtime/Engine/Private/HLSLTree/HLSLTreeCommon.cpp
--- a/Engine/Source/Runtime/Engine/Private/HLSLTree/HLSLTreeCommon.cpp
+++ b/Engine/Source/Runtime/Engine/Private/HLSLTree/HLSLTreeCommon.cpp
@@ -1,5 +1,6 @@
 // Copyright Epic Games, Inc. All Rights Reserved.
 #include "HLSLTree/HLSLTreeCommon.h"
+#include "HLSLTree/HLSLTreeEmit.h"
 #include "Misc/StringBuilder.h"
 #include "MaterialShared.h"
 #include "Engine/Texture.h"
@@ -220,24 +221,24 @@ bool FExpressionExternalInput::PrepareValue(FEmitContext& Context, const FReques
 	return OutResult.SetType(Context, RequestedType, EExpressionEvaluation::Shader, InputDesc.Type);
 }
 
-void FExpressionExternalInput::EmitValueShader(FEmitContext& Context, const FRequestedType& RequestedType, FEmitValueShaderResult& OutResult) const
+void FExpressionExternalInput::EmitValueShader(FEmitContext& Context, FEmitShaderScope& Scope, const FRequestedType& RequestedType, FEmitValueShaderResult& OutResult) const
 {
 	const int32 TypeIndex = (int32)InputType;
 	if (IsTexCoord(InputType))
 	{
 		const int32 TexCoordIndex = TypeIndex - (int32)EExternalInput::TexCoord0;
 		Context.NumTexCoords = FMath::Max(Context.NumTexCoords, TexCoordIndex + 1);
-		OutResult.Code = Context.EmitInlineCode(Shader::EValueType::Float2, TEXT("Parameters.TexCoords[%].xy"), TexCoordIndex);
+		OutResult.Code = Context.EmitInlineExpression(Scope, Shader::EValueType::Float2, TEXT("Parameters.TexCoords[%].xy"), TexCoordIndex);
 	}
 	else if (IsTexCoord_Ddx(InputType))
 	{
 		const int32 TexCoordIndex = TypeIndex - (int32)EExternalInput::TexCoord0_Ddx;
-		OutResult.Code = Context.EmitInlineCode(Shader::EValueType::Float2, TEXT("Parameters.TexCoords_DDX[%].xy"), TexCoordIndex);
+		OutResult.Code = Context.EmitInlineExpression(Scope, Shader::EValueType::Float2, TEXT("Parameters.TexCoords_DDX[%].xy"), TexCoordIndex);
 	}
 	else if (IsTexCoord_Ddy(InputType))
 	{
 		const int32 TexCoordIndex = TypeIndex - (int32)EExternalInput::TexCoord0_Ddy;
-		OutResult.Code = Context.EmitInlineCode(Shader::EValueType::Float2, TEXT("Parameters.TexCoords_DDY[%].xy"), TexCoordIndex);
+		OutResult.Code = Context.EmitInlineExpression(Scope, Shader::EValueType::Float2, TEXT("Parameters.TexCoords_DDY[%].xy"), TexCoordIndex);
 	}
 	else
 	{
@@ -247,22 +248,22 @@ void FExpressionExternalInput::EmitValueShader(FEmitContext& Context, const FReq
 		switch (InputType)
 		{
 		case EExternalInput::WorldPosition:
-			OutResult.Code = Context.EmitInlineCode(InputDesc.Type, TEXT("GetWorldPosition(Parameters)"));
+			OutResult.Code = Context.EmitInlineExpression(Scope, InputDesc.Type, TEXT("GetWorldPosition(Parameters)"));
 			break;
 		case EExternalInput::WorldPosition_NoOffsets:
-			OutResult.Code = Context.EmitInlineCode(InputDesc.Type, TEXT("GetWorldPosition_NoMaterialOffsets(Parameters)"));
+			OutResult.Code = Context.EmitInlineExpression(Scope, InputDesc.Type, TEXT("GetWorldPosition_NoMaterialOffsets(Parameters)"));
 			break;
 		case EExternalInput::TranslatedWorldPosition:
-			OutResult.Code = Context.EmitInlineCode(InputDesc.Type, TEXT("GetTranslatedWorldPosition(Parameters)"));
+			OutResult.Code = Context.EmitInlineExpression(Scope, InputDesc.Type, TEXT("GetTranslatedWorldPosition(Parameters)"));
 			break;
 		case EExternalInput::TranslatedWorldPosition_NoOffsets:
-			OutResult.Code = Context.EmitInlineCode(InputDesc.Type, TEXT("GetTranslatedWorldPosition_NoMaterialOffsets(Parameters)"));
+			OutResult.Code = Context.EmitInlineExpression(Scope, InputDesc.Type, TEXT("GetTranslatedWorldPosition_NoMaterialOffsets(Parameters)"));
 			break;
 		case EExternalInput::WorldPosition_Ddx:
-			OutResult.Code = Context.EmitInlineCode(InputDesc.Type, TEXT("Parameters.WorldPosition_DDX"));
+			OutResult.Code = Context.EmitInlineExpression(Scope, InputDesc.Type, TEXT("Parameters.WorldPosition_DDX"));
 			break;
 		case EExternalInput::WorldPosition_Ddy:
-			OutResult.Code = Context.EmitInlineCode(InputDesc.Type, TEXT("Parameters.WorldPosition_DDY"));
+			OutResult.Code = Context.EmitInlineExpression(Scope, InputDesc.Type, TEXT("Parameters.WorldPosition_DDY"));
 			break;
 		default:
 			checkNoEntry();
@@ -280,7 +281,7 @@ bool FExpressionTextureSample::PrepareValue(FEmitContext& Context, const FReques
 	return OutResult.SetType(Context, RequestedType, EExpressionEvaluation::Shader, Shader::EValueType::Float4);
 }
 
-void FExpressionTextureSample::EmitValueShader(FEmitContext& Context, const FRequestedType& RequestedType, FEmitValueShaderResult& OutResult) const
+void FExpressionTextureSample::EmitValueShader(FEmitContext& Context, FEmitShaderScope& Scope, const FRequestedType& RequestedType, FEmitValueShaderResult& OutResult) const
 {
 	const FTextureDescription& Desc = Declaration->Description;
 	const EMaterialValueType MaterialValueType = Desc.Texture->GetMaterialType();
@@ -406,10 +407,10 @@ void FExpressionTextureSample::EmitValueShader(FEmitContext& Context, const FReq
 		}
 	}
 
-	FEmitShaderCode* TexCoordValue = TexCoordExpression->GetValueShader(Context, Shader::EValueType::Float2);
-	FEmitShaderCode* TexCoordValueDdx = TexCoordDerivatives.ExpressionDdx->GetValueShader(Context, Shader::EValueType::Float2);
-	FEmitShaderCode* TexCoordValueDdy = TexCoordDerivatives.ExpressionDdy->GetValueShader(Context, Shader::EValueType::Float2);
-	OutResult.Code = Context.EmitCode(Shader::EValueType::Float4, TEXT("%(%Grad(%, %, %, %, %))"),
+	FEmitShaderExpression* TexCoordValue = TexCoordExpression->GetValueShader(Context, Scope, Shader::EValueType::Float2);
+	FEmitShaderExpression* TexCoordValueDdx = TexCoordDerivatives.ExpressionDdx->GetValueShader(Context, Scope, Shader::EValueType::Float2);
+	FEmitShaderExpression* TexCoordValueDdy = TexCoordDerivatives.ExpressionDdy->GetValueShader(Context, Scope, Shader::EValueType::Float2);
+	OutResult.Code = Context.EmitExpression(Scope, Shader::EValueType::Float4, TEXT("%(%Grad(%, %, %, %, %))"),
 		SamplerTypeFunction,
 		SampleFunctionName,
 		*TextureName,
@@ -426,9 +427,11 @@ void FExpressionGetStructField::ComputeAnalyticDerivatives(FTree& Tree, FExpress
 	{
 		const Shader::FStructType* DerivativeStructType = StructType->DerivativeType;
 		check(DerivativeStructType);
+		const Shader::FStructField* DerivativeField = DerivativeStructType->FindFieldByName(Field->Name);
+		check(DerivativeField);
 
-		OutResult.ExpressionDdx = Tree.NewExpression<FExpressionGetStructField>(DerivativeStructType, Field->Name, StructDerivatives.ExpressionDdx);
-		OutResult.ExpressionDdy = Tree.NewExpression<FExpressionGetStructField>(DerivativeStructType, Field->Name, StructDerivatives.ExpressionDdy);
+		OutResult.ExpressionDdx = Tree.NewExpression<FExpressionGetStructField>(DerivativeStructType, DerivativeField, StructDerivatives.ExpressionDdx);
+		OutResult.ExpressionDdy = Tree.NewExpression<FExpressionGetStructField>(DerivativeStructType, DerivativeField, StructDerivatives.ExpressionDdy);
 	}
 }
 
@@ -446,14 +449,14 @@ bool FExpressionGetStructField::PrepareValue(FEmitContext& Context, const FReque
 	return OutResult.SetType(Context, RequestedType, StructPreparedType.GetFieldType(Field));
 }
 
-void FExpressionGetStructField::EmitValueShader(FEmitContext& Context, const FRequestedType& RequestedType, FEmitValueShaderResult& OutResult) const
+void FExpressionGetStructField::EmitValueShader(FEmitContext& Context, FEmitShaderScope& Scope, const FRequestedType& RequestedType, FEmitValueShaderResult& OutResult) const
 {
 	FRequestedType RequestedStructType(StructType, false);
 	RequestedStructType.SetField(Field, RequestedType);
 
-	FEmitShaderCode* StructValue = StructExpression->GetValueShader(Context, RequestedStructType);
+	FEmitShaderExpression* StructValue = StructExpression->GetValueShader(Context, Scope, RequestedStructType);
 
-	OutResult.Code = Context.EmitInlineCode(Field->Type, TEXT("%.%"),
+	OutResult.Code = Context.EmitInlineExpression(Scope, Field->Type, TEXT("%.%"),
 		StructValue,
 		Field->Name);
 }
@@ -476,9 +479,11 @@ void FExpressionSetStructField::ComputeAnalyticDerivatives(FTree& Tree, FExpress
 	{
 		const Shader::FStructType* DerivativeStructType = StructType->DerivativeType;
 		check(DerivativeStructType);
+		const Shader::FStructField* DerivativeField = DerivativeStructType->FindFieldByName(Field->Name);
+		check(DerivativeField);
 
-		OutResult.ExpressionDdx = Tree.NewExpression<FExpressionSetStructField>(DerivativeStructType, Field->Name, StructDerivatives.ExpressionDdx, FieldDerivatives.ExpressionDdx);
-		OutResult.ExpressionDdy = Tree.NewExpression<FExpressionSetStructField>(DerivativeStructType, Field->Name, StructDerivatives.ExpressionDdy, FieldDerivatives.ExpressionDdy);
+		OutResult.ExpressionDdx = Tree.NewExpression<FExpressionSetStructField>(DerivativeStructType, DerivativeField, StructDerivatives.ExpressionDdx, FieldDerivatives.ExpressionDdx);
+		OutResult.ExpressionDdy = Tree.NewExpression<FExpressionSetStructField>(DerivativeStructType, DerivativeField, StructDerivatives.ExpressionDdy, FieldDerivatives.ExpressionDdy);
 	}
 }
 
@@ -510,7 +515,7 @@ bool FExpressionSetStructField::PrepareValue(FEmitContext& Context, const FReque
 	}
 }
 
-void FExpressionSetStructField::EmitValueShader(FEmitContext& Context, const FRequestedType& RequestedType, FEmitValueShaderResult& OutResult) const
+void FExpressionSetStructField::EmitValueShader(FEmitContext& Context, FEmitShaderScope& Scope, const FRequestedType& RequestedType, FEmitValueShaderResult& OutResult) const
 {
 	FRequestedType RequestedStructType(RequestedType);
 	RequestedStructType.ClearFieldRequested(Field);
@@ -518,7 +523,7 @@ void FExpressionSetStructField::EmitValueShader(FEmitContext& Context, const FRe
 
 	const FRequestedType RequestedFieldType = RequestedType.GetField(Field);
 	const EExpressionEvaluation FieldEvaluation = FieldExpression->GetEvaluation(RequestedFieldType);
-	FEmitShaderCode* FieldValue = FieldExpression->GetValueShader(Context, RequestedFieldType, Field->Type);
+	FEmitShaderExpression* FieldValue = FieldExpression->GetValueShader(Context, Scope, RequestedFieldType, Field->Type);
 
 	if (StructEvaluation == EExpressionEvaluation::None)
 	{
@@ -526,7 +531,7 @@ void FExpressionSetStructField::EmitValueShader(FEmitContext& Context, const FRe
 
 		// StructExpression is not used, so default to a zero-initialized struct
 		// This will happen if all the accessed struct fields are explicitly defined
-		OutResult.Code = Context.EmitCode(StructType, TEXT("%_Set%((%)0, %)"),
+		OutResult.Code = Context.EmitExpression(Scope, StructType, TEXT("%_Set%((%)0, %)"),
 			StructType->Name,
 			Field->Name,
 			StructType->Name,
@@ -534,8 +539,8 @@ void FExpressionSetStructField::EmitValueShader(FEmitContext& Context, const FRe
 	}
 	else
 	{
-		FEmitShaderCode* StructValue = StructExpression->GetValueShader(Context, RequestedStructType);
-		OutResult.Code = Context.EmitCode(StructType, TEXT("%_Set%(%, %)"),
+		FEmitShaderExpression* StructValue = StructExpression->GetValueShader(Context, Scope, RequestedStructType);
+		OutResult.Code = Context.EmitExpression(Scope, StructType, TEXT("%_Set%(%, %)"),
 			StructType->Name,
 			Field->Name,
 			StructValue,
@@ -586,13 +591,13 @@ bool FExpressionSelect::PrepareValue(FEmitContext& Context, const FRequestedType
 	return OutResult.SetType(Context, RequestedType, MergePreparedTypes(LhsType, RhsType));
 }
 
-void FExpressionSelect::EmitValueShader(FEmitContext& Context, const FRequestedType& RequestedType, FEmitValueShaderResult& OutResult) const
+void FExpressionSelect::EmitValueShader(FEmitContext& Context, FEmitShaderScope& Scope, const FRequestedType& RequestedType, FEmitValueShaderResult& OutResult) const
 {
-	FEmitShaderCode* TrueValue = TrueExpression->GetValueShader(Context, RequestedType);
-	FEmitShaderCode* FalseValue = FalseExpression->GetValueShader(Context, RequestedType);
+	FEmitShaderExpression* TrueValue = TrueExpression->GetValueShader(Context, Scope, RequestedType);
+	FEmitShaderExpression* FalseValue = FalseExpression->GetValueShader(Context, Scope, RequestedType);
 
-	OutResult.Code = Context.EmitCode(RequestedType.GetType(), TEXT("(% ? % : %)"),
-		ConditionExpression->GetValueShader(Context, Shader::EValueType::Bool1),
+	OutResult.Code = Context.EmitExpression(Scope, RequestedType.GetType(), TEXT("(% ? % : %)"),
+		ConditionExpression->GetValueShader(Context, Scope, Shader::EValueType::Bool1),
 		TrueValue,
 		FalseValue);
 }
@@ -657,9 +662,9 @@ bool FExpressionUnaryOp::PrepareValue(FEmitContext& Context, const FRequestedTyp
 	return OutResult.SetType(Context, RequestedType, ResultType);
 }
 
-void FExpressionUnaryOp::EmitValueShader(FEmitContext& Context, const FRequestedType& RequestedType, FEmitValueShaderResult& OutResult) const
+void FExpressionUnaryOp::EmitValueShader(FEmitContext& Context, FEmitShaderScope& Scope, const FRequestedType& RequestedType, FEmitValueShaderResult& OutResult) const
 {
-	FEmitShaderCode* InputValue = Input->GetValueShader(Context, RequestedType);
+	FEmitShaderExpression* InputValue = Input->GetValueShader(Context, Scope, RequestedType);
 	Shader::FType ResultType = InputValue->Type;
 	const Shader::FValueTypeDescription InputTypeDesc = Shader::GetValueTypeDescription(ResultType);
 
@@ -668,22 +673,22 @@ void FExpressionUnaryOp::EmitValueShader(FEmitContext& Context, const FRequested
 	case EUnaryOp::Neg:
 		if (InputTypeDesc.ComponentType == Shader::EValueComponentType::Double)
 		{
-			OutResult.Code = Context.EmitCode(ResultType, TEXT("LWCNegate(%)"), InputValue);
+			OutResult.Code = Context.EmitExpression(Scope, ResultType, TEXT("LWCNegate(%)"), InputValue);
 		}
 		else
 		{
-			OutResult.Code = Context.EmitInlineCode(ResultType, TEXT("(-%)"), InputValue);
+			OutResult.Code = Context.EmitInlineExpression(Scope, ResultType, TEXT("(-%)"), InputValue);
 		}
 		break;
 	case EUnaryOp::Rcp:
 		if (InputTypeDesc.ComponentType == Shader::EValueComponentType::Double)
 		{
 			ResultType = Shader::MakeValueType(Shader::EValueComponentType::Float, ResultType.GetNumComponents());
-			OutResult.Code = Context.EmitCode(ResultType, TEXT("LWCRcp(%)"), InputValue);
+			OutResult.Code = Context.EmitExpression(Scope, ResultType, TEXT("LWCRcp(%)"), InputValue);
 		}
 		else
 		{
-			OutResult.Code = Context.EmitCode(ResultType, TEXT("rcp(%)"), InputValue);
+			OutResult.Code = Context.EmitExpression(Scope, ResultType, TEXT("rcp(%)"), InputValue);
 		}
 		break;
 	default:
@@ -803,65 +808,65 @@ FRequestedType GetBinaryOpInputType(const FRequestedType& RequestedType, EBinary
 }
 } // namespace Private
 
-void FExpressionBinaryOp::EmitValueShader(FEmitContext& Context, const FRequestedType& RequestedType, FEmitValueShaderResult& OutResult) const
+void FExpressionBinaryOp::EmitValueShader(FEmitContext& Context, FEmitShaderScope& Scope, const FRequestedType& RequestedType, FEmitValueShaderResult& OutResult) const
 {
 	const FRequestedType InputType = Private::GetBinaryOpInputType(RequestedType, Op, Lhs->GetRequestedType(), Rhs->GetRequestedType());
 	const Shader::EValueComponentType InputComponentType = InputType.ValueComponentType;
 	Shader::FType ResultType = InputType.GetType();
-	FEmitShaderCode* LhsValue = Lhs->GetValueShader(Context, InputType);
-	FEmitShaderCode* RhsValue = Rhs->GetValueShader(Context, InputType);
+	FEmitShaderExpression* LhsValue = Lhs->GetValueShader(Context, Scope, InputType);
+	FEmitShaderExpression* RhsValue = Rhs->GetValueShader(Context, Scope, InputType);
 
 	switch (Op)
 	{
 	case EBinaryOp::Add:
 		if (InputComponentType == Shader::EValueComponentType::Double)
 		{
-			OutResult.Code = Context.EmitCode(ResultType, TEXT("LWCAdd(%, %)"), LhsValue, RhsValue);
+			OutResult.Code = Context.EmitExpression(Scope, ResultType, TEXT("LWCAdd(%, %)"), LhsValue, RhsValue);
 		}
 		else
 		{
-			OutResult.Code = Context.EmitCode(ResultType, TEXT("(% + %)"), LhsValue, RhsValue);
+			OutResult.Code = Context.EmitExpression(Scope, ResultType, TEXT("(% + %)"), LhsValue, RhsValue);
 		}
 		break;
 	case EBinaryOp::Sub:
 		if (InputComponentType == Shader::EValueComponentType::Double)
 		{
-			OutResult.Code = Context.EmitCode(ResultType, TEXT("LWCSubtract(%, %)"), LhsValue, RhsValue);
+			OutResult.Code = Context.EmitExpression(Scope, ResultType, TEXT("LWCSubtract(%, %)"), LhsValue, RhsValue);
 		}
 		else
 		{
-			OutResult.Code = Context.EmitCode(ResultType, TEXT("(% - %)"), LhsValue, RhsValue);
+			OutResult.Code = Context.EmitExpression(Scope, ResultType, TEXT("(% - %)"), LhsValue, RhsValue);
 		}
 		break;
 	case EBinaryOp::Mul:
 		if (InputComponentType == Shader::EValueComponentType::Double)
 		{
-			OutResult.Code = Context.EmitCode(ResultType, TEXT("LWCMultiply(%, %)"), LhsValue, RhsValue);
+			OutResult.Code = Context.EmitExpression(Scope, ResultType, TEXT("LWCMultiply(%, %)"), LhsValue, RhsValue);
 		}
 		else
 		{
-			OutResult.Code = Context.EmitCode(ResultType, TEXT("(% * %)"), LhsValue, RhsValue);
+			OutResult.Code = Context.EmitExpression(Scope, ResultType, TEXT("(% * %)"), LhsValue, RhsValue);
 		}
 		break;
 	case EBinaryOp::Div:
 		if (InputComponentType == Shader::EValueComponentType::Double)
 		{
-			OutResult.Code = Context.EmitCode(ResultType, TEXT("LWCDivide(%, %)"), LhsValue, RhsValue);
+			OutResult.Code = Context.EmitExpression(Scope, ResultType, TEXT("LWCDivide(%, %)"), LhsValue, RhsValue);
 		}
 		else
 		{
-			OutResult.Code = Context.EmitCode(ResultType, TEXT("(% / %)"), LhsValue, RhsValue);
+			OutResult.Code = Context.EmitExpression(Scope, ResultType, TEXT("(% / %)"), LhsValue, RhsValue);
 		}
 		break;
 	case EBinaryOp::Less:
 		ResultType = Shader::MakeValueType(Shader::EValueComponentType::Bool, ResultType.GetNumComponents());
 		if (InputComponentType == Shader::EValueComponentType::Double)
 		{
-			OutResult.Code = Context.EmitCode(ResultType, TEXT("LWCLess(%, %)"), LhsValue, RhsValue);
+			OutResult.Code = Context.EmitExpression(Scope, ResultType, TEXT("LWCLess(%, %)"), LhsValue, RhsValue);
 		}
 		else
 		{
-			OutResult.Code = Context.EmitCode(ResultType, TEXT("(% < %)"), LhsValue, RhsValue);
+			OutResult.Code = Context.EmitExpression(Scope, ResultType, TEXT("(% < %)"), LhsValue, RhsValue);
 		}
 		break;
 	default:
@@ -906,7 +911,7 @@ bool FExpressionSwizzle::PrepareValue(FEmitContext& Context, const FRequestedTyp
 	return OutResult.SetType(Context, RequestedType, ResultType);
 }
 
-void FExpressionSwizzle::EmitValueShader(FEmitContext& Context, const FRequestedType& RequestedType, FEmitValueShaderResult& OutResult) const
+void FExpressionSwizzle::EmitValueShader(FEmitContext& Context, FEmitShaderScope& Scope, const FRequestedType& RequestedType, FEmitValueShaderResult& OutResult) const
 {
 	const FRequestedType RequestedInputType = Parameters.GetRequestedInputType(RequestedType);
 
@@ -924,10 +929,10 @@ void FExpressionSwizzle::EmitValueShader(FEmitContext& Context, const FRequested
 		}
 	}
 
-	FEmitShaderCode* InputValue = Input->GetValueShader(Context, RequestedInputType);
+	FEmitShaderExpression* InputValue = Input->GetValueShader(Context, Scope, RequestedInputType);
 	if (bHasSwizzleReorder)
 	{
-		OutResult.Code = Context.EmitInlineCode(RequestedType.GetType(), TEXT("%.%"),
+		OutResult.Code = Context.EmitInlineExpression(Scope, RequestedType.GetType(), TEXT("%.%"),
 			InputValue,
 			Swizzle);
 	}
@@ -1026,10 +1031,10 @@ FAppendTypes GetAppendTypes(const FRequestedType& RequestedType, const FRequeste
 }
 } // namespace Private
 
-void FExpressionAppend::EmitValueShader(FEmitContext& Context, const FRequestedType& RequestedType, FEmitValueShaderResult& OutResult) const
+void FExpressionAppend::EmitValueShader(FEmitContext& Context, FEmitShaderScope& Scope, const FRequestedType& RequestedType, FEmitValueShaderResult& OutResult) const
 {
 	const Private::FAppendTypes Types = Private::GetAppendTypes(RequestedType, Lhs->GetRequestedType(), Rhs->GetRequestedType());
-	FEmitShaderCode* LhsValue = Lhs->GetValueShader(Context, Types.LhsType);
+	FEmitShaderExpression* LhsValue = Lhs->GetValueShader(Context, Scope, Types.LhsType);
 
 	if (Types.RhsType.IsVoid())
 	{
@@ -1037,16 +1042,16 @@ void FExpressionAppend::EmitValueShader(FEmitContext& Context, const FRequestedT
 	}
 	else
 	{
-		FEmitShaderCode* RhsValue = Rhs->GetValueShader(Context, Types.RhsType);
+		FEmitShaderExpression* RhsValue = Rhs->GetValueShader(Context, Scope, Types.RhsType);
 		if (Types.ResultComponentType == Shader::EValueComponentType::Double)
 		{
-			OutResult.Code = Context.EmitInlineCode(Types.ResultType, TEXT("MakeLWCVector(%, %)"),
+			OutResult.Code = Context.EmitInlineExpression(Scope, Types.ResultType, TEXT("MakeLWCVector(%, %)"),
 				LhsValue,
 				RhsValue);
 		}
 		else
 		{
-			OutResult.Code = Context.EmitInlineCode(Types.ResultType, TEXT("%(%, %)"),
+			OutResult.Code = Context.EmitInlineExpression(Scope, Types.ResultType, TEXT("%(%, %)"),
 				Types.ResultType.GetName(),
 				LhsValue,
 				RhsValue);
@@ -1071,27 +1076,27 @@ bool FExpressionReflectionVector::PrepareValue(FEmitContext& Context, const FReq
 	return OutResult.SetType(Context, RequestedType, EExpressionEvaluation::Shader, Shader::EValueType::Float3);
 }
 
-void FExpressionReflectionVector::EmitValueShader(FEmitContext& Context, const FRequestedType& RequestedType, FEmitValueShaderResult& OutResult) const
+void FExpressionReflectionVector::EmitValueShader(FEmitContext& Context, FEmitShaderScope& Scope, const FRequestedType& RequestedType, FEmitValueShaderResult& OutResult) const
 {
-	OutResult.Code = Context.EmitInlineCode(Shader::EValueType::Float3, TEXT("Parameters.ReflectionVector"));
+	OutResult.Code = Context.EmitInlineExpression(Scope, Shader::EValueType::Float3, TEXT("Parameters.ReflectionVector"));
 }
 
 void FStatementBreak::Prepare(FEmitContext& Context) const
 {
 }
 
-void FStatementBreak::EmitShader(FEmitContext& Context) const
+void FStatementBreak::EmitShader(FEmitContext& Context, FEmitShaderScope& Scope) const
 {
-	ParentScope->EmitStatementf(Context, TEXT("break;"));
+	Context.EmitStatement(Scope, TEXT("break;"));
 }
 
 void FStatementReturn::Prepare(FEmitContext& Context) const
 {
 }
 
-void FStatementReturn::EmitShader(FEmitContext& Context) const
+void FStatementReturn::EmitShader(FEmitContext& Context, FEmitShaderScope& Scope) const
 {
-	ParentScope->EmitStatementf(Context, TEXT("return %s;"), Expression->GetValueShader(Context)->Reference);
+	Context.EmitStatement(Scope, TEXT("return %;"), Expression->GetValueShader(Context, Scope));
 }
 
 void FStatementIf::Prepare(FEmitContext& Context) const
@@ -1118,34 +1123,38 @@ void FStatementIf::Prepare(FEmitContext& Context) const
 	}
 }
 
-void FStatementIf::EmitShader(FEmitContext& Context) const
+void FStatementIf::EmitShader(FEmitContext& Context, FEmitShaderScope& Scope) const
 {
+	FEmitShaderNode* Dependency = nullptr;
 	if (ConditionExpression->GetEvaluation(Shader::EValueType::Bool1) == EExpressionEvaluation::Constant)
 	{
 		const bool bCondition = ConditionExpression->GetValueConstant(Context, Shader::EValueType::Bool1).AsBoolScalar();
 		if (bCondition && IsScopeLive(ThenScope))
 		{
-			ParentScope->EmitScope(Context, ThenScope);
+			Dependency = Context.EmitNextScope(Scope, ThenScope);
 		}
 		else if(!bCondition && IsScopeLive(ElseScope))
 		{
-			ParentScope->EmitScope(Context, ElseScope);
+			Dependency = Context.EmitNextScope(Scope, ElseScope);
 		}
 	}
 	else
 	{
 		if (IsScopeLive(ThenScope) || IsScopeLive(ElseScope))
 		{
-			FEmitShaderCode* ConditionValue = ConditionExpression->GetValueShader(Context, Shader::EValueType::Bool1);
-			ParentScope->EmitNestedScopef(Context, ThenScope, TEXT("if (%s)"), ConditionValue->Reference);
+			FEmitShaderExpression* ConditionValue = ConditionExpression->GetValueShader(Context, Scope, Shader::EValueType::Bool1);
 			if (IsScopeLive(ElseScope))
 			{
-				ParentScope->EmitNestedScopef(Context, ElseScope, TEXT("else"));
+				Dependency = Context.EmitNestedScopes(Scope, ThenScope, ElseScope, TEXT("if (%)"), TEXT("else"), ConditionValue);
+			}
+			else
+			{
+				Dependency = Context.EmitNestedScope(Scope, ThenScope, TEXT("if (%)"), ConditionValue);
 			}
 		}
 	}
 
-	ParentScope->EmitScope(Context, NextScope);
+	Context.EmitNextScopeWithDependency(Scope, Dependency, NextScope);
 }
 
 void FStatementLoop::Prepare(FEmitContext& Context) const
@@ -1153,13 +1162,14 @@ void FStatementLoop::Prepare(FEmitContext& Context) const
 	MarkScopeLive(LoopScope);
 }
 
-void FStatementLoop::EmitShader(FEmitContext& Context) const
+void FStatementLoop::EmitShader(FEmitContext& Context, FEmitShaderScope& Scope) const
 {
+	FEmitShaderNode* Dependency = nullptr;
 	if (IsScopeLive(LoopScope))
 	{
-		ParentScope->EmitNestedScopef(Context, LoopScope, TEXT("while (true)"));
+		Dependency = Context.EmitNestedScope(Scope, LoopScope, TEXT("while (true)"));
 	}
-	ParentScope->EmitScope(Context, NextScope);
+	Context.EmitNextScopeWithDependency(Scope, Dependency, NextScope);
 }
 
 } // namespace HLSLTree
diff --git a/Engine/Source/Runtime/Engine/Private/HLSLTree/HLSLTreeEmit.cpp b/Engine/Source/Runtime/Engine/Private/HLSLTree/HLSLTreeEmit.cpp
new file mode 100644
--- /dev/null
+++ b/Engine/Source/Runtime/Engine/Private/HLSLTree/HLSLTreeEmit.cpp
@@ -0,0 +1,748 @@
+// Copyright Epic Games, Inc. All Rights Reserved.
+#include "HLSLTree/HLSLTreeEmit.h"
+#include "HLSLTree/HLSLTree.h"
+#include "Misc/StringBuilder.h"
+#include "Misc/MemStack.h"
+#include "Misc/MemStackUtility.h"
+#include "Shader/ShaderTypes.h"
+#include "Shader/Preshader.h"
+#include "MaterialShared.h" // TODO - split preshader out into its own module
+#include "Misc/LargeWorldRenderPosition.h"
+
+namespace UE
+{
+namespace HLSLTree
+{
+
+FEmitShaderNode::FEmitShaderNode(FEmitShaderScope& InScope, TArrayView<FEmitShaderNode*> InDependencies)
+	: Scope(&InScope)
+	, Dependencies(InDependencies)
+{
+	for (FEmitShaderNode* Dependency : InDependencies)
+	{
+		check(Dependency);
+		check(Dependency != (void*)0xcccccccccccccccc);
+	}
+}
+
+namespace Private
+{
+
+void WriteIndent(int32 IndentLevel, FStringBuilderBase& InOutString)
+{
+	const int32 Offset = InOutString.AddUninitialized(IndentLevel);
+	TCHAR* Result = InOutString.GetData() + Offset;
+	for (int32 i = 0; i < IndentLevel; ++i)
+	{
+		*Result++ = TCHAR('\t');
+	}
+}
+
+void EmitShaderCode(FEmitShaderNode* EmitNode, FEmitShaderScopeStack& Stack)
+{
+	if (EmitNode && EmitNode->Scope)
+	{
+		const FEmitShaderScope* Scope = EmitNode->Scope;
+		FEmitShaderScopeEntry EmitEntry;
+		for (int32 Index = Stack.Num() - 1; Index >= 0; --Index)
+		{
+			FEmitShaderScopeEntry& CheckEntry = Stack[Index];
+			if (CheckEntry.Scope == Scope)
+			{
+				EmitEntry = CheckEntry;
+				break;
+			}
+		}
+
+		// LocalPHI can sometimes generate circular dependency on expressions that execute in the future
+		// Should revist this once dependencies are cleaned up
+		if (EmitEntry.Code)
+		{
+			EmitNode->Scope = nullptr; // only emit code once
+			for (FEmitShaderNode* Dependency : EmitNode->Dependencies)
+			{
+				Private::EmitShaderCode(Dependency, Stack);
+			}
+			EmitNode->EmitShaderCode(Stack, EmitEntry.Indent, *EmitEntry.Code);
+		}
+	}
+}
+
+} // namespace Private
+
+void FEmitShaderExpression::EmitShaderCode(FEmitShaderScopeStack& Stack, int32 Indent, FStringBuilderBase& OutString)
+{
+	// Don't need a declaration for inline values
+	if (!IsInline())
+	{
+		Private::WriteIndent(Indent, OutString);
+		OutString.Appendf(TEXT("const %s %s = %s;\n"), Type.GetName(), Reference, Value);
+	}
+}
+
+void FEmitShaderStatement::EmitShaderCode(FEmitShaderScopeStack& Stack, int32 Indent, FStringBuilderBase& OutString)
+{
+	TStringBuilder<1024> ScopeCode;
+	for (int32 i = 0; i < 2; ++i)
+	{
+		if (Code[i].Len() > 0)
+		{
+			Private::WriteIndent(Indent, ScopeCode);
+			ScopeCode.Append(Code[i]);
+			ScopeCode.AppendChar(TEXT('\n'));
+		}
+
+		FEmitShaderScope* NestedScope = NestedScopes[i];
+		if (NestedScope)
+		{
+			bool bNeedToCloseScope = false;
+			int32 NextIndent = Indent;
+
+			if (ScopeFormat == EEmitScopeFormat::Scoped)
+			{
+				Private::WriteIndent(Indent, ScopeCode);
+				ScopeCode.Append(TEXT("{\n"));
+				bNeedToCloseScope = true;
+				NextIndent++;
+			}
+
+			Stack.Emplace(NestedScopes[i], NextIndent, ScopeCode);
+			NestedScope->EmitShaderCode(Stack);
+			Stack.Pop(false);
+
+			if (bNeedToCloseScope)
+			{
+				Private::WriteIndent(Indent, ScopeCode);
+				ScopeCode.Append(TEXT("}\n"));
+			}
+		}
+	}
+	OutString.Append(ScopeCode.ToView());
+}
+
+void FEmitShaderScope::EmitShaderCode(FEmitShaderScopeStack& Stack)
+{
+	FEmitShaderNode* EmitNode = FirstNode;
+	while (EmitNode)
+	{
+		Private::EmitShaderCode(EmitNode, Stack);
+		EmitNode = EmitNode->NextScopedNode;
+	}
+}
+
+FEmitShaderScope* FEmitShaderScope::FindSharedParent(FEmitShaderScope* Lhs, FEmitShaderScope* Rhs)
+{
+	FEmitShaderScope* Scope0 = Lhs;
+	FEmitShaderScope* Scope1 = Rhs;
+	if (Scope1)
+	{
+		while (Scope0 != Scope1)
+		{
+			if (Scope0->NestedLevel > Scope1->NestedLevel)
+			{
+				check(Scope0->ParentScope);
+				Scope0 = Scope0->ParentScope;
+			}
+			else
+			{
+				check(Scope1->ParentScope);
+				Scope1 = Scope1->ParentScope;
+			}
+		}
+	}
+	return Scope0;
+}
+
+FEmitContext::FEmitContext(FMemStackBase& InAllocator, FErrorHandlerInterface& InErrors, const Shader::FStructTypeRegistry& InTypeRegistry)
+	: Allocator(&InAllocator)
+	, Errors(&InErrors)
+	, TypeRegistry(&InTypeRegistry)
+{
+}
+
+FEmitContext::~FEmitContext()
+{
+}
+
+const TCHAR* FEmitContext::AcquireLocalDeclarationCode()
+{
+	return MemStack::AllocateStringf(*Allocator, TEXT("Local%d"), NumExpressionLocals++);
+}
+
+namespace Private
+{
+void MoveToScope(FEmitShaderNode* EmitNode, FEmitShaderScope& Scope)
+{
+	if (EmitNode->Scope != &Scope)
+	{
+		FEmitShaderScope* NewScope = &Scope;
+		if (EmitNode->Scope)
+		{
+			NewScope = FEmitShaderScope::FindSharedParent(EmitNode->Scope, &Scope);
+			check(NewScope);
+		}
+
+		EmitNode->Scope = NewScope;
+		for (FEmitShaderNode* Dependency : EmitNode->Dependencies)
+		{
+			MoveToScope(Dependency, *NewScope);
+		}
+	}
+}
+
+void FormatArg_ShaderValue(FEmitShaderExpression* ShaderValue, FEmitShaderDependencies& OutDependencies, FStringBuilderBase& OutCode)
+{
+	OutDependencies.Add(ShaderValue);
+	OutCode.Append(ShaderValue->Reference);
+}
+
+int32 InternalFormatString(FStringBuilderBase* OutString, FEmitShaderDependencies& OutDependencies, FStringView Format, TArrayView<const FFormatArgVariant> ArgList, int32 BaseArgIndex)
+{
+	check(OutString || Format.Len() == 0);
+
+	int32 ArgIndex = BaseArgIndex;
+	for (TCHAR Char : Format)
+	{
+		if (Char == TEXT('%'))
+		{
+			const FFormatArgVariant& Arg = ArgList[ArgIndex++];
+			switch (Arg.Type)
+			{
+			case EFormatArgType::ShaderValue: FormatArg_ShaderValue(Arg.ShaderValue, OutDependencies, *OutString); break;
+			case EFormatArgType::String: OutString->Append(Arg.String); break;
+			case EFormatArgType::Int: OutString->Appendf(TEXT("%d"), Arg.Int); break;
+			default:
+				checkNoEntry();
+				break;
+			}
+		}
+		else
+		{
+			OutString->AppendChar(Char);
+		}
+	}
+
+	return ArgIndex;
+}
+
+void InternalFormatStrings(FStringBuilderBase* OutString0, FStringBuilderBase* OutString1, FEmitShaderDependencies& OutDependencies, FStringView Format0, FStringView Format1, const FFormatArgList& ArgList)
+{
+	int32 ArgIndex = 0;
+	ArgIndex = InternalFormatString(OutString0, OutDependencies, Format0, ArgList, ArgIndex);
+	ArgIndex = InternalFormatString(OutString1, OutDependencies, Format1, ArgList, ArgIndex);
+	checkf(ArgIndex == ArgList.Num(), TEXT("%d args were provided, but %d were used"), ArgList.Num(), ArgIndex);
+}
+
+} // namespace Private
+
+FEmitShaderExpression* FEmitContext::EmitExpressionInternal(FEmitShaderScope& Scope, TArrayView<FEmitShaderNode*> Dependencies, bool bInline, const Shader::FType& Type, FStringView Code)
+{
+	FEmitShaderExpression* ShaderValue = nullptr;
+
+	FXxHash64Builder Hasher;
+	Hasher.Update(Code.GetData(), Code.Len() * sizeof(TCHAR));
+	if (bInline)
+	{
+		uint8 InlineFlag = 1;
+		Hasher.Update(&InlineFlag, 1);
+	}
+
+	// Check to see if we've already generated code for an equivalent expression
+	const FXxHash64 ShaderHash = Hasher.Finalize();
+	FEmitShaderExpression** const PrevShaderValue = EmitExpressionMap.Find(ShaderHash);
+	if (PrevShaderValue)
+	{
+		ShaderValue = *PrevShaderValue;
+		check(ShaderValue->Type == Type);
+		Private::MoveToScope(ShaderValue, Scope);
+	}
+	else
+	{
+		ShaderValue = new(*Allocator) FEmitShaderExpression(Scope, MemStack::AllocateArrayView(*Allocator, Dependencies), Type, ShaderHash);
+		if (bInline)
+		{
+			ShaderValue->Reference = MemStack::AllocateString(*Allocator, Code);
+		}
+		else
+		{
+			ShaderValue->Reference = AcquireLocalDeclarationCode();
+			ShaderValue->Value = MemStack::AllocateString(*Allocator, Code);
+		}
+		EmitExpressionMap.Add(ShaderHash, ShaderValue);
+		EmitNodes.Add(ShaderValue);
+	}
+
+	return ShaderValue;
+}
+
+FEmitShaderStatement* FEmitContext::EmitStatementInternal(FEmitShaderScope& Scope, TArrayView<FEmitShaderNode*> Dependencies, EEmitScopeFormat ScopeFormat, FEmitShaderScope* NestedScope0, FEmitShaderScope* NestedScope1, FStringView Code0, FStringView Code1)
+{
+	FEmitShaderStatement* EmitStatement = new(*Allocator) FEmitShaderStatement(Scope, MemStack::AllocateArrayView(*Allocator, Dependencies));
+	EmitStatement->ScopeFormat = ScopeFormat;
+	EmitStatement->NestedScopes[0] = NestedScope0;
+	EmitStatement->NestedScopes[1] = NestedScope1;
+	EmitStatement->Code[0] = MemStack::AllocateString(*Allocator, Code0);
+	EmitStatement->Code[1] = MemStack::AllocateString(*Allocator, Code1);
+
+	EmitNodes.Add(EmitStatement);
+	return EmitStatement;
+}
+
+FEmitShaderScope* FEmitContext::AcquireEmitScope(FScope* Scope, FEmitShaderScope* OverrideParent)
+{
+	FEmitShaderScope* EmitScope = nullptr;
+	if (Scope)
+	{
+		FEmitShaderScope* const* PrevEmitScope = EmitScopeMap.Find(Scope);
+		EmitScope = PrevEmitScope ? *PrevEmitScope : nullptr;
+		if (!EmitScope)
+		{
+			FEmitShaderScope* ParentScope = OverrideParent;
+			if (!ParentScope)
+			{
+				ParentScope = AcquireEmitScope(Scope->ParentScope);
+			}
+
+			EmitScope = new(*Allocator) FEmitShaderScope();
+			EmitScope->ParentScope = ParentScope;
+			EmitScope->NestedLevel = ParentScope ? ParentScope->NestedLevel + 1 : 0;
+			EmitScopeMap.Add(Scope, EmitScope);
+			if (Scope->ContainedStatement)
+			{
+				Scope->ContainedStatement->EmitShader(*this, *EmitScope);
+			}
+		}
+		check(!OverrideParent || EmitScope->ParentScope == OverrideParent);
+	}
+	return EmitScope;
+}
+
+FEmitShaderScope* FEmitContext::FindEmitScope(FScope* Scope) const
+{
+	FEmitShaderScope* EmitScope = nullptr;
+	if (Scope)
+	{
+		FEmitShaderScope* const* PrevEmitScope = EmitScopeMap.Find(Scope);
+		EmitScope = PrevEmitScope ? *PrevEmitScope : nullptr;
+	}
+	return EmitScope;
+}
+
+namespace Private
+{
+void WriteMaterialUniformAccess(Shader::EValueComponentType ComponentType, uint32 NumComponents, uint32 UniformOffset, FStringBuilderBase& OutResult)
+{
+	static const TCHAR IndexToMask[] = TEXT("xyzw");
+	uint32 RegisterIndex = UniformOffset / 4;
+	uint32 RegisterOffset = UniformOffset % 4;
+	uint32 NumComponentsToWrite = NumComponents;
+	bool bConstructor = false;
+
+	check(ComponentType == Shader::EValueComponentType::Float || ComponentType == Shader::EValueComponentType::Int);
+	const bool bIsInt = (ComponentType == Shader::EValueComponentType::Int);
+
+	while (NumComponentsToWrite > 0u)
+	{
+		const uint32 NumComponentsInRegister = FMath::Min(NumComponentsToWrite, 4u - RegisterOffset);
+		if (NumComponentsInRegister < NumComponents && !bConstructor)
+		{
+			// Uniform will be split across multiple registers, so add the constructor to concat them together
+			OutResult.Appendf(TEXT("%s%d("), Shader::GetComponentTypeName(ComponentType), NumComponents);
+			bConstructor = true;
+		}
+
+		if (bIsInt)
+		{
+			// PreshaderBuffer is typed as float4, so reinterpret as 'int' if needed
+			OutResult.Append(TEXT("asint("));
+		}
+
+		OutResult.Appendf(TEXT("Material.PreshaderBuffer[%u]"), RegisterIndex);
+		// Can skip writing mask if we're taking all 4 components from the register
+		if (NumComponentsInRegister < 4u)
+		{
+			OutResult.AppendChar(TCHAR('.'));
+			for (uint32 i = 0u; i < NumComponentsInRegister; ++i)
+			{
+				OutResult.AppendChar(IndexToMask[RegisterOffset + i]);
+			}
+		}
+
+		if (bIsInt)
+		{
+			OutResult.Append(TEXT(")"));
+		}
+
+		NumComponentsToWrite -= NumComponentsInRegister;
+		RegisterIndex++;
+		RegisterOffset = 0u;
+		if (NumComponentsToWrite > 0u)
+		{
+			OutResult.Append(TEXT(", "));
+		}
+	}
+	if (bConstructor)
+	{
+		OutResult.Append(TEXT(")"));
+	}
+}
+} // namespace Private
+
+FEmitShaderExpression* FEmitContext::EmitPreshaderOrConstant(FEmitShaderScope& Scope, const FRequestedType& RequestedType, FExpression* Expression)
+{
+	Shader::FPreshaderData LocalPreshader;
+	Expression->EmitValuePreshader(*this, RequestedType, LocalPreshader);
+
+	const Shader::FType Type = RequestedType.GetType();
+
+	FXxHash64Builder Hasher;
+	Hasher.Update(&Type, sizeof(Type));
+	LocalPreshader.AppendHash(Hasher);
+	const FXxHash64 Hash = Hasher.Finalize();
+	FEmitShaderExpression* const* PrevShaderValue = EmitPreshaderMap.Find(Hash);
+	if (PrevShaderValue)
+	{
+		FEmitShaderExpression* ShaderValue = *PrevShaderValue;
+		check(ShaderValue->Type == Type);
+		Private::MoveToScope(ShaderValue, Scope);
+		return ShaderValue;
+	}
+
+	Shader::FPreshaderStack Stack;
+	const Shader::FPreshaderValue ConstantValue = LocalPreshader.EvaluateConstant(*Material, Stack);
+
+	TStringBuilder<1024> FormattedCode;
+	if (Type.IsStruct())
+	{
+		FormattedCode.Append(TEXT("{ "));
+	}
+
+	FMaterialUniformPreshaderHeader* PreshaderHeader = nullptr;
+	uint32 CurrentBoolUniformOffset = ~0u;
+	uint32 CurrentNumBoolComponents = 32u;
+
+	int32 ComponentIndex = 0;
+	for (int32 FieldIndex = 0; FieldIndex < Type.GetNumFlatFields(); ++FieldIndex)
+	{
+		if (FieldIndex > 0)
+		{
+			FormattedCode.Append(TEXT(", "));
+		}
+
+		const Shader::EValueType FieldType = Type.GetFlatFieldType(FieldIndex);
+		const Shader::FValueTypeDescription TypeDesc = Shader::GetValueTypeDescription(FieldType);
+		const int32 NumFieldComponents = TypeDesc.NumComponents;
+		const EExpressionEvaluation FieldEvaluation = Expression->GetPreparedType().GetFieldEvaluation(ComponentIndex, NumFieldComponents);
+
+		if (FieldEvaluation == EExpressionEvaluation::Preshader)
+		{
+			// Only need to allocate uniform buffer for non-constant components
+			// Constant components can have their value inlined into the shader directly
+			FUniformExpressionSet& UniformExpressionSet = MaterialCompilationOutput->UniformExpressionSet;
+			if (!PreshaderHeader)
+			{
+				// Allocate a preshader header the first time we hit a non-constant field
+				PreshaderHeader = &UniformExpressionSet.UniformPreshaders.AddDefaulted_GetRef();
+				PreshaderHeader->FieldIndex = UniformExpressionSet.UniformPreshaderFields.Num();
+				PreshaderHeader->NumFields = 0u;
+				PreshaderHeader->OpcodeOffset = UniformExpressionSet.UniformPreshaderData.Num();
+				Expression->EmitValuePreshader(*this, RequestedType, UniformExpressionSet.UniformPreshaderData);
+				PreshaderHeader->OpcodeSize = UniformExpressionSet.UniformPreshaderData.Num() - PreshaderHeader->OpcodeOffset;
+			}
+
+			FMaterialUniformPreshaderField& PreshaderField = UniformExpressionSet.UniformPreshaderFields.AddDefaulted_GetRef();
+			PreshaderField.ComponentIndex = ComponentIndex;
+			PreshaderField.Type = FieldType;
+			PreshaderHeader->NumFields++;
+
+			if (TypeDesc.ComponentType == Shader::EValueComponentType::Bool)
+			{
+				// 'Bool' uniforms are packed into bits
+				if (CurrentNumBoolComponents + NumFieldComponents > 32u)
+				{
+					CurrentBoolUniformOffset = UniformPreshaderOffset++;
+					CurrentNumBoolComponents = 0u;
+				}
+
+				const uint32 RegisterIndex = CurrentBoolUniformOffset / 4;
+				const uint32 RegisterOffset = CurrentBoolUniformOffset % 4;
+				FormattedCode.Appendf(TEXT("UnpackUniform_%s(asuint(Material.PreshaderBuffer[%u][%u]), %u)"),
+					TypeDesc.Name,
+					RegisterIndex,
+					RegisterOffset,
+					CurrentNumBoolComponents);
+
+				PreshaderField.BufferOffset = CurrentBoolUniformOffset * 32u + CurrentNumBoolComponents;
+				CurrentNumBoolComponents += NumFieldComponents;
+			}
+			else if (TypeDesc.ComponentType == Shader::EValueComponentType::Double)
+			{
+				// Double uniforms are split into Tile/Offset components to make FLWCScalar/FLWCVectors
+				PreshaderField.BufferOffset = UniformPreshaderOffset;
+
+				if (NumFieldComponents > 1)
+				{
+					FormattedCode.Appendf(TEXT("MakeLWCVector%d("), NumFieldComponents);
+				}
+				else
+				{
+					FormattedCode.Append(TEXT("MakeLWCScalar("));
+				}
+
+				// Write the tile uniform
+				Private::WriteMaterialUniformAccess(Shader::EValueComponentType::Float, NumFieldComponents, UniformPreshaderOffset, FormattedCode);
+				UniformPreshaderOffset += NumFieldComponents;
+				FormattedCode.Append(TEXT(", "));
+
+				// Write the offset uniform
+				Private::WriteMaterialUniformAccess(Shader::EValueComponentType::Float, NumFieldComponents, UniformPreshaderOffset, FormattedCode);
+				UniformPreshaderOffset += NumFieldComponents;
+				FormattedCode.Append(TEXT(")"));
+			}
+			else
+			{
+				// Float/Int uniforms are written directly to the uniform buffer
+				const uint32 RegisterOffset = UniformPreshaderOffset % 4;
+				if (RegisterOffset + NumFieldComponents > 4u)
+				{
+					// If this uniform would span multiple registers, align offset to the next register to avoid this
+					// TODO - we could keep track of this empty padding space, and pack other smaller uniform types here
+					UniformPreshaderOffset = Align(UniformPreshaderOffset, 4u);
+				}
+
+				PreshaderField.BufferOffset = UniformPreshaderOffset;
+				Private::WriteMaterialUniformAccess(TypeDesc.ComponentType, NumFieldComponents, UniformPreshaderOffset, FormattedCode);
+				UniformPreshaderOffset += NumFieldComponents;
+			}
+		}
+		else
+		{
+			// We allow FieldEvaluation to be 'None', since in that case we still need to fill in a value for the HLSL initializer
+			check(FieldEvaluation == EExpressionEvaluation::Constant || FieldEvaluation == EExpressionEvaluation::None);
+
+			// The type generated by the preshader might not match the expected type
+			// In the future, with new HLSLTree, preshader could potentially include explicit cast opcodes, and avoid implicit conversions
+			Shader::FValue FieldConstantValue(ConstantValue.Type.GetComponentType(ComponentIndex), NumFieldComponents);
+			for (int32 i = 0; i < NumFieldComponents; ++i)
+			{
+				// Allow replicating scalar values
+				FieldConstantValue.Component[i] = ConstantValue.Component.Num() == 1 ? ConstantValue.Component[0] : ConstantValue.Component[ComponentIndex + i];
+			}
+
+			if (TypeDesc.ComponentType == Shader::EValueComponentType::Double)
+			{
+				const Shader::FDoubleValue DoubleValue = FieldConstantValue.AsDouble();
+				TStringBuilder<256> TileValue;
+				TStringBuilder<256> OffsetValue;
+				for (int32 Index = 0; Index < NumFieldComponents; ++Index)
+				{
+					if (Index > 0)
+					{
+						TileValue.Append(TEXT(", "));
+						OffsetValue.Append(TEXT(", "));
+					}
+
+					const FLargeWorldRenderScalar Value(DoubleValue[Index]);
+					TileValue.Appendf(TEXT("%#.9gf"), Value.GetTile());
+					OffsetValue.Appendf(TEXT("%#.9gf"), Value.GetOffset());
+				}
+
+				if (NumFieldComponents > 1)
+				{
+					FormattedCode.Appendf(TEXT("MakeLWCVector%d(float%d(%s), float%d(%s))"), NumFieldComponents, NumFieldComponents, TileValue.ToString(), NumFieldComponents, OffsetValue.ToString());
+				}
+				else
+				{
+					FormattedCode.Appendf(TEXT("MakeLWCScalar(%s, %s)"), TileValue.ToString(), OffsetValue.ToString());
+				}
+			}
+			else
+			{
+				const Shader::FValue CastFieldConstantValue = Shader::Cast(FieldConstantValue, FieldType);
+				if (NumFieldComponents > 1)
+				{
+					FormattedCode.Appendf(TEXT("%s("), TypeDesc.Name);
+				}
+				for (int32 Index = 0; Index < NumFieldComponents; ++Index)
+				{
+					if (Index > 0)
+					{
+						FormattedCode.Append(TEXT(", "));
+					}
+					CastFieldConstantValue.Component[Index].ToString(TypeDesc.ComponentType, FormattedCode);
+				}
+				if (NumFieldComponents > 1)
+				{
+					FormattedCode.Append(TEXT(")"));
+				}
+			}
+		}
+		ComponentIndex += NumFieldComponents;
+	}
+	check(ComponentIndex == Type.GetNumComponents());
+
+	if (Type.IsStruct())
+	{
+		FormattedCode.Append(TEXT(" }"));
+	}
+
+	const bool bInline = !Type.IsStruct(); // struct declarations can't be inline, due to HLSL syntax
+	FEmitShaderExpression* ShaderValue = EmitExpressionInternal(Scope, TArrayView<FEmitShaderNode*>(), bInline, Type, FormattedCode.ToView());
+	EmitPreshaderMap.Add(Hash, ShaderValue);
+
+	return ShaderValue;
+}
+
+FEmitShaderExpression* FEmitContext::EmitConstantZero(FEmitShaderScope& Scope, const Shader::FType& Type)
+{
+	return EmitInlineExpression(Scope, Type, TEXT("((%)0)"), Type.GetName());
+}
+
+FEmitShaderExpression* FEmitContext::EmitCast(FEmitShaderScope& Scope, FEmitShaderExpression* ShaderValue, const Shader::FType& DestType)
+{
+	check(ShaderValue);
+	check(!DestType.IsVoid());
+
+	if (ShaderValue->Type == DestType)
+	{
+		return ShaderValue;
+	}
+
+	const Shader::FValueTypeDescription SourceTypeDesc = Shader::GetValueTypeDescription(ShaderValue->Type);
+	const Shader::FValueTypeDescription DestTypeDesc = Shader::GetValueTypeDescription(DestType);
+
+	TStringBuilder<1024> FormattedCode;
+	Shader::FType IntermediateType = DestType;
+
+	if (SourceTypeDesc.NumComponents > 0 && DestTypeDesc.NumComponents > 0)
+	{
+		const bool bIsSourceLWC = SourceTypeDesc.ComponentType == Shader::EValueComponentType::Double;
+		const bool bIsLWC = DestTypeDesc.ComponentType == Shader::EValueComponentType::Double;
+
+		if (bIsLWC != bIsSourceLWC)
+		{
+			if (bIsLWC)
+			{
+				// float->LWC
+				ShaderValue = EmitCast(Scope, ShaderValue, Shader::MakeValueType(Shader::EValueComponentType::Float, DestTypeDesc.NumComponents));
+				FormattedCode.Appendf(TEXT("LWCPromote(%s)"), ShaderValue->Reference);
+			}
+			else
+			{
+				//LWC->float
+				FormattedCode.Appendf(TEXT("LWCToFloat(%s)"), ShaderValue->Reference);
+				IntermediateType = Shader::MakeValueType(Shader::EValueComponentType::Float, SourceTypeDesc.NumComponents);
+			}
+		}
+		else
+		{
+			const bool bReplicateScalar = (SourceTypeDesc.NumComponents == 1);
+
+			int32 NumComponents = 0;
+			bool bNeedClosingParen = false;
+			if (bIsLWC)
+			{
+				FormattedCode.Append(TEXT("MakeLWCVector("));
+				bNeedClosingParen = true;
+			}
+			else
+			{
+				if (SourceTypeDesc.NumComponents == 1 || SourceTypeDesc.NumComponents == DestTypeDesc.NumComponents)
+				{
+					NumComponents = DestTypeDesc.NumComponents;
+					// Cast the scalar to the correct type, HLSL language will replicate the scalar if needed when performing this cast
+					FormattedCode.Appendf(TEXT("((%s)%s)"), DestTypeDesc.Name, ShaderValue->Reference);
+				}
+				else
+				{
+					NumComponents = FMath::Min(SourceTypeDesc.NumComponents, DestTypeDesc.NumComponents);
+					if (NumComponents < DestTypeDesc.NumComponents)
+					{
+						FormattedCode.Appendf(TEXT("%s("), DestTypeDesc.Name);
+						bNeedClosingParen = true;
+					}
+					if (NumComponents == SourceTypeDesc.NumComponents && SourceTypeDesc.ComponentType == DestTypeDesc.ComponentType)
+					{
+						// If we're taking all the components from the source, can avoid adding a swizzle
+						FormattedCode.Append(ShaderValue->Reference);
+					}
+					else
+					{
+						// Use a cast to truncate the source to the correct number of types
+						const Shader::EValueType LocalType = Shader::MakeValueType(DestTypeDesc.ComponentType, NumComponents);
+						FormattedCode.Appendf(TEXT("((%s)%s)"), Shader::GetValueTypeDescription(LocalType).Name, ShaderValue->Reference);
+					}
+				}
+			}
+
+			if (bNeedClosingParen)
+			{
+				const Shader::FValue ZeroValue(DestTypeDesc.ComponentType, 1);
+				for (int32 ComponentIndex = NumComponents; ComponentIndex < DestTypeDesc.NumComponents; ++ComponentIndex)
+				{
+					if (ComponentIndex > 0u)
+					{
+						FormattedCode.Append(TEXT(","));
+					}
+					if (bIsLWC)
+					{
+						if (!bReplicateScalar && ComponentIndex >= SourceTypeDesc.NumComponents)
+						{
+							FormattedCode.Append(TEXT("LWCPromote(0.0f)"));
+						}
+						else
+						{
+							FormattedCode.Appendf(TEXT("LWCGetComponent(%s, %d)"), ShaderValue->Reference, bReplicateScalar ? 0 : ComponentIndex);
+						}
+					}
+					else
+					{
+						// Non-LWC case should only be zero-filling here, other cases should have already been handled
+						check(!bReplicateScalar);
+						check(ComponentIndex >= SourceTypeDesc.NumComponents);
+						ZeroValue.ToString(Shader::EValueStringFormat::HLSL, FormattedCode);
+					}
+				}
+				NumComponents = DestTypeDesc.NumComponents;
+				FormattedCode.Append(TEXT(")"));
+			}
+
+			check(NumComponents == DestTypeDesc.NumComponents);
+		}
+	}
+	else
+	{
+		Errors->AddErrorf(TEXT("Cannot cast between non-numeric types %s to %s."), SourceTypeDesc.Name, DestTypeDesc.Name);
+		FormattedCode.Appendf(TEXT("((%s)0)"), DestType.GetName());
+	}
+
+	check(IntermediateType != ShaderValue->Type);
+	ShaderValue = EmitInlineExpressionWithDependency(Scope, ShaderValue, IntermediateType, FormattedCode.ToView());
+	if (ShaderValue->Type != DestType)
+	{
+		// May need to cast through multiple intermediate types to reach our destination type
+		ShaderValue = EmitCast(Scope, ShaderValue, DestType);
+	}
+	return ShaderValue;
+}
+
+void FEmitContext::Finalize()
+{
+	// Unlink all nodes from scopes
+	for (FEmitShaderNode* EmitNode : EmitNodes)
+	{
+		EmitNode->Scope = nullptr;
+		EmitNode->NextScopedNode = nullptr;
+	}
+	
+	// Don't reset Expression/Preshader maps, allow future passes to share matching preshaders/expressions
+
+	EmitScopeMap.Reset();
+	EmitFunctionMap.Reset();
+	EmitLocalPHIMap.Reset();
+
+	MaterialCompilationOutput->UniformExpressionSet.UniformPreshaderBufferSize = (UniformPreshaderOffset + 3u) / 4u;
+}
+
+} // namespace HLSLTree
+} // namespace UE
diff --git a/Engine/Source/Runtime/Engine/Private/HLSLTree/HLSLTreeIntrinsics.cpp b/Engine/Source/Runtime/Engine/Private/HLSLTree/HLSLTreeIntrinsics.cpp
deleted file mode 100644
diff --git a/Engine/Source/Runtime/Engine/Private/Materials/HLSLMaterialTranslator.cpp b/Engine/Source/Runtime/Engine/Private/Materials/HLSLMaterialTranslator.cpp
--- a/Engine/Source/Runtime/Engine/Private/Materials/HLSLMaterialTranslator.cpp
+++ b/Engine/Source/Runtime/Engine/Private/Materials/HLSLMaterialTranslator.cpp
@@ -648,7 +648,7 @@ bool FHLSLMaterialTranslator::Translate()
 		Material->CompileErrors.Empty();
 		Material->ErrorExpressions.Empty();
 
-		bEnableExecutionFlow = Material->IsCompiledWithExecutionFlow();
+		bEnableExecutionFlow = Material->IsUsingControlFlow();
 
 		bCompileForComputeShader = Material->IsLightFunction();
 
diff --git a/Engine/Source/Runtime/Engine/Private/Materials/Material.cpp b/Engine/Source/Runtime/Engine/Private/Materials/Material.cpp
--- a/Engine/Source/Runtime/Engine/Private/Materials/Material.cpp
+++ b/Engine/Source/Runtime/Engine/Private/Materials/Material.cpp
@@ -3305,7 +3305,7 @@ void UMaterial::PropagateDataToMaterialProxy()
 	UpdateMaterialRenderProxy(*DefaultMaterialInstance);
 }
 
-bool UMaterial::IsCompiledWithExecutionFlow() const
+bool UMaterial::IsUsingControlFlow() const
 {
 	if (bEnableExecWire)
 	{
@@ -3538,7 +3538,7 @@ bool UMaterial::CanEditChange(const FProperty* InProperty) const
 
 void UMaterial::CreateExecutionFlowExpressions()
 {
-	if (IsCompiledWithExecutionFlow())
+	if (IsUsingControlFlow())
 	{
 		if (!ExpressionExecBegin)
 		{
@@ -4986,7 +4986,7 @@ int32 UMaterial::CompilePropertyEx( FMaterialCompiler* Compiler, const FGuid& At
 {
 	const EMaterialProperty Property = FMaterialAttributeDefinitionMap::GetProperty(AttributeID);
 
-	if (IsCompiledWithExecutionFlow())
+	if (IsUsingControlFlow())
 	{
 		check(ExpressionExecBegin);
 		return ExpressionExecBegin->Compile(Compiler, UMaterialExpression::CompileExecutionOutputIndex);
diff --git a/Engine/Source/Runtime/Engine/Private/Materials/MaterialExpressionHLSL.cpp b/Engine/Source/Runtime/Engine/Private/Materials/MaterialExpressionHLSL.cpp
--- a/Engine/Source/Runtime/Engine/Private/Materials/MaterialExpressionHLSL.cpp
+++ b/Engine/Source/Runtime/Engine/Private/Materials/MaterialExpressionHLSL.cpp
@@ -27,6 +27,7 @@
 #include "Materials/MaterialExpressionTextureObject.h"
 #include "Materials/MaterialExpressionTextureObjectParameter.h"
 #include "Materials/MaterialExpressionFunctionInput.h"
+#include "Materials/MaterialExpressionFunctionOutput.h"
 #include "Materials/MaterialExpressionMaterialFunctionCall.h"
 #include "Materials/MaterialExpressionStaticSwitch.h"
 #include "Materials/MaterialExpressionGetLocal.h"
@@ -186,12 +187,8 @@ bool UMaterialExpressionTextureSample::GenerateHLSLExpressionBase(FMaterialHLSLG
 	}
 
 	FExpression* TexCoordExpression = Coordinates.GetTracedInput().Expression ? Coordinates.TryAcquireHLSLExpression(Generator, Scope) : Generator.NewTexCoord(ConstCoordinate);
-	FExpressionTextureSample* ExpressionTextureSample = Generator.GetTree().NewExpression<FExpressionTextureSample>(TextureDeclaration, TexCoordExpression);
-	ExpressionTextureSample->SamplerSource = SamplerSource;
-	ExpressionTextureSample->MipValueMode = MipValueMode;
-	ExpressionTextureSample->TexCoordDerivatives = Generator.GetTree().GetAnalyticDerivatives(TexCoordExpression);
-
-	OutExpression = ExpressionTextureSample;
+	const FExpressionDerivatives TexCoordDerivatives = Generator.GetTree().GetAnalyticDerivatives(TexCoordExpression);
+	OutExpression = Generator.GetTree().NewExpression<FExpressionTextureSample>(TextureDeclaration, TexCoordExpression, TexCoordDerivatives, SamplerSource, MipValueMode);
 	return true;
 }
 
@@ -321,7 +318,8 @@ bool UMaterialExpressionGetMaterialAttributes::GenerateHLSLExpression(FMaterialH
 
 	const FGuid& AttributeID = AttributeGetTypes[AttributeIndex];
 	const FString& AttributeName = FMaterialAttributeDefinitionMap::GetAttributeName(AttributeID);
-	OutExpression = Generator.GetTree().NewExpression<UE::HLSLTree::FExpressionGetStructField>(Generator.GetMaterialAttributesType(), *AttributeName, AttributesExpression);
+	const UE::Shader::FStructField* AttributeField = Generator.GetMaterialAttributesType()->FindFieldByName(*AttributeName);
+	OutExpression = Generator.GetTree().NewExpression<UE::HLSLTree::FExpressionGetStructField>(Generator.GetMaterialAttributesType(), AttributeField, AttributesExpression);
 
 	return true;
 
@@ -345,7 +343,8 @@ bool UMaterialExpressionSetMaterialAttributes::GenerateHLSLExpression(FMaterialH
 				UE::HLSLTree::FExpression* ValueExpression = AttributeInput.TryAcquireHLSLExpression(Generator, Scope);
 				if (ValueExpression)
 				{
-					AttributesExpression = Generator.GetTree().NewExpression<UE::HLSLTree::FExpressionSetStructField>(Generator.GetMaterialAttributesType(), *AttributeName, AttributesExpression, ValueExpression);
+					const UE::Shader::FStructField* AttributeField = Generator.GetMaterialAttributesType()->FindFieldByName(*AttributeName);
+					AttributesExpression = Generator.GetTree().NewExpression<UE::HLSLTree::FExpressionSetStructField>(Generator.GetMaterialAttributesType(), AttributeField, AttributesExpression, ValueExpression);
 				}
 			}
 		}
@@ -363,24 +362,23 @@ bool UMaterialExpressionReflectionVectorWS::GenerateHLSLExpression(FMaterialHLSL
 	return true;
 }
 
+bool UMaterialExpressionFunctionOutput::GenerateHLSLExpression(FMaterialHLSLGenerator& Generator, UE::HLSLTree::FScope& Scope, int32 OutputIndex, UE::HLSLTree::FExpression*& OutExpression)
+{
+	// This should only be called when editing/previewing the function directly
+	OutExpression = A.AcquireHLSLExpression(Generator, Scope);
+	return true;
+}
+
 bool UMaterialExpressionFunctionInput::GenerateHLSLExpression(FMaterialHLSLGenerator& Generator, UE::HLSLTree::FScope& Scope, int32 OutputIndex, UE::HLSLTree::FExpression*& OutExpression)
 {
 	OutExpression = Generator.AcquireFunctionInputExpression(Scope, this);
-	if (OutExpression)
-	{
-		return true;
-	}
-	return Generator.GetErrors().AddError(TEXT("Invalid material function input"));
+	return true;
 }
 
 bool UMaterialExpressionMaterialFunctionCall::GenerateHLSLExpression(FMaterialHLSLGenerator& Generator, UE::HLSLTree::FScope& Scope, int32 OutputIndex, UE::HLSLTree::FExpression*& OutExpression)
 {
 	OutExpression = Generator.GenerateFunctionCall(Scope, MaterialFunction, FunctionInputs, OutputIndex);
-	if (OutExpression)
-	{
-		return true;
-	}
-	return false;
+	return true;
 }
 
 bool UMaterialExpressionExecBegin::GenerateHLSLStatements(FMaterialHLSLGenerator& Generator, UE::HLSLTree::FScope& Scope)
diff --git a/Engine/Source/Runtime/Engine/Private/Materials/MaterialExpressions.cpp b/Engine/Source/Runtime/Engine/Private/Materials/MaterialExpressions.cpp
--- a/Engine/Source/Runtime/Engine/Private/Materials/MaterialExpressions.cpp
+++ b/Engine/Source/Runtime/Engine/Private/Materials/MaterialExpressions.cpp
@@ -12669,6 +12669,27 @@ void UMaterialFunction::UpdateDependentFunctionCandidates()
 void UMaterialFunction::PostEditChangeProperty(FPropertyChangedEvent& PropertyChangedEvent)
 {
 	Super::PostEditChangeProperty(PropertyChangedEvent);
+
+	FProperty* PropertyThatChanged = PropertyChangedEvent.Property;
+	if (PropertyChangedEvent.GetPropertyName() == GET_MEMBER_NAME_CHECKED(UMaterialFunction, bEnableExecWire))
+	{
+		CreateExecutionFlowExpressions();
+	}
+
+	if (PropertyChangedEvent.GetPropertyName() == GET_MEMBER_NAME_CHECKED(UMaterialFunction, bEnableNewHLSLGenerator))
+	{
+		if (EditorMaterial)
+		{
+			EditorMaterial->bEnableNewHLSLGenerator = bEnableNewHLSLGenerator;
+		}
+	}
+
+	// many property changes can require rebuild of graph so always mark as changed
+	// not interested in PostEditChange calls though as the graph may have instigated it
+	if (PropertyThatChanged && MaterialGraph)
+	{
+		MaterialGraph->NotifyGraphChanged();
+	}
 }
 
 void UMaterialFunction::ForceRecompileForRendering(FMaterialUpdateContext& UpdateContext, UMaterial* InPreviewMaterial)
@@ -13295,7 +13316,55 @@ bool UMaterialFunction::SetStaticComponentMaskParameterValueEditorOnly(FName Par
 	Meta.ExpressionGuid = OutExpressionGuid;
 	return SetParameterValueEditorOnly(ParameterName, Meta);
 };
-#endif //WITH_EDITORONLY_DATA
+
+bool UMaterialFunction::IsUsingControlFlow() const
+{
+	if (bEnableExecWire)
+	{
+		static const auto CVar = IConsoleManager::Get().FindTConsoleVariableDataInt(TEXT("r.MaterialEnableControlFlow"));
+		return CVar->GetValueOnAnyThread() != 0;
+	}
+	return false;
+}
+
+bool UMaterialFunction::IsUsingNewHLSLGenerator() const
+{
+	if (bEnableNewHLSLGenerator)
+	{
+		static const auto CVar = IConsoleManager::Get().FindTConsoleVariableDataInt(TEXT("r.MaterialEnableNewHLSLGenerator"));
+		return CVar->GetValueOnAnyThread() != 0;
+	}
+	return false;
+}
+
+void UMaterialFunction::CreateExecutionFlowExpressions()
+{
+	const bool bUsingControlFlow = IsUsingControlFlow();
+	if (bUsingControlFlow)
+	{
+		if (!ExpressionExecBegin)
+		{
+			ExpressionExecBegin = NewObject<UMaterialExpressionExecBegin>(this);
+			ExpressionExecBegin->Function = this;
+			FunctionExpressions.Add(ExpressionExecBegin);
+		}
+
+		if (!ExpressionExecEnd)
+		{
+			ExpressionExecEnd = NewObject<UMaterialExpressionExecEnd>(this);
+			ExpressionExecEnd->Function = this;
+			FunctionExpressions.Add(ExpressionExecEnd);
+		}
+	}
+
+	if (EditorMaterial)
+	{
+		EditorMaterial->bEnableExecWire = bUsingControlFlow;
+		EditorMaterial->ExpressionExecBegin = ExpressionExecBegin;
+		EditorMaterial->ExpressionExecEnd = ExpressionExecEnd;
+	}
+}
+#endif // WITH_EDITOR
 
 ///////////////////////////////////////////////////////////////////////////////
 // UMaterialFunctionInstance
diff --git a/Engine/Source/Runtime/Engine/Private/Materials/MaterialHLSLEmitter.cpp b/Engine/Source/Runtime/Engine/Private/Materials/MaterialHLSLEmitter.cpp
--- a/Engine/Source/Runtime/Engine/Private/Materials/MaterialHLSLEmitter.cpp
+++ b/Engine/Source/Runtime/Engine/Private/Materials/MaterialHLSLEmitter.cpp
@@ -14,6 +14,7 @@
 #include "MaterialHLSLGenerator.h"
 #include "ShaderCore.h"
 #include "HLSLTree/HLSLTree.h"
+#include "HLSLTree/HLSLTreeEmit.h"
 #include "Containers/LazyPrintf.h"
 
 static bool SharedPixelProperties[CompiledMP_MAX];
diff --git a/Engine/Source/Runtime/Engine/Private/Materials/MaterialHLSLGenerator.cpp b/Engine/Source/Runtime/Engine/Private/Materials/MaterialHLSLGenerator.cpp
--- a/Engine/Source/Runtime/Engine/Private/Materials/MaterialHLSLGenerator.cpp
+++ b/Engine/Source/Runtime/Engine/Private/Materials/MaterialHLSLGenerator.cpp
@@ -139,10 +139,17 @@ bool FMaterialHLSLGenerator::Generate()
 	FScope& RootScope = HLSLTree->GetRootScope();
 
 	bool bResult = false;
-	if (TargetMaterial->IsCompiledWithExecutionFlow())
+	if (TargetMaterial->IsUsingControlFlow())
 	{
 		UMaterialExpression* BaseExpression = TargetMaterial->ExpressionExecBegin;
-		bResult = GenerateStatements(RootScope, BaseExpression);
+		if (!BaseExpression)
+		{
+			bResult = Errors.AddError(TEXT("Missing ExpressionExecBegin"));
+		}
+		else
+		{
+			bResult = GenerateStatements(RootScope, BaseExpression);
+		}
 	}
 	else
 	{
@@ -216,7 +223,7 @@ static UE::HLSLTree::FExpression* CompileMaterialInput(FMaterialHLSLGenerator& G
 					const int32 TexCoordIndex = (int32)InputProperty - MP_CustomizedUVs0;
 					if (TexCoordIndex < Material->NumCustomizedUVs)
 					{
-						Expression = InputDescription.Input->AcquireHLSLExpression(Generator, Scope);
+						Expression = InputDescription.Input->TryAcquireHLSLExpression(Generator, Scope);
 					}
 					if (!Expression)
 					{
@@ -225,7 +232,7 @@ static UE::HLSLTree::FExpression* CompileMaterialInput(FMaterialHLSLGenerator& G
 				}
 				else
 				{
-					Expression = InputDescription.Input->AcquireHLSLExpression(Generator, Scope);
+					Expression = InputDescription.Input->TryAcquireHLSLExpression(Generator, Scope);
 				}
 			}
 		}
@@ -233,7 +240,7 @@ static UE::HLSLTree::FExpression* CompileMaterialInput(FMaterialHLSLGenerator& G
 
 	if (Expression)
 	{
-		FExpressionSetStructField* SetAttributeExpression = Generator.GetTree().NewExpression<FExpressionSetStructField>(
+		FExpression* SetAttributeExpression = Generator.GetTree().NewExpression<FExpressionSetStructField>(
 			Generator.GetMaterialAttributesType(),
 			AttributeField,
 			AttributesExpression,
@@ -248,8 +255,22 @@ bool FMaterialHLSLGenerator::GenerateResult(UE::HLSLTree::FScope& Scope)
 {
 	using namespace UE::HLSLTree;
 
+	FFunctionCallEntry* FunctionEntry = FunctionCallStack.Last();
+
 	bool bResult = false;
-	if (bGeneratedResult)
+	if (FunctionEntry->MaterialFunction)
+	{
+		// Result for function call
+		FFunction* HLSLFunction = FunctionEntry->HLSLFunction;
+		HLSLFunction->OutputExpressions.Reserve(FunctionEntry->FunctionOutputs.Num());
+		for (const UMaterialExpressionFunctionOutput* ExpressionOutput : FunctionEntry->FunctionOutputs)
+		{
+			HLSLFunction->OutputExpressions.Add(ExpressionOutput->A.TryAcquireHLSLExpression(*this, Scope));
+		}
+		FunctionEntry->bGeneratedResult = true;
+		bResult = true;
+	}
+	else if (bGeneratedResult)
 	{
 		return Errors.AddError(TEXT("Multiple connections to execution output"));
 	}
@@ -343,7 +364,7 @@ UE::HLSLTree::FTextureParameterDeclaration* FMaterialHLSLGenerator::AcquireTextu
 	FString SamplerTypeError;
 	if (!UMaterialExpressionTextureBase::VerifySamplerType(CompileTarget.FeatureLevel, CompileTarget.TargetPlatform, Value.Texture, Value.SamplerType, SamplerTypeError))
 	{
-		Errors.AddErrorf(TEXT("%s"), *SamplerTypeError);
+		Errors.AddError(SamplerTypeError);
 		return nullptr;
 	}
 
@@ -360,7 +381,7 @@ UE::HLSLTree::FTextureParameterDeclaration* FMaterialHLSLGenerator::AcquireTextu
 	FString SamplerTypeError;
 	if (!UMaterialExpressionTextureBase::VerifySamplerType(CompileTarget.FeatureLevel, CompileTarget.TargetPlatform, DefaultValue.Texture, DefaultValue.SamplerType, SamplerTypeError))
 	{
-		Errors.AddErrorf(TEXT("%s"), *SamplerTypeError);
+		Errors.AddError(SamplerTypeError);
 		return nullptr;
 	}
 
@@ -388,23 +409,63 @@ UE::HLSLTree::FExpression* FMaterialHLSLGenerator::AcquireExpression(UE::HLSLTre
 
 UE::HLSLTree::FExpression* FMaterialHLSLGenerator::AcquireFunctionInputExpression(UE::HLSLTree::FScope& Scope, const UMaterialExpressionFunctionInput* MaterialExpression)
 {
+	using namespace UE::HLSLTree;
 	const FFunctionCallEntry* FunctionEntry = FunctionCallStack.Last();
-	const FFunctionInput* ConnectedInput = nullptr;
-	for (const FFunctionInput& Input : FunctionEntry->Inputs)
+	FExpression* InputExpression = nullptr;
+	if (FunctionEntry->MaterialFunction)
 	{
-		if (Input.FunctionInputExpression == MaterialExpression)
+		bool bFoundInput = false;
+		for (int32 Index = 0; Index < FunctionEntry->FunctionInputs.Num(); ++Index)
 		{
-			ConnectedInput = &Input;
-			break;
+			if (FunctionEntry->FunctionInputs[Index] == MaterialExpression)
+			{
+				bFoundInput = true;
+				InputExpression = FunctionEntry->ConnectedInputs[Index];
+				break;
+			}
+		}
+
+		if (!bFoundInput)
+		{
+			// Finding a connected input is always expected if we're in a function call
+			Errors.AddError(TEXT("Invalid function input"));
+			return nullptr;
 		}
 	}
 
-	if (!ConnectedInput)
+	if (!InputExpression && (MaterialExpression->bUsePreviewValueAsDefault || !FunctionEntry->MaterialFunction))
 	{
-		return nullptr;
+		// Either we're previewing the material function, or the input isn't connected and we're using preview as default value
+		InputExpression = MaterialExpression->Preview.TryAcquireHLSLExpression(*this, Scope);
+		if (!InputExpression)
+		{
+			const FVector4& PreviewValue = MaterialExpression->PreviewValue;
+			UE::Shader::FValue DefaultValue;
+			switch (MaterialExpression->InputType)
+			{
+			case FunctionInput_Scalar: DefaultValue = PreviewValue.X; break;
+			case FunctionInput_Vector2: DefaultValue = FVector2f(PreviewValue.X, PreviewValue.Y); break;
+			case FunctionInput_Vector3: DefaultValue = FVector3f(PreviewValue.X, PreviewValue.Y, PreviewValue.Z); break;
+			case FunctionInput_Vector4: DefaultValue = PreviewValue; break;
+			case FunctionInput_MaterialAttributes: DefaultValue = GetMaterialAttributesDefaultValue(); break;
+			case FunctionInput_Texture2D:
+			case FunctionInput_TextureCube:
+			case FunctionInput_Texture2DArray:
+			case FunctionInput_VolumeTexture:
+			case FunctionInput_StaticBool:
+			case FunctionInput_TextureExternal:
+				Errors.AddErrorf(TEXT("Missing Preview connection for function input '%s'"), *MaterialExpression->InputName.ToString());
+				return nullptr;
+			default:
+				Errors.AddError(TEXT("Unknown input type"));
+				return nullptr;
+			}
+
+			InputExpression = NewConstant(DefaultValue);
+		}
 	}
 
-	return AcquireExpression(Scope, ConnectedInput->ConnectedExpression, ConnectedInput->ConnectedOutputIndex);
+	return InputExpression;
 }
 
 UE::HLSLTree::FTextureParameterDeclaration* FMaterialHLSLGenerator::AcquireTextureDeclaration(UE::HLSLTree::FScope& Scope, UMaterialExpression* MaterialExpression, int32 OutputIndex)
@@ -461,11 +522,11 @@ bool FMaterialHLSLGenerator::GenerateStatements(UE::HLSLTree::FScope& Scope, UMa
 	return bResult;
 }
 
-UE::HLSLTree::FExpression* FMaterialHLSLGenerator::GenerateFunctionCall(UE::HLSLTree::FScope& Scope, UMaterialFunctionInterface* Function, TArrayView<const FFunctionExpressionInput> ConnectedInputs, int32 OutputIndex)
+UE::HLSLTree::FExpression* FMaterialHLSLGenerator::GenerateFunctionCall(UE::HLSLTree::FScope& Scope, UMaterialFunctionInterface* MaterialFunction, TArrayView<const FFunctionExpressionInput> ConnectedInputs, int32 OutputIndex)
 {
 	using namespace UE::HLSLTree;
 
-	if (!Function)
+	if (!MaterialFunction)
 	{
 		Errors.AddError(TEXT("Missing material function"));
 		return nullptr;
@@ -473,7 +534,7 @@ UE::HLSLTree::FExpression* FMaterialHLSLGenerator::GenerateFunctionCall(UE::HLSL
 
 	TArray<FFunctionExpressionInput> FunctionInputs;
 	TArray<FFunctionExpressionOutput> FunctionOutputs;
-	Function->GetInputsAndOutputs(FunctionInputs, FunctionOutputs);
+	MaterialFunction->GetInputsAndOutputs(FunctionInputs, FunctionOutputs);
 
 	if (FunctionInputs.Num() != ConnectedInputs.Num())
 	{
@@ -488,41 +549,76 @@ UE::HLSLTree::FExpression* FMaterialHLSLGenerator::GenerateFunctionCall(UE::HLSL
 		return nullptr;
 	}
 
+	const bool bInlineFunction = !MaterialFunction->IsUsingControlFlow();
+
 	FSHA1 Hasher;
-	Hasher.Update((uint8*)&Function, sizeof(UMaterialFunctionInterface*));
+	Hasher.Update((uint8*)&MaterialFunction, sizeof(UMaterialFunctionInterface*));
 
-	FFunctionInputArray Inputs;
+	FFunctionInputArray LocalFunctionInputs;
+	FConnectedInputArray LocalConnectedInputs;
 	for (int32 InputIndex = 0; InputIndex < ConnectedInputs.Num(); ++InputIndex)
 	{
 		// FunctionInputs are the inputs from the UMaterialFunction object
-		// ConnectedInputs are the inputs from the UMaterialFunctionCall object
-		// We want to connect the UMaterialExpressionFunctionInput from the UMaterialFunction to whatever UMaterialExpression is passed to the UMaterialFunctionCall
 		const FFunctionExpressionInput& FunctionInput = FunctionInputs[InputIndex];
-		const FFunctionExpressionInput& ConnectedInput = ConnectedInputs[InputIndex];
 
-		FFunctionInput& Input = Inputs.AddDefaulted_GetRef();
-		Input.FunctionInputExpression = FunctionInput.ExpressionInput;
-		Input.ConnectedExpression = ConnectedInput.Input.Expression;
-		Input.ConnectedOutputIndex = ConnectedInput.Input.OutputIndex;
+		// ConnectedInputs are the inputs from the UMaterialFunctionCall object
+		// We want to connect the UMaterialExpressionFunctionInput from the UMaterialFunction to whatever UMaterialExpression is passed to the UMaterialFunctionCall
+		FExpression* ConnectedInput = ConnectedInputs[InputIndex].Input.TryAcquireHLSLExpression(*this, Scope);
 
-		Hasher.Update((uint8*)&Input.ConnectedExpression, sizeof(Input.ConnectedExpression));
-		Hasher.Update((uint8*)&Input.ConnectedOutputIndex, sizeof(Input.ConnectedOutputIndex));
+		LocalFunctionInputs.Add(FunctionInput.ExpressionInput);
+		LocalConnectedInputs.Add(ConnectedInput);
+		Hasher.Update((uint8*)&ConnectedInput, sizeof(ConnectedInput));
 	}
-
 	const FSHAHash Hash = Hasher.Finalize();
 
 	FFunctionCallEntry** ExistingFunctionCall = FunctionCallMap.Find(Hash);
 	FFunctionCallEntry* FunctionCall = ExistingFunctionCall ? *ExistingFunctionCall : nullptr;
 	if (!FunctionCall)
 	{
+		// Generate an HLSL function object, if this is not an inline function call
+		FFunction* HLSLFunction = !bInlineFunction ? HLSLTree->NewFunction() : nullptr;
 		FunctionCall = new(HLSLTree->GetAllocator()) FFunctionCallEntry();
-		FunctionCall->Function = Function;
-		FunctionCall->Inputs = MoveTemp(Inputs);
+		FunctionCall->MaterialFunction = MaterialFunction;
+		FunctionCall->HLSLFunction = HLSLFunction;
+		FunctionCall->FunctionInputs = MoveTemp(LocalFunctionInputs);
+		FunctionCall->ConnectedInputs = MoveTemp(LocalConnectedInputs);
+		FunctionCall->FunctionOutputs.Reserve(FunctionOutputs.Num());
+		for (const FFunctionExpressionOutput& Output : FunctionOutputs)
+		{
+			FunctionCall->FunctionOutputs.Add(Output.ExpressionOutput);
+		}
+
 		FunctionCallMap.Add(Hash, FunctionCall);
+
+		if (HLSLFunction)
+		{
+			UMaterialFunction* BaseMaterialFunction = Cast<UMaterialFunction>(MaterialFunction->GetBaseFunction());
+			FunctionCallStack.Add(FunctionCall);
+			GenerateStatements(HLSLFunction->GetRootScope(), BaseMaterialFunction->ExpressionExecBegin);
+			verify(FunctionCallStack.Pop() == FunctionCall);
+			check(FunctionCall->bGeneratedResult);
+		}
 	}
 
+	FExpression* Result = nullptr;
 	FunctionCallStack.Add(FunctionCall);
-	UE::HLSLTree::FExpression* Result = ExpressionOutput->A.AcquireHLSLExpression(*this, Scope);
+	if (bInlineFunction)
+	{
+		Result = ExpressionOutput->A.AcquireHLSLExpression(*this, Scope);
+	}
+	else
+	{
+		FFunction* HLSLFunction = FunctionCall->HLSLFunction;
+		check(HLSLFunction->OutputExpressions.Num() == FunctionOutputs.Num());
+		if (HLSLFunction->OutputExpressions[OutputIndex])
+		{
+			Result = HLSLTree->NewFunctionCall(Scope, HLSLFunction, FunctionCall->ConnectedInputs, OutputIndex);
+		}
+		else
+		{
+			Errors.AddError(TEXT("Invalid function output"));
+		}
+	}
 	verify(FunctionCallStack.Pop() == FunctionCall);
 
 	return Result;
diff --git a/Engine/Source/Runtime/Engine/Private/Materials/MaterialShared.cpp b/Engine/Source/Runtime/Engine/Private/Materials/MaterialShared.cpp
--- a/Engine/Source/Runtime/Engine/Private/Materials/MaterialShared.cpp
+++ b/Engine/Source/Runtime/Engine/Private/Materials/MaterialShared.cpp
@@ -1341,7 +1341,7 @@ bool FMaterialResource::IsVolumetricPrimitive() const { return Material->Materia
 bool FMaterialResource::IsSpecialEngineMaterial() const { return Material->bUsedAsSpecialEngineMaterial; }
 bool FMaterialResource::HasVertexPositionOffsetConnected() const { return Material->HasVertexPositionOffsetConnected(); }
 bool FMaterialResource::HasPixelDepthOffsetConnected() const { return Material->HasPixelDepthOffsetConnected(); }
-bool FMaterialResource::HasMaterialAttributesConnected() const { return (Material->bUseMaterialAttributes && Material->MaterialAttributes.IsConnected()) || Material->bEnableExecWire; }
+bool FMaterialResource::HasMaterialAttributesConnected() const { return (Material->bUseMaterialAttributes && Material->MaterialAttributes.IsConnected()); }
 EMaterialShadingRate FMaterialResource::GetShadingRate() const { return Material->ShadingRate; }
 FString FMaterialResource::GetBaseMaterialPathName() const { return Material->GetPathName(); }
 FString FMaterialResource::GetDebugName() const
@@ -1839,11 +1839,11 @@ bool FMaterialResource::ShouldInlineShaderCode() const
 	return bNeedsToBeInlined;
 }
 
-bool FMaterialResource::IsCompiledWithExecutionFlow() const
+bool FMaterialResource::IsUsingControlFlow() const
 {
 	if (Material)
 	{
-		return Material->IsCompiledWithExecutionFlow();
+		return Material->IsUsingControlFlow();
 	}
 	return false;
 }
diff --git a/Engine/Source/Runtime/Engine/Private/Shader/Preshader.cpp b/Engine/Source/Runtime/Engine/Private/Shader/Preshader.cpp
--- a/Engine/Source/Runtime/Engine/Private/Shader/Preshader.cpp
+++ b/Engine/Source/Runtime/Engine/Private/Shader/Preshader.cpp
@@ -5,6 +5,7 @@
 #include "Materials/MaterialUniformExpressions.h"
 #include "Engine/Texture.h"
 #include "VT/RuntimeVirtualTexture.h"
+#include "Hash/xxhash.h"
 #include "ExternalTexture.h"
 
 IMPLEMENT_TYPE_LAYOUT(UE::Shader::FPreshaderData);
@@ -276,9 +277,21 @@ static void EvaluateSetField(FPreshaderStack& Stack, FPreshaderDataContext& REST
 	const int32 ComponentNum = ReadPreshaderValue<int32>(Data);
 	
 	// Modify the struct value in-place
-	for (int32 Index = 0; Index < ComponentNum; ++Index)
+	if (Value.Component.Num() == 1)
 	{
-		StructValue.Component[ComponentIndex + Index] = Value.Component[Index];
+		// Splat scalar
+		const FValueComponent Component = Value.Component[0];
+		for (int32 Index = 0; Index < ComponentNum; ++Index)
+		{
+			StructValue.Component[ComponentIndex + Index] = Component;
+		}
+	}
+	else
+	{
+		for (int32 Index = 0; Index < ComponentNum; ++Index)
+		{
+			StructValue.Component[ComponentIndex + Index] = Value.Component[Index];
+		}
 	}
 }
 
@@ -575,19 +588,12 @@ FPreshaderValue FPreshaderData::EvaluateConstant(const FMaterial& Material, FPre
 	return EvaluatePreshader(nullptr, FMaterialRenderContext(nullptr, Material, nullptr), Stack, PreshaderContext);
 }
 
-void FPreshaderData::AppendHash(FSHA1& OutHasher) const
-{
-	OutHasher.Update((uint8*)Names.GetData(), Names.Num() * Names.GetTypeSize());
-	OutHasher.Update((uint8*)StructTypes.GetData(), StructTypes.Num() * StructTypes.GetTypeSize());
-	OutHasher.Update((uint8*)StructComponentTypes.GetData(), StructComponentTypes.Num() * StructComponentTypes.GetTypeSize());
-	OutHasher.Update((uint8*)Data.GetData(), Data.Num() * Data.GetTypeSize());
-}
-
-FSHAHash FPreshaderData::GetHash() const
+void FPreshaderData::AppendHash(FXxHash64Builder& OutHasher) const
 {
-	FSHA1 Hasher;
-	AppendHash(Hasher);
-	return Hasher.Finalize();
+	OutHasher.Update(Names.GetData(), Names.Num() * Names.GetTypeSize());
+	OutHasher.Update(StructTypes.GetData(), StructTypes.Num() * StructTypes.GetTypeSize());
+	OutHasher.Update(StructComponentTypes.GetData(), StructComponentTypes.Num() * StructComponentTypes.GetTypeSize());
+	OutHasher.Update(Data.GetData(), Data.Num() * Data.GetTypeSize());
 }
 
 } // namespace Shader
diff --git a/Engine/Source/Runtime/Engine/Public/HLSLTree/HLSLTree.h b/Engine/Source/Runtime/Engine/Public/HLSLTree/HLSLTree.h
--- a/Engine/Source/Runtime/Engine/Public/HLSLTree/HLSLTree.h
+++ b/Engine/Source/Runtime/Engine/Public/HLSLTree/HLSLTree.h
@@ -9,6 +9,7 @@
 #include "Containers/List.h"
 #include "Misc/StringBuilder.h"
 #include "Misc/MemStack.h"
+#include "Hash/xxhash.h"
 #include "HLSLTree/HLSLTreeTypes.h"
 
 class FMaterial;
@@ -33,9 +34,13 @@ namespace HLSLTree
 
 class FNode;
 class FScope;
+class FFunction;
 class FExpressionLocalPHI;
 class FRequestedType;
-struct FEmitShaderCode;
+
+class FEmitContext;
+class FEmitShaderScope;
+class FEmitShaderExpression;
 
 static constexpr int32 MaxNumPreviousScopes = 2;
 
@@ -104,123 +109,59 @@ private:
 	int32 NumErrors = 0;
 };
 
-enum class ECastFlags : uint32
-{
-	None = 0u,
-	ReplicateScalar = (1u << 0),
-	AllowTruncate = (1u << 1),
-	AllowAppendZeroes = (1u << 2),
-
-	ValidCast = ReplicateScalar | AllowTruncate,
-};
-ENUM_CLASS_FLAGS(ECastFlags);
-
-struct FEmitShaderCode
-{
-	FEmitShaderCode(FScope* InScope, const Shader::FType& InType) : Scope(InScope), Type(InType) {}
-
-	inline bool IsInline() const { return Value == nullptr; }
-
-	FScope* Scope = nullptr;
-	const TCHAR* Reference = nullptr;
-	const TCHAR* Value = nullptr;
-	Shader::FType Type;
-	TArrayView<FEmitShaderCode*> Dependencies;
-	FSHAHash Hash;
-};
-
-using FEmitShaderValueDependencies = TArray<FEmitShaderCode*, TInlineAllocator<8>>;
-
-enum class EFormatArgType : uint8
-{
-	Void,
-	ShaderValue,
-	String,
-	Int,
-};
-
-struct FFormatArgVariant
+class FHasher
 {
-	FFormatArgVariant() {}
-	FFormatArgVariant(FEmitShaderCode* InValue) : Type(EFormatArgType::ShaderValue), ShaderValue(InValue) { check(InValue); }
-	FFormatArgVariant(const TCHAR* InValue) : Type(EFormatArgType::String), String(InValue) { check(InValue); }
-	FFormatArgVariant(int32 InValue) : Type(EFormatArgType::Int), Int(InValue) {}
+public:
+	FXxHash64 Finalize() { return Builder.Finalize(); }
+	FHasher& AppendData(const void* Data, uint64 Size) { Builder.Update(Data, Size); return *this; }
 
-	EFormatArgType Type = EFormatArgType::Void;
-	union
+	template<typename T>
+	inline FHasher& AppendValue(const T& Value)
 	{
-		FEmitShaderCode* ShaderValue;
-		const TCHAR* String;
-		int32 Int;
-	};
-};
-
-using FFormatArgList = TArray<FFormatArgVariant, TInlineAllocator<8>>;
-
-inline void BuildFormatArgList(FFormatArgList&) {}
-
-template<typename Type, typename... Types>
-inline void BuildFormatArgList(FFormatArgList& OutList, Type Arg, Types... Args)
-{
-	OutList.Add(Arg);
-	BuildFormatArgList(OutList, Forward<Types>(Args)...);
-}
-
-/** Tracks shared state while emitting HLSL code */
-class FEmitContext
-{
-public:
-	explicit FEmitContext(FMemStackBase& InAllocator, FErrorHandlerInterface& InErrors, const Shader::FStructTypeRegistry& InTypeRegistry);
-	~FEmitContext();
+		return AppendData(&Value, sizeof(Value));
+	}
 
-	void Finalize();
+	template<typename T>
+	inline FHasher& AppendValue(TArrayView<T> Value)
+	{
+		for (const T& Element : Value)
+		{
+			AppendValue(Element);
+		}
+		return *this;
+	}
 
-	/** Get a unique local variable name */
-	const TCHAR* AcquireLocalDeclarationCode();
+	inline FHasher& AppendValue(const FName& Value)
+	{
+		return AppendValue(Value.GetComparisonIndex()).AppendValue(Value.GetNumber());
+	}
 
-	FEmitShaderCode* EmitCodeInternal(const Shader::FType& Type, FStringView Code, bool bInline, TArrayView<FEmitShaderCode*> Dependencies);
-	FEmitShaderCode* EmitFormatCodeInternal(const Shader::FType& Type, const TCHAR* Format, bool bInline, const FFormatArgList& ArgList);
+	inline FHasher& AppendValue(const Shader::FType& Type)
+	{
+		if (Type.IsStruct()) return AppendValue(Type.StructType);
+		else return AppendValue(Type.ValueType);
+	}
 
-	template<typename FormatType, typename... Types>
-	FEmitShaderCode* EmitCode(const Shader::FType& Type, const FormatType& Format, Types... Args)
+	inline FHasher& AppendValue(const Shader::FValue& Value)
 	{
-		FFormatArgList ArgList;
-		BuildFormatArgList(ArgList, Forward<Types>(Args)...);
-		return EmitFormatCodeInternal(Type, Format, false, ArgList);
+		AppendValue(Value.Type);
+		for (int32 i = 0; i < Value.Type.GetNumComponents(); ++i)
+		{
+			AppendValue(Value.TryGetComponent(i));
+		}
+		return *this;
 	}
 
-	template<typename FormatType, typename... Types>
-	FEmitShaderCode* EmitInlineCode(const Shader::FType& Type, const FormatType& Format, Types... Args)
+	inline FHasher& AppendValues() { return *this; }
+
+	template<typename T, typename... ArgTypes>
+	inline FHasher& AppendValues(const T& Value, ArgTypes&&... Args)
 	{
-		FFormatArgList ArgList;
-		BuildFormatArgList(ArgList, Forward<Types>(Args)...);
-		return EmitFormatCodeInternal(Type, Format, true, ArgList);
+		return AppendValue(Value).AppendValues(Forward<ArgTypes>(Args)...);
 	}
 
-	
-	FEmitShaderCode* EmitPreshaderOrConstant(const FRequestedType& RequestedType, FExpression* Expression);
-	FEmitShaderCode* EmitConstantZero(const Shader::FType& Type);
-	FEmitShaderCode* EmitCast(FEmitShaderCode* ShaderValue, const Shader::FType& DestType);
-	
-	FMemStackBase* Allocator = nullptr;
-	FErrorHandlerInterface* Errors = nullptr;
-	const Shader::FStructTypeRegistry* TypeRegistry = nullptr;
-
-	TArray<FScope*, TInlineAllocator<16>> ScopeStack;
-	TMap<FSHAHash, FEmitShaderCode*> ShaderValueMap;
-	TMap<FSHAHash, FEmitShaderCode*> PreshaderValueMap;
-	TArray<const FExpressionLocalPHI*> LocalPHIs;
-
-	// TODO - remove preshader material dependency
-	const FMaterial* Material = nullptr;
-	const FStaticParameterSet* StaticParameters = nullptr;
-	FMaterialCompilationOutput* MaterialCompilationOutput = nullptr;
-	TMap<Shader::FValue, uint32> DefaultUniformValues;
-	uint32 UniformPreshaderOffset = 0u;
-	bool bReadMaterialNormal = false;
-
-	int32 NumExpressionLocals = 0;
-	int32 NumTexCoords = 0;
+private:
+	FXxHash64Builder Builder;
 };
 
 /** Root class of the HLSL AST */
@@ -250,12 +191,10 @@ private:
 class FStatement : public FNode
 {
 public:
-	virtual void Reset() override;
 	virtual void Prepare(FEmitContext& Context) const = 0;
-	virtual void EmitShader(FEmitContext& Context) const = 0;
+	virtual void EmitShader(FEmitContext& Context, FEmitShaderScope& Scope) const = 0;
 
 	FScope* ParentScope = nullptr;
-	bool bEmitShader = false;
 };
 
 /**
@@ -378,7 +317,7 @@ private:
 
 struct FEmitValueShaderResult
 {
-	FEmitShaderCode* Code = nullptr;
+	FEmitShaderExpression* Code = nullptr;
 };
 
 enum class EDerivativeCoordinate : uint8
@@ -415,17 +354,17 @@ public:
 
 	friend const FPreparedType& PrepareExpressionValue(FEmitContext& Context, FExpression* InExpression, const FRequestedType& RequestedType);
 
-	FEmitShaderCode* GetValueShader(FEmitContext& Context, const FRequestedType& RequestedType, const Shader::FType& ResultType);
+	FEmitShaderExpression* GetValueShader(FEmitContext& Context, FEmitShaderScope& Scope, const FRequestedType& RequestedType, const Shader::FType& ResultType);
 	void GetValuePreshader(FEmitContext& Context, const FRequestedType& RequestedType, Shader::FPreshaderData& OutPreshader);
 	Shader::FValue GetValueConstant(FEmitContext& Context, const FRequestedType& RequestedType);
 
-	FEmitShaderCode* GetValueShader(FEmitContext& Context, const FRequestedType& RequestedType);
-	FEmitShaderCode* GetValueShader(FEmitContext& Context);
+	FEmitShaderExpression* GetValueShader(FEmitContext& Context, FEmitShaderScope& Scope, const FRequestedType& RequestedType);
+	FEmitShaderExpression* GetValueShader(FEmitContext& Context, FEmitShaderScope& Scope);
 
 protected:
 	virtual void ComputeAnalyticDerivatives(FTree& Tree, FExpressionDerivatives& OutResult) const;
 	virtual bool PrepareValue(FEmitContext& Context, const FRequestedType& RequestedType, FPrepareValueResult& OutResult) const = 0;
-	virtual void EmitValueShader(FEmitContext& Context, const FRequestedType& RequestedType, FEmitValueShaderResult& OutResult) const;
+	virtual void EmitValueShader(FEmitContext& Context, FEmitShaderScope& Scope, const FRequestedType& RequestedType, FEmitValueShaderResult& OutResult) const;
 	virtual void EmitValuePreshader(FEmitContext& Context, const FRequestedType& RequestedType, Shader::FPreshaderData& OutPreshader) const;
 
 private:
@@ -464,27 +403,14 @@ public:
 	FExpression* Expression;
 };
 
-
-/**
- * Represents a phi node (see various topics on single static assignment)
- * A phi node takes on a value based on the previous scope that was executed.
- * In practice, this means the generated HLSL code will declare a local variable before all the previous scopes, then assign that variable the proper value from within each scope
- */
-class FExpressionLocalPHI final : public FExpression
+class FFunction final : public FNode
 {
 public:
-	FExpressionLocalPHI() = default;
-	FExpressionLocalPHI(const FExpressionLocalPHI* Source, EDerivativeCoordinate Coord);
-
-	virtual void ComputeAnalyticDerivatives(FTree& Tree, FExpressionDerivatives& OutResult) const override;
-	virtual bool PrepareValue(FEmitContext& Context, const FRequestedType& RequestedType, FPrepareValueResult& OutResult) const override;
-	virtual void EmitValueShader(FEmitContext& Context, const FRequestedType& RequestedType, FEmitValueShaderResult& OutResult) const override;
-
-	TArray<EDerivativeCoordinate, TInlineAllocator<8>> DerivativeChain;
-	FName LocalName;
-	FScope* Scopes[MaxNumPreviousScopes];
-	FExpression* Values[MaxNumPreviousScopes];
-	int32 NumValues = 0;
+	FScope& GetRootScope() const { return *RootScope; }
+
+	FScope* RootScope = nullptr;
+	FScope* CalledScope = nullptr;
+	TArray<FExpression*, TInlineAllocator<8>> OutputExpressions;
 };
 
 /**
@@ -512,10 +438,10 @@ enum class EScopeState : uint8
 class FScope final : public FNode
 {
 public:
-	virtual void Reset() override;
-
 	static FScope* FindSharedParent(FScope* Lhs, FScope* Rhs);
 
+	virtual void Reset() override;
+
 	inline FScope* GetParentScope() const { return ParentScope; }
 	inline bool IsLive() const { return State == EScopeState::Live; }
 	inline bool IsDead() const { return State == EScopeState::Dead; }
@@ -532,79 +458,20 @@ public:
 
 	friend bool PrepareScope(FEmitContext& Context, FScope* InScope);
 
-	template<typename FormatType, typename... Types>
-	void EmitDeclarationf(FEmitContext& Context, const FormatType& Format, Types... Args)
-	{
-		InternalEmitCodef(Context, Declarations, ENextScopeFormat::None, nullptr, Format, Forward<Types>(Args)...);
-	}
-
-	template<typename FormatType, typename... Types>
-	void EmitStatementf(FEmitContext& Context, const FormatType& Format, Types... Args)
-	{
-		InternalEmitCodef(Context, Statements, ENextScopeFormat::None, nullptr, Format, Forward<Types>(Args)...);
-	}
-
-	void EmitScope(FEmitContext& Context, FScope* NestedScope)
-	{
-		InternalEmitCode(Context, Statements, ENextScopeFormat::Unscoped, NestedScope, nullptr, 0);
-	}
-
-	template<typename FormatType, typename... Types>
-	void EmitNestedScopef(FEmitContext& Context, FScope* NestedScope, const FormatType& Format, Types... Args)
-	{
-		InternalEmitCodef(Context, Statements, ENextScopeFormat::Scoped, NestedScope, Format, Forward<Types>(Args)...);
-	}
-
 	void MarkLive();
 	void MarkLiveRecursive();
 	void MarkDead();
 
-	void WriteHLSL(int32 Indent, FStringBuilderBase& OutString) const;
-
 private:
 	friend class FTree;
 	friend class FExpression;
-
-	enum class ENextScopeFormat : uint8
-	{
-		None,
-		Unscoped,
-		Scoped,
-	};
-
-	struct FCodeEntry
-	{
-		FCodeEntry* Next;
-		FScope* Scope;
-		int32 Length;
-		ENextScopeFormat ScopeFormat;
-		TCHAR String[1];
-	};
-
-	struct FCodeList
-	{
-		FCodeEntry* First = nullptr;
-		FCodeEntry* Last = nullptr;
-		int32 Num = 0;
-	};
-
-	void InternalEmitCode(FEmitContext& Context, FCodeList& List, ENextScopeFormat ScopeFormat, FScope* Scope, const TCHAR* String, int32 Length);
-
-	template<typename FormatType, typename... Types>
-	void InternalEmitCodef(FEmitContext& Context, FCodeList& List, ENextScopeFormat ScopeFormat, FScope* Scope, const FormatType& Format, Types... Args)
-	{
-		TStringBuilder<2048> String;
-		String.Appendf(Format, Forward<Types>(Args)...);
-		InternalEmitCode(Context, List, ScopeFormat, Scope, String.ToString(), String.Len());
-	}
+	friend class FEmitContext;
 
 	FStatement* OwnerStatement = nullptr;
 	FScope* ParentScope = nullptr;
 	FStatement* ContainedStatement = nullptr;
 	FScope* PreviousScope[MaxNumPreviousScopes];
 	TMap<FName, FExpression*> LocalMap;
-	FCodeList Declarations;
-	FCodeList Statements;
 	int32 NumPreviousScopes = 0;
 	int32 NestedLevel = 0;
 	EScopeState State = EScopeState::Uninitialized;
@@ -631,6 +498,20 @@ inline void MarkScopeDead(FScope* InScope)
 	}
 }
 
+namespace Private
+{
+uint64 GetNextTypeHash();
+template<typename T>
+struct TTypeHasher
+{
+	static uint64 GetHash()
+	{
+		static uint64 Hash = GetNextTypeHash();
+		return Hash;
+	}
+};
+}
+
 /**
  * The HLSL AST.  Basically a wrapper around the root scope, with some helper methods
  */
@@ -651,10 +532,19 @@ public:
 	FScope& GetRootScope() const { return *RootScope; }
 
 	template<typename T, typename... ArgTypes>
-	inline T* NewExpression(ArgTypes&&... Args)
+	inline FExpression* NewExpression(ArgTypes&&... Args)
 	{
-		T* Expression = NewNode<T>(Forward<ArgTypes>(Args)...);
-		RegisterExpression(Expression);
+		FHasher Hasher;
+		Hasher.AppendValue(Private::TTypeHasher<T>::GetHash());
+		Hasher.AppendValues(Forward<ArgTypes>(Args)...);
+		const FXxHash64 Hash = Hasher.Finalize();
+		FExpression* Expression = FindExpression(Hash);
+		if (!Expression)
+		{
+			T* TypedExpression = NewNode<T>(Forward<ArgTypes>(Args)...);
+			RegisterExpression(TypedExpression, Hash);
+			Expression = TypedExpression;
+		}
 		return Expression;
 	}
 
@@ -669,10 +559,13 @@ public:
 	void AssignLocal(FScope& Scope, const FName& LocalName, FExpression* Value);
 	FExpression* AcquireLocal(FScope& Scope, const FName& LocalName);
 
+	FExpression* NewFunctionCall(FScope& Scope, FFunction* Function, TArrayView<FExpression*> InputExpressions, int32 OutputIndex);
+
 	const FExpressionDerivatives& GetAnalyticDerivatives(FExpression* InExpression);
 
 	FScope* NewScope(FScope& Scope);
 	FScope* NewOwnedScope(FStatement& Owner);
+	FFunction* NewFunction();
 
 	/** Shortcuts to create various common expression types */
 	FExpression* NewConstant(const Shader::FValue& Value);
@@ -700,12 +593,15 @@ private:
 	}
 
 	void RegisterNode(FNode* Node);
-	void RegisterExpression(FExpression* Expression);
+	void RegisterExpression(FExpression* Expression, FXxHash64 Hash);
+	void RegisterExpression(FExpressionLocalPHI* Expression, FXxHash64 Hash);
 	void RegisterStatement(FScope& Scope, FStatement* Statement);
+	FExpression* FindExpression(FXxHash64 Hash) const;
 
 	FMemStackBase* Allocator = nullptr;
 	FNode* Nodes = nullptr;
 	FScope* RootScope = nullptr;
+	TMap<FXxHash64, FExpression*> ExpressionMap;
 	TArray<FExpressionLocalPHI*> PHIExpressions;
 	TArray<UObject*, TInlineAllocator<8>> OwnerStack;
 
diff --git a/Engine/Source/Runtime/Engine/Public/HLSLTree/HLSLTreeCommon.h b/Engine/Source/Runtime/Engine/Public/HLSLTree/HLSLTreeCommon.h
--- a/Engine/Source/Runtime/Engine/Public/HLSLTree/HLSLTreeCommon.h
+++ b/Engine/Source/Runtime/Engine/Public/HLSLTree/HLSLTreeCommon.h
@@ -123,17 +123,18 @@ public:
 
 	virtual void ComputeAnalyticDerivatives(FTree& Tree, FExpressionDerivatives& OutResult) const override;
 	virtual bool PrepareValue(FEmitContext& Context, const FRequestedType& RequestedType, FPrepareValueResult& OutResult) const override;
-	virtual void EmitValueShader(FEmitContext& Context, const FRequestedType& RequestedType, FEmitValueShaderResult& OutResult) const override;
+	virtual void EmitValueShader(FEmitContext& Context, FEmitShaderScope& Scope, const FRequestedType& RequestedType, FEmitValueShaderResult& OutResult) const override;
 };
 
 class FExpressionTextureSample : public FExpression
 {
 public:
-	FExpressionTextureSample(FTextureParameterDeclaration* InDeclaration, FExpression* InTexCoordExpression)
+	FExpressionTextureSample(FTextureParameterDeclaration* InDeclaration, FExpression* InTexCoordExpression, const FExpressionDerivatives& InTexCoordDerivatives, ESamplerSourceMode InSamplerSource, ETextureMipValueMode InMipValueMode)
 		: Declaration(InDeclaration)
 		, TexCoordExpression(InTexCoordExpression)
-		, SamplerSource(SSM_FromTextureAsset)
-		, MipValueMode(TMVM_None)
+		, TexCoordDerivatives(InTexCoordDerivatives)
+		, SamplerSource(InSamplerSource)
+		, MipValueMode(InMipValueMode)
 	{}
 
 	FTextureParameterDeclaration* Declaration;
@@ -143,18 +144,17 @@ public:
 	ETextureMipValueMode MipValueMode;
 
 	virtual bool PrepareValue(FEmitContext& Context, const FRequestedType& RequestedType, FPrepareValueResult& OutResult) const override;
-	virtual void EmitValueShader(FEmitContext& Context, const FRequestedType& RequestedType, FEmitValueShaderResult& OutResult) const override;
+	virtual void EmitValueShader(FEmitContext& Context, FEmitShaderScope& Scope, const FRequestedType& RequestedType, FEmitValueShaderResult& OutResult) const override;
 };
 
 class FExpressionGetStructField : public FExpression
 {
 public:
-	FExpressionGetStructField(const Shader::FStructType* InStructType, const TCHAR* InFieldName, FExpression* InStructExpression)
+	FExpressionGetStructField(const Shader::FStructType* InStructType, const Shader::FStructField* InField, FExpression* InStructExpression)
 		: StructType(InStructType)
-		, Field(InStructType->FindFieldByName(InFieldName))
+		, Field(InField)
 		, StructExpression(InStructExpression)
 	{
-		check(Field);
 	}
 
 	const Shader::FStructType* StructType;
@@ -163,7 +163,7 @@ public:
 
 	virtual void ComputeAnalyticDerivatives(FTree& Tree, FExpressionDerivatives& OutResult) const override;
 	virtual bool PrepareValue(FEmitContext& Context, const FRequestedType& RequestedType, FPrepareValueResult& OutResult) const override;
-	virtual void EmitValueShader(FEmitContext& Context, const FRequestedType& RequestedType, FEmitValueShaderResult& OutResult) const override;
+	virtual void EmitValueShader(FEmitContext& Context, FEmitShaderScope& Scope, const FRequestedType& RequestedType, FEmitValueShaderResult& OutResult) const override;
 	virtual void EmitValuePreshader(FEmitContext& Context, const FRequestedType& RequestedType, Shader::FPreshaderData& OutPreshader) const override;
 };
 
@@ -180,15 +180,6 @@ public:
 		check(InField);
 	}
 
-	FExpressionSetStructField(const Shader::FStructType* InStructType, const TCHAR* InFieldName, FExpression* InStructExpression, FExpression* InFieldExpression)
-		: StructType(InStructType)
-		, Field(InStructType->FindFieldByName(InFieldName))
-		, StructExpression(InStructExpression)
-		, FieldExpression(InFieldExpression)
-	{
-		check(Field);
-	}
-
 	const Shader::FStructType* StructType;
 	const Shader::FStructField* Field;
 	FExpression* StructExpression;
@@ -196,7 +187,7 @@ public:
 
 	virtual void ComputeAnalyticDerivatives(FTree& Tree, FExpressionDerivatives& OutResult) const override;
 	virtual bool PrepareValue(FEmitContext& Context, const FRequestedType& RequestedType, FPrepareValueResult& OutResult) const override;
-	virtual void EmitValueShader(FEmitContext& Context, const FRequestedType& RequestedType, FEmitValueShaderResult& OutResult) const override;
+	virtual void EmitValueShader(FEmitContext& Context, FEmitShaderScope& Scope, const FRequestedType& RequestedType, FEmitValueShaderResult& OutResult) const override;
 	virtual void EmitValuePreshader(FEmitContext& Context, const FRequestedType& RequestedType, Shader::FPreshaderData& OutPreshader) const override;
 };
 
@@ -214,7 +205,7 @@ public:
 	FExpression* FalseExpression;
 
 	virtual bool PrepareValue(FEmitContext& Context, const FRequestedType& RequestedType, FPrepareValueResult& OutResult) const override;
-	virtual void EmitValueShader(FEmitContext& Context, const FRequestedType& RequestedType, FEmitValueShaderResult& OutResult) const override;
+	virtual void EmitValueShader(FEmitContext& Context, FEmitShaderScope& Scope, const FRequestedType& RequestedType, FEmitValueShaderResult& OutResult) const override;
 	virtual void EmitValuePreshader(FEmitContext& Context, const FRequestedType& RequestedType, Shader::FPreshaderData& OutPreshader) const override;
 };
 
@@ -231,7 +222,7 @@ public:
 
 	virtual void ComputeAnalyticDerivatives(FTree& Tree, FExpressionDerivatives& OutResult) const override;
 	virtual bool PrepareValue(FEmitContext& Context, const FRequestedType& RequestedType, FPrepareValueResult& OutResult) const override;
-	virtual void EmitValueShader(FEmitContext& Context, const FRequestedType& RequestedType, FEmitValueShaderResult& OutResult) const override;
+	virtual void EmitValueShader(FEmitContext& Context, FEmitShaderScope& Scope, const FRequestedType& RequestedType, FEmitValueShaderResult& OutResult) const override;
 	virtual void EmitValuePreshader(FEmitContext& Context, const FRequestedType& RequestedType, Shader::FPreshaderData& OutPreshader) const override;
 };
 
@@ -250,7 +241,7 @@ public:
 
 	virtual void ComputeAnalyticDerivatives(FTree& Tree, FExpressionDerivatives& OutResult) const override;
 	virtual bool PrepareValue(FEmitContext& Context, const FRequestedType& RequestedType, FPrepareValueResult& OutResult) const override;
-	virtual void EmitValueShader(FEmitContext& Context, const FRequestedType& RequestedType, FEmitValueShaderResult& OutResult) const override;
+	virtual void EmitValueShader(FEmitContext& Context, FEmitShaderScope& Scope, const FRequestedType& RequestedType, FEmitValueShaderResult& OutResult) const override;
 	virtual void EmitValuePreshader(FEmitContext& Context, const FRequestedType& RequestedType, Shader::FPreshaderData& OutPreshader) const override;
 };
 
@@ -280,7 +271,7 @@ public:
 
 	virtual void ComputeAnalyticDerivatives(FTree& Tree, FExpressionDerivatives& OutResult) const override;
 	virtual bool PrepareValue(FEmitContext& Context, const FRequestedType& RequestedType, FPrepareValueResult& OutResult) const override;
-	virtual void EmitValueShader(FEmitContext& Context, const FRequestedType& RequestedType, FEmitValueShaderResult& OutResult) const override;
+	virtual void EmitValueShader(FEmitContext& Context, FEmitShaderScope& Scope, const FRequestedType& RequestedType, FEmitValueShaderResult& OutResult) const override;
 	virtual void EmitValuePreshader(FEmitContext& Context, const FRequestedType& RequestedType, Shader::FPreshaderData& OutPreshader) const override;
 };
 
@@ -297,7 +288,7 @@ public:
 
 	virtual void ComputeAnalyticDerivatives(FTree& Tree, FExpressionDerivatives& OutResult) const override;
 	virtual bool PrepareValue(FEmitContext& Context, const FRequestedType& RequestedType, FPrepareValueResult& OutResult) const override;
-	virtual void EmitValueShader(FEmitContext& Context, const FRequestedType& RequestedType, FEmitValueShaderResult& OutResult) const override;
+	virtual void EmitValueShader(FEmitContext& Context, FEmitShaderScope& Scope, const FRequestedType& RequestedType, FEmitValueShaderResult& OutResult) const override;
 	virtual void EmitValuePreshader(FEmitContext& Context, const FRequestedType& RequestedType, Shader::FPreshaderData& OutPreshader) const override;
 };
 
@@ -305,7 +296,7 @@ class FExpressionReflectionVector : public FExpression
 {
 public:
 	virtual bool PrepareValue(FEmitContext& Context, const FRequestedType& RequestedType, FPrepareValueResult& OutResult) const override;
-	virtual void EmitValueShader(FEmitContext& Context, const FRequestedType& RequestedType, FEmitValueShaderResult& OutResult) const override;
+	virtual void EmitValueShader(FEmitContext& Context, FEmitShaderScope& Scope, const FRequestedType& RequestedType, FEmitValueShaderResult& OutResult) const override;
 };
 
 class FStatementReturn : public FStatement
@@ -315,7 +306,7 @@ public:
 	FExpression* Expression;
 
 	virtual void Prepare(FEmitContext& Context) const override;
-	virtual void EmitShader(FEmitContext& Context) const override;
+	virtual void EmitShader(FEmitContext& Context, FEmitShaderScope& Scope) const override;
 };
 
 class FStatementBreak : public FStatement
@@ -324,7 +315,7 @@ public:
 	//static constexpr bool MarkScopeLive = true;
 
 	virtual void Prepare(FEmitContext& Context) const override;
-	virtual void EmitShader(FEmitContext& Context) const override;
+	virtual void EmitShader(FEmitContext& Context, FEmitShaderScope& Scope) const override;
 };
 
 class FStatementIf : public FStatement
@@ -336,7 +327,7 @@ public:
 	FScope* NextScope;
 
 	virtual void Prepare(FEmitContext& Context) const override;
-	virtual void EmitShader(FEmitContext& Context) const override;
+	virtual void EmitShader(FEmitContext& Context, FEmitShaderScope& Scope) const override;
 };
 
 class FStatementLoop : public FStatement
@@ -346,7 +337,7 @@ public:
 	FScope* NextScope;
 
 	virtual void Prepare(FEmitContext& Context) const override;
-	virtual void EmitShader(FEmitContext& Context) const override;
+	virtual void EmitShader(FEmitContext& Context, FEmitShaderScope& Scope) const override;
 };
 
 }
diff --git a/Engine/Source/Runtime/Engine/Public/HLSLTree/HLSLTreeEmit.h b/Engine/Source/Runtime/Engine/Public/HLSLTree/HLSLTreeEmit.h
new file mode 100644
--- /dev/null
+++ b/Engine/Source/Runtime/Engine/Public/HLSLTree/HLSLTreeEmit.h
@@ -0,0 +1,324 @@
+// Copyright Epic Games, Inc. All Rights Reserved.
+#pragma once
+
+#include "CoreMinimal.h"
+#include "Engine/EngineTypes.h"
+#include "Containers/ArrayView.h"
+#include "Containers/StringView.h"
+#include "Misc/StringBuilder.h"
+#include "Misc/MemStack.h"
+#include "Hash/xxhash.h"
+#include "HLSLTree/HLSLTreeTypes.h"
+
+class FMaterial;
+class FMaterialCompilationOutput;
+struct FStaticParameterSet;
+
+namespace UE
+{
+namespace HLSLTree
+{
+
+class FErrorHandlerInterface;
+class FNode;
+class FScope;
+class FExpression;
+class FFunction;
+class FRequestedType;
+class FEmitShaderScope;
+class FEmitShaderExpression;
+class FEmitShaderStatement;
+
+struct FEmitShaderScopeEntry
+{
+	FEmitShaderScopeEntry() = default;
+	FEmitShaderScopeEntry(FEmitShaderScope* InScope, int32 InIndent, FStringBuilderBase& InCode) : Scope(InScope), Code(&InCode), Indent(InIndent) {}
+
+	FEmitShaderScope* Scope = nullptr;
+	FStringBuilderBase* Code = nullptr;
+	int32 Indent = 0;
+};
+using FEmitShaderScopeStack = TArray<FEmitShaderScopeEntry, TInlineAllocator<16>>;
+
+class FEmitShaderNode
+{
+public:
+	virtual void EmitShaderCode(FEmitShaderScopeStack& Stack, int32 Indent, FStringBuilderBase& OutString) = 0;
+	virtual FEmitShaderExpression* AsExpression() { return nullptr; }
+	virtual FEmitShaderStatement* AsStatement() { return nullptr; }
+
+	FEmitShaderNode(FEmitShaderScope& InScope, TArrayView<FEmitShaderNode*> InDependencies);
+
+	FEmitShaderScope* Scope = nullptr;
+	FEmitShaderNode* NextScopedNode = nullptr;
+	TArrayView<FEmitShaderNode*> Dependencies;
+};
+using FEmitShaderDependencies = TArray<FEmitShaderNode*, TInlineAllocator<8>>;
+
+class FEmitShaderExpression final : public FEmitShaderNode
+{
+public:
+	FEmitShaderExpression(FEmitShaderScope& InScope, TArrayView<FEmitShaderNode*> InDependencies, const Shader::FType& InType, FXxHash64 InHash)
+		: FEmitShaderNode(InScope, InDependencies)
+		, Type(InType)
+		, Hash(InHash)
+	{}
+
+	virtual void EmitShaderCode(FEmitShaderScopeStack& Stack, int32 Indent, FStringBuilderBase& OutString) override;
+	virtual FEmitShaderExpression* AsExpression() override { return this; }
+	inline bool IsInline() const { return Value == nullptr; }
+
+	const TCHAR* Reference = nullptr;
+	const TCHAR* Value = nullptr;
+	Shader::FType Type;
+	FXxHash64 Hash;
+};
+
+enum class EEmitScopeFormat : uint8
+{
+	None,
+	Unscoped,
+	Scoped,
+};
+
+class FEmitShaderStatement final : public FEmitShaderNode
+{
+public:
+	FEmitShaderStatement(FEmitShaderScope& InScope, TArrayView<FEmitShaderNode*> InDependencies)
+		: FEmitShaderNode(InScope, InDependencies)
+	{}
+
+	virtual void EmitShaderCode(FEmitShaderScopeStack& Stack, int32 Indent, FStringBuilderBase& OutString) override;
+	virtual FEmitShaderStatement* AsStatement() override { return this; }
+
+	FEmitShaderScope* NestedScopes[2] = { nullptr };
+	FStringView Code[2];
+	EEmitScopeFormat ScopeFormat;
+	int8 NumEntries = 0;
+};
+
+enum class EFormatArgType : uint8
+{
+	Void,
+	ShaderValue,
+	String,
+	Int,
+};
+
+struct FFormatArgVariant
+{
+	FFormatArgVariant() {}
+	FFormatArgVariant(FEmitShaderExpression* InValue) : Type(EFormatArgType::ShaderValue), ShaderValue(InValue) { check(InValue); }
+	FFormatArgVariant(const TCHAR* InValue) : Type(EFormatArgType::String), String(InValue) { check(InValue); }
+	FFormatArgVariant(int32 InValue) : Type(EFormatArgType::Int), Int(InValue) {}
+
+	EFormatArgType Type = EFormatArgType::Void;
+	union
+	{
+		FEmitShaderExpression* ShaderValue;
+		const TCHAR* String;
+		int32 Int;
+	};
+};
+
+using FFormatArgList = TArray<FFormatArgVariant, TInlineAllocator<8>>;
+
+namespace Private
+{
+inline void BuildFormatArgList(FFormatArgList&) {}
+
+template<typename Type, typename... Types>
+inline void BuildFormatArgList(FFormatArgList& OutList, Type Arg, Types... Args)
+{
+	OutList.Add(Arg);
+	BuildFormatArgList(OutList, Forward<Types>(Args)...);
+}
+
+void InternalFormatStrings(FStringBuilderBase* OutString0,
+	FStringBuilderBase* OutString1,
+	FEmitShaderDependencies& OutDependencies,
+	FStringView Format0,
+	FStringView Format1,
+	const FFormatArgList& ArgList);
+} // namespace Private
+
+template<typename FormatType, typename... Types>
+void FormatString(FStringBuilderBase& OutString, FEmitShaderDependencies& OutDependencies, const FormatType& Format, Types... Args)
+{
+	FFormatArgList ArgList;
+	Private::BuildFormatArgList(ArgList, Forward<Types>(Args)...);
+	Private::InternalFormatStrings(&OutString, nullptr, OutDependencies, Format, FStringView(), ArgList);
+}
+
+template<typename FormatType0, typename FormatType1, typename... Types>
+void FormatStrings(FStringBuilderBase& OutString0, FStringBuilderBase& OutString1, FEmitShaderDependencies& OutDependencies, const FormatType0& Format0, const FormatType1& Format1, Types... Args)
+{
+	FFormatArgList ArgList;
+	Private::BuildFormatArgList(ArgList, Forward<Types>(Args)...);
+	Private::InternalFormatStrings(&OutString0, &OutString1, OutDependencies, Format0, Format1, ArgList);
+}
+
+class FEmitShaderScope
+{
+public:
+	static FEmitShaderScope* FindSharedParent(FEmitShaderScope* Lhs, FEmitShaderScope* Rhs);
+
+	void EmitShaderCode(FEmitShaderScopeStack& Stack);
+
+	FEmitShaderScope* ParentScope = nullptr;
+	FEmitShaderNode* FirstNode;
+	int32 NestedLevel = 0;
+};
+
+/** Tracks shared state while emitting HLSL code */
+class FEmitContext
+{
+public:
+	explicit FEmitContext(FMemStackBase& InAllocator, FErrorHandlerInterface& InErrors, const Shader::FStructTypeRegistry& InTypeRegistry);
+	~FEmitContext();
+
+	template<typename T>
+	static TArrayView<FEmitShaderNode*> MakeDependencies(T*& Dependency)
+	{
+		return Dependency ? TArrayView<FEmitShaderNode*>(&Dependency, 1) : TArrayView<FEmitShaderNode*>();
+	}
+
+	void Finalize();
+
+	/** Get a unique local variable name */
+	const TCHAR* AcquireLocalDeclarationCode();
+
+	FEmitShaderExpression* EmitExpressionInternal(FEmitShaderScope& Scope, TArrayView<FEmitShaderNode*> Dependencies, bool bInline, const Shader::FType& Type, FStringView Code);
+
+	template<typename FormatType, typename... Types>
+	FEmitShaderExpression* EmitExpressionWithDependencies(FEmitShaderScope& Scope, TArrayView<FEmitShaderNode*> Dependencies, const Shader::FType& Type, const FormatType& Format, Types... Args)
+	{
+		TStringBuilder<2048> String;
+		FEmitShaderDependencies LocalDependencies(Dependencies);
+		FormatString(String, LocalDependencies, Format, Forward<Types>(Args)...);
+		return EmitExpressionInternal(Scope, LocalDependencies, false, Type, String.ToView());
+	}
+
+	template<typename FormatType, typename... Types>
+	FEmitShaderExpression* EmitInlineExpressionWithDependencies(FEmitShaderScope& Scope, TArrayView<FEmitShaderNode*> Dependencies, const Shader::FType& Type, const FormatType& Format, Types... Args)
+	{
+		TStringBuilder<2048> String;
+		FEmitShaderDependencies LocalDependencies(Dependencies);
+		FormatString(String, LocalDependencies, Format, Forward<Types>(Args)...);
+		return EmitExpressionInternal(Scope, LocalDependencies, true, Type, String.ToView());
+	}
+
+	template<typename FormatType, typename... Types>
+	FEmitShaderExpression* EmitInlineExpressionWithDependency(FEmitShaderScope& Scope, FEmitShaderNode* Dependency, const Shader::FType& Type, const FormatType& Format, Types... Args)
+	{
+		return EmitInlineExpressionWithDependencies(Scope, MakeDependencies(Dependency), Type, Format, Forward<Types>(Args)...);
+	}
+
+	template<typename FormatType, typename... Types>
+	FEmitShaderExpression* EmitExpression(FEmitShaderScope& Scope, const Shader::FType& Type, const FormatType& Format, Types... Args)
+	{
+		return EmitExpressionWithDependencies(Scope, TArrayView<FEmitShaderNode*>(), Type, Format, Forward<Types>(Args)...);
+	}
+
+	template<typename FormatType, typename... Types>
+	FEmitShaderExpression* EmitInlineExpression(FEmitShaderScope& Scope, const Shader::FType& Type, const FormatType& Format, Types... Args)
+	{
+		return EmitInlineExpressionWithDependencies(Scope, TArrayView<FEmitShaderNode*>(), Type, Format, Forward<Types>(Args)...);
+	}
+
+	FEmitShaderStatement* EmitStatementInternal(FEmitShaderScope& Scope, TArrayView<FEmitShaderNode*> Dependencies, EEmitScopeFormat ScopeFormat, FEmitShaderScope* NestedScope0, FEmitShaderScope* NestedScope1, FStringView Code0, FStringView Code1);
+
+	template<typename FormatType0, typename FormatType1, typename... Types>
+	FEmitShaderStatement* EmitFormatStatementInternal(FEmitShaderScope& Scope,
+		TArrayView<FEmitShaderNode*> Dependencies,
+		EEmitScopeFormat ScopeFormat,
+		FEmitShaderScope* NestedScope0,
+		FEmitShaderScope* NestedScope1,
+		const FormatType0& Format0,
+		const FormatType1& Format1,
+		Types... Args)
+	{
+		TStringBuilder<1024> String0;
+		TStringBuilder<1024> String1;
+		FEmitShaderDependencies LocalDependencies(Dependencies);
+		FormatStrings(String0, String1, LocalDependencies, Format0, Format1, Forward<Types>(Args)...);
+		return EmitStatementInternal(Scope, LocalDependencies, ScopeFormat, NestedScope0, NestedScope1, String0.ToView(), String1.ToView());
+	}
+
+	template<typename FormatType, typename... Types>
+	FEmitShaderStatement* EmitStatementWithDependencies(FEmitShaderScope& Scope, TArrayView<FEmitShaderNode*> Dependencies, const FormatType& Format, Types... Args)
+	{
+		return EmitFormatStatementInternal(Scope, Dependencies, EEmitScopeFormat::None, nullptr, nullptr, Format, FStringView(), Forward<Types>(Args)...);
+	}
+
+	template<typename FormatType, typename... Types>
+	FEmitShaderStatement* EmitStatementWithDependency(FEmitShaderScope& Scope, FEmitShaderNode* Dependency, const FormatType& Format, Types... Args)
+	{
+		return EmitStatementWithDependencies(Scope, MakeDependencies(Dependency), Format, Forward<Types>(Args)...);
+	}
+
+	template<typename FormatType, typename... Types>
+	FEmitShaderStatement* EmitStatement(FEmitShaderScope& Scope, const FormatType& Format, Types... Args)
+	{
+		return EmitStatementWithDependencies(Scope, TArrayView<FEmitShaderNode*>(), Format, Forward<Types>(Args)...);
+	}
+
+	FEmitShaderStatement* EmitNextScopeWithDependency(FEmitShaderScope& Scope, FEmitShaderNode* Dependency, FScope* NextScope)
+	{
+		FEmitShaderScope* EmitScope = AcquireEmitScope(NextScope, &Scope);
+		return EmitStatementInternal(Scope, MakeDependencies(Dependency), EEmitScopeFormat::Unscoped, EmitScope, nullptr, FStringView(), FStringView());
+	}
+
+	FEmitShaderStatement* EmitNextScope(FEmitShaderScope& Scope, FScope* NextScope)
+	{
+		return EmitNextScopeWithDependency(Scope, nullptr, NextScope);
+	}
+
+	template<typename FormatType, typename... Types>
+	FEmitShaderStatement* EmitNestedScope(FEmitShaderScope& Scope, FScope* NestedScope, const FormatType& Format, Types... Args)
+	{
+		FEmitShaderScope* EmitScope = AcquireEmitScope(NestedScope, &Scope);
+		return EmitFormatStatementInternal(Scope, TArrayView<FEmitShaderNode*>(), EEmitScopeFormat::Scoped, EmitScope, nullptr, Format, FStringView(), Forward<Types>(Args)...);
+	}
+
+	template<typename FormatType0, typename FormatType1, typename... Types>
+	FEmitShaderStatement* EmitNestedScopes(FEmitShaderScope& Scope, FScope* NestedScope0, FScope* NestedScope1, const FormatType0& Format0, const FormatType1& Format1, Types... Args)
+	{
+		FEmitShaderScope* EmitScope0 = AcquireEmitScope(NestedScope0, &Scope);
+		FEmitShaderScope* EmitScope1 = AcquireEmitScope(NestedScope1, &Scope);
+		return EmitFormatStatementInternal(Scope, TArrayView<FEmitShaderNode*>(), EEmitScopeFormat::Scoped, EmitScope0, EmitScope1, Format0, Format1, Forward<Types>(Args)...);
+	}
+
+	FEmitShaderScope* AcquireEmitScope(FScope* Scope, FEmitShaderScope* OverrideParent = nullptr);
+	FEmitShaderScope* FindEmitScope(FScope* Scope) const;
+
+	FEmitShaderExpression* EmitPreshaderOrConstant(FEmitShaderScope& Scope, const FRequestedType& RequestedType, FExpression* Expression);
+	FEmitShaderExpression* EmitConstantZero(FEmitShaderScope& Scope, const Shader::FType& Type);
+	FEmitShaderExpression* EmitCast(FEmitShaderScope& Scope, FEmitShaderExpression* ShaderValue, const Shader::FType& DestType);
+
+	FMemStackBase* Allocator = nullptr;
+	FErrorHandlerInterface* Errors = nullptr;
+	const Shader::FStructTypeRegistry* TypeRegistry = nullptr;
+
+	TArray<FEmitShaderNode*> EmitNodes;
+	TMap<const FScope*, FEmitShaderScope*> EmitScopeMap;
+	TMap<const FExpression*, FEmitShaderExpression*> EmitLocalPHIMap;
+	TMap<FXxHash64, FEmitShaderExpression*> EmitExpressionMap;
+	TMap<FXxHash64, FEmitShaderExpression*> EmitPreshaderMap;
+	TMap<const FFunction*, FEmitShaderNode*> EmitFunctionMap;
+
+	// TODO - remove preshader material dependency
+	const FMaterial* Material = nullptr;
+	const FStaticParameterSet* StaticParameters = nullptr;
+	FMaterialCompilationOutput* MaterialCompilationOutput = nullptr;
+	TMap<Shader::FValue, uint32> DefaultUniformValues;
+	uint32 UniformPreshaderOffset = 0u;
+	bool bReadMaterialNormal = false;
+
+	int32 NumExpressionLocals = 0;
+	int32 NumExpressionLocalPHIs = 0;
+	int32 NumTexCoords = 0;
+};
+
+} // namespace HLSLTree
+} // namespace UE
diff --git a/Engine/Source/Runtime/Engine/Public/MaterialHLSLGenerator.h b/Engine/Source/Runtime/Engine/Public/MaterialHLSLGenerator.h
--- a/Engine/Source/Runtime/Engine/Public/MaterialHLSLGenerator.h
+++ b/Engine/Source/Runtime/Engine/Public/MaterialHLSLGenerator.h
@@ -161,18 +161,18 @@ private:
 		}
 	};
 
-	struct FFunctionInput
-	{
-		const UMaterialExpressionFunctionInput* FunctionInputExpression = nullptr;
-		UMaterialExpression* ConnectedExpression = nullptr;
-		int32 ConnectedOutputIndex = INDEX_NONE;
-	};
-	using FFunctionInputArray = TArray<FFunctionInput, TInlineAllocator<4>>;
+	using FFunctionInputArray = TArray<const UMaterialExpressionFunctionInput*, TInlineAllocator<4>>;
+	using FFunctionOutputArray = TArray<const UMaterialExpressionFunctionOutput*, TInlineAllocator<4>>;
+	using FConnectedInputArray = TArray<UE::HLSLTree::FExpression*, TInlineAllocator<4>>;
 
 	struct FFunctionCallEntry
 	{
-		UMaterialFunctionInterface* Function = nullptr;
-		FFunctionInputArray Inputs;
+		UMaterialFunctionInterface* MaterialFunction = nullptr;
+		UE::HLSLTree::FFunction* HLSLFunction = nullptr;
+		FFunctionInputArray FunctionInputs;
+		FFunctionOutputArray FunctionOutputs;
+		FConnectedInputArray ConnectedInputs;
+		bool bGeneratedResult = false;
 	};
 
 	struct FStatementEntry
@@ -201,6 +201,7 @@ private:
 	TMap<UE::HLSLTree::FTextureDescription, UE::HLSLTree::FTextureParameterDeclaration*> TextureDeclarationMap;
 	TMap<FName, UE::HLSLTree::FTextureParameterDeclaration*> TextureParameterDeclarationMap;
 	TMap<FSHAHash, FFunctionCallEntry*> FunctionCallMap;
+	TMap<UMaterialFunctionInterface*, UE::HLSLTree::FFunction*> FunctionMap;
 	TMap<UMaterialExpression*, FStatementEntry> StatementMap;
 	TMap<FExpressionDataKey, void*> ExpressionDataMap;
 	bool bGeneratedResult;
diff --git a/Engine/Source/Runtime/Engine/Public/MaterialShared.h b/Engine/Source/Runtime/Engine/Public/MaterialShared.h
--- a/Engine/Source/Runtime/Engine/Public/MaterialShared.h
+++ b/Engine/Source/Runtime/Engine/Public/MaterialShared.h
@@ -1871,7 +1871,7 @@ public:
 	ENGINE_API bool IsCompilationFinished() const;
 
 	/** Should the material be compiled using exec pin? */
-	virtual bool IsCompiledWithExecutionFlow() const { return false; }
+	virtual bool IsUsingControlFlow() const { return false; }
 
 	/** Is the material using the new (WIP) HLSL generator? */
 	virtual bool IsUsingNewHLSLGenerator() const { return false; }
@@ -2772,7 +2772,7 @@ public:
 	/** Allows to associate the shader resources with the asset for load order. */
 	ENGINE_API virtual FName GetAssetPath() const override;
 	ENGINE_API virtual bool ShouldInlineShaderCode() const override;
-	ENGINE_API virtual bool IsCompiledWithExecutionFlow() const override;
+	ENGINE_API virtual bool IsUsingControlFlow() const override;
 	ENGINE_API virtual bool IsUsingNewHLSLGenerator() const override;
 #endif // WITH_EDITOR
 
diff --git a/Engine/Source/Runtime/Engine/Public/Shader/Preshader.h b/Engine/Source/Runtime/Engine/Public/Shader/Preshader.h
--- a/Engine/Source/Runtime/Engine/Public/Shader/Preshader.h
+++ b/Engine/Source/Runtime/Engine/Public/Shader/Preshader.h
@@ -9,6 +9,7 @@
 class FUniformExpressionSet;
 struct FMaterialRenderContext;
 class FMaterial;
+class FXxHash64Builder;
 
 namespace UE
 {
@@ -87,8 +88,7 @@ public:
 		return !operator==(Lhs, Rhs);
 	}
 
-	FSHAHash GetHash() const;
-	void AppendHash(FSHA1& OutHasher) const;
+	void AppendHash(FXxHash64Builder& OutHasher) const;
 
 	FPreshaderValue Evaluate(FUniformExpressionSet* UniformExpressionSet, const struct FMaterialRenderContext& Context, FPreshaderStack& Stack) const;
 	FPreshaderValue EvaluateConstant(const FMaterial& Material, FPreshaderStack& Stack) const;
