Add ListTextures support for volume and cube textures

diff --git a/Engine/Source/Runtime/Engine/Private/UnrealEngine.cpp b/Engine/Source/Runtime/Engine/Private/UnrealEngine.cpp
--- a/Engine/Source/Runtime/Engine/Private/UnrealEngine.cpp
+++ b/Engine/Source/Runtime/Engine/Private/UnrealEngine.cpp
@@ -72,6 +72,7 @@ UnrealEngine.cpp: Implements the UEngine class and helpers.
 #include "Components/SkeletalMeshComponent.h"
 #include "Engine/Texture.h"
 #include "Engine/Texture2D.h"
+#include "Engine/VolumeTexture.h"
 #include "ParticleHelper.h"
 #include "Particles/ParticleModule.h"
 #include "Particles/ParticleSystemComponent.h"
@@ -5761,6 +5762,7 @@ bool UEngine::HandleListTexturesCommand( const TCHAR* Cmd, FOutputDevice& Ar )
 		UTexture*			Texture				= *It;
 		UTexture2D*			Texture2D			= Cast<UTexture2D>(Texture);
 		UTextureCube*		TextureCube			= Cast<UTextureCube>(Texture);
+		UVolumeTexture*		Texture3D			= Cast<UVolumeTexture>(Texture);
 
 		int32				LODGroup			= Texture->LODGroup;
 		int32				NumMips				= 0;
@@ -5818,7 +5820,26 @@ bool UEngine::HandleListTexturesCommand( const TCHAR* Cmd, FOutputDevice& Ar )
 		else if (TextureCube != nullptr)
 		{
 			NumMips				= TextureCube->GetNumMips();
+			MaxResLODBias		= TextureCube->GetCachedLODBias();
+			MaxAllowedSizeX		= FMath::Max<int32>(TextureCube->GetSizeX() >> MaxResLODBias, 1);
+			MaxAllowedSizeY		= FMath::Max<int32>(TextureCube->GetSizeY() >> MaxResLODBias, 1);
 			Format				= TextureCube->GetPixelFormat();
+			DroppedMips			= MaxResLODBias;
+			CurSizeX			= FMath::Max<int32>(TextureCube->GetSizeX() >> DroppedMips, 1);
+			CurSizeY			= FMath::Max<int32>(TextureCube->GetSizeY() >> DroppedMips, 1);
+			bIsUncompressed		= TextureCube->IsUncompressed();
+		}
+		else if (Texture3D != nullptr)
+		{
+			NumMips				= Texture3D->GetNumMips();
+			MaxResLODBias		= Texture3D->GetCachedLODBias();
+			MaxAllowedSizeX		= FMath::Max<int32>(Texture3D->GetSizeX() >> MaxResLODBias, 1);
+			MaxAllowedSizeY		= FMath::Max<int32>(Texture3D->GetSizeY() >> MaxResLODBias, 1);
+			Format				= Texture3D->GetPixelFormat();
+			DroppedMips			= MaxResLODBias;
+			CurSizeX			= FMath::Max<int32>(Texture3D->GetSizeX() >> DroppedMips, 1);
+			CurSizeY			= FMath::Max<int32>(Texture3D->GetSizeY() >> DroppedMips, 1);
+			bIsUncompressed		= Texture3D->IsUncompressed();
 		}
 
 		if (bListUnused && UsageCount != 0)
