Update Advanced Volumetric Output node to refelct the fact that multiplescattering value of 2 is expenssive and 0 or 1 is recommended.

diff --git a/Engine/Source/Runtime/Engine/Classes/Materials/MaterialExpressionVolumetricAdvancedMaterialOutput.h b/Engine/Source/Runtime/Engine/Classes/Materials/MaterialExpressionVolumetricAdvancedMaterialOutput.h
--- a/Engine/Source/Runtime/Engine/Classes/Materials/MaterialExpressionVolumetricAdvancedMaterialOutput.h
+++ b/Engine/Source/Runtime/Engine/Classes/Materials/MaterialExpressionVolumetricAdvancedMaterialOutput.h
@@ -60,8 +60,8 @@ class UMaterialExpressionVolumetricAdvancedMaterialOutput : public UMaterialExpr
 	UPROPERTY(EditAnywhere, Category = "Phase")
 	bool PerSamplePhaseEvaluation;
 
-	/** How many octave to use for the multiple-scattering approximation. This makes the shader more expensive so try to only a single octave. 0 means single scattering only. */
-	UPROPERTY(EditAnywhere, Category = "Multi-Scattering", meta = (UIMin = 0, UIMax = 2, ClampMin = 0, ClampMax = 2))
+	/** How many octave to use for the multiple-scattering approximation. This makes the shader more expensive so you should only use 0 or 1 for better performance, and tweak multiple scattering parameters accordingly. 0 means single scattering only. The maximum value is 2 (expenssive). */
+	UPROPERTY(EditAnywhere, Category = "Multi-Scattering", meta = (UIMin = 0, UIMax = 1, ClampMin = 0, ClampMax = 2))
 	uint32 MultiScatteringApproximationOctaveCount;
 
 	/** Only used if MultiScatteringContribution is not hooked up. Multi-scattering approximation: represents how much contribution each successive octave will add. Valid range is [0,1], from low to high contribution */
