World Partition: move data layers specific code out of actor descriptors code.

diff --git a/Engine/Source/Editor/SceneOutliner/Private/SceneOutlinerModule.cpp b/Engine/Source/Editor/SceneOutliner/Private/SceneOutlinerModule.cpp
--- a/Engine/Source/Editor/SceneOutliner/Private/SceneOutlinerModule.cpp
+++ b/Engine/Source/Editor/SceneOutliner/Private/SceneOutlinerModule.cpp
@@ -234,7 +234,6 @@ void FSceneOutlinerModule::CreateActorInfoColumns(FSceneOutlinerInitializationOp
 
 		TArray<const UDataLayer*> DataLayerObjects;
 
-
 		if (const FActorTreeItem* ActorItem = Item.CastTo<FActorTreeItem>())
 		{
 			AActor* Actor = ActorItem->Actor.Get();
@@ -250,7 +249,13 @@ void FSceneOutlinerModule::CreateActorInfoColumns(FSceneOutlinerInitializationOp
 
 			if (ActorDesc)
 			{
-				DataLayerObjects = ActorDesc->GetDataLayerObjects();
+				if (UWorld* World = ActorDesc->GetWorld())
+				{
+					if (AWorldDataLayers* WorldDataLayers = World->GetWorldDataLayers())
+					{
+						DataLayerObjects = WorldDataLayers->GetDataLayerObjects(ActorDesc->GetDataLayers());
+					}
+				}
 			}
 		}
 
diff --git a/Engine/Source/Runtime/Engine/Private/WorldPartition/WorldPartitionActorDesc.cpp b/Engine/Source/Runtime/Engine/Private/WorldPartition/WorldPartitionActorDesc.cpp
--- a/Engine/Source/Runtime/Engine/Private/WorldPartition/WorldPartitionActorDesc.cpp
+++ b/Engine/Source/Runtime/Engine/Private/WorldPartition/WorldPartitionActorDesc.cpp
@@ -10,16 +10,12 @@
 #include "Algo/Transform.h"
 #include "Engine/World.h"
 #include "Engine/Level.h"
-#include "Editor/GroupActor.h"
 #include "UObject/UE5MainStreamObjectVersion.h"
 #include "UObject/UE5ReleaseStreamObjectVersion.h"
 #include "WorldPartition/WorldPartitionLog.h"
 #include "WorldPartition/ActorDescContainer.h"
-#include "WorldPartition/DataLayer/DataLayer.h"
-#include "WorldPartition/DataLayer/WorldDataLayers.h"
 #include "WorldPartition/HLOD/HLODLayer.h"
 #include "Engine/Public/ActorReferencesUtils.h"
-#include "ActorFolder.h"
 #endif
 
 #if WITH_EDITOR
@@ -157,6 +153,11 @@ void FWorldPartitionActorDesc::SerializeTo(TArray<uint8>& OutData)
 	OutData.Append(PayloadData);
 }
 
+UWorld* FWorldPartitionActorDesc::GetWorld() const
+{
+	return Container ? Container->GetWorld() : nullptr;
+}
+
 void FWorldPartitionActorDesc::TransformInstance(const FString& From, const FString& To)
 {
 	check(!HardRefCount);
@@ -271,30 +272,6 @@ FName FWorldPartitionActorDesc::GetActorName() const
 FName FWorldPartitionActorDesc::GetActorLabelOrName() const
 {
 	return GetActorLabel().IsNone() ? GetActorName() : GetActorLabel();
-};
-
-
-TArray<const UDataLayer*> FWorldPartitionActorDesc::GetDataLayerObjects() const
-{
-	if (Container)
-	{
-		if (AWorldDataLayers* WorldDataLayers = Container->GetWorld()->GetWorldDataLayers())
-		{
-			TArray<const UDataLayer*> DataLayerObjects;
-			DataLayerObjects.Reserve(DataLayers.Num());
-			for (const FName& DataLayerName : DataLayers)
-			{
-				if (const UDataLayer* DataLayer = WorldDataLayers->GetDataLayerFromName(DataLayerName))
-				{
-					DataLayerObjects.Add(DataLayer);
-				}
-			}
-
-			return DataLayerObjects;
-		}
-	}
-
-	return TArray<const UDataLayer*>();
 }
 
 UHLODLayer* FWorldPartitionActorDesc::GetHLODLayer() const
@@ -403,5 +380,4 @@ void FWorldPartitionActorDesc::UnregisterActor()
 		Container->OnActorDescUnregistered(*this);
 	}
 }
-
 #endif
\ No newline at end of file
diff --git a/Engine/Source/Runtime/Engine/Public/WorldPartition/WorldPartitionActorDesc.h b/Engine/Source/Runtime/Engine/Public/WorldPartition/WorldPartitionActorDesc.h
--- a/Engine/Source/Runtime/Engine/Public/WorldPartition/WorldPartitionActorDesc.h
+++ b/Engine/Source/Runtime/Engine/Public/WorldPartition/WorldPartitionActorDesc.h
@@ -68,7 +68,6 @@ public:
 	inline bool GetActorIsHLODRelevant() const { return bActorIsHLODRelevant; }
 	class UHLODLayer* GetHLODLayer() const;
 	inline const TArray<FName>& GetDataLayers() const { return DataLayers; }
-	TArray<const UDataLayer*> GetDataLayerObjects() const;
 	inline FName GetActorPackage() const { return ActorPackage; }
 	inline FName GetActorPath() const { return ActorPath; }
 	inline FName GetActorLabel() const { return ActorLabel; }
@@ -161,6 +160,8 @@ public:
 
 	void SerializeTo(TArray<uint8>& OutData);
 
+	UWorld* GetWorld() const;
+
 protected:
 	FWorldPartitionActorDesc();
 
