Blueprint: Rely on GetWorld for determining the world to use for debugging instances

diff --git a/Engine/Source/Runtime/Engine/Private/Blueprint.cpp b/Engine/Source/Runtime/Engine/Private/Blueprint.cpp
--- a/Engine/Source/Runtime/Engine/Private/Blueprint.cpp
+++ b/Engine/Source/Runtime/Engine/Private/Blueprint.cpp
@@ -792,9 +792,18 @@ void UBlueprint::DebuggingWorldRegistrationHelper(UObject* ObjectProvidingWorld,
 			ObjOuter = ObjOuter->GetOuter();
 		}
 
+		// if we can't find the world on the outer chain, fallback to the GetWorld method
+		if (ObjWorld == NULL)
+		{
+			ObjWorld = ObjectProvidingWorld->GetWorld();
+		}
+
 		if (ObjWorld != NULL)
 		{
-			ObjWorld->NotifyOfBlueprintDebuggingAssociation(this, ValueToRegister);
+			if( !ObjWorld->HasAnyFlags(RF_BeginDestroyed))
+			{
+				ObjWorld->NotifyOfBlueprintDebuggingAssociation(this, ValueToRegister);
+			}
 			OnSetObjectBeingDebuggedDelegate.Broadcast(ValueToRegister);
 		}
 	}
