Fix a bug that the brightness of the reflectioncapture couldn't be changed at runtime on mobile.

diff --git a/Engine/Shaders/Private/MobileBasePassPixelShader.usf b/Engine/Shaders/Private/MobileBasePassPixelShader.usf
--- a/Engine/Shaders/Private/MobileBasePassPixelShader.usf
+++ b/Engine/Shaders/Private/MobileBasePassPixelShader.usf
@@ -336,7 +336,8 @@ void GatherSpecularIBL(FMaterialPixelParameters MaterialParameters
 		CompositedAverageBrightness.x += AverageBrightness * DistanceAlpha * CompositedAverageBrightness.y;
 		CompositedAverageBrightness.y *= 1 - DistanceAlpha;
 #else
-		ImageBasedReflections.rgb = SpecularIBL * ImageBasedReflections.a;
+		half ReflectionCaptureBrightness = MobileReflectionCapture.Params.w;
+		ImageBasedReflections.rgb = SpecularIBL * ImageBasedReflections.a * ReflectionCaptureBrightness;
 #if ALLOW_STATIC_LIGHTING
 		//To keep ImageBasedReflectionLighting conherence with PC, use ComputeMixingWeight instead of InvReflectionAverageBrightness to calulate the IBL constribution
 		ImageBasedReflections.rgb *= MobileComputeMixingWeight(IndirectIrradiance, MobileReflectionCapture.Params.x, Roughness);
diff --git a/Engine/Source/Editor/UnrealEd/Private/EditorEngine.cpp b/Engine/Source/Editor/UnrealEd/Private/EditorEngine.cpp
--- a/Engine/Source/Editor/UnrealEd/Private/EditorEngine.cpp
+++ b/Engine/Source/Editor/UnrealEd/Private/EditorEngine.cpp
@@ -3848,7 +3848,7 @@ void UEditorEngine::BuildReflectionCaptures(UWorld* World)
 					UMapBuildDataRegistry* Registry = StorageLevel->GetOrCreateMapBuildData();
 					if (!CaptureComponent->bModifyMaxValueRGBM)
 					{
-						CaptureComponent->MaxValueRGBM = GetMaxValueRGBM(ReadbackCaptureData.FullHDRCapturedData, ReadbackCaptureData.CubemapSize, ReadbackCaptureData.Brightness);
+						CaptureComponent->MaxValueRGBM = GetMaxValueRGBM(ReadbackCaptureData.FullHDRCapturedData, ReadbackCaptureData.CubemapSize);
 					}
 					FString TextureName = CaptureComponent->GetName() + TEXT("Texture");
 					TextureName += LexToString(CaptureComponent->MapBuildDataId);
diff --git a/Engine/Source/Runtime/Core/Public/UObject/UE5ReleaseStreamObjectVersion.h b/Engine/Source/Runtime/Core/Public/UObject/UE5ReleaseStreamObjectVersion.h
--- a/Engine/Source/Runtime/Core/Public/UObject/UE5ReleaseStreamObjectVersion.h
+++ b/Engine/Source/Runtime/Core/Public/UObject/UE5ReleaseStreamObjectVersion.h
@@ -78,6 +78,9 @@ struct CORE_API FUE5ReleaseStreamObjectVersion
 		// Remove FSkeletalMeshLODModel bulk datas
 		RemoveSkeletalMeshLODModelBulkDatas,
 		
+		// Exclude brightness from the EncodedHDRCubemap,
+		ExcludeBrightnessFromEncodedHDRCubemap,
+
 		// -----<new versions can be added above this line>-------------------------------------------------
 		VersionPlusOne,
 		LatestVersion = VersionPlusOne - 1
diff --git a/Engine/Source/Runtime/Engine/Classes/Components/ReflectionCaptureComponent.h b/Engine/Source/Runtime/Engine/Classes/Components/ReflectionCaptureComponent.h
--- a/Engine/Source/Runtime/Engine/Classes/Components/ReflectionCaptureComponent.h
+++ b/Engine/Source/Runtime/Engine/Classes/Components/ReflectionCaptureComponent.h
@@ -169,6 +169,6 @@ private:
 	friend class FReflectionCaptureProxy;
 };
 
-ENGINE_API extern float GetMaxValueRGBM(const TArray<uint8>& FullHDRData, int32 CubemapSize, float Brightness);
-ENGINE_API extern void GenerateEncodedHDRData(const TArray<uint8>& FullHDRData, int32 CubemapSize, float Brightness, float MaxValueRGBM, TArray<uint8>& OutEncodedHDRData);
+ENGINE_API extern float GetMaxValueRGBM(const TArray<uint8>& FullHDRData, int32 CubemapSize);
+ENGINE_API extern void GenerateEncodedHDRData(const TArray<uint8>& FullHDRData, int32 CubemapSize, float MaxValueRGBM, TArray<uint8>& OutEncodedHDRData);
 ENGINE_API extern void GenerateEncodedHDRTextureCube(class UMapBuildDataRegistry* Registry, class FReflectionCaptureData& CaptureBuildData, FString& TextureName, float MaxValueRGBM, class UReflectionCaptureComponent* CaptureComponent = nullptr, bool bIsReflectionCaptureCompressionProjectSetting = false);
\ No newline at end of file
diff --git a/Engine/Source/Runtime/Engine/Classes/Engine/MapBuildDataRegistry.h b/Engine/Source/Runtime/Engine/Classes/Engine/MapBuildDataRegistry.h
--- a/Engine/Source/Runtime/Engine/Classes/Engine/MapBuildDataRegistry.h
+++ b/Engine/Source/Runtime/Engine/Classes/Engine/MapBuildDataRegistry.h
@@ -185,18 +185,24 @@ class FReflectionCaptureData
 public:
 	int32 CubemapSize;
 	float AverageBrightness;
-	/** Needed for EncodedHDRCapturedData cooking. */
-	float Brightness;
 
 	TArray<uint8> FullHDRCapturedData;
 	UTextureCube* EncodedCaptureData;
 
+#if WITH_EDITOR
+	/** Whether the brightness is baked in the EncodedHDRCubemap*/
+	bool bBrightnessBakedInEncodedHDRCubemap;
+#endif
+
 	FReflectionCaptureData() :
 		CubemapSize(0),
 		AverageBrightness(0.0f),
-		Brightness(0.0f),
 		EncodedCaptureData(nullptr),
+#if WITH_EDITOR
+		bBrightnessBakedInEncodedHDRCubemap(false),
+#endif
 		bUploadedFinal(false)
+
 	{}
 
 	bool HasBeenUploadedFinal() const
diff --git a/Engine/Source/Runtime/Engine/Private/Components/ReflectionCaptureComponent.cpp b/Engine/Source/Runtime/Engine/Private/Components/ReflectionCaptureComponent.cpp
--- a/Engine/Source/Runtime/Engine/Private/Components/ReflectionCaptureComponent.cpp
+++ b/Engine/Source/Runtime/Engine/Private/Components/ReflectionCaptureComponent.cpp
@@ -505,7 +505,7 @@ static void EdgeWalkSetup( bool ReverseDirection, int32 Edge, int32 MipSize, int
 	}
 }
 
-float GetMaxValueRGBM(const TArray<uint8>& FullHDRData, int32 CubemapSize, float Brightness)
+float GetMaxValueRGBM(const TArray<uint8>& FullHDRData, int32 CubemapSize)
 {
 	const int32 NumMips = FMath::CeilLogTwo(CubemapSize) + 1;
 	// get MaxValue from Mip0
@@ -524,7 +524,7 @@ float GetMaxValueRGBM(const TArray<uint8>& FullHDRData, int32 CubemapSize, float
 			for (int32 x = 0; x < MipSize; x++)
 			{
 				int32 TexelIndex = x + y * MipSize;
-				const FLinearColor LinearColor = FLinearColor(FaceSourceData[TexelIndex]) * Brightness;
+				const FLinearColor LinearColor = FLinearColor(FaceSourceData[TexelIndex]);
 				float MaxValueTexel = FMath::Max(FMath::Max(LinearColor.R, LinearColor.G), FMath::Max(LinearColor.B, DELTA));
 				if (MaxValue < MaxValueTexel)
 				{
@@ -538,7 +538,7 @@ float GetMaxValueRGBM(const TArray<uint8>& FullHDRData, int32 CubemapSize, float
 	return MaxValue;
 }
 
-void GenerateEncodedHDRData(const TArray<uint8>& FullHDRData, int32 CubemapSize, float Brightness, float MaxValueRGBM, TArray<uint8>& OutEncodedHDRData)
+void GenerateEncodedHDRData(const TArray<uint8>& FullHDRData, int32 CubemapSize, float MaxValueRGBM, TArray<uint8>& OutEncodedHDRData)
 {
 	check(FullHDRData.Num() > 0);
 	const int32 NumMips = FMath::CeilLogTwo(CubemapSize) + 1;
@@ -593,7 +593,7 @@ void GenerateEncodedHDRData(const TArray<uint8>& FullHDRData, int32 CubemapSize,
 			for( int32 Corner = 0; Corner < 4; Corner++ )
 			{
 				const FLinearColor LinearColor = AvgCornerColors[ CubeCornerList[Face][Corner] ] / 3.0f;
-				FaceDstData[ CornerTable[Corner] ] = RGBMEncode( LinearColor * Brightness, MaxValueRGBM);
+				FaceDstData[ CornerTable[Corner] ] = RGBMEncode( LinearColor, MaxValueRGBM);
 			}
 		}
 
@@ -634,7 +634,7 @@ void GenerateEncodedHDRData(const TArray<uint8>& FullHDRData, int32 CubemapSize,
 				const FLinearColor EdgeColorB = FLinearColor( FaceSrcDataB[ EdgeTexelB ] );
 				const FLinearColor AvgColor = 0.5f * ( EdgeColorA + EdgeColorB );
 				
-				FaceDstDataA[ EdgeTexelA ] = FaceDstDataB[ EdgeTexelB ] = RGBMEncode( AvgColor * Brightness, MaxValueRGBM);
+				FaceDstDataA[ EdgeTexelA ] = FaceDstDataB[ EdgeTexelB ] = RGBMEncode( AvgColor, MaxValueRGBM);
 			}
 		}
 #endif // MOBILE_AVERAGE_CUBEMAP_EDGES
@@ -648,7 +648,6 @@ void GenerateEncodedHDRData(const TArray<uint8>& FullHDRData, int32 CubemapSize,
 			FColor* FaceDestData = (FColor*)&OutEncodedHDRData[FaceDestIndex];
 
 			// Convert each texel from linear space FP16 to RGBM FColor
-			// Note: Brightness on the capture is baked into the encoded HDR data
 			// Skip edges
 			const int32 SkipEdges = MOBILE_AVERAGE_CUBEMAP_EDGES ? 1 : 0;
 
@@ -658,7 +657,7 @@ void GenerateEncodedHDRData(const TArray<uint8>& FullHDRData, int32 CubemapSize,
 				for (int32 x = SkipEdges; x < MipSize - SkipEdges; x++)
 				{
 					int32 TexelIndex = x + y * MipSize;
-					const FLinearColor LinearColor = FLinearColor(FaceSourceData[TexelIndex]) * Brightness;
+					const FLinearColor LinearColor = FLinearColor(FaceSourceData[TexelIndex]);
 					FaceDestData[ TexelIndex ] = RGBMEncode( LinearColor, MaxValueRGBM);
 				}
 			}
@@ -676,13 +675,13 @@ void GenerateEncodedHDRTextureCube(UMapBuildDataRegistry* Registry, FReflectionC
 	UTextureFactory* TextureFactory = NewObject<UTextureFactory>();
 
 	TextureFactory->CompressionSettings = TC_EncodedReflectionCapture;
-	UTextureCube* TextureCube = TextureFactory->CreateTextureCube(Registry, FName(TextureName), RF_Public);
+	UTextureCube* TextureCube = ReflectionCaptureData.EncodedCaptureData != nullptr ? ReflectionCaptureData.EncodedCaptureData : TextureFactory->CreateTextureCube(Registry, FName(TextureName), RF_Public);
 
 	if (TextureCube)
 	{
 		TArray<uint8> TemporaryEncodedHDRCapturedData;
 
-		GenerateEncodedHDRData(ReflectionCaptureData.FullHDRCapturedData, ReflectionCaptureData.CubemapSize, ReflectionCaptureData.Brightness, MaxValueRGBM, TemporaryEncodedHDRCapturedData);
+		GenerateEncodedHDRData(ReflectionCaptureData.FullHDRCapturedData, ReflectionCaptureData.CubemapSize, MaxValueRGBM, TemporaryEncodedHDRCapturedData);
 		const int32 NumMips = FMath::CeilLogTwo(ReflectionCaptureData.CubemapSize) + 1;
 		TextureCube->Source.Init(
 			ReflectionCaptureData.CubemapSize,
@@ -715,6 +714,7 @@ void GenerateEncodedHDRTextureCube(UMapBuildDataRegistry* Registry, FReflectionC
 		TextureCube->MarkPackageDirty();
 	}
 	ReflectionCaptureData.EncodedCaptureData = TextureCube;
+	ReflectionCaptureData.bBrightnessBakedInEncodedHDRCubemap = false;
 #endif
 }
 
@@ -830,6 +830,7 @@ void UReflectionCaptureComponent::SerializeLegacyData(FArchive& Ar)
 {
 	Ar.UsingCustomVersion(FRenderingObjectVersion::GUID);
 	Ar.UsingCustomVersion(FReflectionCaptureObjectVersion::GUID);
+	Ar.UsingCustomVersion(FUE5ReleaseStreamObjectVersion::GUID);
 
 	if (Ar.CustomVer(FReflectionCaptureObjectVersion::GUID) < FReflectionCaptureObjectVersion::MoveReflectionCaptureDataToMapBuildData)
 	{
@@ -901,7 +902,6 @@ void UReflectionCaptureComponent::SerializeLegacyData(FArchive& Ar)
 					}
 
 					LegacyMapBuildData->AverageBrightness = AverageBrightness;
-					LegacyMapBuildData->Brightness = Brightness;
 
 					FReflectionCaptureMapBuildLegacyData LegacyComponentData;
 					LegacyComponentData.Id = MapBuildDataId;
@@ -1286,7 +1286,7 @@ void FReflectionCaptureProxy::UpdateMobileUniformBuffer()
 		
 	FMobileReflectionCaptureShaderParameters Parameters;
 	//To keep ImageBasedReflectionLighting coherence with PC, use AverageBrightness instead of InvAverageBrightness to calculate the IBL contribution
-	Parameters.Params = FVector4(EncodedHDRAverageBrightness, 0.f, MaxValueRGBM <= 0.0f ? 16.0f: MaxValueRGBM, 0.f);
+	Parameters.Params = FVector4(EncodedHDRAverageBrightness, 0.f, MaxValueRGBM <= 0.0f ? 16.0f: MaxValueRGBM, Brightness);
 	Parameters.Texture = CaptureTexture->TextureRHI;
 	Parameters.TextureSampler = CaptureTexture->SamplerStateRHI;
 
diff --git a/Engine/Source/Runtime/Engine/Private/MapBuildData.cpp b/Engine/Source/Runtime/Engine/Private/MapBuildData.cpp
--- a/Engine/Source/Runtime/Engine/Private/MapBuildData.cpp
+++ b/Engine/Source/Runtime/Engine/Private/MapBuildData.cpp
@@ -299,11 +299,20 @@ FArchive& operator<<(FArchive& Ar, FReflectionCaptureMapBuildData& ReflectionCap
 	Ar << ReflectionCaptureMapBuildData.CubemapSize;
 	Ar << ReflectionCaptureMapBuildData.AverageBrightness;
 
-	if (Ar.CustomVer(FRenderingObjectVersion::GUID) >= FRenderingObjectVersion::StoreReflectionCaptureBrightnessForCooking)
+	float Brightness = 1.0f;
+	if (Ar.CustomVer(FRenderingObjectVersion::GUID) >= FRenderingObjectVersion::StoreReflectionCaptureBrightnessForCooking 
+		&& Ar.CustomVer(FUE5ReleaseStreamObjectVersion::GUID) < FUE5ReleaseStreamObjectVersion::ExcludeBrightnessFromEncodedHDRCubemap)
 	{
-		Ar << ReflectionCaptureMapBuildData.Brightness;
+		Ar << Brightness;
 	}
 
+#if WITH_EDITOR
+	if (Ar.CustomVer(FUE5ReleaseStreamObjectVersion::GUID) < FUE5ReleaseStreamObjectVersion::ExcludeBrightnessFromEncodedHDRCubemap && Brightness != 1.0f)
+	{
+		ReflectionCaptureMapBuildData.bBrightnessBakedInEncodedHDRCubemap = true;
+	}
+#endif
+
 	static FName FullHDR(TEXT("FullHDR"));
 	static FName EncodedHDR(TEXT("EncodedHDR"));
 
@@ -449,6 +458,7 @@ void UMapBuildDataRegistry::Serialize(FArchive& Ar)
 	Ar.UsingCustomVersion(FRenderingObjectVersion::GUID);
 	Ar.UsingCustomVersion(FMobileObjectVersion::GUID);
 	Ar.UsingCustomVersion(FReflectionCaptureObjectVersion::GUID);
+	Ar.UsingCustomVersion(FUE5ReleaseStreamObjectVersion::GUID);
 
 	if (!StripFlags.IsDataStrippedForServer())
 	{
@@ -533,7 +543,7 @@ void UMapBuildDataRegistry::HandleLegacyEncodedCubemapData()
 		for (TMap<FGuid, FReflectionCaptureMapBuildData>::TIterator It(ReflectionCaptureBuildData); It; ++It)
 		{
 			FReflectionCaptureMapBuildData& CaptureBuildData = It.Value();
-			if (CaptureBuildData.EncodedCaptureData == nullptr && CaptureBuildData.FullHDRCapturedData.Num() != 0)
+			if ((CaptureBuildData.EncodedCaptureData == nullptr || CaptureBuildData.bBrightnessBakedInEncodedHDRCubemap) && CaptureBuildData.FullHDRCapturedData.Num() != 0)
 			{
 				FString TextureName = TEXT("DeprecatedTexture");
 				TextureName += LexToString(It.Key());
diff --git a/Engine/Source/Runtime/Engine/Public/SceneManagement.h b/Engine/Source/Runtime/Engine/Public/SceneManagement.h
--- a/Engine/Source/Runtime/Engine/Public/SceneManagement.h
+++ b/Engine/Source/Runtime/Engine/Public/SceneManagement.h
@@ -1103,7 +1103,7 @@ inline bool PrimitiveNeedsDistanceFieldSceneData(bool bTrackAllPrimitives,
 
 
 BEGIN_GLOBAL_SHADER_PARAMETER_STRUCT(FMobileReflectionCaptureShaderParameters,ENGINE_API)
-	SHADER_PARAMETER(FVector4f, Params) // x - inv average brightness, y - sky cubemap max mip, z - Max value for RGBM, w - unused
+	SHADER_PARAMETER(FVector4f, Params) // x - inv average brightness, y - sky cubemap max mip, z - Max value for RGBM, w - brightness of reflection capture
 	SHADER_PARAMETER_TEXTURE(TextureCube, Texture)
 	SHADER_PARAMETER_SAMPLER(SamplerState, TextureSampler)
 END_GLOBAL_SHADER_PARAMETER_STRUCT()
diff --git a/Engine/Source/Runtime/Renderer/Private/ReflectionEnvironmentCapture.cpp b/Engine/Source/Runtime/Renderer/Private/ReflectionEnvironmentCapture.cpp
--- a/Engine/Source/Runtime/Renderer/Private/ReflectionEnvironmentCapture.cpp
+++ b/Engine/Source/Runtime/Renderer/Private/ReflectionEnvironmentCapture.cpp
@@ -1214,9 +1214,6 @@ void FScene::GetReflectionCaptureData(UReflectionCaptureComponent* Component, FR
 
 	// Necessary since the RT is writing to OutDerivedData directly
 	FlushRenderingCommands();
-
-	// Required for cooking of Encoded HDR data
-	OutCaptureData.Brightness = Component->Brightness;
 }
 
 void UploadReflectionCapture_RenderingThread(FScene* Scene, const FReflectionCaptureData* CaptureData, const UReflectionCaptureComponent* CaptureComponent)
