AltSkinWeights - Add to Registry
#fix Added "SkinWeightProfiles" asset registry to SkeletalMesh asset

diff --git a/Engine/Source/Runtime/Engine/Private/SkeletalMesh.cpp b/Engine/Source/Runtime/Engine/Private/SkeletalMesh.cpp
--- a/Engine/Source/Runtime/Engine/Private/SkeletalMesh.cpp
+++ b/Engine/Source/Runtime/Engine/Private/SkeletalMesh.cpp
@@ -3833,6 +3833,7 @@ void USkeletalMesh::GetAssetRegistryTags(TArray<FAssetRegistryTag>& OutTags) con
 	OutTags.Add(FAssetRegistryTag("LODs", FString::FromInt(NumLODs), FAssetRegistryTag::TT_Numerical));
 	OutTags.Add(FAssetRegistryTag("Bones", FString::FromInt(GetRefSkeleton().GetRawBoneNum()), FAssetRegistryTag::TT_Numerical));
 	OutTags.Add(FAssetRegistryTag("MorphTargets", FString::FromInt(GetMorphTargets().Num()), FAssetRegistryTag::TT_Numerical));
+	OutTags.Add(FAssetRegistryTag("SkinWeightProfiles", FString::FromInt(GetSkinWeightProfiles().Num()), FAssetRegistryTag::TT_Numerical));
 
 #if WITH_EDITORONLY_DATA
 	if (GetAssetImportData())
