Focus Interpolation bug fix + cvar to opt out in case it affects mix negatively

- Rather than FocusFactor always interpolating up from 0 to the correct value, add in functionality to direct-set on first update of FocusData

diff --git a/Engine/Source/Runtime/Engine/Private/ActiveSound.cpp b/Engine/Source/Runtime/Engine/Private/ActiveSound.cpp
--- a/Engine/Source/Runtime/Engine/Private/ActiveSound.cpp
+++ b/Engine/Source/Runtime/Engine/Private/ActiveSound.cpp
@@ -40,6 +40,14 @@ FAutoConsoleVariableRef CVarGatherInteriorDataFromIActiveSoundUpdate(
 	TEXT("0: Disabled, 1: Enabled (default)"),
 	ECVF_Default);
 
+static int32 InitializeFocusFactorOnFirstUpdateCVar = 1;
+FAutoConsoleVariableRef CVarInitializeFocusFactorOnFirstUpdateCVar(
+	TEXT("au.FocusData.InitializeFocusFactorOnFirstUpdate"),
+	InitializeFocusFactorOnFirstUpdateCVar,
+	TEXT("When set to 1, focus factor will be initialized on first update to the proper value, instead of interpolating from 0 to the proper value.\n")
+	TEXT("0: Disabled, 1: Enabled (default)"),
+	ECVF_Default);
+
 FTraceDelegate FActiveSound::ActiveSoundTraceDelegate;
 TMap<FTraceHandle, FActiveSound::FAsyncTraceDetails> FActiveSound::TraceToActiveSoundMap;
 
@@ -1496,8 +1504,12 @@ void FActiveSound::UpdateFocusData(float DeltaTime, const FAttenuationListenerDa
 	const FGlobalFocusSettings& FocusSettings = AudioDevice->GetGlobalFocusSettings();
 	const float TargetFocusFactor = AudioDevice->GetFocusFactor(FocusDataToUpdate->Azimuth, *ListenerData.AttenuationSettings);
 
+	// Enabling InitializeFocusVolumeBeforeInterpCVar will fix a bug related to focus factor always needing to interpolate up from 0 to the correct value.
+	// Could affect game mix, so enable with caution.
+	const bool bShouldInterpolateFocusFactor = (!FocusDataToUpdate->bFirstFocusUpdate || !InitializeFocusFactorOnFirstUpdateCVar);
+
 	// User opt-in for focus interpolation
-	if (ListenerData.AttenuationSettings->bEnableFocusInterpolation)
+	if (ListenerData.AttenuationSettings->bEnableFocusInterpolation && bShouldInterpolateFocusFactor)
 	{
 		// Determine which interpolation speed to use (attack/release)
 		float InterpSpeed;
@@ -1518,6 +1530,9 @@ void FActiveSound::UpdateFocusData(float DeltaTime, const FAttenuationListenerDa
 		FocusDataToUpdate->FocusFactor = TargetFocusFactor;
 	}
 
+	// No longer first update
+	FocusDataToUpdate->bFirstFocusUpdate = false;
+
 	// Scale the volume-weighted priority scale value we use for sorting this sound for voice-stealing
 	FocusDataToUpdate->PriorityScale = ListenerData.AttenuationSettings->GetFocusPriorityScale(FocusSettings, FocusDataToUpdate->FocusFactor);
 	FocusDataToUpdate->DistanceScale = ListenerData.AttenuationSettings->GetFocusDistanceScale(FocusSettings, FocusDataToUpdate->FocusFactor);
@@ -1628,6 +1643,9 @@ void FActiveSound::UpdateAttenuation(float DeltaTime, FSoundParseParameters& Par
 		// Feed prior focus factor on update to allow for proper interpolation.
 		FocusDataToApply.FocusFactor = FocusData.FocusFactor;
 
+		// Feed first update flag
+		FocusDataToApply.bFirstFocusUpdate = FocusData.bFirstFocusUpdate;
+
 		// Update azimuth angles prior to updating focus as it uses this in calculating
 		// in and out of focus values.
 		UpdateFocusData(DeltaTime, ListenerData, &FocusDataToApply);
diff --git a/Engine/Source/Runtime/Engine/Public/ActiveSound.h b/Engine/Source/Runtime/Engine/Public/ActiveSound.h
--- a/Engine/Source/Runtime/Engine/Public/ActiveSound.h
+++ b/Engine/Source/Runtime/Engine/Public/ActiveSound.h
@@ -34,35 +34,28 @@ struct FAttenuationListenerData;
 struct FAttenuationFocusData
 {
 	/** Azimuth of the active sound relative to the listener. Used by sound focus. */
-	float Azimuth;
+	float Azimuth = 0.0f;
 
 	/** Absolute azimuth of the active sound relative to the listener. Used for 3d audio calculations. */
-	float AbsoluteAzimuth;
+	float AbsoluteAzimuth = 0.0f;
 
 	/** Value used to allow smooth interpolation in/out of focus */
-	float FocusFactor;
+	float FocusFactor = 1.0f;
 
 	/** Cached calculation of the amount distance is scaled due to focus */
-	float DistanceScale;
+	float DistanceScale = 1.0f;
 
 	/** The amount priority is scaled due to focus */
-	float PriorityScale;
+	float PriorityScale = 1.0f;
 
 	/** Cached highest priority of the parent active sound's wave instances. */
-	float PriorityHighest;
+	float PriorityHighest = 1.0f;
 
 	/** The amount volume is scaled due to focus */
-	float VolumeScale;
+	float VolumeScale = 1.0f;
 
-	FAttenuationFocusData()
-		: Azimuth(0.0f)
-		, AbsoluteAzimuth(0.0f)
-		, FocusFactor(1.0f)
-		, DistanceScale(1.0f)
-		, PriorityScale(1.0f)
-		, VolumeScale(1.0f)
-	{
-	}
+	/** If this is the first update for focus. Handles edge case of starting a sound in-focus or out-of-focus. */
+	bool bFirstFocusUpdate = true;
 };
 
 /**
