Reducing verbosity of "Setting CVAR" logs.
Differentiated the log message coming from scalability settings.

diff --git a/Engine/Source/Runtime/Core/Private/Misc/ConfigCacheIni.cpp b/Engine/Source/Runtime/Core/Private/Misc/ConfigCacheIni.cpp
--- a/Engine/Source/Runtime/Core/Private/Misc/ConfigCacheIni.cpp
+++ b/Engine/Source/Runtime/Core/Private/Misc/ConfigCacheIni.cpp
@@ -5075,7 +5075,7 @@ CORE_API void OnSetCVarFromIniEntry(const TCHAR *IniFile, const TCHAR *Key, cons
 
 		if(bAllowChange)
 		{
-			UE_LOG(LogConfig, Log, TEXT("Setting CVar [[%s:%s]]"), Key, Value);
+			UE_LOG(LogConfig, Verbose, TEXT("Setting CVar [[%s:%s]]"), Key, Value);
 			if (SetBy == ECVF_SetByMask)
 			{
 				CVar->SetWithCurrentPriority(Value);
diff --git a/Engine/Source/Runtime/Engine/Private/Scalability.cpp b/Engine/Source/Runtime/Engine/Private/Scalability.cpp
--- a/Engine/Source/Runtime/Engine/Private/Scalability.cpp
+++ b/Engine/Source/Runtime/Engine/Private/Scalability.cpp
@@ -351,7 +351,7 @@ void ApplyScalabilityGroupFromPlatformIni(const TCHAR* InSectionName, const TCHA
 			}
 
 			// Apply the platform override
-			UE_LOG(LogConfig, Log, TEXT("Setting CVar [[%s:%s]]"), *KeyString, *ValueString);
+			UE_LOG(LogConfig, Verbose, TEXT("Setting scalability CVar [[%s:%s]]"), *KeyString, *ValueString);
 			CVar->Set(*ValueString, ECVF_SetByScalability);
 		}			   
 	};
