World Partition: move HLOD layers specific code out of actor descriptors code.

diff --git a/Engine/Source/Runtime/Engine/Private/WorldPartition/HLOD/HLODLayer.cpp b/Engine/Source/Runtime/Engine/Private/WorldPartition/HLOD/HLODLayer.cpp
--- a/Engine/Source/Runtime/Engine/Private/WorldPartition/HLOD/HLODLayer.cpp
+++ b/Engine/Source/Runtime/Engine/Private/WorldPartition/HLOD/HLODLayer.cpp
@@ -60,7 +60,8 @@ UHLODLayer* UHLODLayer::GetHLODLayer(const FWorldPartitionActorDescView& InActor
 {
 	check(InWorldPartition);
 
-	if (UHLODLayer* HLODLayer = InActorDesc.GetHLODLayer())
+	const FName HLODLayerName = InActorDesc.GetHLODLayer();
+	if (UHLODLayer* HLODLayer = HLODLayerName.IsNone() ? nullptr : Cast<UHLODLayer>(FSoftObjectPath(HLODLayerName).TryLoad()))
 	{
 		return HLODLayer;
 	}
diff --git a/Engine/Source/Runtime/Engine/Private/WorldPartition/WorldPartitionActorDesc.cpp b/Engine/Source/Runtime/Engine/Private/WorldPartition/WorldPartitionActorDesc.cpp
--- a/Engine/Source/Runtime/Engine/Private/WorldPartition/WorldPartitionActorDesc.cpp
+++ b/Engine/Source/Runtime/Engine/Private/WorldPartition/WorldPartitionActorDesc.cpp
@@ -274,11 +274,6 @@ FName FWorldPartitionActorDesc::GetActorLabelOrName() const
 	return GetActorLabel().IsNone() ? GetActorName() : GetActorLabel();
 }
 
-UHLODLayer* FWorldPartitionActorDesc::GetHLODLayer() const
-{
-	return HLODLayer.IsNone() ? nullptr : Cast<UHLODLayer>(FSoftObjectPath(HLODLayer).TryLoad());
-}
-
 bool FWorldPartitionActorDesc::IsLoaded(bool bEvenIfPendingKill) const
 {
 	return ActorPtr.IsValid(bEvenIfPendingKill);
diff --git a/Engine/Source/Runtime/Engine/Private/WorldPartition/WorldPartitionActorDescView.cpp b/Engine/Source/Runtime/Engine/Private/WorldPartition/WorldPartitionActorDescView.cpp
--- a/Engine/Source/Runtime/Engine/Private/WorldPartition/WorldPartitionActorDescView.cpp
+++ b/Engine/Source/Runtime/Engine/Private/WorldPartition/WorldPartitionActorDescView.cpp
@@ -67,7 +67,7 @@ bool FWorldPartitionActorDescView::GetActorIsHLODRelevant() const
 	return ActorDesc->GetActorIsHLODRelevant();
 }
 
-UHLODLayer* FWorldPartitionActorDescView::GetHLODLayer() const
+FName FWorldPartitionActorDescView::GetHLODLayer() const
 {
 	return ActorDesc->GetHLODLayer();
 }
diff --git a/Engine/Source/Runtime/Engine/Public/WorldPartition/WorldPartitionActorDesc.h b/Engine/Source/Runtime/Engine/Public/WorldPartition/WorldPartitionActorDesc.h
--- a/Engine/Source/Runtime/Engine/Public/WorldPartition/WorldPartitionActorDesc.h
+++ b/Engine/Source/Runtime/Engine/Public/WorldPartition/WorldPartitionActorDesc.h
@@ -66,7 +66,7 @@ public:
 	inline bool GetActorIsEditorOnly() const { return bActorIsEditorOnly; }
 	inline bool GetLevelBoundsRelevant() const { return bLevelBoundsRelevant; }
 	inline bool GetActorIsHLODRelevant() const { return bActorIsHLODRelevant; }
-	class UHLODLayer* GetHLODLayer() const;
+	inline FName GetHLODLayer() const { return HLODLayer; }
 	inline const TArray<FName>& GetDataLayers() const { return DataLayers; }
 	inline FName GetActorPackage() const { return ActorPackage; }
 	inline FName GetActorPath() const { return ActorPath; }
diff --git a/Engine/Source/Runtime/Engine/Public/WorldPartition/WorldPartitionActorDescView.h b/Engine/Source/Runtime/Engine/Public/WorldPartition/WorldPartitionActorDescView.h
--- a/Engine/Source/Runtime/Engine/Public/WorldPartition/WorldPartitionActorDescView.h
+++ b/Engine/Source/Runtime/Engine/Public/WorldPartition/WorldPartitionActorDescView.h
@@ -4,7 +4,6 @@
 #include "CoreMinimal.h"
 
 #if WITH_EDITOR
-class UHLODLayer;
 class FWorldPartitionActorDesc;
 class UActorDescContainer;
 class UWorldPartition;
@@ -31,7 +30,7 @@ public:
 	bool GetActorIsEditorOnly() const;
 	bool GetLevelBoundsRelevant() const;
 	bool GetActorIsHLODRelevant() const;
-	UHLODLayer* GetHLODLayer() const;
+	inline FName GetHLODLayer() const;
 	const TArray<FName>& GetDataLayers() const;
 	FName GetActorPackage() const;
 	FName GetActorPath() const;
