World Partition: fix unset container for unsaved actor descriptors.

diff --git a/Engine/Source/Runtime/Engine/Private/WorldPartition/WorldPartitionStreamingGeneration.cpp b/Engine/Source/Runtime/Engine/Private/WorldPartition/WorldPartitionStreamingGeneration.cpp
--- a/Engine/Source/Runtime/Engine/Private/WorldPartition/WorldPartitionStreamingGeneration.cpp
+++ b/Engine/Source/Runtime/Engine/Private/WorldPartition/WorldPartitionStreamingGeneration.cpp
@@ -61,6 +61,7 @@ class FWorldPartitionStreamingGenerator
 					{
 						// Dirty, unsaved actor for PIE
 						FWorldPartitionActorDesc* ActorDesc = ModifiedActorsDescList->AddActor(Actor);
+						ActorDesc->SetContainer(const_cast<UActorDescContainer*>(InContainer));
 						ActorDesc->OnRegister(Actor->GetWorld());
 						OutActorDescViewMap.Emplace(ActorDescIt->GetGuid(), ActorDesc);
 						continue;
@@ -80,6 +81,7 @@ class FWorldPartitionStreamingGenerator
 				if (IsValid(Actor) && Actor->IsPackageExternal() && Actor->IsMainPackageActor() && !Actor->IsEditorOnly() && !InContainer->GetActorDesc(Actor->GetActorGuid()))
 				{
 					FWorldPartitionActorDesc* ActorDesc = ModifiedActorsDescList->AddActor(Actor);
+					ActorDesc->SetContainer(const_cast<UActorDescContainer*>(InContainer));
 					ActorDesc->OnRegister(Actor->GetWorld());
 					OutActorDescViewMap.Emplace(ActorDesc->GetGuid(), ActorDesc);
 				}
