Fixed bad management of unsaved level instance in world partition PIE. Was causing Level instance content to be added twice in the PIE world and caused a crash when exiting PIE.

diff --git a/Engine/Source/Runtime/Engine/Private/WorldPartition/WorldPartitionLevelStreamingDynamic.cpp b/Engine/Source/Runtime/Engine/Private/WorldPartition/WorldPartitionLevelStreamingDynamic.cpp
--- a/Engine/Source/Runtime/Engine/Private/WorldPartition/WorldPartitionLevelStreamingDynamic.cpp
+++ b/Engine/Source/Runtime/Engine/Private/WorldPartition/WorldPartitionLevelStreamingDynamic.cpp
@@ -55,7 +55,7 @@ void UWorldPartitionLevelStreamingDynamic::Initialize(const UWorldPartitionRunti
 	bShouldBeAlwaysLoaded = InCell.IsAlwaysLoaded();
 	StreamingPriority = 0;
 	ChildPackages = InCell.GetPackages();
-	ActorContainer = InCell.ActorContainer;
+	UnsavedActorsContainer = InCell.UnsavedActorsContainer;
 	ActorFolders = InCell.GetActorFolders().Array();
 
 	UWorld* OuterWorld = InCell.GetOuterUWorldPartition()->GetTypedOuter<UWorld>();
@@ -253,13 +253,13 @@ bool UWorldPartitionLevelStreamingDynamic::IssueLoadRequests()
 		bool bNeedDup = false;
 		if (ChildPackage.ContainerID.IsMainContainer())
 		{
-			if (ActorContainer)
+			if (UnsavedActorsContainer)
 			{
 				FString SubObjectName;
 				FString SubObjectContext;
 				if (ChildPackage.LoadedPath.ToString().Split(TEXT("."), &SubObjectContext, &SubObjectName, ESearchCase::IgnoreCase, ESearchDir::FromEnd))
 				{
-					if (AActor* ActorModifiedForPIE = ActorContainer->Actors.FindRef(*SubObjectName))
+					if (AActor* ActorModifiedForPIE = UnsavedActorsContainer->Actors.FindRef(*SubObjectName))
 					{
 						bNeedDup = true;
 					}
@@ -274,10 +274,10 @@ bool UWorldPartitionLevelStreamingDynamic::IssueLoadRequests()
 	}
 
 	// Duplicate unsaved actors
-	if (ActorContainer)
+	if (UnsavedActorsContainer)
 	{
-		FObjectDuplicationParameters Parameters(ActorContainer, RuntimeLevel);
-		Parameters.DestClass = ActorContainer->GetClass();
+		FObjectDuplicationParameters Parameters(UnsavedActorsContainer, RuntimeLevel);
+		Parameters.DestClass = UnsavedActorsContainer->GetClass();
 		Parameters.FlagMask = RF_AllFlags & ~(RF_MarkAsRootSet | RF_MarkAsNative | RF_HasExternalPackage);
 		Parameters.InternalFlagMask = EInternalObjectFlags::AllFlags;
 		Parameters.DuplicateMode = EDuplicateMode::PIE;
diff --git a/Engine/Source/Runtime/Engine/Private/WorldPartition/WorldPartitionRuntimeSpatialHash.cpp b/Engine/Source/Runtime/Engine/Private/WorldPartition/WorldPartitionRuntimeSpatialHash.cpp
--- a/Engine/Source/Runtime/Engine/Private/WorldPartition/WorldPartitionRuntimeSpatialHash.cpp
+++ b/Engine/Source/Runtime/Engine/Private/WorldPartition/WorldPartitionRuntimeSpatialHash.cpp
@@ -870,7 +870,7 @@ bool UWorldPartitionRuntimeSpatialHash::CreateStreamingGrid(const FSpatialHashRu
 
 				UE_LOG(LogWorldPartition, Verbose, TEXT("Cell%s %s Actors = %d Bounds (%s)"), bIsCellAlwaysLoaded ? TEXT(" (AlwaysLoaded)") : TEXT(""), *StreamingCell->GetName(), FilteredActors.Num(), *Bounds.ToString());
 
-				check(!StreamingCell->ActorContainer);
+				check(!StreamingCell->UnsavedActorsContainer);
 				for (const FActorInstance& ActorInstance : FilteredActors)
 				{
 					const FWorldPartitionActorDescView& ActorDescView = ActorInstance.GetActorDescView();
@@ -879,7 +879,7 @@ bool UWorldPartitionRuntimeSpatialHash::CreateStreamingGrid(const FSpatialHashRu
 						if (ModifiedActorDescListForPIE.GetActorDesc(ActorDescView.GetGuid()) != nullptr)
 						{
 							// Create an actor container to make sure duplicated actors will share an outer to properly remap inter-actors references
-							StreamingCell->ActorContainer = NewObject<UActorContainer>(StreamingCell);
+							StreamingCell->UnsavedActorsContainer = NewObject<UActorContainer>(StreamingCell);
 							break;
 						}
 					}
@@ -890,11 +890,11 @@ bool UWorldPartitionRuntimeSpatialHash::CreateStreamingGrid(const FSpatialHashRu
 					const FWorldPartitionActorDescView& ActorDescView = ActorInstance.GetActorDescView();
 					StreamingCell->AddActorToCell(ActorDescView, ActorInstance.ContainerInstance->ID, ActorInstance.ContainerInstance->Transform, ActorInstance.ContainerInstance->Container);
 
-					if (StreamingCell->ActorContainer)
+					if (ActorInstance.ContainerInstance->ID.IsMainContainer() && StreamingCell->UnsavedActorsContainer)
 					{
 						if (AActor* Actor = FindObject<AActor>(nullptr, *ActorDescView.GetActorPath().ToString()))
 						{
-							StreamingCell->ActorContainer->Actors.Add(Actor->GetFName(), Actor);
+							StreamingCell->UnsavedActorsContainer->Actors.Add(Actor->GetFName(), Actor);
 						}
 					}
 					UE_LOG(LogWorldPartition, Verbose, TEXT("  Actor : %s (%s) (Container %s)"), *(ActorDescView.GetActorPath().ToString()), *ActorDescView.GetGuid().ToString(EGuidFormats::UniqueObjectGuid), *ActorInstance.ContainerInstance->ID.ToString());
diff --git a/Engine/Source/Runtime/Engine/Private/WorldPartition/WorldPartitionRuntimeSpatialHashCell.cpp b/Engine/Source/Runtime/Engine/Private/WorldPartition/WorldPartitionRuntimeSpatialHashCell.cpp
--- a/Engine/Source/Runtime/Engine/Private/WorldPartition/WorldPartitionRuntimeSpatialHashCell.cpp
+++ b/Engine/Source/Runtime/Engine/Private/WorldPartition/WorldPartitionRuntimeSpatialHashCell.cpp
@@ -28,15 +28,15 @@ void UWorldPartitionRuntimeSpatialHashCell::PostDuplicate(bool bDuplicateForPIE)
 {
 	Super::PostDuplicate(bDuplicateForPIE);
 
-	if (ActorContainer)
+	if (UnsavedActorsContainer)
 	{
 		// Make sure actor container isn't under PIE World so those template actors will never be considered part of the world.
-		ActorContainer->Rename(nullptr, GetPackage());
+		UnsavedActorsContainer->Rename(nullptr, GetPackage());
 
-		for (auto& ActorPair : ActorContainer->Actors)
+		for (auto& ActorPair : UnsavedActorsContainer->Actors)
 		{
 			// Don't use AActor::Rename here since the actor is not par of the world, it's only a duplication template.
-			ActorPair.Value->UObject::Rename(nullptr, ActorContainer);
+			ActorPair.Value->UObject::Rename(nullptr, UnsavedActorsContainer);
 		}
 	}
 }
diff --git a/Engine/Source/Runtime/Engine/Public/WorldPartition/WorldPartitionLevelStreamingDynamic.h b/Engine/Source/Runtime/Engine/Public/WorldPartition/WorldPartitionLevelStreamingDynamic.h
--- a/Engine/Source/Runtime/Engine/Public/WorldPartition/WorldPartitionLevelStreamingDynamic.h
+++ b/Engine/Source/Runtime/Engine/Public/WorldPartition/WorldPartitionLevelStreamingDynamic.h
@@ -70,7 +70,7 @@ private:
 
 #if WITH_EDITORONLY_DATA
 	UPROPERTY()
-	TObjectPtr<UActorContainer> ActorContainer;
+	TObjectPtr<UActorContainer> UnsavedActorsContainer;
 #endif
 
 	UPROPERTY()
diff --git a/Engine/Source/Runtime/Engine/Public/WorldPartition/WorldPartitionRuntimeSpatialHashCell.h b/Engine/Source/Runtime/Engine/Public/WorldPartition/WorldPartitionRuntimeSpatialHashCell.h
--- a/Engine/Source/Runtime/Engine/Public/WorldPartition/WorldPartitionRuntimeSpatialHashCell.h
+++ b/Engine/Source/Runtime/Engine/Public/WorldPartition/WorldPartitionRuntimeSpatialHashCell.h
@@ -48,6 +48,6 @@ class UWorldPartitionRuntimeSpatialHashCell : public UWorldPartitionRuntimeCell
 
 #if WITH_EDITORONLY_DATA
 	UPROPERTY()
-	TObjectPtr<UActorContainer> ActorContainer;
+	TObjectPtr<UActorContainer> UnsavedActorsContainer;
 #endif
 };
\ No newline at end of file
