Move code around to make static analysis happy

diff --git a/Engine/Source/Runtime/Engine/Private/HLSLTree/HLSLTreeEmit.cpp b/Engine/Source/Runtime/Engine/Private/HLSLTree/HLSLTreeEmit.cpp
--- a/Engine/Source/Runtime/Engine/Private/HLSLTree/HLSLTreeEmit.cpp
+++ b/Engine/Source/Runtime/Engine/Private/HLSLTree/HLSLTreeEmit.cpp
@@ -198,30 +198,31 @@ void FormatArg_ShaderValue(FEmitShaderExpression* ShaderValue, FEmitShaderDepend
 
 int32 InternalFormatString(FStringBuilderBase* OutString, FEmitShaderDependencies& OutDependencies, FStringView Format, TArrayView<const FFormatArgVariant> ArgList, int32 BaseArgIndex)
 {
-	check(OutString || Format.Len() == 0);
-
 	int32 ArgIndex = BaseArgIndex;
-	for (TCHAR Char : Format)
+	if (Format.Len() > 0)
 	{
-		if (Char == TEXT('%'))
+		check(OutString);
+		for (TCHAR Char : Format)
 		{
-			const FFormatArgVariant& Arg = ArgList[ArgIndex++];
-			switch (Arg.Type)
+			if (Char == TEXT('%'))
 			{
-			case EFormatArgType::ShaderValue: FormatArg_ShaderValue(Arg.ShaderValue, OutDependencies, *OutString); break;
-			case EFormatArgType::String: OutString->Append(Arg.String); break;
-			case EFormatArgType::Int: OutString->Appendf(TEXT("%d"), Arg.Int); break;
-			default:
-				checkNoEntry();
-				break;
+				const FFormatArgVariant& Arg = ArgList[ArgIndex++];
+				switch (Arg.Type)
+				{
+				case EFormatArgType::ShaderValue: FormatArg_ShaderValue(Arg.ShaderValue, OutDependencies, *OutString); break;
+				case EFormatArgType::String: OutString->Append(Arg.String); break;
+				case EFormatArgType::Int: OutString->Appendf(TEXT("%d"), Arg.Int); break;
+				default:
+					checkNoEntry();
+					break;
+				}
+			}
+			else
+			{
+				OutString->AppendChar(Char);
 			}
-		}
-		else
-		{
-			OutString->AppendChar(Char);
 		}
 	}
-
 	return ArgIndex;
 }
 
