Add a check() to make static analysis happy

diff --git a/Engine/Source/Runtime/Engine/Private/Materials/MaterialHLSLGenerator.cpp b/Engine/Source/Runtime/Engine/Private/Materials/MaterialHLSLGenerator.cpp
--- a/Engine/Source/Runtime/Engine/Private/Materials/MaterialHLSLGenerator.cpp
+++ b/Engine/Source/Runtime/Engine/Private/Materials/MaterialHLSLGenerator.cpp
@@ -613,6 +613,7 @@ UE::HLSLTree::FExpression* FMaterialHLSLGenerator::GenerateFunctionCall(UE::HLSL
 	else
 	{
 		FFunction* HLSLFunction = FunctionCall->HLSLFunction;
+		check(HLSLFunction);
 		check(HLSLFunction->OutputExpressions.Num() == FunctionOutputs.Num());
 		if (HLSLFunction->OutputExpressions[OutputIndex])
 		{
