Generate an error if MF with control flow is called from legacy material translator

diff --git a/Engine/Source/Runtime/Engine/Private/Materials/MaterialExpressions.cpp b/Engine/Source/Runtime/Engine/Private/Materials/MaterialExpressions.cpp
--- a/Engine/Source/Runtime/Engine/Private/Materials/MaterialExpressions.cpp
+++ b/Engine/Source/Runtime/Engine/Private/Materials/MaterialExpressions.cpp
@@ -14307,6 +14307,11 @@ int32 UMaterialExpressionMaterialFunctionCall::Compile(class FMaterialCompiler*
 		return Compiler->Errorf(TEXT("Missing Material Function"));
 	}
 
+	if (MaterialFunction->IsUsingControlFlow())
+	{
+		return Compiler->Errorf(TEXT("Material Functions with control flow are only supported with new HLSL translator"));
+	}
+
 	// Verify that all function inputs and outputs are in a valid state to be linked into this material for compiling
 	for (int32 i = 0; i < FunctionInputs.Num(); i++)
 	{
