Fixed layer names in TerrainLayerWeightParameters that were incorrectly overridden when the UMaterialInstance asset didn't have the MaterialLayerStacksAreNotParameters version yet. This was leading to landscape materials being detected as invalid and rebuilt on load.

#tests editor

diff --git a/Engine/Source/Runtime/Engine/Private/Materials/MaterialInstance.cpp b/Engine/Source/Runtime/Engine/Private/Materials/MaterialInstance.cpp
--- a/Engine/Source/Runtime/Engine/Private/Materials/MaterialInstance.cpp
+++ b/Engine/Source/Runtime/Engine/Private/Materials/MaterialInstance.cpp
@@ -2303,10 +2303,14 @@ void UMaterialInstance::Serialize(FArchive& Ar)
 		}
 	}
 
-	if (Ar.CustomVer(FUE5ReleaseStreamObjectVersion::GUID) < FUE5ReleaseStreamObjectVersion::MaterialLayerStacksAreNotParameters ||
-		Ar.CustomVer(FFortniteMainBranchObjectVersion::GUID) < FFortniteMainBranchObjectVersion::TerrainLayerWeightsAreNotParameters)
+	if (Ar.CustomVer(FUE5ReleaseStreamObjectVersion::GUID) < FUE5ReleaseStreamObjectVersion::MaterialLayerStacksAreNotParameters)
 	{
-		StaticParameters.UpdateLegacyData();
+		StaticParameters.UpdateLegacyMaterialLayersData();
+	}
+
+	if (Ar.CustomVer(FFortniteMainBranchObjectVersion::GUID) < FFortniteMainBranchObjectVersion::TerrainLayerWeightsAreNotParameters)
+	{
+		StaticParameters.UpdateLegacyTerrainLayerWeightData();
 	}
 #endif // WITH_EDITOR
 
diff --git a/Engine/Source/Runtime/Engine/Private/Materials/MaterialShader.cpp b/Engine/Source/Runtime/Engine/Private/Materials/MaterialShader.cpp
--- a/Engine/Source/Runtime/Engine/Private/Materials/MaterialShader.cpp
+++ b/Engine/Source/Runtime/Engine/Private/Materials/MaterialShader.cpp
@@ -304,7 +304,7 @@ void FStaticParameterSet::SerializeLegacy(FArchive& Ar)
 	}
 }
 
-void FStaticParameterSet::UpdateLegacyData()
+void FStaticParameterSet::UpdateLegacyTerrainLayerWeightData()
 {
 	int32 ParameterIndex = 0;
 	while (ParameterIndex < TerrainLayerWeightParameters.Num())
@@ -321,7 +321,10 @@ void FStaticParameterSet::UpdateLegacyData()
 			TerrainLayerWeightParameters.RemoveAt(ParameterIndex);
 		}
 	}
+}
 
+void FStaticParameterSet::UpdateLegacyMaterialLayersData()
+{
 	PRAGMA_DISABLE_DEPRECATION_WARNINGS
 	if (MaterialLayersParameters_DEPRECATED.Num() > 0)
 	{
diff --git a/Engine/Source/Runtime/Engine/Public/StaticParameterSet.h b/Engine/Source/Runtime/Engine/Public/StaticParameterSet.h
--- a/Engine/Source/Runtime/Engine/Public/StaticParameterSet.h
+++ b/Engine/Source/Runtime/Engine/Public/StaticParameterSet.h
@@ -412,7 +412,8 @@ struct FStaticParameterSet
 
 #if WITH_EDITOR
 	void SerializeLegacy(FArchive& Ar);
-	void UpdateLegacyData();
+	void UpdateLegacyTerrainLayerWeightData();
+	void UpdateLegacyMaterialLayersData();
 #endif // WITH_EDITOR
 
 	/** 
diff --git a/Engine/Source/Runtime/Landscape/Private/LandscapeEdit.cpp b/Engine/Source/Runtime/Landscape/Private/LandscapeEdit.cpp
--- a/Engine/Source/Runtime/Landscape/Private/LandscapeEdit.cpp
+++ b/Engine/Source/Runtime/Landscape/Private/LandscapeEdit.cpp
@@ -353,7 +353,9 @@ UMaterialInstanceConstant* ULandscapeComponent::GetCombinationMaterial(FMaterial
 			{
 				if (Allocation.LayerInfo)
 				{
-					StaticParameters.TerrainLayerWeightParameters.Add(FStaticTerrainLayerWeightParameter(Allocation.GetLayerName(), Allocation.WeightmapTextureIndex, !Allocation.LayerInfo->bNoWeightBlend));
+					FName LayerName = Allocation.GetLayerName();
+					check(LayerName != NAME_None);
+					StaticParameters.TerrainLayerWeightParameters.Add(FStaticTerrainLayerWeightParameter(LayerName, Allocation.WeightmapTextureIndex, !Allocation.LayerInfo->bNoWeightBlend));
 				}
 			}
 
