Set motion vectors on passthrough vertex factory when it is updated through mesh deformer path.
This is posiibly pushing the paradigm where we poke and reinit the vertex factory to its limit. So next step could be to clean that up a bit.

diff --git a/Engine/Source/Runtime/Engine/Private/GPUSkinCache.cpp b/Engine/Source/Runtime/Engine/Private/GPUSkinCache.cpp
--- a/Engine/Source/Runtime/Engine/Private/GPUSkinCache.cpp
+++ b/Engine/Source/Runtime/Engine/Private/GPUSkinCache.cpp
@@ -1647,25 +1647,23 @@ void FGPUSkinCache::Release(FGPUSkinCacheEntry*& SkinCacheEntry)
 }
 
 void FGPUSkinCache::GetShaderBindings(
-	FGPUSkinCacheEntry* Entry, 
+	const FGPUSkinCacheEntry* Entry, 
 	int32 Section,
-	const FShader* Shader, 
+	bool bVerticesInMotion,
 	const FGPUSkinPassthroughVertexFactory* VertexFactory,
-	uint32 BaseVertexIndex, 
 	FShaderResourceParameter GPUSkinCachePositionBuffer,
 	FShaderResourceParameter GPUSkinCachePreviousPositionBuffer,
-	class FMeshDrawSingleShaderBindings& ShaderBindings,
-	FVertexInputStreamArray& VertexStreams,
-	const FSceneView* View)
+	FMeshDrawSingleShaderBindings& ShaderBindings,
+	FVertexInputStreamArray& VertexStreams)
 {
 	INC_DWORD_STAT(STAT_GPUSkinCache_NumSetVertexStreams);
 	check(Entry);
 	check(Entry->IsSectionValid(Section));
 	check(Entry->SkinCache);
 
-	FGPUSkinCacheEntry::FSectionDispatchData& DispatchData = Entry->DispatchData[Section];
+	FGPUSkinCacheEntry::FSectionDispatchData const& DispatchData = Entry->DispatchData[Section];
 
-	//UE_LOG(LogSkinCache, Warning, TEXT("*** SetVertexStreams E %p All %p Sec %d(%p) LOD %d"), Entry, Entry->DispatchData[Section].Allocation, Section, Entry->DispatchData[Section].Section, Entry->LOD);
+	//UE_LOG(LogSkinCache, Warning, TEXT("*** SetVertexStreams E %p Sec %d(%p) LOD %d"), Entry, Section, Entry->DispatchData[Section].Section, Entry->LOD);
 
 	VertexStreams.Add(FVertexInputStream(VertexFactory->GetPositionStreamIndex(), 0, DispatchData.GetPositionRWBuffer()->Buffer.Buffer));
 
@@ -1676,15 +1674,8 @@ void FGPUSkinCache::GetShaderBindings(
 
 	ShaderBindings.Add(GPUSkinCachePositionBuffer, DispatchData.GetPositionRWBuffer()->Buffer.SRV);
 
-	FGPUBaseSkinVertexFactory::FShaderDataType& ShaderData = DispatchData.SourceVertexFactory->GetShaderData();
-	// Bone data is updated whenever animation triggers a dynamic update, animation can skip frames hence the frequency is not necessary every frame.
-	// So check if bone data is updated this frame, if not then the previous frame data is stale and not suitable for motion blur.
-	bool bBoneDataUpdatedThisFrame = (View->Family->FrameNumber == ShaderData.UpdatedFrameNumber);
-	// If world is paused, use current frame bone matrices, so velocity is canceled and skeletal mesh isn't blurred from motion.
-	bool bPrevious = !View->Family->bWorldIsPaused_IncludingSimulatingInEditor && bBoneDataUpdatedThisFrame;
-
 	// If world is paused, use current frame bone matrices, so velocity is canceled and skeletal mesh isn't blurred from motion.
-	ShaderBindings.Add(GPUSkinCachePreviousPositionBuffer, bPrevious ? DispatchData.GetPreviousPositionRWBuffer()->Buffer.SRV : DispatchData.GetPositionRWBuffer()->Buffer.SRV);
+	ShaderBindings.Add(GPUSkinCachePreviousPositionBuffer, bVerticesInMotion ? DispatchData.GetPreviousPositionRWBuffer()->Buffer.SRV : DispatchData.GetPositionRWBuffer()->Buffer.SRV);
 }
 
 void FGPUSkinCache::PrepareUpdateSkinning(FGPUSkinCacheEntry* Entry, int32 Section, uint32 RevisionNumber, TArray<FSkinCacheRWBuffer*>* OverlappedUAVs)
diff --git a/Engine/Source/Runtime/Engine/Private/GPUSkinVertexFactory.cpp b/Engine/Source/Runtime/Engine/Private/GPUSkinVertexFactory.cpp
--- a/Engine/Source/Runtime/Engine/Private/GPUSkinVertexFactory.cpp
+++ b/Engine/Source/Runtime/Engine/Private/GPUSkinVertexFactory.cpp
@@ -778,66 +778,71 @@ public:
 		class FMeshDrawSingleShaderBindings& ShaderBindings,
 		FVertexInputStreamArray& VertexStreams) const
 	{
+		// #dxr_todo do we need this call to the base?
+		FLocalVertexFactoryShaderParametersBase::GetElementShaderBindingsBase(Scene, View, Shader, InputStreamType, FeatureLevel, VertexFactory, BatchElement, nullptr, ShaderBindings, VertexStreams);
+
 		// todo: Add more context into VertexFactoryUserData about whether this is skin cache/mesh deformer/ray tracing etc.
 		// For now it just holds a skin cache pointer which is null if using other mesh deformers.
 		FGPUSkinBatchElementUserData* BatchUserData = (FGPUSkinBatchElementUserData*)BatchElement.VertexFactoryUserData;
+		const bool bUsesSkinCache = BatchUserData != nullptr;
 
 		check(VertexFactory->GetType() == &FGPUSkinPassthroughVertexFactory::StaticType);
-		FGPUSkinPassthroughVertexFactory const* LocalVertexFactory = static_cast<FGPUSkinPassthroughVertexFactory const*>(VertexFactory);
-		
-		FRHIUniformBuffer* VertexFactoryUniformBuffer = nullptr;
-		VertexFactoryUniformBuffer = LocalVertexFactory->GetUniformBuffer();
+		FGPUSkinPassthroughVertexFactory const* PassthroughVertexFactory = static_cast<FGPUSkinPassthroughVertexFactory const*>(VertexFactory);
 
-		// #dxr_todo do we need this call to the base?
-		FLocalVertexFactoryShaderParametersBase::GetElementShaderBindingsBase(Scene, View, Shader, InputStreamType, FeatureLevel, VertexFactory, BatchElement, VertexFactoryUniformBuffer, ShaderBindings, VertexStreams);
+		// Bone data is updated whenever animation triggers a dynamic update, animation can skip frames hence the frequency is not necessary every frame.
+		// So check if bone data is updated this frame, if not then the previous frame data is stale and not suitable for motion blur.
+		bool bBoneDataUpdatedThisFrame = (View->Family->FrameNumber == PassthroughVertexFactory->GetUpdatedFrameNumber());
+		// If world is paused, use current frame bone matrices, so velocity is canceled and skeletal mesh isn't blurred from motion.
+		bool bVerticesInMotion = !View->Family->bWorldIsPaused_IncludingSimulatingInEditor && bBoneDataUpdatedThisFrame;
 
-		const bool bUsesSkinCache = BatchUserData != nullptr;
 		if (bUsesSkinCache)
 		{
-			GetElementShaderBindingsSkinCache(View, Shader, LocalVertexFactory, BatchElement, BatchUserData, ShaderBindings, VertexStreams);
+			GetElementShaderBindingsSkinCache(PassthroughVertexFactory, BatchUserData, bVerticesInMotion, ShaderBindings, VertexStreams);
 		}
 		else
 		{
-			GetElementShaderBindingsMeshDeformer(View, Shader, LocalVertexFactory, BatchElement, ShaderBindings, VertexStreams);
+			GetElementShaderBindingsMeshDeformer(PassthroughVertexFactory, bVerticesInMotion, ShaderBindings, VertexStreams);
 		}
 	}
 
 	void GetElementShaderBindingsSkinCache(
-		FSceneView const* View,
-		FMeshMaterialShader const* Shader,
-		FGPUSkinPassthroughVertexFactory const* VertexFactory, 
-		FMeshBatchElement const& BatchElement,
+		FGPUSkinPassthroughVertexFactory const* PassthroughVertexFactory,
 		FGPUSkinBatchElementUserData* BatchUserData,
-		class FMeshDrawSingleShaderBindings& ShaderBindings,
+		bool bVerticesInMotion,
+		FMeshDrawSingleShaderBindings& ShaderBindings,
 		FVertexInputStreamArray& VertexStreams) const
 	{
 		FGPUSkinCache::GetShaderBindings(
-			BatchUserData->Entry, BatchUserData->Section, 
-			Shader, VertexFactory, BatchElement.MinVertexIndex, 
+			BatchUserData->Entry, BatchUserData->Section, bVerticesInMotion,
+			PassthroughVertexFactory,
 			GPUSkinCachePositionBuffer, GPUSkinCachePreviousPositionBuffer, 
-			ShaderBindings, VertexStreams, View);
+			ShaderBindings, VertexStreams);
 	}
 
 	void GetElementShaderBindingsMeshDeformer(
-		FSceneView const* View,
-		FMeshMaterialShader const* Shader,
-		FGPUSkinPassthroughVertexFactory const* VertexFactory,
-		FMeshBatchElement const& BatchElement,
-		class FMeshDrawSingleShaderBindings& ShaderBindings,
+		FGPUSkinPassthroughVertexFactory const* PassthroughVertexFactory,
+		bool bVerticesInMotion,
+		FMeshDrawSingleShaderBindings& ShaderBindings,
 		FVertexInputStreamArray& VertexStreams) const
 	{
-		if (VertexFactory->PositionRDG.IsValid())
+		if (PassthroughVertexFactory->PositionRDG.IsValid())
 		{
-			VertexStreams.Add(FVertexInputStream(VertexFactory->GetPositionStreamIndex(), 0, VertexFactory->PositionRDG->GetRHI()));
-			ShaderBindings.Add(GPUSkinCachePositionBuffer, VertexFactory->PositionRDG->GetOrCreateSRV(FRHIBufferSRVCreateInfo(PF_R32_FLOAT)));
+			VertexStreams.Add(FVertexInputStream(PassthroughVertexFactory->GetPositionStreamIndex(), 0, PassthroughVertexFactory->PositionRDG->GetRHI()));
+			ShaderBindings.Add(GPUSkinCachePositionBuffer, PassthroughVertexFactory->GetPositionsSRV());
 		}
-		if (VertexFactory->TangentRDG.IsValid() && VertexFactory->GetTangentStreamIndex() > -1)
+		if (PassthroughVertexFactory->TangentRDG.IsValid() && PassthroughVertexFactory->GetTangentStreamIndex() > -1)
 		{
-			VertexStreams.Add(FVertexInputStream(VertexFactory->GetTangentStreamIndex(), 0, VertexFactory->TangentRDG->GetRHI()));
+			VertexStreams.Add(FVertexInputStream(PassthroughVertexFactory->GetTangentStreamIndex(), 0, PassthroughVertexFactory->TangentRDG->GetRHI()));
 		}
 
-		// todo: Need to track previous position buffer for mesh deformers
-		ShaderBindings.Add(GPUSkinCachePreviousPositionBuffer, VertexFactory->GetPositionsSRV());
+		if (bVerticesInMotion && PassthroughVertexFactory->PrevPositionRDG)
+		{
+			ShaderBindings.Add(GPUSkinCachePreviousPositionBuffer, PassthroughVertexFactory->PrevPositionRDG->GetOrCreateSRV(FRHIBufferSRVCreateInfo(PF_R32_FLOAT)));
+		}
+		else
+		{
+			ShaderBindings.Add(GPUSkinCachePreviousPositionBuffer, PassthroughVertexFactory->GetPositionsSRV());
+		}
 	}
 
 private:
@@ -886,6 +891,7 @@ void FGPUSkinPassthroughVertexFactory::ReleaseRHI()
 
 	// When adding anything else to this function be aware of the bypassing code in InternalUpdateVertexDeclaration.
 	PositionRDG = nullptr;
+	PrevPositionRDG = nullptr;
 	TangentRDG = nullptr;
 	ColorRDG = nullptr;
 	PositionVBAlias.ReleaseRHI();
@@ -953,7 +959,6 @@ void FGPUSkinPassthroughVertexFactory::InternalUpdateVertexDeclaration(FGPUBaseS
 	// Find added streams
 	PositionStreamIndex = -1;
 	TangentStreamIndex = -1;
-	ColorStreamIndex = -1;
 
 	for (int32 Index = 0; Index < Streams.Num(); ++Index)
 	{
@@ -968,13 +973,6 @@ void FGPUSkinPassthroughVertexFactory::InternalUpdateVertexDeclaration(FGPUBaseS
 				TangentStreamIndex = Index;
 			}
 		}
-		if (ColorVBAlias.VertexBufferRHI != nullptr)
-		{
-			if (Streams[Index].VertexBuffer->VertexBufferRHI.GetReference() == Data.ColorComponent.VertexBuffer->VertexBufferRHI.GetReference())
-			{
-				ColorStreamIndex = Index;
-			}
-		}
 	}
 	checkf(PositionStreamIndex != -1, TEXT("Unable to find stream for RWBuffer Vertex buffer!"));
 }
@@ -985,6 +983,7 @@ void FGPUSkinPassthroughVertexFactory::InternalUpdateVertexDeclaration(
 	struct FRWBuffer* TangentRWBuffer)
 {
 	PositionRDG = nullptr;
+	PrevPositionRDG = nullptr;
 	TangentRDG = nullptr;
 	ColorRDG = nullptr;
 
@@ -993,6 +992,7 @@ void FGPUSkinPassthroughVertexFactory::InternalUpdateVertexDeclaration(
 	ColorVBAlias.VertexBufferRHI = nullptr;
 
 	PositionSRVAlias = PositionRWBuffer ? PositionRWBuffer->SRV : nullptr;
+	PrevPositionSRVAlias = nullptr;
 	TangentSRVAlias = TangentRWBuffer ? TangentRWBuffer->SRV : nullptr;
 	ColorSRVAlias = nullptr;
 
@@ -1008,6 +1008,7 @@ void FGPUSkinPassthroughVertexFactory::InternalUpdateVertexDeclaration(
 {
 	if (EnumHasAnyFlags(OverrideFlags, EOverrideFlags::Position))
 	{
+		PrevPositionRDG = PositionBuffer.IsValid() ? PositionRDG : nullptr;
 		PositionRDG = PositionBuffer;
 	}
 	if (EnumHasAnyFlags(OverrideFlags, EOverrideFlags::Tangent))
@@ -1024,6 +1025,7 @@ void FGPUSkinPassthroughVertexFactory::InternalUpdateVertexDeclaration(
 	ColorVBAlias.VertexBufferRHI = ColorRDG.IsValid() ? ColorRDG->GetRHI() : nullptr;
 
 	PositionSRVAlias = PositionRDG.IsValid() ? PositionRDG->GetOrCreateSRV(FRHIBufferSRVCreateInfo(PF_R32_FLOAT)) : nullptr;
+	PrevPositionSRVAlias = PrevPositionRDG.IsValid() ? PrevPositionRDG->GetOrCreateSRV(FRHIBufferSRVCreateInfo(PF_R32_FLOAT)) : nullptr;
 	const EPixelFormat TangentsFormat = IsOpenGLPlatform(GMaxRHIShaderPlatform) ? PF_R16G16B16A16_SINT : PF_R16G16B16A16_SNORM;
 	TangentSRVAlias = TangentRDG.IsValid() ? TangentRDG->GetOrCreateSRV(FRHIBufferSRVCreateInfo(TangentsFormat)) : nullptr;
 	ColorSRVAlias = ColorRDG.IsValid() ? ColorRDG->GetOrCreateSRV(FRHIBufferSRVCreateInfo(PF_R8G8B8A8)) : nullptr;
diff --git a/Engine/Source/Runtime/Engine/Public/GPUSkinCache.h b/Engine/Source/Runtime/Engine/Public/GPUSkinCache.h
--- a/Engine/Source/Runtime/Engine/Public/GPUSkinCache.h
+++ b/Engine/Source/Runtime/Engine/Public/GPUSkinCache.h
@@ -175,16 +175,14 @@ public:
 		);
 
 	static void GetShaderBindings(
-		FGPUSkinCacheEntry* Entry,
+		const FGPUSkinCacheEntry* Entry,
 		int32 Section,
-		const FShader* Shader,
+		bool bVerticesInMotion,
 		const FGPUSkinPassthroughVertexFactory* VertexFactory,
-		uint32 BaseVertexIndex,
 		FShaderResourceParameter GPUSkinCachePositionBuffer,
 		FShaderResourceParameter GPUSkinCachePreviousPositionBuffer,
 		class FMeshDrawSingleShaderBindings& ShaderBindings,
-		FVertexInputStreamArray& VertexStreams,
-		const FSceneView* View);
+		FVertexInputStreamArray& VertexStreams);
 
 	static void Release(FGPUSkinCacheEntry*& SkinCacheEntry);
 
diff --git a/Engine/Source/Runtime/Engine/Public/GPUSkinVertexFactory.h b/Engine/Source/Runtime/Engine/Public/GPUSkinVertexFactory.h
--- a/Engine/Source/Runtime/Engine/Public/GPUSkinVertexFactory.h
+++ b/Engine/Source/Runtime/Engine/Public/GPUSkinVertexFactory.h
@@ -467,6 +467,8 @@ public:
 
 	void SetData(const FDataType& InData);
 
+	uint32 GetUpdatedFrameNumber() const { return UpdatedFrameNumber; }
+
 	//TODO should be supported
 	bool SupportsPositionOnlyStream() const override { return false; }
 	bool SupportsPositionAndNormalOnlyStream() const override { return false; }
@@ -475,7 +477,7 @@ public:
 	{
 		PositionStreamIndex = -1;
 		TangentStreamIndex = -1;
-		ColorStreamIndex = -1;
+
 		PositionVBAlias.ReleaseRHI();
 		TangentVBAlias.ReleaseRHI();
 		ColorVBAlias.ReleaseRHI();
@@ -489,6 +491,7 @@ public:
 		struct FRWBuffer* PositionRWBuffer, 
 		struct FRWBuffer* TangentRWBuffer)
 	{
+		UpdatedFrameNumber = SourceVertexFactory->GetShaderData().UpdatedFrameNumber;
 		if (PositionStreamIndex == -1)
 		{
 			InternalUpdateVertexDeclaration(SourceVertexFactory, PositionRWBuffer, TangentRWBuffer);
@@ -513,6 +516,7 @@ public:
 		TRefCountPtr<FRDGPooledBuffer> const& TangentBuffer,
 		TRefCountPtr<FRDGPooledBuffer> const& ColorBuffer)
 	{
+		UpdatedFrameNumber = SourceVertexFactory->GetShaderData().UpdatedFrameNumber;
 		if (PositionStreamIndex == -1)
 		{
 			InternalUpdateVertexDeclaration(OverrideFlags, SourceVertexFactory, PositionBuffer, TangentBuffer, ColorBuffer);
@@ -524,6 +528,7 @@ protected:
 
 	// Reference holders for RDG buffers
 	TRefCountPtr<FRDGPooledBuffer> PositionRDG;
+	TRefCountPtr<FRDGPooledBuffer> PrevPositionRDG;
 	TRefCountPtr<FRDGPooledBuffer> TangentRDG;
 	TRefCountPtr<FRDGPooledBuffer> ColorRDG;
 	// Vertex buffer required for creating the Vertex Declaration
@@ -532,11 +537,14 @@ protected:
 	FVertexBuffer ColorVBAlias;
 	// SRVs required for binding
 	FRHIShaderResourceView* PositionSRVAlias;
+	FRHIShaderResourceView* PrevPositionSRVAlias;
 	FRHIShaderResourceView* TangentSRVAlias;
 	FRHIShaderResourceView* ColorSRVAlias;
 	// Cached stream indices
 	int32 PositionStreamIndex;
 	int32 TangentStreamIndex;
+	// Frame number of the bone data that is last updated
+	uint32 UpdatedFrameNumber = 0;
 
 	void InternalUpdateVertexDeclaration(
 		FGPUBaseSkinVertexFactory* SourceVertexFactory);
