Fixed crash when calling UInheritableComponentHandler::FindBestArchetype in a cooked editor

Blueprint assets don't exist in a cooked editor, so the CastChecked was asserting

diff --git a/Engine/Source/Runtime/Engine/Private/InheritableComponentHandler.cpp b/Engine/Source/Runtime/Engine/Private/InheritableComponentHandler.cpp
--- a/Engine/Source/Runtime/Engine/Private/InheritableComponentHandler.cpp
+++ b/Engine/Source/Runtime/Engine/Private/InheritableComponentHandler.cpp
@@ -385,9 +385,18 @@ UActorComponent* UInheritableComponentHandler::FindBestArchetype(const FComponen
 	UBlueprintGeneratedClass* ActualBPGC = Cast<UBlueprintGeneratedClass>(GetOuter());
 	if (ActualBPGC && Key.GetComponentOwner() && (ActualBPGC != Key.GetComponentOwner()))
 	{
-		// During reparenting the outer's Class isn't always the Blueprint's class when the ICH is updating so reference
-		// the Blueprint's ParentClass instead
-		ActualBPGC = Cast<UBlueprintGeneratedClass>(CastChecked<UBlueprint>(ActualBPGC->ClassGeneratedBy)->ParentClass);
+		if (UBlueprint* ClassGeneratedByBP = Cast<UBlueprint>(ActualBPGC->ClassGeneratedBy))
+		{
+			// During reparenting the outer's Class isn't always the Blueprint's class when the ICH is updating so reference
+			// the Blueprint's ParentClass instead
+			ActualBPGC = Cast<UBlueprintGeneratedClass>(ClassGeneratedByBP->ParentClass);
+		}
+		else
+		{
+			// Blueprint assets aren't available in cooked editors (only the BPGC), so just use the super class directly
+			// since we know it will be up-to-date
+			ActualBPGC = Cast<UBlueprintGeneratedClass>(ActualBPGC->GetSuperClass());
+		}
 
 		while (!ClosestArchetype && ActualBPGC)
 		{
