Fix renderer issues when rendering multiple ISR passes.

- New functions are added to get the Primary View or all Secondary Views and takes advantage of that to replace hardcoded indices.
- Added support for retrieving a monoscopic projection/pose from IStereoRendering when a single frustum is needed for culling purposes.

diff --git a/Engine/Plugins/FX/Niagara/Source/Niagara/Private/NiagaraRendererMeshes.cpp b/Engine/Plugins/FX/Niagara/Source/Niagara/Private/NiagaraRendererMeshes.cpp
--- a/Engine/Plugins/FX/Niagara/Source/Niagara/Private/NiagaraRendererMeshes.cpp
+++ b/Engine/Plugins/FX/Niagara/Source/Niagara/Private/NiagaraRendererMeshes.cpp
@@ -19,10 +19,6 @@
 #include "NiagaraCullProxyComponent.h"
 #include "IXRTrackingSystem.h"
 
-#ifdef HMD_MODULE_INCLUDED
-	#include "IXRTrackingSystem.h"
-#endif
-
 DECLARE_CYCLE_STAT(TEXT("Generate Mesh Vertex Data [GT]"), STAT_NiagaraGenMeshVertexData, STATGROUP_Niagara);
 DECLARE_CYCLE_STAT(TEXT("Render Meshes [RT]"), STAT_NiagaraRenderMeshes, STATGROUP_Niagara);
 DECLARE_CYCLE_STAT(TEXT("Render Meshes - Allocate GPU Data [RT]"), STAT_NiagaraRenderMeshes_AllocateGPUData, STATGROUP_Niagara);
@@ -497,23 +493,10 @@ void FNiagaraRendererMeshes::InitializeSortInfo(const FParticleMeshRenderData& P
 	const FViewMatrices& ViewMatrices = GetViewMatrices(View, OutSortInfo.ViewOrigin);
 	OutSortInfo.ViewDirection = ViewMatrices.GetViewMatrix().GetColumn(2);
 
-#ifdef HMD_MODULE_INCLUDED
-	if (View.StereoPass != EStereoscopicPass::eSSP_FULL && GEngine->XRSystem.IsValid() && (GEngine->XRSystem->GetHMDDevice() != nullptr))
-#else
-	if (View.StereoPass != EStereoscopicPass::eSSP_FULL && AllViewsInFamily.Num() > 1)
-#endif	
+	if (bIsInstancedStereo)
 	{
 		// For VR, do distance culling and sorting from a central eye position to prevent differences between views
-		const uint32 PairedViewIdx = (ViewIndex & 1) ? (ViewIndex - 1) : (ViewIndex + 1);
-		if (AllViewsInFamily.IsValidIndex(PairedViewIdx))
-		{
-			const FSceneView* PairedView = AllViewsInFamily[PairedViewIdx];
-			check(PairedView);
-
-			FVector PairedViewOrigin;
-			GetViewMatrices(*PairedView, PairedViewOrigin);
-			OutSortInfo.ViewOrigin = 0.5f * (OutSortInfo.ViewOrigin + PairedViewOrigin);
-		}
+		OutSortInfo.ViewOrigin = View.CullingOrigin;
 	}
 
 	if (bEnableFrustumCulling)
@@ -534,31 +517,14 @@ void FNiagaraRendererMeshes::InitializeSortInfo(const FParticleMeshRenderData& P
 		}
 		else
 		{
-			OutSortInfo.CullPlanes.SetNumZeroed(6);
-
-			// Gather the culling planes from the view projection matrix
-			const FMatrix& ViewProj = ViewMatrices.GetViewProjectionMatrix();
-			ViewProj.GetFrustumNearPlane(OutSortInfo.CullPlanes[0]);
-			ViewProj.GetFrustumFarPlane(OutSortInfo.CullPlanes[1]);
-			ViewProj.GetFrustumTopPlane(OutSortInfo.CullPlanes[2]);
-			ViewProj.GetFrustumBottomPlane(OutSortInfo.CullPlanes[3]);
-
-			ViewProj.GetFrustumLeftPlane(OutSortInfo.CullPlanes[4]);
-			if (bIsInstancedStereo && AllViewsInFamily.IsValidIndex(ViewIndex + 1))
+			if (bIsInstancedStereo)
 			{
 				// For Instanced Stereo, cull using an extended frustum that encompasses both eyes
-				ensure(View.StereoPass == EStereoscopicPass::eSSP_PRIMARY); // Sanity check that this is a primary view
-				const FSceneView* RightEyeView = AllViewsInFamily[ViewIndex + 1];
-				check(RightEyeView);
-				ensure(RightEyeView->StereoPass == EStereoscopicPass::eSSP_SECONDARY);
-				FVector RightEyePos;
-				FPlane CullPlane;
-				GetViewMatrices(*RightEyeView, RightEyePos).GetViewProjectionMatrix().GetFrustumRightPlane(CullPlane);
-				OutSortInfo.CullPlanes[5] = CullPlane;	// LWC_TODO: Perf pessimization
+				OutSortInfo.CullPlanes = View.CullingFrustum.Planes;
 			}
 			else
 			{
-				ViewProj.GetFrustumRightPlane(OutSortInfo.CullPlanes[5]);
+				OutSortInfo.CullPlanes = View.ViewFrustum.Planes;
 			}
 		}
 	}
diff --git a/Engine/Plugins/Runtime/Oculus/OculusVR/Source/OculusHMD/Private/OculusHMD.cpp b/Engine/Plugins/Runtime/Oculus/OculusVR/Source/OculusHMD/Private/OculusHMD.cpp
--- a/Engine/Plugins/Runtime/Oculus/OculusVR/Source/OculusHMD/Private/OculusHMD.cpp
+++ b/Engine/Plugins/Runtime/Oculus/OculusVR/Source/OculusHMD/Private/OculusHMD.cpp
@@ -273,6 +273,9 @@ namespace OculusHMD
 		case EStereoscopicEye::eSSE_RIGHT_EYE:
 			Node = ovrpNode_EyeRight;
 			break;
+		case EStereoscopicEye::eSSE_MONOSCOPIC:
+			Node = ovrpNode_EyeCenter;
+			break;
 		default:
 			return false;
 		}
@@ -1345,7 +1348,8 @@ namespace OculusHMD
 
 		check(IsStereoEnabled());
 
-		FMatrix proj = ToFMatrix(Settings->EyeProjectionMatrices[ViewIndex]);
+		FMatrix proj = (ViewIndex == EStereoscopicEye::eSSE_MONOSCOPIC) ?
+			ToFMatrix(Settings->MonoProjectionMatrix) : ToFMatrix(Settings->EyeProjectionMatrices[ViewIndex]);
 
 		// correct far and near planes for reversed-Z projection matrix
 		const float WorldScale = GetWorldToMetersScale() * (1.0 / 100.0f); // physical scale is 100 UUs/meter
@@ -2623,9 +2627,11 @@ namespace OculusHMD
 			// Update projection matrices
 			ovrpFrustum2f frustumLeft = { 0.001f, 1000.0f, EyeLayerDesc.Fov[0] };
 			ovrpFrustum2f frustumRight = { 0.001f, 1000.0f, EyeLayerDesc.Fov[1] };
+			ovrpFrustum2f frustumCenter = { 0.001f, 1000.0f,{ EyeLayerDesc.Fov[0].UpTan, EyeLayerDesc.Fov[0].DownTan, EyeLayerDesc.Fov[0].LeftTan, EyeLayerDesc.Fov[1].RightTan } };
 
 			Settings->EyeProjectionMatrices[0] = ovrpMatrix4f_Projection(frustumLeft, true);
 			Settings->EyeProjectionMatrices[1] = ovrpMatrix4f_Projection(frustumRight, true);
+			Settings->MonoProjectionMatrix = ovrpMatrix4f_Projection(frustumCenter, true);
 
 			if (Frame.IsValid())
 			{
diff --git a/Engine/Plugins/Runtime/Oculus/OculusVR/Source/OculusHMD/Private/OculusHMD_Settings.h b/Engine/Plugins/Runtime/Oculus/OculusVR/Source/OculusHMD/Private/OculusHMD_Settings.h
--- a/Engine/Plugins/Runtime/Oculus/OculusVR/Source/OculusHMD/Private/OculusHMD_Settings.h
+++ b/Engine/Plugins/Runtime/Oculus/OculusVR/Source/OculusHMD/Private/OculusHMD_Settings.h
@@ -81,7 +81,8 @@ public:
 	/** Viewports for each eye, without DynamicResolution scaling applied */
 	FIntRect EyeUnscaledRenderViewport[2];
 
-	ovrpMatrix4f EyeProjectionMatrices[2]; // 0 - left, 1 - right, same as Views
+	ovrpMatrix4f EyeProjectionMatrices[2]; // 0 - left, 1 - right same as Views
+	ovrpMatrix4f MonoProjectionMatrix;
 
 	FIntPoint RenderTargetSize;
 	float PixelDensity;
diff --git a/Engine/Plugins/Runtime/OpenXR/Source/OpenXRHMD/Private/OpenXRHMD.cpp b/Engine/Plugins/Runtime/OpenXR/Source/OpenXRHMD/Private/OpenXRHMD.cpp
--- a/Engine/Plugins/Runtime/OpenXR/Source/OpenXRHMD/Private/OpenXRHMD.cpp
+++ b/Engine/Plugins/Runtime/OpenXR/Source/OpenXRHMD/Private/OpenXRHMD.cpp
@@ -691,6 +691,11 @@ EStereoscopicPass FOpenXRHMD::GetViewPassForIndex(bool bStereoRequested, int32 V
 		// Views provided by a plugin should be considered a new primary pass
 		return EStereoscopicPass::eSSP_PRIMARY;
 	}
+
+	if (SelectedViewConfigurationType == XR_VIEW_CONFIGURATION_TYPE_PRIMARY_QUAD_VARJO)
+	{
+		return ViewIndex % 2 == 0 ? EStereoscopicPass::eSSP_PRIMARY : EStereoscopicPass::eSSP_SECONDARY;
+	}
 	return ViewIndex == EStereoscopicEye::eSSE_LEFT_EYE ? EStereoscopicPass::eSSP_PRIMARY : EStereoscopicPass::eSSP_SECONDARY;
 }
 
@@ -735,15 +740,32 @@ bool FOpenXRHMD::GetRelativeEyePose(int32 InDeviceId, int32 InViewIndex, FQuat&
 FMatrix FOpenXRHMD::GetStereoProjectionMatrix(const int32 ViewIndex) const
 {
 	const FPipelinedFrameState& FrameState = GetPipelinedFrameStateForThread();
-	XrFovf Fov = (ViewIndex < FrameState.Views.Num()) ? FrameState.Views[ViewIndex].fov 
-		: XrFovf{ -PI / 4.0f, PI / 4.0f, PI / 4.0f, -PI / 4.0f };
-	float ZNear = GNearClippingPlane;
+
+	XrFovf Fov = {};
+	if (ViewIndex == eSSE_MONOSCOPIC)
+	{
+		// The monoscopic projection matrix uses the combined field-of-view of both eyes
+		for (int32 Index = 0; Index < FrameState.Views.Num(); Index++)
+		{
+			const XrFovf& ViewFov = FrameState.Views[Index].fov;
+			Fov.angleUp = FMath::Max(Fov.angleUp, ViewFov.angleUp);
+			Fov.angleDown = FMath::Min(Fov.angleDown, ViewFov.angleDown);
+			Fov.angleLeft = FMath::Min(Fov.angleLeft, ViewFov.angleLeft);
+			Fov.angleRight = FMath::Max(Fov.angleRight, ViewFov.angleRight);
+		}
+	}
+	else
+	{
+		Fov = (ViewIndex < FrameState.Views.Num()) ? FrameState.Views[ViewIndex].fov
+			: XrFovf{ -PI / 4.0f, PI / 4.0f, PI / 4.0f, -PI / 4.0f };
+	}
 
 	Fov.angleUp = tan(Fov.angleUp);
 	Fov.angleDown = tan(Fov.angleDown);
 	Fov.angleLeft = tan(Fov.angleLeft);
 	Fov.angleRight = tan(Fov.angleRight);
 
+	float ZNear = GNearClippingPlane;
 	float SumRL = (Fov.angleRight + Fov.angleLeft);
 	float SumTB = (Fov.angleUp + Fov.angleDown);
 	float InvRL = (1.0f / (Fov.angleRight - Fov.angleLeft));
diff --git a/Engine/Plugins/Runtime/Steam/SteamVR/Source/SteamVR/Private/SteamVRHMD.cpp b/Engine/Plugins/Runtime/Steam/SteamVR/Source/SteamVR/Private/SteamVRHMD.cpp
--- a/Engine/Plugins/Runtime/Steam/SteamVR/Source/SteamVR/Private/SteamVRHMD.cpp
+++ b/Engine/Plugins/Runtime/Steam/SteamVR/Source/SteamVR/Private/SteamVRHMD.cpp
@@ -628,8 +628,16 @@ bool FSteamVRHMD::GetHMDMonitorInfo(MonitorInfo& MonitorDesc)
 
 void FSteamVRHMD::GetFieldOfView(float& OutHFOVInDegrees, float& OutVFOVInDegrees) const
 {
-	OutHFOVInDegrees = 0.0f;
-	OutVFOVInDegrees = 0.0f;
+	float HFOV, VFOV, Left, Right, Top, Bottom;
+	VRSystem->GetProjectionRaw(vr::Eye_Left, &Right, &Left, &Top, &Bottom);
+	HFOV = FMath::Atan(-Left);
+	VFOV = FMath::Atan(Top - Bottom);
+	VRSystem->GetProjectionRaw(vr::Eye_Right, &Right, &Left, &Top, &Bottom);
+	HFOV += FMath::Atan(Right);
+	VFOV = FMath::Max(VFOV, FMath::Atan(Top - Bottom));
+
+	OutHFOVInDegrees = FMath::RadiansToDegrees(HFOV);
+	OutVFOVInDegrees = FMath::RadiansToDegrees(VFOV);
 }
 
 bool FSteamVRHMD::DoesSupportPositionalTracking() const
@@ -1470,7 +1478,7 @@ FMatrix FSteamVRHMD::GetStereoProjectionMatrix(const int32 StereoViewIndex) cons
 {
 	check(IsStereoEnabled() || IsHeadTrackingEnforced());
 
-	vr::Hmd_Eye HmdEye = (StereoViewIndex == EStereoscopicEye::eSSE_LEFT_EYE) ? vr::Eye_Left : vr::Eye_Right;
+	vr::Hmd_Eye HmdEye = (StereoViewIndex == EStereoscopicEye::eSSE_RIGHT_EYE) ? vr::Eye_Right : vr::Eye_Left;
 	float Left, Right, Top, Bottom;
 
 	VRSystem->GetProjectionRaw(HmdEye, &Right, &Left, &Top, &Bottom);
@@ -1479,6 +1487,13 @@ FMatrix FSteamVRHMD::GetStereoProjectionMatrix(const int32 StereoViewIndex) cons
 	Right *= -1.0f;
 	Left *= -1.0f;
 
+	if (StereoViewIndex == EStereoscopicEye::eSSE_MONOSCOPIC)
+	{
+		float Dummy;
+		VRSystem->GetProjectionRaw(vr::Eye_Right, &Right, &Dummy, &Dummy, &Dummy);
+		Right *= -1.0f;
+	}
+
 	float ZNear = GNearClippingPlane;
 
 	float SumRL = (Right + Left);
diff --git a/Engine/Plugins/Runtime/WindowsMixedReality/Source/WindowsMixedRealityHMD/Private/WindowsMixedRealityHMD.cpp b/Engine/Plugins/Runtime/WindowsMixedReality/Source/WindowsMixedRealityHMD/Private/WindowsMixedRealityHMD.cpp
--- a/Engine/Plugins/Runtime/WindowsMixedReality/Source/WindowsMixedRealityHMD/Private/WindowsMixedRealityHMD.cpp
+++ b/Engine/Plugins/Runtime/WindowsMixedReality/Source/WindowsMixedRealityHMD/Private/WindowsMixedRealityHMD.cpp
@@ -1390,7 +1390,8 @@ namespace WindowsMixedReality
 			}
 			else
 			{
-				return FMatrix::Identity;
+				// FIXME: This will result in the left eye matrix being used for monoscopic projections
+				return Frame_RenderThread.ProjectionMatrixL;
 			}
 		}
 		else
@@ -1407,11 +1408,12 @@ namespace WindowsMixedReality
 			}
 			else if (ViewIndex == eSSE_THIRD_CAMERA_EYE)
 			{
-				return Frame_RenderThread.ProjectionMatrixThirdCamera;
+				return Frame_GameThread.ProjectionMatrixThirdCamera;
 			}
 			else
 			{
-				return FMatrix::Identity;
+				// FIXME: This will result in the left eye matrix being used for monoscopic projections
+				return Frame_GameThread.ProjectionMatrixL;
 			}
 		}
 	}
diff --git a/Engine/Source/Runtime/Engine/Private/HierarchicalInstancedStaticMesh.cpp b/Engine/Source/Runtime/Engine/Private/HierarchicalInstancedStaticMesh.cpp
--- a/Engine/Source/Runtime/Engine/Private/HierarchicalInstancedStaticMesh.cpp
+++ b/Engine/Source/Runtime/Engine/Private/HierarchicalInstancedStaticMesh.cpp
@@ -1613,34 +1613,28 @@ void FHierarchicalStaticMeshSceneProxy::GetDynamicMeshElements(const TArray<cons
 
 					// Instanced stereo needs to use the right plane from the right eye when constructing the frustum bounds to cull against.
 					// Otherwise we'll cull objects visible in the right eye, but not the left.
-					if ((Views[0]->IsInstancedStereoPass() || Views[0]->bIsMobileMultiViewEnabled) && IStereoRendering::IsStereoEyeView(*Views[0]) && ViewIndex == 0)
+					if ((View->IsInstancedStereoPass() || View->bIsMobileMultiViewEnabled) && IStereoRendering::IsStereoEyeView(*View) && GEngine->StereoRenderingDevice.IsValid())
 					{
-						check(Views.Num() == 2);
-
-						const FMatrix LeftEyeLocalViewProjForCulling  = GetLocalToWorld() * Views[0]->ViewMatrices.GetViewProjectionMatrix();
-						const FMatrix RightEyeLocalViewProjForCulling = GetLocalToWorld() * Views[1]->ViewMatrices.GetViewProjectionMatrix();
-
-						FConvexVolume LeftEyeBounds, RightEyeBounds;
-						GetViewFrustumBounds(LeftEyeBounds, LeftEyeLocalViewProjForCulling, bViewFrustumUsesNearPlane, bViewFrustumUsesFarPlane);
-						GetViewFrustumBounds(RightEyeBounds, RightEyeLocalViewProjForCulling, bViewFrustumUsesNearPlane, bViewFrustumUsesFarPlane);
+						// TODO: Stereo culling frustum needs to use the culling origin instead of the view origin.
+						InstanceParams.ViewFrustumLocal = View->CullingFrustum;
+						for (FPlane& Plane : InstanceParams.ViewFrustumLocal.Planes)
+						{
+							Plane = Plane.TransformBy(WorldToLocal);
+						}
+						InstanceParams.ViewFrustumLocal.Init();
 
 						// Invalid bounds retrieved, so skip render of this frame :
-						if (LeftEyeBounds.Planes.Num() != 4 || RightEyeBounds.Planes.Num() != 4)
+						if (bIsPerspectiveProjection && (InstanceParams.ViewFrustumLocal.Planes.Num() != 4))
 						{
-							ensureMsgf(false, TEXT("Invalid frustum, skipping render of HISM : culling view projection matrix:\n- Left eye: %s\n-Right eye: %s"), *LeftEyeLocalViewProjForCulling.ToString(), *RightEyeLocalViewProjForCulling.ToString());
+							// Report the error as a warning (instead of an ensure or a check) as the problem can come from improper user data (invalid transform or view-proj matrix) : 
+							ensureMsgf(false, TEXT("Invalid frustum, skipping render of HISM"));
 							continue;
 						}
-						
-						InstanceParams.ViewFrustumLocal.Planes.Empty(4);
-						InstanceParams.ViewFrustumLocal.Planes.Add(LeftEyeBounds.Planes[0]);
-						InstanceParams.ViewFrustumLocal.Planes.Add(RightEyeBounds.Planes[1]);
-						InstanceParams.ViewFrustumLocal.Planes.Add(LeftEyeBounds.Planes[2]);
-						InstanceParams.ViewFrustumLocal.Planes.Add(LeftEyeBounds.Planes[3]);
-						InstanceParams.ViewFrustumLocal.Init();
 					}
 					else
 					{
-						const FMatrix LocalViewProjForCulling = GetLocalToWorld() * View->ViewMatrices.GetViewProjectionMatrix();
+						FMatrix LocalViewProjForCulling = GetLocalToWorld() * View->ViewMatrices.GetViewProjectionMatrix();
+
 						GetViewFrustumBounds(InstanceParams.ViewFrustumLocal, LocalViewProjForCulling, bViewFrustumUsesNearPlane, bViewFrustumUsesFarPlane);
 
 						// Invalid bounds retrieved, so skip render of this frame :
diff --git a/Engine/Source/Runtime/Engine/Private/LocalPlayer.cpp b/Engine/Source/Runtime/Engine/Private/LocalPlayer.cpp
--- a/Engine/Source/Runtime/Engine/Private/LocalPlayer.cpp
+++ b/Engine/Source/Runtime/Engine/Private/LocalPlayer.cpp
@@ -826,8 +826,8 @@ FSceneView* ULocalPlayer::CalcSceneView( class FSceneViewFamily* ViewFamily,
 	// but it makes GetProjectionData better
 	FMinimalViewInfo ViewInfo;
 	GetViewPoint(ViewInfo);
-	OutViewLocation = ViewInfo.Location;
-	OutViewRotation = ViewInfo.Rotation;
+	ViewInitOptions.ViewLocation = ViewInfo.Location;
+	ViewInitOptions.ViewRotation = ViewInfo.Rotation;
 	ViewInitOptions.bUseFieldOfViewForLOD = ViewInfo.bUseFieldOfViewForLOD;
 	ViewInitOptions.FOV = ViewInfo.FOV;
 	ViewInitOptions.DesiredFOV = ViewInfo.DesiredFOV;
@@ -843,7 +843,7 @@ FSceneView* ULocalPlayer::CalcSceneView( class FSceneViewFamily* ViewFamily,
 	else
 	{
 		QUICK_SCOPE_CYCLE_COUNTER(STAT_BuildHiddenComponentList);
-		PlayerController->BuildHiddenComponentList(OutViewLocation, /*out*/ ViewInitOptions.HiddenPrimitives);
+		PlayerController->BuildHiddenComponentList(ViewInfo.Location, /*out*/ ViewInitOptions.HiddenPrimitives);
 	}
 
 	//@TODO: SPLITSCREEN: This call will have an issue with splitscreen, as the show flags are shared across the view family
@@ -851,15 +851,15 @@ FSceneView* ULocalPlayer::CalcSceneView( class FSceneViewFamily* ViewFamily,
 
 	FSceneView* const View = new FSceneView(ViewInitOptions);
 
-	View->ViewLocation = OutViewLocation;
-	View->ViewRotation = OutViewRotation;
+	OutViewLocation = View->ViewLocation;
+	OutViewRotation = View->ViewRotation;
 	// Pass on the previous view transform from the view info (probably provided by the camera if set)
 	View->PreviousViewTransform = ViewInfo.PreviousViewTransform;
 
 	ViewFamily->Views.Add(View);
 
 	{
-		View->StartFinalPostprocessSettings(OutViewLocation);
+		View->StartFinalPostprocessSettings(ViewInfo.Location);
 
 		// CameraAnim override
 		if (PlayerController->PlayerCameraManager)
diff --git a/Engine/Source/Runtime/Engine/Private/SceneView.cpp b/Engine/Source/Runtime/Engine/Private/SceneView.cpp
--- a/Engine/Source/Runtime/Engine/Private/SceneView.cpp
+++ b/Engine/Source/Runtime/Engine/Private/SceneView.cpp
@@ -660,6 +660,40 @@ static void SetupViewFrustum(FSceneView& View)
 		GetViewFrustumBounds(View.ViewFrustum, View.ViewMatrices.GetViewProjectionMatrix(), false);
 	}
 
+	// Use a unified frustum for culling an instanced stereo pass.
+	if (IStereoRendering::IsStereoEyeView(View) && GEngine->StereoRenderingDevice.IsValid())
+	{
+		FVector MonoLocation = View.ViewLocation;
+		FRotator MonoRotation = View.ViewRotation;
+		GEngine->StereoRenderingDevice->CalculateStereoViewOffset(eSSE_MONOSCOPIC, MonoRotation, View.WorldToMetersScale, MonoLocation);
+		const FMatrix ViewRotationMatrix = FInverseRotationMatrix(MonoRotation) * FMatrix(
+			FPlane(0, 0, 1, 0),
+			FPlane(1, 0, 0, 0),
+			FPlane(0, 1, 0, 0),
+			FPlane(0, 0, 0, 1));
+		const FMatrix ViewMatrixForCulling = FTranslationMatrix(-MonoLocation) * ViewRotationMatrix;
+		const FMatrix ViewProjForCulling = ViewMatrixForCulling * GEngine->StereoRenderingDevice->GetStereoProjectionMatrix(eSSE_MONOSCOPIC);
+
+		if (View.SceneViewInitOptions.OverrideFarClippingPlaneDistance > 0.0f)
+		{
+			const FPlane FarPlane(View.ViewMatrices.GetViewOrigin() + View.GetViewDirection() * View.SceneViewInitOptions.OverrideFarClippingPlaneDistance, View.GetViewDirection());
+			// Derive the frustum from the view projection matrix, overriding the far plane
+			GetViewFrustumBounds(View.CullingFrustum, ViewProjForCulling, FarPlane, true, false);
+		}
+		else
+		{
+			// Derive the frustum from the view projection matrix.
+			GetViewFrustumBounds(View.CullingFrustum, ViewProjForCulling, false);
+		}
+
+		View.CullingOrigin = MonoLocation;
+	}
+	else
+	{
+		View.CullingFrustum = View.ViewFrustum;
+		View.CullingOrigin = View.ViewMatrices.GetViewOrigin();
+	}
+
 	// Derive the view's near clipping distance and plane.
 	static_assert((int32)ERHIZBuffer::IsInverted != 0, "Fix Near Clip distance!");
 	FPlane NearClippingPlane;
@@ -689,8 +723,8 @@ FSceneView::FSceneView(const FSceneViewInitOptions& InitOptions)
 	, UnconstrainedViewRect(InitOptions.GetViewRect())
 	, MaxShadowCascades(10)
 	, ViewMatrices(InitOptions)
-	, ViewLocation(ForceInitToZero)
-	, ViewRotation(ForceInitToZero)
+	, ViewLocation(InitOptions.ViewLocation)
+	, ViewRotation(InitOptions.ViewRotation)
 	, BaseHmdOrientation(EForceInit::ForceInit)
 	, BaseHmdLocation(ForceInitToZero)
 	, WorldToMetersScale(InitOptions.WorldToMetersScale)
@@ -701,6 +735,7 @@ FSceneView::FSceneView(const FSceneViewInitOptions& InitOptions)
 	, ColorScale(InitOptions.ColorScale)
 	, StereoPass(InitOptions.StereoPass)
 	, StereoViewIndex(InitOptions.StereoViewIndex)
+	, PrimaryViewIndex(INDEX_NONE)
 	, StereoIPD(InitOptions.StereoIPD)
 	, bAllowCrossGPUTransfer(true)
 	, bOverrideGPUMask(false)
@@ -862,18 +897,36 @@ FSceneView::FSceneView(const FSceneViewInitOptions& InitOptions)
 		PrimaryScreenPercentageMethod = EPrimaryScreenPercentageMethod::TemporalUpscale;
 	}
 
-	if (Family && Family->bResolveScene && Family->EngineShowFlags.PostProcessing)
+	if (Family)
 	{
-		EyeAdaptationViewState = State;
-
-		// When rendering in stereo we want to use the same exposure for both eyes.
+		// Find the primary view in the view family, if this is the primary view itself assume it'll be added at the back of the family.
+		PrimaryViewIndex = Family->Views.Num();
 		if (IStereoRendering::IsASecondaryView(*this))
 		{
 			check(Family->Views.Num() >= 1);
-			const FSceneView* PrimaryView = Family->Views[0];
-			if (IStereoRendering::IsAPrimaryView(*PrimaryView))
+			while (--PrimaryViewIndex >= 0)
+			{
+				const FSceneView* PrimaryView = Family->Views[PrimaryViewIndex];
+				if (IStereoRendering::IsAPrimaryView(*PrimaryView))
+				{
+					break;
+				}
+			}
+		}
+
+		if (Family->bResolveScene && Family->EngineShowFlags.PostProcessing)
+		{
+			EyeAdaptationViewState = State;
+
+			// When rendering in stereo we want to use the same exposure for both eyes.
+			if (IStereoRendering::IsASecondaryView(*this))
 			{
-				EyeAdaptationViewState = PrimaryView->State;
+				check(Family->Views.Num() >= 1);
+				const FSceneView* PrimaryView = Family->Views[PrimaryViewIndex];
+				if (IStereoRendering::IsAPrimaryView(*PrimaryView))
+				{
+					EyeAdaptationViewState = PrimaryView->State;
+				}
 			}
 		}
 	}
@@ -2585,6 +2638,36 @@ FRDGPooledBuffer* FSceneView::GetEyeAdaptationBuffer() const
 	return nullptr;
 }
 
+const FSceneView* FSceneView::GetPrimarySceneView() const
+{
+	// It is valid for this function to return itself if it's already the primary view.
+	return Family->Views[PrimaryViewIndex];
+}
+
+const FSceneView* FSceneView::GetInstancedSceneView() const
+{
+	// If called on the first secondary view it'll return itself.
+	if (Family->Views.IsValidIndex(PrimaryViewIndex + 1))
+	{
+		const FSceneView* SecondaryView = Family->Views[PrimaryViewIndex + 1];
+		return IStereoRendering::IsASecondaryView(*SecondaryView) ? SecondaryView : nullptr;
+	}
+	return nullptr;
+}
+
+TArray<const FSceneView*> FSceneView::GetSecondaryViews() const
+{
+	// If called on a secondary view we'll return all other secondary views, including this view.
+	TArray<const FSceneView*> Views;
+	Views.Reserve(Family->Views.Num());
+	for (int32 ViewIndex = PrimaryViewIndex + 1; ViewIndex < Family->Views.Num() &&
+		IStereoRendering::IsASecondaryView(*Family->Views[ViewIndex]); ViewIndex++)
+	{
+		Views.Add(Family->Views[ViewIndex]);
+	}
+	return Views;
+}
+
 FSceneViewFamily::FSceneViewFamily(const ConstructionValues& CVS)
 	:
 	ViewMode(VMI_Lit),
diff --git a/Engine/Source/Runtime/Engine/Private/UnrealEngine.cpp b/Engine/Source/Runtime/Engine/Private/UnrealEngine.cpp
--- a/Engine/Source/Runtime/Engine/Private/UnrealEngine.cpp
+++ b/Engine/Source/Runtime/Engine/Private/UnrealEngine.cpp
@@ -3429,10 +3429,12 @@ public:
 		: FOVInDegrees(100)
 		, Width(640)
 		, Height(480)
+		, NumViews(2)
 	{
 		static TAutoConsoleVariable<float> CVarEmulateStereoFOV(TEXT("r.StereoEmulationFOV"), 0, TEXT("FOV in degrees, of the imaginable HMD for stereo emulation"));
 		static TAutoConsoleVariable<int32> CVarEmulateStereoWidth(TEXT("r.StereoEmulationWidth"), 0, TEXT("Width of the imaginable HMD for stereo emulation"));
 		static TAutoConsoleVariable<int32> CVarEmulateStereoHeight(TEXT("r.StereoEmulationHeight"), 0, TEXT("Height of the imaginable HMD for stereo emulation"));
+		static TAutoConsoleVariable<int32> CVarEmulateStereoViews(TEXT("r.StereoEmulationViews"), 0, TEXT("Number of views in the imaginable HMD for stereo emulation"));
 		float FOV = CVarEmulateStereoFOV.GetValueOnAnyThread();
 		if (FOV != 0)
 		{
@@ -3440,6 +3442,7 @@ public:
 		}
 		int32 W = CVarEmulateStereoWidth.GetValueOnAnyThread();
 		int32 H = CVarEmulateStereoHeight.GetValueOnAnyThread();
+		int32 V = CVarEmulateStereoViews.GetValueOnAnyThread();
 		if (W != 0)
 		{
 			Width = FMath::Clamp(W, 100, 10000);
@@ -3448,6 +3451,10 @@ public:
 		{
 			Height = FMath::Clamp(H, 100, 10000);
 		}
+		if (V != 0)
+		{
+			NumViews = FMath::Clamp(V, 1, 32);
+		}
 	}
 
 	virtual ~FFakeStereoRenderingDevice() {}
@@ -3456,6 +3463,8 @@ public:
 
 	virtual bool EnableStereo(bool stereo = true) override { return true; }
 
+	virtual int32 GetDesiredNumberOfViews(bool bStereoRequested) const { return (bStereoRequested) ? NumViews : 1; }
+
 	virtual EStereoscopicPass GetViewPassForIndex(bool bStereoRequested, int32 ViewIndex) const override
 	{
 #if WITH_MGPU
@@ -3467,22 +3476,31 @@ public:
 			return EStereoscopicPass::eSSP_PRIMARY;
 		}
 #endif
-		return IStereoRendering::GetViewPassForIndex(bStereoRequested, ViewIndex);
+		return ViewIndex % 2 == 0 ? EStereoscopicPass::eSSP_PRIMARY : EStereoscopicPass::eSSP_SECONDARY;
 	}
 
 	virtual void AdjustViewRect(int32 ViewIndex, int32& X, int32& Y, uint32& SizeX, uint32& SizeY) const override
 	{
 		SizeX = SizeX / 2;
-		X += SizeX * ViewIndex;
+		SizeY = SizeY / (NumViews / 2);
+		X += SizeX * (ViewIndex % 2);
+		Y += SizeY * (ViewIndex / 2);
 	}
 
 	virtual void CalculateStereoViewOffset(const int32 ViewIndex, FRotator& ViewRotation, const float WorldToMeters, FVector& ViewLocation) override
 	{
-		if (ViewIndex != INDEX_NONE)
+		// 32mm, 1/2 average interpupillary distance
+		float EyeOffset = .0320000005f * WorldToMeters;
+
+		if (ViewIndex == eSSE_MONOSCOPIC)
+		{
+			const float HalfFov = FMath::DegreesToRadians(FOVInDegrees) / 2.f;
+			const float CenterOffset = GNearClippingPlane + EyeOffset / FMath::Tan(HalfFov);
+			ViewLocation += ViewRotation.Quaternion().RotateVector(FVector(CenterOffset, 0, 0));
+		}
+		else
 		{
-			// 32mm, 1/2 average interpupillary distance
-			float EyeOffset = .0320000005f * WorldToMeters;
-			const float PassOffset = (ViewIndex == EStereoscopicEye::eSSE_LEFT_EYE) ? -EyeOffset : EyeOffset;
+			const float PassOffset = (ViewIndex % 2 == 0) ? -EyeOffset : EyeOffset;
 			ViewLocation += ViewRotation.Quaternion().RotateVector(FVector(0,PassOffset,0));
 		}
 	}
@@ -3522,6 +3540,7 @@ public:
 
 	float FOVInDegrees;		// max(HFOV, VFOV) in degrees of imaginable HMD
 	int32 Width, Height;	// resolution of imaginable HMD
+	int32 NumViews;			// views of imaginable HMD
 };
 
 bool UEngine::InitializeHMDDevice()
diff --git a/Engine/Source/Runtime/Engine/Public/SceneView.h b/Engine/Source/Runtime/Engine/Public/SceneView.h
--- a/Engine/Source/Runtime/Engine/Public/SceneView.h
+++ b/Engine/Source/Runtime/Engine/Public/SceneView.h
@@ -137,6 +137,10 @@ struct FSceneViewInitOptions : public FSceneViewProjectionData
 	/** Conversion from world units (uu) to meters, so we can scale motion to the world appropriately */
 	float WorldToMetersScale;
 
+	/** The view origin and rotation without any stereo offsets applied to it */
+	FVector ViewLocation;
+	FRotator ViewRotation;
+
 	TSet<FPrimitiveComponentId> HiddenPrimitives;
 
 	/** The primitives which are visible for this view. If the array is not empty, all other primitives will be hidden. */
@@ -190,6 +194,8 @@ struct FSceneViewInitOptions : public FSceneViewProjectionData
 		, StereoViewIndex(INDEX_NONE)
 		, StereoIPD(0.0f)
 		, WorldToMetersScale(100.f)
+		, ViewLocation(ForceInitToZero)
+		, ViewRotation(ForceInitToZero)
 		, CursorPos(-1, -1)
 		, LODDistanceFactor(1.0f)
 		, OverrideFarClippingPlaneDistance(-1.0f)
@@ -990,6 +996,9 @@ public:
 	/** For stereoscopic rendering, unique index identifying the view across view families */
 	int32 StereoViewIndex;
 
+	/** For stereoscopic rendering, view family index of the primary view associated with this view */
+	int32 PrimaryViewIndex;
+
 	/** Half of the view's stereo IPD (- for lhs, + for rhs) */
 	float StereoIPD;
 
@@ -1064,6 +1073,12 @@ public:
 
 	float NearClippingDistance;
 
+	/** Monoscopic culling frustum, same as ViewFrustum in case of non-stereo */
+	FConvexVolume CullingFrustum;
+
+	/** Monoscopic culling origin, same as the view matrix origin in case of non-stereo */
+	FVector CullingOrigin;
+
 	/** true if ViewMatrix.Determinant() is negative. */
 	bool bReverseCulling;
 
@@ -1422,6 +1437,14 @@ public:
 	/** Returns the eye adaptation buffer (mobile) or null if it doesn't exist. */
 	FRDGPooledBuffer* GetEyeAdaptationBuffer() const;
 
+	/** Get the primary view associated with the secondary view. */
+	const FSceneView* GetPrimarySceneView() const;
+
+	/** Get the first secondary view associated with the primary view. */
+	const FSceneView* GetInstancedSceneView() const;
+
+	/** Get all secondary views associated with the primary view. */
+	TArray<const FSceneView*> GetSecondaryViews() const;
 
 protected:
 	FSceneViewStateInterface* EyeAdaptationViewState = nullptr;
diff --git a/Engine/Source/Runtime/Engine/Public/StereoRendering.h b/Engine/Source/Runtime/Engine/Public/StereoRendering.h
--- a/Engine/Source/Runtime/Engine/Public/StereoRendering.h
+++ b/Engine/Source/Runtime/Engine/Public/StereoRendering.h
@@ -141,11 +141,13 @@ public:
 
 	/**
 	 * Calculates the offset for the camera position, given the specified position, rotation, and world scale
+	 * Specifying eSSE_MONOSCOPIC for the view index returns a center offset behind the stereo views
 	 */
 	virtual void CalculateStereoViewOffset(const int32 ViewIndex, FRotator& ViewRotation, const float WorldToMeters, FVector& ViewLocation) = 0;
 
 	/**
-	 * Gets a projection matrix for the device, given the specified eye setup
+	 * Gets a projection matrix for the device, given the specified view index
+	 * Specifying eSSE_MONOSCOPIC for the view index returns a center projection matrix encompassing all views
 	 */
 	virtual FMatrix GetStereoProjectionMatrix(const int32 ViewIndex) const = 0;
 
diff --git a/Engine/Source/Runtime/HeadMountedDisplay/Private/DefaultXRCamera.cpp b/Engine/Source/Runtime/HeadMountedDisplay/Private/DefaultXRCamera.cpp
--- a/Engine/Source/Runtime/HeadMountedDisplay/Private/DefaultXRCamera.cpp
+++ b/Engine/Source/Runtime/HeadMountedDisplay/Private/DefaultXRCamera.cpp
@@ -72,24 +72,34 @@ void FDefaultXRCamera::SetupLateUpdate(const FTransform& ParentToWorld, USceneCo
 
 void FDefaultXRCamera::CalculateStereoCameraOffset(const int32 ViewIndex, FRotator& ViewRotation, FVector& ViewLocation)
 {
-	if (ViewIndex != INDEX_NONE)
+	FQuat EyeOrientation;
+	FVector EyeOffset;
+	if (TrackingSystem->GetRelativeEyePose(DeviceId, ViewIndex, EyeOrientation, EyeOffset))
 	{
-		FQuat EyeOrientation;
-		FVector EyeOffset;
-		if (TrackingSystem->GetRelativeEyePose(DeviceId, ViewIndex, EyeOrientation, EyeOffset))
+		ViewLocation += ViewRotation.Quaternion().RotateVector(EyeOffset);
+		ViewRotation = FRotator(ViewRotation.Quaternion() * EyeOrientation);
+	}
+	else if (ViewIndex == EStereoscopicEye::eSSE_MONOSCOPIC && TrackingSystem->GetHMDDevice())
+	{
+		float HFov, VFov;
+		TrackingSystem->GetHMDDevice()->GetFieldOfView(HFov, VFov);
+		if (HFov > 0.0f)
 		{
-			ViewLocation += ViewRotation.Quaternion().RotateVector(EyeOffset);
-			ViewRotation = FRotator(ViewRotation.Quaternion() * EyeOrientation);
-
-			if (!bUseImplicitHMDPosition)
-			{
-				FQuat DeviceOrientation; // Unused
-				FVector DevicePosition;
-				TrackingSystem->GetCurrentPose(DeviceId, DeviceOrientation, DevicePosition);
-				ViewLocation += DeltaControlOrientation.RotateVector(DevicePosition);
-			}
+			const float CenterOffset = GNearClippingPlane + (TrackingSystem->GetHMDDevice()->GetInterpupillaryDistance() / 2.0f) * (1.0f / FMath::Tan(FMath::DegreesToRadians(HFov)));
+			ViewLocation += ViewRotation.Quaternion().RotateVector(FVector(CenterOffset, 0, 0));
 		}
-		
+	}
+	else
+	{
+		return;
+	}
+
+	if (!bUseImplicitHMDPosition)
+	{
+		FQuat DeviceOrientation; // Unused
+		FVector DevicePosition;
+		TrackingSystem->GetCurrentPose(DeviceId, DeviceOrientation, DevicePosition);
+		ViewLocation += DeltaControlOrientation.RotateVector(DevicePosition);
 	}
 }
 
diff --git a/Engine/Source/Runtime/Renderer/Private/BasePassRendering.cpp b/Engine/Source/Runtime/Renderer/Private/BasePassRendering.cpp
--- a/Engine/Source/Runtime/Renderer/Private/BasePassRendering.cpp
+++ b/Engine/Source/Runtime/Renderer/Private/BasePassRendering.cpp
@@ -512,9 +512,9 @@ void SetupSharedBasePassParameters(
 
 	if (View.IsInstancedStereoPass())
 	{
-		const FViewInfo& RightEye = static_cast<const FViewInfo&>(*View.Family->Views[1]);
-		SharedParameters.ForwardISR = *RightEye.ForwardLightingResources.ForwardLightData;
-		SetupFogUniformParameters(GraphBuilder, (FViewInfo&)RightEye, SharedParameters.FogISR);
+		const FViewInfo& InstancedView = *View.GetInstancedView();
+		SharedParameters.ForwardISR = *InstancedView.ForwardLightingResources.ForwardLightData;
+		SetupFogUniformParameters(GraphBuilder, (FViewInfo&)InstancedView, SharedParameters.FogISR);
 	}
 	else
 	{
diff --git a/Engine/Source/Runtime/Renderer/Private/DeferredShadingRenderer.cpp b/Engine/Source/Runtime/Renderer/Private/DeferredShadingRenderer.cpp
--- a/Engine/Source/Runtime/Renderer/Private/DeferredShadingRenderer.cpp
+++ b/Engine/Source/Runtime/Renderer/Private/DeferredShadingRenderer.cpp
@@ -2380,9 +2380,6 @@ void FDeferredShadingSceneRenderer::Render(FRDGBuilder& GraphBuilder)
 
 		RDG_GPU_STAT_SCOPE(GraphBuilder, NaniteRaster);
 		const FIntPoint RasterTextureSize = SceneTextures.Depth.Target->Desc.Extent;
-
-		const FViewInfo& PrimaryViewRef = Views[0];
-		const FIntRect PrimaryViewRect = PrimaryViewRef.ViewRect;
 		
 		// Primary raster view
 		{
@@ -2486,6 +2483,8 @@ void FDeferredShadingSceneRenderer::Render(FRDGBuilder& GraphBuilder)
 					FRDGTextureRef SceneDepth = SystemTextures.Black;
 					FRDGTextureRef GraphHZB = nullptr;
 
+					const FIntRect PrimaryViewRect = View.GetPrimaryView()->ViewRect;
+
 					BuildHZBFurthest(
 						GraphBuilder,
 						SceneDepth,
@@ -2500,12 +2499,12 @@ void FDeferredShadingSceneRenderer::Render(FRDGBuilder& GraphBuilder)
 				}
 
 				Nanite::ExtractResults(GraphBuilder, CullingContext, RasterContext, RasterResults);
-			}
-		}
 
-		if (GNaniteShowStats != 0)
-		{
-			Nanite::PrintStats(GraphBuilder, PrimaryViewRef);
+				if (GNaniteShowStats != 0 && IStereoRendering::IsAPrimaryView(View))
+				{
+					Nanite::PrintStats(GraphBuilder, View);
+				}
+			}
 		}
 	}
 
diff --git a/Engine/Source/Runtime/Renderer/Private/DistortionRendering.cpp b/Engine/Source/Runtime/Renderer/Private/DistortionRendering.cpp
--- a/Engine/Source/Runtime/Renderer/Private/DistortionRendering.cpp
+++ b/Engine/Source/Runtime/Renderer/Private/DistortionRendering.cpp
@@ -86,7 +86,6 @@ void SetupDistortionParams(FVector4f& DistortionParams, const FViewInfo& View)
 	// When ISR is enabled we store two FOVs in the distortion parameters and compute the aspect ratio in the shader instead.
 	if (View.IsInstancedStereoPass() || View.bIsMobileMultiViewEnabled)
 	{
-		// When drawing the left eye in a stereo scene, copy the right eye view values into the instanced view uniform buffer.
 		const FSceneView* InstancedView = View.GetInstancedView();
 		if (InstancedView)
 		{
diff --git a/Engine/Source/Runtime/Renderer/Private/MobileShadingRenderer.cpp b/Engine/Source/Runtime/Renderer/Private/MobileShadingRenderer.cpp
--- a/Engine/Source/Runtime/Renderer/Private/MobileShadingRenderer.cpp
+++ b/Engine/Source/Runtime/Renderer/Private/MobileShadingRenderer.cpp
@@ -1016,7 +1016,7 @@ void FMobileSceneRenderer::RenderForward(FRDGBuilder& GraphBuilder, FRDGTextureR
 	{
 		FViewInfo& View = Views[ViewIndex];
 
-		SCOPED_GPU_MASK(GraphBuilder.RHICmdList, !View.IsInstancedStereoPass() ? View.GPUMask : (Views[0].GPUMask | Views[1].GPUMask));
+		SCOPED_GPU_MASK(GraphBuilder.RHICmdList, !View.IsInstancedStereoPass() ? View.GPUMask : (View.GPUMask | View.GetInstancedView()->GPUMask));
 		SCOPED_CONDITIONAL_DRAW_EVENTF(GraphBuilder.RHICmdList, EventView, Views.Num() > 1, TEXT("View%d"), ViewIndex);
 
 		if (!View.ShouldRenderView())
@@ -1312,7 +1312,7 @@ void FMobileSceneRenderer::RenderDeferred(FRDGBuilder& GraphBuilder, const FSort
 	{
 		FViewInfo& View = Views[ViewIndex];
 
-		SCOPED_GPU_MASK(GraphBuilder.RHICmdList, !View.IsInstancedStereoPass() ? View.GPUMask : (Views[0].GPUMask | Views[1].GPUMask));
+		SCOPED_GPU_MASK(GraphBuilder.RHICmdList, !View.IsInstancedStereoPass() ? View.GPUMask : (View.GPUMask | View.GetInstancedView()->GPUMask));
 		SCOPED_CONDITIONAL_DRAW_EVENTF(GraphBuilder.RHICmdList, EventView, NumViews > 1, TEXT("View%d"), ViewIndex);
 
 		if (!View.ShouldRenderView())
diff --git a/Engine/Source/Runtime/Renderer/Private/SceneRendering.cpp b/Engine/Source/Runtime/Renderer/Private/SceneRendering.cpp
--- a/Engine/Source/Runtime/Renderer/Private/SceneRendering.cpp
+++ b/Engine/Source/Runtime/Renderer/Private/SceneRendering.cpp
@@ -1592,18 +1592,7 @@ void FViewInfo::SetupUniformBufferParameters(
 
 	{
 		// If rendering in stereo, the other stereo passes uses the left eye's translucency lighting volume.
-		const FViewInfo* PrimaryView = this;
-		if (IStereoRendering::IsASecondaryView(*this))
-		{
-			if (Family->Views.IsValidIndex(0))
-			{
-				const FSceneView* LeftEyeView = Family->Views[0];
-				if (LeftEyeView->bIsViewInfo && IStereoRendering::IsAPrimaryView(*LeftEyeView))
-				{
-					PrimaryView = static_cast<const FViewInfo*>(LeftEyeView);
-				}
-			}
-		}
+		const FViewInfo* PrimaryView = GetPrimaryView();
 		PrimaryView->CalcTranslucencyLightingVolumeBounds(OutTranslucentCascadeBoundsArray, NumTranslucentCascades);
 	}
 
@@ -2033,17 +2022,19 @@ FViewShaderParameters FViewInfo::GetShaderParameters() const
 	return Parameters;
 }
 
+const FViewInfo* FViewInfo::GetPrimaryView() const
+{
+	// It is valid for this function to return itself if it's already the primary view.
+	const FSceneView* PrimaryView = Family->Views[PrimaryViewIndex];
+	check(PrimaryView->bIsViewInfo);
+	return static_cast<const FViewInfo*>(PrimaryView);
+}
+
 const FViewInfo* FViewInfo::GetInstancedView() const
 {
 	if ((IsInstancedStereoPass() || bIsMobileMultiViewEnabled) && Family->Views.Num() > 0)
 	{
-		// When drawing the left eye in a stereo scene, copy the right eye view values into the instanced view uniform buffer.
-		const int32 ViewIndex = IStereoRendering::IsStereoEyeView(*this) ? EStereoscopicEye::eSSE_RIGHT_EYE : 0;
-
-		if (Family->Views.IsValidIndex(ViewIndex))
-		{
-			return static_cast<const FViewInfo*>(Family->Views[ViewIndex]);
-		}
+		return static_cast<const FViewInfo*>(GetInstancedSceneView());
 	}
 	return nullptr;
 }
@@ -3042,8 +3033,8 @@ void FSceneRenderer::ComputeFamilySize()
 		MaxFamilyX = FMath::Max(MaxFamilyX, FinalViewMaxX);
 		MaxFamilyY = FMath::Max(MaxFamilyY, FinalViewMaxY);
 
-			InstancedStereoWidth = FPlatformMath::Max(InstancedStereoWidth, static_cast<uint32>(View.ViewRect.Max.X));
-		}
+		InstancedStereoWidth = FPlatformMath::Max(InstancedStereoWidth, static_cast<uint32>(View.ViewRect.Max.X));
+	}
 
 	for (FViewInfo& View : Views)
 	{
@@ -4838,14 +4829,12 @@ void FSceneRenderer::SetStereoViewport(FRHICommandList& RHICmdList, const FViewI
 	{
 		if (View.bIsMultiViewEnabled)
 		{
-			check(&View == &Views[0]);
-
-			const FViewInfo& LeftView = Views[0];
+			const FViewInfo& LeftView = View;
 			const uint32 LeftMinX = LeftView.ViewRect.Min.X * ViewportScale;
 			const uint32 LeftMaxX = LeftView.ViewRect.Max.X * ViewportScale;
 			const uint32 LeftMaxY = LeftView.ViewRect.Max.Y * ViewportScale;
 
-			const FViewInfo& RightView = Views[1];
+			const FViewInfo& RightView = static_cast<const FViewInfo&>(*View.GetInstancedView());
 			const uint32 RightMinX = RightView.ViewRect.Min.X * ViewportScale;
 			const uint32 RightMaxX = RightView.ViewRect.Max.X * ViewportScale;
 			const uint32 RightMaxY = RightView.ViewRect.Max.Y * ViewportScale;
@@ -4854,7 +4843,7 @@ void FSceneRenderer::SetStereoViewport(FRHICommandList& RHICmdList, const FViewI
 		}
 		else
 		{
-			RHICmdList.SetViewport(0, 0, 0, InstancedStereoWidth * ViewportScale, View.ViewRect.Max.Y * ViewportScale, 1.0f);
+			RHICmdList.SetViewport(View.ViewRect.Min.X * ViewportScale, View.ViewRect.Min.Y * ViewportScale, 0.0f, InstancedStereoWidth * ViewportScale, View.ViewRect.Max.Y * ViewportScale, 1.0f);
 		}
 	}
 	else
diff --git a/Engine/Source/Runtime/Renderer/Private/SceneRendering.h b/Engine/Source/Runtime/Renderer/Private/SceneRendering.h
--- a/Engine/Source/Runtime/Renderer/Private/SceneRendering.h
+++ b/Engine/Source/Runtime/Renderer/Private/SceneRendering.h
@@ -1653,6 +1653,9 @@ public:
 	/** Returns the set of view uniform buffers representing this view. */
 	FViewShaderParameters GetShaderParameters() const;
 
+	/** Returns the primary view associated with the input view, or null if none exists. */
+	const FViewInfo* GetPrimaryView() const;
+
 	/** Returns the instanced view associated with the input view, or null if none exists. */
 	const FViewInfo* GetInstancedView() const;
 
diff --git a/Engine/Source/Runtime/Renderer/Private/SceneVisibility.cpp b/Engine/Source/Runtime/Renderer/Private/SceneVisibility.cpp
--- a/Engine/Source/Runtime/Renderer/Private/SceneVisibility.cpp
+++ b/Engine/Source/Runtime/Renderer/Private/SceneVisibility.cpp
@@ -3348,7 +3348,6 @@ void FSceneRenderer::GatherDynamicMeshElements(
 				&DynamicReadBuffer);
 		}
 
-		const bool bIsInstancedStereo = (ViewCount > 0) ? (InViews[0].IsInstancedStereoPass() || InViews[0].bIsMobileMultiViewEnabled) : false;
 		const EShadingPath ShadingPath = Scene->GetShadingPath();
 
 		for (int32 PrimitiveIndex = 0; PrimitiveIndex < NumPrimitives; ++PrimitiveIndex)
@@ -3357,8 +3356,16 @@ void FSceneRenderer::GatherDynamicMeshElements(
 
 			if (ViewMask != 0)
 			{
-				// Don't cull a single eye when drawing a stereo pair
-				const uint8 ViewMaskFinal = (bIsInstancedStereo) ? ViewMask | 0x3 : ViewMask;
+				// If a mesh is visible in a secondary view, mark it as visible in the primary view
+				uint8 ViewMaskFinal = ViewMask;
+				for (int32 ViewIndex = 0; ViewIndex < ViewCount; ViewIndex++)
+				{
+					FViewInfo& View = InViews[ViewIndex];
+					if (ViewMask & (1 << ViewIndex) && IStereoRendering::IsASecondaryView(View))
+					{
+						ViewMaskFinal |= 1 << InViews[ViewIndex].PrimaryViewIndex;
+					}
+				}
 
 				FPrimitiveSceneInfo* PrimitiveSceneInfo = InScene->Primitives[PrimitiveIndex];
 				const FPrimitiveBounds& Bounds = InScene->PrimitiveBounds[PrimitiveIndex];
@@ -4308,7 +4315,6 @@ void FSceneRenderer::ComputeViewVisibility(
 		HasDynamicEditorMeshElementsMasks.AddZeroed(NumPrimitives);
 	}
 
-	const bool bIsInstancedStereo = (Views.Num() > 0) ? (Views[0].IsInstancedStereoPass() || Views[0].bIsMobileMultiViewEnabled) : false;
 	UpdateReflectionSceneData(Scene);
 
 	{
@@ -4345,6 +4351,8 @@ void FSceneRenderer::ComputeViewVisibility(
 			FViewCommands& ViewCommands = ViewCommandsPerView[ViewIndex];
 			FSceneViewState* ViewState = (FSceneViewState*)View.State;
 
+			const bool bIsSinglePassStereo = View.bIsInstancedStereoEnabled || View.bIsMobileMultiViewEnabled;
+
 			// Allocate the view's visibility maps.
 			View.PrimitiveVisibilityMap.Init(false, Scene->Primitives.Num());
 			View.PrimitiveRayTracingVisibilityMap.Init(false, Scene->Primitives.Num());
@@ -4528,8 +4536,8 @@ void FSceneRenderer::ComputeViewVisibility(
 				}
 			}
 
-			// ISR views can't compute relevance until all views are frustum culled
-			if (!bIsInstancedStereo)
+			// Single-pass stereo views can't compute relevance until all views are visibility culled
+			if (!bIsSinglePassStereo)
 			{
 				SCOPE_CYCLE_COUNTER(STAT_ViewRelevance);
 				ComputeAndMarkRelevanceForViewParallel(RHICmdList, Scene, View, ViewCommands, ViewBit, HasDynamicMeshElementsMasks, HasDynamicEditorMeshElementsMasks);
@@ -4560,39 +4568,38 @@ void FSceneRenderer::ComputeViewVisibility(
 			// TODO: right now decals visibility computed right before rendering them, ideally it should be done in InitViews and this flag should be replaced with list of visible decals  
 			// Currently used to disable stencil operations in forward base pass when scene has no any decals
 			View.bSceneHasDecals = (Scene->Decals.Num() > 0) || (GForceSceneHasDecals != 0);
-		}
-	}
 
-	if ((Views.Num() > 1) && bIsInstancedStereo)
-	{
-		// Ensure primitives from the right-eye view are visible in the left-eye (instanced) view
-		FSceneBitArray& LeftView = Views[0].PrimitiveVisibilityMap;
-		const FSceneBitArray& RightView = Views[1].PrimitiveVisibilityMap;
+			if (bIsSinglePassStereo && IStereoRendering::IsASecondaryView(View) && Views.IsValidIndex(View.PrimaryViewIndex))
+			{
+				// Ensure primitives from the secondary view are visible in the primary view
+				FSceneBitArray& PrimaryVis = Views[View.PrimaryViewIndex].PrimitiveVisibilityMap;
+				const FSceneBitArray& SecondaryVis = View.PrimitiveVisibilityMap;
 
-		check(LeftView.Num() == RightView.Num())
+				check(PrimaryVis.Num() == SecondaryVis.Num());
 
-		const uint32 NumWords = FMath::DivideAndRoundUp(LeftView.Num(), NumBitsPerDWORD);
-		uint32* const LeftData = LeftView.GetData();
-		const uint32* const RightData = RightView.GetData();
+				const uint32 NumWords = FMath::DivideAndRoundUp(PrimaryVis.Num(), NumBitsPerDWORD);
+				uint32* const PrimaryData = PrimaryVis.GetData();
+				const uint32* const SecondaryData = SecondaryVis.GetData();
 
-		for (uint32 Index = 0; Index < NumWords; ++Index)
-		{
-			LeftData[Index] |= RightData[Index];
+				for (uint32 Index = 0; Index < NumWords; ++Index)
+				{
+					PrimaryData[Index] |= SecondaryData[Index];
+				}
+			}
 		}
 	}
 	
 	ViewBit = 0x1;
-	for (int32 ViewIndex = 0; ViewIndex < Views.Num(); ++ViewIndex)
+	for (int32 ViewIndex = 0; ViewIndex < Views.Num(); ++ViewIndex, ViewBit <<= 1)
 	{
 		FViewInfo& View = Views[ViewIndex];
 		FViewCommands& ViewCommands = ViewCommandsPerView[ViewIndex];
 		
-		if (bIsInstancedStereo)
+		if (View.bIsInstancedStereoEnabled || View.bIsMobileMultiViewEnabled)
 		{
 			SCOPE_CYCLE_COUNTER(STAT_ViewRelevance);
 			ComputeAndMarkRelevanceForViewParallel(RHICmdList, Scene, View, ViewCommands, ViewBit, HasDynamicMeshElementsMasks, HasDynamicEditorMeshElementsMasks);
 		}
-		ViewBit <<= 1;
 	}
 
 	{
diff --git a/Engine/Source/Runtime/Renderer/Private/ShadowRendering.cpp b/Engine/Source/Runtime/Renderer/Private/ShadowRendering.cpp
--- a/Engine/Source/Runtime/Renderer/Private/ShadowRendering.cpp
+++ b/Engine/Source/Runtime/Renderer/Private/ShadowRendering.cpp
@@ -906,14 +906,14 @@ void FProjectedShadowInfo::SetupProjectionStencilMask(
 		//       loop is over each view, the stencil mask is not retained when the right view comes around.
 		// TODO: Support instanced stereo properly in the projection stenciling pass.
 		const bool bIsInstancedStereoEmulated = View->bIsInstancedStereoEnabled && !View->bIsMultiViewEnabled && IStereoRendering::IsStereoEyeView(*View);
-		if (bIsInstancedStereoEmulated)
+		if (bIsInstancedStereoEmulated && ProjectionStencilingPasses.IsValidIndex(View->PrimaryViewIndex))
 		{
-			ensure(ProjectionStencilingPasses[0]->GetInstanceCullingMode() == EInstanceCullingMode::Stereo);
+			ensure(ProjectionStencilingPasses[View->PrimaryViewIndex]->GetInstanceCullingMode() == EInstanceCullingMode::Stereo);
 
-			RHICmdList.SetViewport(0, 0, 0, SceneRender->InstancedStereoWidth, View->ViewRect.Max.Y, 1);
+			RHICmdList.SetViewport(View->ViewRect.Min.X, View->ViewRect.Min.Y, 0.0f, SceneRender->InstancedStereoWidth, View->ViewRect.Max.Y, 1.0f);
 			RHICmdList.SetScissorRect(true, View->ViewRect.Min.X, View->ViewRect.Min.Y, View->ViewRect.Max.X, View->ViewRect.Max.Y);
 			// Submit the first (and only pass - we share that at least) as the pass is set up for stereo.
-			ProjectionStencilingPasses[0]->SubmitDraw(RHICmdList, InstanceCullingDrawParams);
+			ProjectionStencilingPasses[View->PrimaryViewIndex]->SubmitDraw(RHICmdList, InstanceCullingDrawParams);
 
 			RHICmdList.SetScissorRect(false, 0, 0, 0, 0);
 			RHICmdList.SetViewport(View->ViewRect.Min.X, View->ViewRect.Min.Y, 0.0f, View->ViewRect.Max.X, View->ViewRect.Max.Y, 1.0f);
@@ -1183,7 +1183,7 @@ void FProjectedShadowInfo::RenderProjection(
 	}
 
 	const bool bIsInstancedStereoEmulated = View->bIsInstancedStereoEnabled && !View->bIsMultiViewEnabled && IStereoRendering::IsStereoEyeView(*View);
-	if (ViewIndex < ProjectionStencilingPasses.Num())
+	if (ViewIndex < ProjectionStencilingPasses.Num() && ProjectionStencilingPasses[ViewIndex] != nullptr)
 	{
 		// GPUCULL_TODO: get rid of const cast
 		FSimpleMeshDrawCommandPass& ProjectionStencilingPass = *const_cast<FSimpleMeshDrawCommandPass*>(ProjectionStencilingPasses[ViewIndex]);
@@ -1198,14 +1198,7 @@ void FProjectedShadowInfo::RenderProjection(
 		// GPUCULL_TODO: get rid of const cast
 		FSimpleMeshDrawCommandPass& ProjectionStencilingPass = *const_cast<FSimpleMeshDrawCommandPass*>(ProjectionStencilingPasses[0]);
 		ensure(ProjectionStencilingPass.GetInstanceCullingMode() == EInstanceCullingMode::Stereo);
-		for (int32 StereoViewIndex = View->StereoViewIndex; StereoViewIndex >= 0; StereoViewIndex--)
-		{
-			const FSceneView* PrimaryView = View->Family->Views[StereoViewIndex];
-			if (PrimaryView->StereoPass == EStereoscopicPass::eSSP_PRIMARY)
-			{
-				ProjectionStencilingPass.BuildRenderingCommands(GraphBuilder, *PrimaryView, *SceneRender->Scene, PassParameters->InstanceCullingDrawParams);
-			}
-		}
+		ProjectionStencilingPass.BuildRenderingCommands(GraphBuilder, *View->GetPrimaryView(), *SceneRender->Scene, PassParameters->InstanceCullingDrawParams);
 	}
 
 	const FInstanceCullingDrawParams& InstanceCullingDrawParams = PassParameters->InstanceCullingDrawParams;
diff --git a/Engine/Source/Runtime/Renderer/Private/ShadowSetup.cpp b/Engine/Source/Runtime/Renderer/Private/ShadowSetup.cpp
--- a/Engine/Source/Runtime/Renderer/Private/ShadowSetup.cpp
+++ b/Engine/Source/Runtime/Renderer/Private/ShadowSetup.cpp
@@ -2207,7 +2207,7 @@ void FProjectedShadowInfo::SetupMeshDrawCommandsForProjectionStenciling(FSceneRe
 			{
 				// Note: not using FMemstack as the data is not cleared in the right order.
 				ProjectionStencilingPasses.Add(new FSimpleMeshDrawCommandPass(View, &InstanceCullingManager, true));
-				FSimpleMeshDrawCommandPass& ProjectionStencilingPass = *ProjectionStencilingPasses[ViewIndex];
+				FSimpleMeshDrawCommandPass& ProjectionStencilingPass = *ProjectionStencilingPasses.Last();
 
 				FMeshPassProcessorRenderState DrawRenderState;
 				DrawRenderState.SetBlendState(TStaticBlendState<CW_NONE>::GetRHI());
@@ -2301,6 +2301,10 @@ void FProjectedShadowInfo::SetupMeshDrawCommandsForProjectionStenciling(FSceneRe
 					}
 				}
 			}
+			else
+			{
+				ProjectionStencilingPasses.Add(nullptr);
+			}
 		}
 	}
 }
diff --git a/Engine/Source/Runtime/Renderer/Public/MaterialShader.h b/Engine/Source/Runtime/Renderer/Public/MaterialShader.h
--- a/Engine/Source/Runtime/Renderer/Public/MaterialShader.h
+++ b/Engine/Source/Runtime/Renderer/Public/MaterialShader.h
@@ -98,14 +98,15 @@ public:
 		const auto& ViewUniformBufferParameter = GetUniformBufferParameter<FViewUniformShaderParameters>();
 		SetUniformBufferParameter(RHICmdList, ShaderRHI, ViewUniformBufferParameter, ViewUniformBuffer);
 
-		if (View.bShouldBindInstancedViewUB && View.Family->Views.Num() > 0)
+		if (View.bShouldBindInstancedViewUB)
 		{
 			// When drawing the left eye in a stereo scene, copy the right eye view values into the instanced view uniform buffer.
-			const int32 StereoViewIndex = IStereoRendering::IsStereoEyeView(View) ? EStereoscopicEye::eSSE_RIGHT_EYE : 0;
-			
-			const FSceneView* InstancedView = View.Family->Views[StereoViewIndex];
-			const auto& InstancedViewUniformBufferParameter = GetUniformBufferParameter<FInstancedViewUniformShaderParameters>();
-			SetUniformBufferParameter(RHICmdList, ShaderRHI, InstancedViewUniformBufferParameter, InstancedView->ViewUniformBuffer);
+			const FSceneView* InstancedView = View.GetInstancedSceneView();
+			if (InstancedView)
+			{
+				const auto& InstancedViewUniformBufferParameter = GetUniformBufferParameter<FInstancedViewUniformShaderParameters>();
+				SetUniformBufferParameter(RHICmdList, ShaderRHI, InstancedViewUniformBufferParameter, InstancedView->ViewUniformBuffer);
+			}
 		}
 	}
 
