Converted Volumetric cloud UI so that relelction sample count scaless matches the main view sample count scale.
This will help users understand why cloud in reflections do not match the main view. With equal values, they will 100% match (visual and perf too)
Also updated the new refelction default scalability settings to be closer to main view.

diff --git a/Engine/Source/Runtime/Core/Public/UObject/UE5ReleaseStreamObjectVersion.h b/Engine/Source/Runtime/Core/Public/UObject/UE5ReleaseStreamObjectVersion.h
--- a/Engine/Source/Runtime/Core/Public/UObject/UE5ReleaseStreamObjectVersion.h
+++ b/Engine/Source/Runtime/Core/Public/UObject/UE5ReleaseStreamObjectVersion.h
@@ -81,6 +81,9 @@ struct CORE_API FUE5ReleaseStreamObjectVersion
 		// Exclude brightness from the EncodedHDRCubemap,
 		ExcludeBrightnessFromEncodedHDRCubemap,
 
+		// Unified volumetric cloud component quality sample count slider between main and reflection views for consistency
+		VolumetricCloudSampleCountUnification,
+
 		// -----<new versions can be added above this line>-------------------------------------------------
 		VersionPlusOne,
 		LatestVersion = VersionPlusOne - 1
diff --git a/Engine/Source/Runtime/Engine/Classes/Components/VolumetricCloudComponent.h b/Engine/Source/Runtime/Engine/Classes/Components/VolumetricCloudComponent.h
--- a/Engine/Source/Runtime/Engine/Classes/Components/VolumetricCloudComponent.h
+++ b/Engine/Source/Runtime/Engine/Classes/Components/VolumetricCloudComponent.h
@@ -70,27 +70,31 @@ class UVolumetricCloudComponent : public USceneComponent
 	 * Scale the tracing sample count in primary views. Quality level scalability CVARs affect the maximum range.
 	 * The sample count resolution is still clamped according to scalability setting to 'r.VolumetricCloud.ViewRaySampleCountMax'.
 	 */
-	UPROPERTY(EditAnywhere, BlueprintReadOnly, Category = "Cloud Tracing", AdvancedDisplay, meta = (UIMin = "0.25", UIMax = "8", ClampMin = "0.25", SliderExponent = 1.0))
+	UPROPERTY(EditAnywhere, BlueprintReadOnly, Category = "Cloud Tracing", AdvancedDisplay, meta = (UIMin = "0.05", UIMax = "8", ClampMin = "0.05", SliderExponent = 1.0))
 	float ViewSampleCountScale;
 	/**
 	 * Scale the tracing sample count in reflection views. Quality level scalability CVARs affect the maximum range.
 	 * The sample count resolution is still clamped according to scalability setting to 'r.VolumetricCloud.ReflectionRaySampleMaxCount'.
 	 */
-	UPROPERTY(EditAnywhere, BlueprintReadOnly, Category = "Cloud Tracing", AdvancedDisplay, meta = (UIMin = "0.25", UIMax = "8", ClampMin = "0.25", SliderExponent = 1.0))
-	float ReflectionSampleCountScale;
+	UPROPERTY(EditAnywhere, BlueprintReadOnly, Category = "Cloud Tracing", AdvancedDisplay, meta = (UIMin = "0.05", UIMax = "8", ClampMin = "0.05", SliderExponent = 1.0))
+	float ReflectionViewSampleCountScale;
+	UPROPERTY()
+	float ReflectionSampleCountScale_DEPRECATED;
 
 	/**
 	 * Scale the shadow tracing sample count in primary views, only used with Advanced Output ray marched shadows. Quality level scalability CVARs affect the maximum range.
 	 * The sample count resolution is still clamped according to scalability setting to 'r.VolumetricCloud.Shadow.ViewRaySampleMaxCount'.
 	 */
-	UPROPERTY(EditAnywhere, BlueprintReadOnly, Category = "Cloud Tracing", AdvancedDisplay, meta = (UIMin = "0.25", UIMax = "8", ClampMin = "0.25", SliderExponent = 1.0))
+	UPROPERTY(EditAnywhere, BlueprintReadOnly, Category = "Cloud Tracing", AdvancedDisplay, meta = (UIMin = "0.05", UIMax = "8", ClampMin = "0.05", SliderExponent = 1.0))
 	float ShadowViewSampleCountScale;
 	/**
 	 * Scale the shadow tracing sample count in reflection views, only used with Advanced Output ray marched shadows. Quality level scalability CVARs affect the maximum range.
 	 * The sample count resolution is still clamped according to scalability setting to 'r.VolumetricCloud.Shadow.ReflectionRaySampleMaxCount'.
 	 */
-	UPROPERTY(EditAnywhere, BlueprintReadOnly, Category = "Cloud Tracing", AdvancedDisplay, meta = (UIMin = "0.25", UIMax = "8", ClampMin = "0.25", SliderExponent = 1.0))
-	float ShadowReflectionSampleCountScale;
+	UPROPERTY(EditAnywhere, BlueprintReadOnly, Category = "Cloud Tracing", AdvancedDisplay, meta = (UIMin = "0.05", UIMax = "8", ClampMin = "0.05", SliderExponent = 1.0))
+	float ShadowReflectionViewSampleCountScale;
+	UPROPERTY()
+	float ShadowReflectionSampleCountScale_DEPRECATED;
 
 	/**
 	 * The shadow tracing distance in kilometers, only used with Advanced Output ray marched shadows.
@@ -124,11 +128,11 @@ class UVolumetricCloudComponent : public USceneComponent
 	UFUNCTION(BlueprintCallable, Category = "Rendering")
 	ENGINE_API void SetViewSampleCountScale(float NewValue);
 	UFUNCTION(BlueprintCallable, Category = "Rendering")
-	ENGINE_API void SetReflectionSampleCountScale(float NewValue);
+	ENGINE_API void SetReflectionViewSampleCountScale(float NewValue);
 	UFUNCTION(BlueprintCallable, Category = "Rendering")
 	ENGINE_API void SetShadowViewSampleCountScale(float NewValue);
 	UFUNCTION(BlueprintCallable, Category = "Rendering")
-	ENGINE_API void SetShadowReflectionSampleCountScale(float NewValue);
+	ENGINE_API void SetShadowReflectionViewSampleCountScale(float NewValue);
 	UFUNCTION(BlueprintCallable, Category = "Rendering")
 	ENGINE_API void SetShadowTracingDistance(float NewValue);
 	UFUNCTION(BlueprintCallable, Category = "Rendering")
@@ -136,6 +140,14 @@ class UVolumetricCloudComponent : public USceneComponent
 	UFUNCTION(BlueprintCallable, Category = "Rendering")
 	ENGINE_API void SetMaterial(UMaterialInterface* NewValue);
 
+	// Deprecated functions but still valid because they forward data correctly.
+	UE_DEPRECATED(5.0, "This function has been replaced by SetReflectionViewSampleCountScale.")
+	UFUNCTION(BlueprintCallable, Category = "Rendering", meta = (DeprecatedFunction, DeprecationMessage = "This function has been replaced by SetReflectionViewSampleCountScale."))
+	ENGINE_API void SetReflectionSampleCountScale(float NewValue);
+	UE_DEPRECATED(5.0, "This function has been replaced by SetShadowReflectionViewSampleCountScale.")
+	UFUNCTION(BlueprintCallable, Category = "Rendering", meta = (DeprecatedFunction, DeprecationMessage = "This function has been replaced by SetShadowReflectionViewSampleCountScale."))
+	ENGINE_API void SetShadowReflectionSampleCountScale(float NewValue);
+
 
 protected:
 	//~ Begin UActorComponent Interface.
@@ -159,6 +171,14 @@ public:
 #endif // WITH_EDITOR
 	//~ End UActorComponent Interface.
 
+
+	// Those values should never be changed wihtout data conversion, that in order to maintain performance in case default values are used.
+	static constexpr float BaseViewRaySampleCount = 96.0f;
+	static constexpr float BaseShadowRaySampleCount = 10.0f;
+	// Those values are part of a data conversion and should never be changed. CVars and component sample count controls should be enough.
+	static constexpr float OldToNewReflectionViewRaySampleCount = 10.0f / BaseViewRaySampleCount;
+	static constexpr float OldToNewReflectionShadowRaySampleCount = 3.0f / BaseShadowRaySampleCount;
+
 private:
 
 	FVolumetricCloudSceneProxy* VolumetricCloudSceneProxy;
diff --git a/Engine/Source/Runtime/Engine/Private/Components/VolumetricCloudComponent.cpp b/Engine/Source/Runtime/Engine/Private/Components/VolumetricCloudComponent.cpp
--- a/Engine/Source/Runtime/Engine/Private/Components/VolumetricCloudComponent.cpp
+++ b/Engine/Source/Runtime/Engine/Private/Components/VolumetricCloudComponent.cpp
@@ -14,6 +14,7 @@
 #include "UObject/UObjectIterator.h"
 #include "UObject/ConstructorHelpers.h"
 #include "Materials/MaterialInstance.h"
+#include "UObject/UE5ReleaseStreamObjectVersion.h"
 
 #if WITH_EDITOR
 #include "ObjectEditorUtils.h"
@@ -39,9 +40,11 @@ UVolumetricCloudComponent::UVolumetricCloudComponent(const FObjectInitializer& O
 	, bUsePerSampleAtmosphericLightTransmittance(false)
 	, SkyLightCloudBottomOcclusion(0.5f)
 	, ViewSampleCountScale(1.0f)
-	, ReflectionSampleCountScale(1.0f)
+	, ReflectionViewSampleCountScale(0.15f)			// Roughly equivalent to previous default 1.0f, scaled by OldToNewReflectionViewRaySampleCount
+	, ReflectionSampleCountScale_DEPRECATED(1.0f)
 	, ShadowViewSampleCountScale(1.0f)
-	, ShadowReflectionSampleCountScale(1.0f)
+	, ShadowReflectionViewSampleCountScale(0.3f)	// Roughly equivalent to previous default 1.0f, scaled by OldToNewReflectionShadowRaySampleCount
+	, ShadowReflectionSampleCountScale_DEPRECATED(1.0f)
 	, ShadowTracingDistance(15.0f)
 	, StopTracingTransmittanceThreshold(0.005f)
 	, VolumetricCloudSceneProxy(nullptr)
@@ -121,6 +124,14 @@ void UVolumetricCloudComponent::PostInterpChange(FProperty* PropertyThatChanged)
 void UVolumetricCloudComponent::Serialize(FArchive& Ar)
 {
 	Super::Serialize(Ar);
+
+	Ar.UsingCustomVersion(FUE5ReleaseStreamObjectVersion::GUID);
+
+	if (Ar.CustomVer(FUE5ReleaseStreamObjectVersion::GUID) < FUE5ReleaseStreamObjectVersion::VolumetricCloudSampleCountUnification && Ar.IsLoading())
+	{
+		ReflectionViewSampleCountScale		= ReflectionSampleCountScale_DEPRECATED		  * UVolumetricCloudComponent::OldToNewReflectionViewRaySampleCount;
+		ShadowReflectionViewSampleCountScale= ShadowReflectionSampleCountScale_DEPRECATED * UVolumetricCloudComponent::OldToNewReflectionShadowRaySampleCount;
+	}
 }
 
 
@@ -142,13 +153,30 @@ CLOUD_DECLARE_BLUEPRINT_SETFUNCTION(FColor, GroundAlbedo);
 CLOUD_DECLARE_BLUEPRINT_SETFUNCTION(bool, bUsePerSampleAtmosphericLightTransmittance);
 CLOUD_DECLARE_BLUEPRINT_SETFUNCTION(float, SkyLightCloudBottomOcclusion);
 CLOUD_DECLARE_BLUEPRINT_SETFUNCTION(float, ViewSampleCountScale);
-CLOUD_DECLARE_BLUEPRINT_SETFUNCTION(float, ReflectionSampleCountScale);
+CLOUD_DECLARE_BLUEPRINT_SETFUNCTION(float, ReflectionViewSampleCountScale);
 CLOUD_DECLARE_BLUEPRINT_SETFUNCTION(float, ShadowViewSampleCountScale);
-CLOUD_DECLARE_BLUEPRINT_SETFUNCTION(float, ShadowReflectionSampleCountScale);
+CLOUD_DECLARE_BLUEPRINT_SETFUNCTION(float, ShadowReflectionViewSampleCountScale);
 CLOUD_DECLARE_BLUEPRINT_SETFUNCTION(float, ShadowTracingDistance);
 CLOUD_DECLARE_BLUEPRINT_SETFUNCTION(float, StopTracingTransmittanceThreshold);
 CLOUD_DECLARE_BLUEPRINT_SETFUNCTION(UMaterialInterface*, Material);
 
+void UVolumetricCloudComponent::SetReflectionSampleCountScale(float NewValue)
+{
+	if (AreDynamicDataChangesAllowed() && ReflectionViewSampleCountScale != NewValue)
+	{
+		ReflectionViewSampleCountScale = NewValue * UVolumetricCloudComponent::OldToNewReflectionViewRaySampleCount;
+		MarkRenderStateDirty();
+	}
+}
+void UVolumetricCloudComponent::SetShadowReflectionSampleCountScale(float NewValue)
+{
+	if (AreDynamicDataChangesAllowed() && ShadowReflectionViewSampleCountScale != NewValue)
+	{
+		ShadowReflectionViewSampleCountScale = NewValue * UVolumetricCloudComponent::OldToNewReflectionShadowRaySampleCount;
+		MarkRenderStateDirty();
+	}
+}
+
 /*=============================================================================
 	AVolumetricCloud implementation.
 =============================================================================*/
diff --git a/Engine/Source/Runtime/Engine/Private/VolumetricCloudProxy.cpp b/Engine/Source/Runtime/Engine/Private/VolumetricCloudProxy.cpp
--- a/Engine/Source/Runtime/Engine/Private/VolumetricCloudProxy.cpp
+++ b/Engine/Source/Runtime/Engine/Private/VolumetricCloudProxy.cpp
@@ -19,9 +19,9 @@ FVolumetricCloudSceneProxy::FVolumetricCloudSceneProxy(const UVolumetricCloudCom
 	, bUsePerSampleAtmosphericLightTransmittance(InComponent->bUsePerSampleAtmosphericLightTransmittance)
 	, SkyLightCloudBottomOcclusion(InComponent->SkyLightCloudBottomOcclusion)
 	, ViewSampleCountScale(InComponent->ViewSampleCountScale)
-	, ReflectionSampleCountScale(InComponent->ReflectionSampleCountScale)
+	, ReflectionViewSampleCountScale(InComponent->ReflectionViewSampleCountScale)
 	, ShadowViewSampleCountScale(InComponent->ShadowViewSampleCountScale)
-	, ShadowReflectionSampleCountScale(InComponent->ShadowReflectionSampleCountScale)
+	, ShadowReflectionViewSampleCountScale(InComponent->ShadowReflectionViewSampleCountScale)
 	, ShadowTracingDistance(InComponent->ShadowTracingDistance)
 	, StopTracingTransmittanceThreshold(InComponent->StopTracingTransmittanceThreshold)
 	, CloudVolumeMaterial(InComponent->Material)
diff --git a/Engine/Source/Runtime/Engine/Public/VolumetricCloudProxy.h b/Engine/Source/Runtime/Engine/Public/VolumetricCloudProxy.h
--- a/Engine/Source/Runtime/Engine/Public/VolumetricCloudProxy.h
+++ b/Engine/Source/Runtime/Engine/Public/VolumetricCloudProxy.h
@@ -41,9 +41,9 @@ public:
 	float SkyLightCloudBottomOcclusion;
 	
 	float ViewSampleCountScale;
-	float ReflectionSampleCountScale;
+	float ReflectionViewSampleCountScale;
 	float ShadowViewSampleCountScale;
-	float ShadowReflectionSampleCountScale;
+	float ShadowReflectionViewSampleCountScale;
 	float ShadowTracingDistance;
 	float StopTracingTransmittanceThreshold;
 private:
diff --git a/Engine/Source/Runtime/Renderer/Private/VolumetricCloudRendering.cpp b/Engine/Source/Runtime/Renderer/Private/VolumetricCloudRendering.cpp
--- a/Engine/Source/Runtime/Renderer/Private/VolumetricCloudRendering.cpp
+++ b/Engine/Source/Runtime/Renderer/Private/VolumetricCloudRendering.cpp
@@ -1489,11 +1489,9 @@ void FSceneRenderer::InitVolumetricCloudsForViews(FRDGBuilder& GraphBuilder, boo
 			CloudGlobalShaderParams.TracingStartMaxDistance = KilometersToCentimeters * CloudProxy.TracingStartMaxDistance;
 			CloudGlobalShaderParams.TracingMaxDistance		= KilometersToCentimeters * CloudProxy.TracingMaxDistance;
 
-			const float BaseViewRaySampleCount = 96.0f;
-			const float BaseShadowRaySampleCount = 10.0f;
 			CloudGlobalShaderParams.SampleCountMin		= FMath::Max(0, CVarVolumetricCloudSampleMinCount.GetValueOnAnyThread());
-			CloudGlobalShaderParams.SampleCountMax		= FMath::Max(2.0f, FMath::Min(BaseViewRaySampleCount   * CloudProxy.ViewSampleCountScale,       CVarVolumetricCloudViewRaySampleMaxCount.GetValueOnAnyThread()));
-			CloudGlobalShaderParams.ShadowSampleCountMax= FMath::Max(2.0f, FMath::Min(BaseShadowRaySampleCount * CloudProxy.ShadowViewSampleCountScale, CVarVolumetricCloudShadowViewRaySampleMaxCount.GetValueOnAnyThread()));
+			CloudGlobalShaderParams.SampleCountMax		= FMath::Max(2.0f, FMath::Min(UVolumetricCloudComponent::BaseViewRaySampleCount   * CloudProxy.ViewSampleCountScale,       CVarVolumetricCloudViewRaySampleMaxCount.GetValueOnAnyThread()));
+			CloudGlobalShaderParams.ShadowSampleCountMax= FMath::Max(2.0f, FMath::Min(UVolumetricCloudComponent::BaseShadowRaySampleCount * CloudProxy.ShadowViewSampleCountScale, CVarVolumetricCloudShadowViewRaySampleMaxCount.GetValueOnAnyThread()));
 			CloudGlobalShaderParams.ShadowTracingMaxDistance = KilometersToCentimeters * FMath::Max(0.1f, CloudProxy.ShadowTracingDistance);
 			CloudGlobalShaderParams.InvDistanceToSampleCountMax = 1.0f / FMath::Max(1.0f, KilometersToCentimeters * CVarVolumetricCloudDistanceToSampleMaxCount.GetValueOnAnyThread());
 			CloudGlobalShaderParams.StopTracingTransmittanceThreshold = FMath::Clamp(CloudProxy.StopTracingTransmittanceThreshold, 0.0f, 1.0f);
@@ -1995,14 +1993,11 @@ static TRDGUniformBufferRef<FRenderVolumetricCloudGlobalParameters> CreateCloudP
 	if (CloudRC.bIsReflectionRendering 
 		&& !(bIsPathTracing && CloudRC.bIsSkyRealTimeReflectionRendering))	// When using path tracing and real time sky capture, the captured sky cubemap is used to render the sky in the main view. As such, we always use the view settings.
 	{
-		const float BaseReflectionRaySampleCount = 10.0f;
-		const float BaseReflectionShadowRaySampleCount = 3.0f;
-
 		VolumetricCloudParams.VolumetricCloud.SampleCountMax = FMath::Max(2.0f, FMath::Min(
-			BaseReflectionRaySampleCount * CloudInfo.GetVolumetricCloudSceneProxy().ReflectionSampleCountScale,
+			UVolumetricCloudComponent::BaseViewRaySampleCount * CloudInfo.GetVolumetricCloudSceneProxy().ReflectionViewSampleCountScale,
 			CVarVolumetricCloudReflectionRaySampleMaxCount.GetValueOnAnyThread()));
 		VolumetricCloudParams.VolumetricCloud.ShadowSampleCountMax = FMath::Max(2.0f, FMath::Min(
-			BaseReflectionShadowRaySampleCount * CloudInfo.GetVolumetricCloudSceneProxy().ShadowReflectionSampleCountScale,
+			UVolumetricCloudComponent::BaseShadowRaySampleCount * CloudInfo.GetVolumetricCloudSceneProxy().ShadowReflectionViewSampleCountScale,
 			CVarVolumetricCloudShadowReflectionRaySampleMaxCount.GetValueOnAnyThread()));
 	}
 
