- Move the code that put back the running platform after the FSkeletalMeshRenderData::Serialize call so the serialize will use the correct morphtargets list.
- Initialize the morph targets during the cache when there is a valid DDC. I remove FFinishBuildInternalData and fix thread safe issue with smartnames, so we can initialize them outside of the main thread.

#rnx

diff --git a/Engine/Source/Runtime/Engine/Classes/Animation/SmartName.h b/Engine/Source/Runtime/Engine/Classes/Animation/SmartName.h
--- a/Engine/Source/Runtime/Engine/Classes/Animation/SmartName.h
+++ b/Engine/Source/Runtime/Engine/Classes/Animation/SmartName.h
@@ -153,6 +153,18 @@ struct ENGINE_API FSmartNameMapping
 	void Iterate(TFunction<void(const FSmartNameMapping* Mapping, SmartName::UID_Type ID)> Callback) const;
 
 private:
+	/*Internal no lock function to prevent re-entrant locking, see API function GetName for documentation.*/
+	bool GetName_NoLock(const SmartName::UID_Type& Uid, FName& OutName) const;
+
+	/*Internal no lock function to prevent re-entrant locking, see API function Exists for documentation.*/
+	bool Exists_NoLock(const SmartName::UID_Type& Uid) const;
+	
+	/*Internal no lock function to prevent re-entrant locking, see API function Exists for documentation.*/
+	bool Exists_NoLock(const FName& Name) const;
+
+	/*Internal no lock function to prevent re-entrant locking, see API function FindUID for documentation.*/
+	SmartName::UID_Type FindUID_NoLock(const FName& Name) const;
+
 	// List of curve names, indexed by UID
 	TArray<FName> CurveNameList;
 
diff --git a/Engine/Source/Runtime/Engine/Classes/Engine/SkeletalMesh.h b/Engine/Source/Runtime/Engine/Classes/Engine/SkeletalMesh.h
--- a/Engine/Source/Runtime/Engine/Classes/Engine/SkeletalMesh.h
+++ b/Engine/Source/Runtime/Engine/Classes/Engine/SkeletalMesh.h
@@ -125,17 +125,6 @@ enum class ESkeletalMeshAsyncProperties : uint64
 
 ENUM_CLASS_FLAGS(ESkeletalMeshAsyncProperties);
 
-#if WITH_EDITOR
-struct FFinishBuildInternalData
-{
-public:
-	void ApplyEditorData(USkeletalMesh* SkeletalMesh) const;
-
-	/** Editor MorphTargets data to set. */
-	TUniquePtr<FFinishBuildMorphTargetData> MorphTargetData;
-};
-#endif
-
 class FSkeletalMeshCompilationContext
 {
 public:
@@ -146,9 +135,6 @@ public:
 	// Movable
 	FSkeletalMeshCompilationContext(FSkeletalMeshCompilationContext&&) = default;
 	FSkeletalMeshCompilationContext& operator=(FSkeletalMeshCompilationContext&&) = default;
-#if WITH_EDITOR
-	FFinishBuildInternalData FinishBuildInternalData;
-#endif
 };
 
 class FSkeletalMeshPostLoadContext : public FSkeletalMeshCompilationContext
diff --git a/Engine/Source/Runtime/Engine/Private/Animation/SmartName.cpp b/Engine/Source/Runtime/Engine/Private/Animation/SmartName.cpp
--- a/Engine/Source/Runtime/Engine/Private/Animation/SmartName.cpp
+++ b/Engine/Source/Runtime/Engine/Private/Animation/SmartName.cpp
@@ -50,15 +50,15 @@ FSmartName FSmartNameMapping::AddName(FName Name)
 
 FCurveMetaData* FSmartNameMapping::GetCurveMetaData(FName CurveName)
 {
-	checkSlow(Exists(CurveName));
 	FWriteScopeLock Lock(SmartNameRWLock);
+	checkSlow(Exists_NoLock(CurveName));
 	return &CurveMetaDataMap.FindOrAdd(CurveName);
 }
 
 const FCurveMetaData* FSmartNameMapping::GetCurveMetaData(FName CurveName) const
 {
-	checkSlow(Exists(CurveName));
 	FReadScopeLock Lock(SmartNameRWLock);
+	checkSlow(Exists_NoLock(CurveName));
 	return CurveMetaDataMap.Find(CurveName);
 }
 
@@ -70,10 +70,8 @@ const FCurveMetaData& FSmartNameMapping::GetCurveMetaData(SmartName::UID_Type Cu
 }
 #endif
 
-bool FSmartNameMapping::GetName(const SmartName::UID_Type& Uid, FName& OutName) const
+bool FSmartNameMapping::GetName_NoLock(const SmartName::UID_Type& Uid, FName& OutName) const
 {
-	FReadScopeLock Lock(SmartNameRWLock);
-	
 	if (CurveNameList.IsValidIndex(Uid))
 	{
 		OutName = CurveNameList[Uid];
@@ -82,13 +80,19 @@ bool FSmartNameMapping::GetName(const SmartName::UID_Type& Uid, FName& OutName)
 	return false;
 }
 
+bool FSmartNameMapping::GetName(const SmartName::UID_Type& Uid, FName& OutName) const
+{
+	FReadScopeLock Lock(SmartNameRWLock);
+	return GetName_NoLock(Uid, OutName);
+}
+
 #if WITH_EDITOR
 bool FSmartNameMapping::Rename(const SmartName::UID_Type& Uid, FName NewName)
 {
 	FWriteScopeLock Lock(SmartNameRWLock);
 	
 	FName ExistingName;
-	if(GetName(Uid, ExistingName))
+	if(GetName_NoLock(Uid, ExistingName))
 	{
 		// fix up meta data
 		FCurveMetaData* MetaDataToCopy = CurveMetaDataMap.Find(ExistingName);
@@ -112,7 +116,7 @@ bool FSmartNameMapping::Remove(const SmartName::UID_Type& Uid)
 	FWriteScopeLock Lock(SmartNameRWLock);
 	
 	FName ExistingName;
-	if (GetName(Uid, ExistingName))
+	if (GetName_NoLock(Uid, ExistingName))
 	{
 		CurveMetaDataMap.Remove(ExistingName);
 		CurveNameList[Uid] = NAME_None;
@@ -126,7 +130,7 @@ bool FSmartNameMapping::Remove(const FName& Name)
 {
 	FWriteScopeLock Lock(SmartNameRWLock);
 	
-	const SmartName::UID_Type Uid = FindUID(Name);
+	const SmartName::UID_Type Uid = FindUID_NoLock(Name);
 	if (Uid != SmartName::MaxUID)
 	{
 		CurveMetaDataMap.Remove(Name);
@@ -265,23 +269,37 @@ void FSmartNameMapping::FillUIDToCurveTypeArray(TArray<FAnimCurveType>& Array) c
 	}
 }
 
+bool FSmartNameMapping::Exists_NoLock(const SmartName::UID_Type& Uid) const
+{
+	return CurveNameList.IsValidIndex(Uid) && CurveNameList[Uid] != NAME_None;
+}
 
 bool FSmartNameMapping::Exists(const SmartName::UID_Type& Uid) const
 {
 	FReadScopeLock Lock(SmartNameRWLock);
-	return CurveNameList.IsValidIndex(Uid) && CurveNameList[Uid] != NAME_None;
+	return Exists_NoLock(Uid);
+}
+
+bool FSmartNameMapping::Exists_NoLock(const FName& Name) const
+{
+	return CurveNameList.Contains(Name);
 }
 
 bool FSmartNameMapping::Exists(const FName& Name) const
 {
 	FReadScopeLock Lock(SmartNameRWLock);
-	return CurveNameList.Contains(Name);
+	return Exists_NoLock(Name);
+}
+
+SmartName::UID_Type FSmartNameMapping::FindUID_NoLock(const FName& Name) const
+{
+	return CurveNameList.IndexOfByKey(Name);
 }
 
 SmartName::UID_Type FSmartNameMapping::FindUID(const FName& Name) const
 {
 	FReadScopeLock Lock(SmartNameRWLock);
-	return CurveNameList.IndexOfByKey(Name);
+	return FindUID_NoLock(Name);
 }
 
 FArchive& operator<<(FArchive& Ar, FSmartNameMapping& Elem)
@@ -294,7 +312,7 @@ FArchive& operator<<(FArchive& Ar, FSmartNameMapping& Elem)
 bool FSmartNameMapping::FindSmartName(FName Name, FSmartName& OutName) const
 {
 	FReadScopeLock Lock(SmartNameRWLock);
-	SmartName::UID_Type ExistingUID = FindUID(Name);
+	SmartName::UID_Type ExistingUID = FindUID_NoLock(Name);
 	if (ExistingUID != SmartName::MaxUID)
 	{
 		OutName = FSmartName(Name, ExistingUID);
@@ -308,7 +326,7 @@ bool FSmartNameMapping::FindSmartNameByUID(SmartName::UID_Type UID, FSmartName&
 {
 	FReadScopeLock Lock(SmartNameRWLock);
 	FName ExistingName;
-	if (GetName(UID, ExistingName))
+	if (GetName_NoLock(UID, ExistingName))
 	{
 		OutName.DisplayName = ExistingName;
 		OutName.UID = UID;
diff --git a/Engine/Source/Runtime/Engine/Private/SkeletalMesh.cpp b/Engine/Source/Runtime/Engine/Private/SkeletalMesh.cpp
--- a/Engine/Source/Runtime/Engine/Private/SkeletalMesh.cpp
+++ b/Engine/Source/Runtime/Engine/Private/SkeletalMesh.cpp
@@ -1684,30 +1684,12 @@ bool USkeletalMesh::IsReadyForFinishDestroy()
 }
 
 #if WITH_EDITOR
-void CacheRunningPlatform(USkeletalMesh* Mesh)
+void CachePlatform(USkeletalMesh* Mesh, const ITargetPlatform* TargetPlatform, FSkeletalMeshRenderData* PlatformRenderData)
 {
-	//Cache the running platform, dcc should be valid so it will be fast
-	ITargetPlatform* RunningPlatform = GetTargetPlatformManagerRef().GetRunningTargetPlatform();
-	check(RunningPlatform);
+	//Cache the platform, dcc should be valid so it will be fast
+	check(TargetPlatform);
 	FSkeletalMeshBuildContext Context;
-	FSkeletalMeshRenderData RunningPlatformRenderData;
-	RunningPlatformRenderData.Cache(RunningPlatform, Mesh, &Context);
-
-	auto ApplyContext = [Mesh, Context = MoveTemp(Context)]()
-	{
-		check(IsInGameThread());
-		Context.FinishBuildInternalData.ApplyEditorData(Mesh);
-	};
-
-	if (IsInGameThread())
-	{
-		ApplyContext();
-	}
-	else
-	{
-		//Always apply context data on the game thread
-		Async(EAsyncExecution::TaskGraphMainThread, MoveTemp(ApplyContext));
-	}
+	PlatformRenderData->Cache(TargetPlatform, Mesh, &Context);
 }
 
 FString BuildSkeletalMeshDerivedDataKey(const ITargetPlatform* TargetPlatform, USkeletalMesh* SkelMesh);
@@ -1731,30 +1713,10 @@ static FSkeletalMeshRenderData& GetPlatformSkeletalMeshRenderData(USkeletalMesh*
 	{
 		// Cache render data for this platform and insert it in to the linked list.
 		PlatformRenderData = new FSkeletalMeshRenderData();
-		FSkeletalMeshBuildContext Context;
-		PlatformRenderData->Cache(TargetPlatform, Mesh, &Context);
-		//No need to apply the context editor data to the mesh when the target is not the running target
-		
+		CachePlatform(Mesh, TargetPlatform, PlatformRenderData);
 		check(PlatformRenderData->DerivedDataKey == PlatformDerivedDataKey);
 		Swap(PlatformRenderData->NextCachedRenderData, Mesh->GetResourceForRendering()->NextCachedRenderData);
 		Mesh->GetResourceForRendering()->NextCachedRenderData = TUniquePtr<FSkeletalMeshRenderData>(PlatformRenderData);
-
-		{
-			//If the running platform DDC key is not equal to the target platform DDC key.
-			//We need to cache the skeletalmesh ddc with the running platform to retrieve the ddc editor data LODModel which can be different because of chunking and reduction
-			//Normally it should just take back the ddc for the running platform, since the ddc was cache when we have load the asset to cook it.
-			ITargetPlatform* RunningPlatform = GetTargetPlatformManagerRef().GetRunningTargetPlatform();
-			check(RunningPlatform);
-			//The running platform was loaded before getting here, so we should not have an null PlatformRenderData if someone call
-			//GetPlatformSkeletalMeshRenderData with a platform using the same ddc key
-			const FString RunningPlatformDerivedDataKey = BuildSkeletalMeshDerivedDataKey(RunningPlatform, Mesh);
-			if (RunningPlatformDerivedDataKey != PlatformDerivedDataKey)
-			{
-				//Force a cache of the running platform to put back the running platform editor data (FSkeletalMeshLODModel and MorphTargets)
-				//Editor data do not have "DDC key linked list". Maybe we should move this data in the FSkeletalMeshRenderData.
-				CacheRunningPlatform(Mesh);
-			}
-		}
 	}
 	check(PlatformRenderData->DerivedDataKey == PlatformDerivedDataKey);
 	check(PlatformRenderData);
@@ -1841,6 +1803,7 @@ void USkeletalMesh::Serialize( FArchive& Ar )
 			}
 			else if (Ar.IsSaving())
 			{
+				ITargetPlatform* RunningPlatform = GetTargetPlatformManagerRef().GetRunningTargetPlatform();
 				FSkeletalMeshRenderData* LocalSkeletalMeshRenderData = GetSkeletalMeshRenderData();
 #if WITH_EDITORONLY_DATA
 				const ITargetPlatform* ArchiveCookingTarget = Ar.CookingTarget();
@@ -1851,7 +1814,6 @@ void USkeletalMesh::Serialize( FArchive& Ar )
 				else
 				{
 					//Fall back in case we use an archive that the cooking target has not been set (i.e. Duplicate archive)
-					ITargetPlatform* RunningPlatform = GetTargetPlatformManagerRef().GetRunningTargetPlatform();
 					check(RunningPlatform != NULL);
 					LocalSkeletalMeshRenderData = &GetPlatformSkeletalMeshRenderData(this, RunningPlatform);
 				}
@@ -1879,6 +1841,22 @@ void USkeletalMesh::Serialize( FArchive& Ar )
 					}
 				}
 				LocalSkeletalMeshRenderData->Serialize(Ar, this);
+
+#if WITH_EDITOR
+				//We need to get the skeletalmesh ddc cache with the running platform to retrieve the ddc editor data LODModel which can be different because of chunking, reduction and morph targets.
+				//Normally it should just take back the ddc for the running platform, since the ddc was cache when we have load the asset to cook it.
+				if(ArchiveCookingTarget && ArchiveCookingTarget != RunningPlatform)
+				{
+					check(RunningPlatform != NULL);
+					const FString CookPlatformDerivedDataKey = BuildSkeletalMeshDerivedDataKey(ArchiveCookingTarget, this);
+					const FString RunningPlatformDerivedDataKey = BuildSkeletalMeshDerivedDataKey(RunningPlatform, this);
+					if (CookPlatformDerivedDataKey != RunningPlatformDerivedDataKey)
+					{
+						FSkeletalMeshRenderData RunningPlatformSkeletalMeshRenderData;
+						CachePlatform(this, RunningPlatform, &RunningPlatformSkeletalMeshRenderData);
+					}
+				}
+#endif
 			}
 		}
 	}
@@ -2093,15 +2071,6 @@ uint32 USkeletalMesh::GetVertexBufferFlags() const
 
 #if WITH_EDITOR
 
-void FFinishBuildInternalData::ApplyEditorData(USkeletalMesh* SkeletalMesh) const
-{
-	check(IsInGameThread());
-	if (MorphTargetData)
-	{
-		MorphTargetData->ApplyEditorData(SkeletalMesh);
-	}
-}
-
 thread_local const USkeletalMesh* FSkeletalMeshAsyncBuildScope::SkeletalMeshBeingAsyncCompiled = nullptr;
 
 FSkeletalMeshAsyncBuildScope::FSkeletalMeshAsyncBuildScope(const USkeletalMesh* SkeletalMesh)
@@ -2300,7 +2269,6 @@ void USkeletalMesh::ApplyFinishBuildInternalData(FSkeletalMeshCompilationContext
 	//We cannot execute this code outside of the game thread
 	checkf(IsInGameThread(), TEXT("Cannot execute function USkeletalMesh::ApplyFinishBuildInternalData asynchronously. Asset: %s"), *this->GetFullName());
 	check(ContextPtr);
-	ContextPtr->FinishBuildInternalData.ApplyEditorData(this);
 }
 
 void USkeletalMesh::FinishBuildInternal(FSkeletalMeshBuildContext& Context)
@@ -3671,11 +3639,6 @@ void USkeletalMesh::FinishPostLoadInternal(FSkeletalMeshPostLoadContext& Context
 
 #if WITH_EDITOR
 	ApplyFinishBuildInternalData(&Context);
-	//If we do not have a valid LOD import data for base LOD, init the morph target. Asset imported in 4.23 or before will not have available build data for the base LOD
-	if (!IsLODImportedDataBuildAvailable(0))
-	{
-		InitMorphTargets();
-	}
 #endif
 	// should do this before InitResources.
 	InitMorphTargets();
diff --git a/Engine/Source/Runtime/Engine/Private/SkeletalMeshRenderData.cpp b/Engine/Source/Runtime/Engine/Private/SkeletalMeshRenderData.cpp
--- a/Engine/Source/Runtime/Engine/Private/SkeletalMeshRenderData.cpp
+++ b/Engine/Source/Runtime/Engine/Private/SkeletalMeshRenderData.cpp
@@ -21,6 +21,7 @@
 #include "Serialization/LargeMemoryReader.h"
 #include "Serialization/LargeMemoryWriter.h"
 #include "Async/Async.h"
+#include "UObject/GarbageCollection.h"
 
 #if ENABLE_COOK_STATS
 namespace SkeletalMeshCookStats
@@ -157,6 +158,69 @@ namespace DDCUtils64Bit
 	}
 } //namespace DDCUtils64Bit
 
+namespace MorphTargetUtils
+{
+	void ApplyMorphTargetsEditorData(USkeletalMesh* SkeletalMesh, TMap<FName, TArray<FMorphTargetLODModel> >& MorphLODModelsPerTargetName)
+	{
+		//Make sure we do not create new morph target during a gc
+		FGCScopeGuard GCScopeGuard;
+
+		FSkeletalMeshModel* SkelMeshModel = SkeletalMesh->GetImportedModel();
+		check(SkelMeshModel);
+
+		TMap<FName, UMorphTarget*> ExistingMorphTargets;
+		ExistingMorphTargets.Reserve(SkeletalMesh->GetMorphTargets().Num());
+		for (UMorphTarget* MorphTarget : SkeletalMesh->GetMorphTargets())
+		{
+			ExistingMorphTargets.Add(MorphTarget->GetFName(), MorphTarget);
+		}
+		int32 MorphTargetNumber = MorphLODModelsPerTargetName.Num();
+		TArray<UMorphTarget*> ToDeleteMorphTargets;
+		ToDeleteMorphTargets.Append(SkeletalMesh->GetMorphTargets());
+		SkeletalMesh->GetMorphTargets().Empty();
+		//Rebuild the MorphTarget object
+		for (TPair<FName, TArray<FMorphTargetLODModel>> TargetNameAndMorphLODModels : MorphLODModelsPerTargetName)
+		{
+			FName MorphTargetName = TargetNameAndMorphLODModels.Key;
+			UMorphTarget* MorphTarget = ExistingMorphTargets.FindRef(MorphTargetName);
+			if (!MorphTarget)
+			{
+				MorphTarget = NewObject<UMorphTarget>(SkeletalMesh, MorphTargetName);
+				check(MorphTarget);
+			}
+			else
+			{
+				ToDeleteMorphTargets.Remove(MorphTarget);
+			}
+			MorphTarget->EmptyMorphLODModels();
+			SkeletalMesh->GetMorphTargets().Add(MorphTarget);
+			const TArray<FMorphTargetLODModel>& MorphTargetLODModels = TargetNameAndMorphLODModels.Value;
+			int32 MorphLODModelNumber = MorphTargetLODModels.Num();
+			MorphTarget->GetMorphLODModels().AddDefaulted(MorphLODModelNumber);
+			for (int32 MorphDataIndex = 0; MorphDataIndex < MorphLODModelNumber; ++MorphDataIndex)
+			{
+				MorphTarget->GetMorphLODModels()[MorphDataIndex] = MorphTargetLODModels[MorphDataIndex];
+			}
+		}
+		//Rebuild the mapping and rehook the curve data
+		SkeletalMesh->InitMorphTargets();
+
+		//Clear any async flags after the morphtargets have been set to the skeletalmesh
+		for (UMorphTarget* MorphTarget : SkeletalMesh->GetMorphTargets())
+		{
+			const EInternalObjectFlags AsyncFlags = EInternalObjectFlags::Async | EInternalObjectFlags::AsyncLoading;
+			MorphTarget->ClearInternalFlags(AsyncFlags);
+		}
+
+		//Make sure the old unused morph targets are clean up properly
+		for (UMorphTarget* ToDeleteMorphTarget : ToDeleteMorphTargets)
+		{
+			ToDeleteMorphTarget->BaseSkelMesh = nullptr;
+			ToDeleteMorphTarget->EmptyMorphLODModels();
+			ToDeleteMorphTarget->MarkAsGarbage();
+		}
+	}
+}
 //Serialize the LODInfo and append the result to the KeySuffix to build the LODInfo part of the DDC KEY
 //Note: this serializer is only used to build the mesh DDC key, no versioning is required
 static void SerializeLODInfoForDDC(USkeletalMesh* SkeletalMesh, FString& KeySuffix)
@@ -359,6 +423,9 @@ void FSkeletalMeshRenderData::Cache(const ITargetPlatform* TargetPlatform, USkel
 
 			FLargeMemoryReader Ar(DerivedData.GetData(), DerivedData.Num(), ELargeMemoryReaderFlags::Persistent);
 
+			//Helper structure to change the morph targets
+			TMap<FName, TArray<FMorphTargetLODModel> > MorphLODModelsPerTargetName;
+
 			//With skeletal mesh build refactor we serialize the LODModel data into the DDC
 			//We need to store those so we do not have to rerun the reduction to make them up to date
 			//with the serialize renderdata. This allow to use DDC when changing the reduction settings.
@@ -370,17 +437,24 @@ void FSkeletalMeshRenderData::Cache(const ITargetPlatform* TargetPlatform, USkel
 
 				//Get the morph target data, we put it in the compilation context to apply them in the game thread before the InitResources
 				{
-					if (Owner->GetMorphTargets().Num() > 0)
-					{
-						ContextPtr->FinishBuildInternalData.MorphTargetData = Owner->GetMorphTargets()[0]->CreateFinishBuildMorphTargetData();
-					}
-					else
+					int32 MorphTargetNumber = 0;
+					Ar << MorphTargetNumber;
+					MorphLODModelsPerTargetName.Reserve(MorphTargetNumber);
+					//We cannot serialize directly the UMorphTarget with a FMemoryArchive. This is not supported.
+					for (int32 MorphTargetIndex = 0; MorphTargetIndex < MorphTargetNumber; ++MorphTargetIndex)
 					{
-						// Create and initialize the FinishBuildInternalData, use the class default object to call the virtual function
-						ContextPtr->FinishBuildInternalData.MorphTargetData = UMorphTarget::StaticClass()->GetDefaultObject<UMorphTarget>()->CreateFinishBuildMorphTargetData();
+						FName MorphTargetName = NAME_None;
+						Ar << MorphTargetName;
+						TArray<FMorphTargetLODModel>& MorphTargetLODModels = MorphLODModelsPerTargetName.FindOrAdd(MorphTargetName);
+						int32 MorphLODModelNumber = 0;
+						Ar << MorphLODModelNumber;
+						MorphTargetLODModels.Empty(MorphLODModelNumber);
+						MorphTargetLODModels.AddDefaulted(MorphLODModelNumber);
+						for (int32 MorphDataIndex = 0; MorphDataIndex < MorphLODModelNumber; ++MorphDataIndex)
+						{
+							Ar << MorphTargetLODModels[MorphDataIndex];
+						}
 					}
-					check(ContextPtr->FinishBuildInternalData.MorphTargetData);
-					ContextPtr->FinishBuildInternalData.MorphTargetData->LoadFromMemoryArchive(Ar);
 				}
 
 				//Serialize the LODModel sections since they are dependent on the reduction
@@ -406,6 +480,12 @@ void FSkeletalMeshRenderData::Cache(const ITargetPlatform* TargetPlatform, USkel
 				LODData.SerializeStreamedData(Ar, Owner, LODIndex, DummyStripFlags, bNeedsCPUAccess, bForceKeepCPUResources);
 			}
 
+			//Apply the morphtargets change if any
+			if (MorphLODModelsPerTargetName.Num() > 0)
+			{
+				MorphTargetUtils::ApplyMorphTargetsEditorData(Owner, MorphLODModelsPerTargetName);
+			}
+
 			int32 T1 = FPlatformTime::Cycles();
 			UE_LOG(LogSkeletalMesh, Verbose, TEXT("Skeletal Mesh found in DDC [%fms] %s"), FPlatformTime::ToMilliseconds(T1 - T0), *Owner->GetPathName());
 		}
