World Partition
- Handle already uninitialized world partition object during garbage collection (unordered BeginDestroy calls).
- Move world partition initialization delegates from the world partition subsystem to the world partition object.
- Remove remnants of partitioned subworlds support.

diff --git a/Engine/Source/Runtime/Engine/Private/WorldPartition/HLOD/HLODSubsystem.cpp b/Engine/Source/Runtime/Engine/Private/WorldPartition/HLOD/HLODSubsystem.cpp
--- a/Engine/Source/Runtime/Engine/Private/WorldPartition/HLOD/HLODSubsystem.cpp
+++ b/Engine/Source/Runtime/Engine/Private/WorldPartition/HLOD/HLODSubsystem.cpp
@@ -123,14 +123,17 @@ void UHLODSubsystem::Initialize(FSubsystemCollectionBase& Collection)
 
 	if (World->IsGameWorld())
 	{
-		World->GetSubsystem<UWorldPartitionSubsystem>()->OnWorldPartitionRegistered.AddUObject(this, &UHLODSubsystem::OnWorldPartitionRegistered);
-		World->GetSubsystem<UWorldPartitionSubsystem>()->OnWorldPartitionUnregistered.AddUObject(this, &UHLODSubsystem::OnWorldPartitionUnregistered);
+		UWorldPartition* WorldPartition = GetWorld()->GetWorldPartition();
+		check(WorldPartition);
+
+		WorldPartition->OnWorldPartitionInitialized.AddUObject(this, &UHLODSubsystem::OnWorldPartitionInitialized);
+		WorldPartition->OnWorldPartitionUninitialized.AddUObject(this, &UHLODSubsystem::OnWorldPartitionUninitialized);
 
 		SceneViewExtension = FSceneViewExtensions::NewExtension<FHLODResourcesResidencySceneViewExtension>(World);
 	}
 }
 
-void UHLODSubsystem::OnWorldPartitionRegistered(UWorldPartition* InWorldPartition)
+void UHLODSubsystem::OnWorldPartitionInitialized(UWorldPartition* InWorldPartition)
 {
 	check(InWorldPartition == GetWorld()->GetWorldPartition());
 	check(CellsData.IsEmpty());
@@ -145,7 +148,7 @@ void UHLODSubsystem::OnWorldPartitionRegistered(UWorldPartition* InWorldPartitio
 	}
 }
 
-void UHLODSubsystem::OnWorldPartitionUnregistered(UWorldPartition* InWorldPartition)
+void UHLODSubsystem::OnWorldPartitionUninitialized(UWorldPartition* InWorldPartition)
 {
 	check(InWorldPartition == GetWorld()->GetWorldPartition());
 	CellsData.Reset();
diff --git a/Engine/Source/Runtime/Engine/Private/WorldPartition/WorldPartition.cpp b/Engine/Source/Runtime/Engine/Private/WorldPartition/WorldPartition.cpp
--- a/Engine/Source/Runtime/Engine/Private/WorldPartition/WorldPartition.cpp
+++ b/Engine/Source/Runtime/Engine/Private/WorldPartition/WorldPartition.cpp
@@ -261,13 +261,6 @@ UWorldPartition::UWorldPartition(const FObjectInitializer& ObjectInitializer)
 	, Replay(nullptr)
 #endif
 {
-	static bool bRegisteredDelegate = false;
-	if (!bRegisteredDelegate)
-	{
-		FWorldDelegates::LevelRemovedFromWorld.AddStatic(&UWorldPartition::WorldPartitionOnLevelRemovedFromWorld);
-		bRegisteredDelegate = true;
-	}
-
 #if WITH_EDITOR
 	WorldPartitionStreamingPolicyClass = UWorldPartitionLevelStreamingPolicy::StaticClass();
 #endif
@@ -502,6 +495,8 @@ void UWorldPartition::Initialize(UWorld* InWorld, const FTransform& InTransform)
 		FWorldPartitionLevelHelper::RemapLevelSoftObjectPaths(World->PersistentLevel, this);
 	}
 #endif
+
+	OnWorldPartitionInitialized.Broadcast(this);
 }
 
 void UWorldPartition::RegisterStreamingSourceProvider(IWorldPartitionStreamingSourceProvider* StreamingSource)
@@ -514,15 +509,6 @@ bool UWorldPartition::UnregisterStreamingSourceProvider(IWorldPartitionStreaming
 	return !!StreamingSourceProviders.Remove(StreamingSource);
 }
 
-// This will trap all broadcast of LevelRemovedFromWorld and Uninitialize world partition if existing
-void UWorldPartition::WorldPartitionOnLevelRemovedFromWorld(ULevel* Level, UWorld* InWorld)
-{
-	if (UWorldPartition* WorldPartition = Level ? Level->GetWorldPartition() : nullptr)
-	{
-		WorldPartition->Uninitialize();
-	}
-}
-
 void UWorldPartition::Uninitialize()
 {
 	if (IsInitialized())
@@ -573,6 +559,8 @@ void UWorldPartition::Uninitialize()
 #endif		
 
 		InitState = EWorldPartitionInitState::Uninitialized;
+
+		OnWorldPartitionUninitialized.Broadcast(this);
 	}
 
 	Super::Uninitialize();
diff --git a/Engine/Source/Runtime/Engine/Private/WorldPartition/WorldPartitionSubsystem.cpp b/Engine/Source/Runtime/Engine/Private/WorldPartition/WorldPartitionSubsystem.cpp
--- a/Engine/Source/Runtime/Engine/Private/WorldPartition/WorldPartitionSubsystem.cpp
+++ b/Engine/Source/Runtime/Engine/Private/WorldPartition/WorldPartitionSubsystem.cpp
@@ -143,8 +143,6 @@ void UWorldPartitionSubsystem::PostInitialize()
 	{
 		WorldPartition->Initialize(GetWorld(), FTransform::Identity);
 
-		OnWorldPartitionRegistered.Broadcast(WorldPartition);
-
 		if (WorldPartition->CanDrawRuntimeHash() && (GetWorld()->GetNetMode() != NM_DedicatedServer))
 		{
 			DrawHandle = UDebugDrawService::Register(TEXT("Game"), FDebugDrawDelegate::CreateUObject(this, &UWorldPartitionSubsystem::Draw));
@@ -172,11 +170,6 @@ void UWorldPartitionSubsystem::Deinitialize()
 	}
 #endif 
 
-	UWorldPartition* WorldPartition = GetWorldPartition();
-	
-	check(WorldPartition);
-	check(WorldPartition->IsInitialized());
-
 	if (GetWorld()->IsGameWorld())
 	{
 		GLevelStreamingContinuouslyIncrementalGCWhileLevelsPendingPurge = LevelStreamingContinuouslyIncrementalGCWhileLevelsPendingPurge;
@@ -189,10 +182,12 @@ void UWorldPartitionSubsystem::Deinitialize()
 		DrawHandle.Reset();
 	}
 
-	// Uninitialize registered world partitions
-	WorldPartition->Uninitialize();
-
-	OnWorldPartitionUnregistered.Broadcast(WorldPartition);
+	// During garbage collection, the world partition object can be uninitialized before the subsystem due to unordered calls to BeginDestroy
+	if (UWorldPartition* WorldPartition = GetWorldPartition(); WorldPartition && WorldPartition->IsInitialized())
+	{
+		// Uninitialize registered world partitions
+		WorldPartition->Uninitialize();
+	}
 
 	Super::Deinitialize();
 }
diff --git a/Engine/Source/Runtime/Engine/Public/WorldPartition/HLOD/HLODSubsystem.h b/Engine/Source/Runtime/Engine/Public/WorldPartition/HLOD/HLODSubsystem.h
--- a/Engine/Source/Runtime/Engine/Public/WorldPartition/HLOD/HLODSubsystem.h
+++ b/Engine/Source/Runtime/Engine/Public/WorldPartition/HLOD/HLODSubsystem.h
@@ -74,8 +74,8 @@ private:
 	TMap<FName, FCellData> CellsData;
 	TSet<FCellData*> CellsToWarmup;
 
-	void OnWorldPartitionRegistered(UWorldPartition* InWorldPartition);
-	void OnWorldPartitionUnregistered(UWorldPartition* InWorldPartition);
+	void OnWorldPartitionInitialized(UWorldPartition* InWorldPartition);
+	void OnWorldPartitionUninitialized(UWorldPartition* InWorldPartition);
 	void OnBeginRenderViews(const FSceneViewFamily& InViewFamily);
 
 	void MakeRenderResourcesResident(const FCellData& CellHLODs, const FSceneViewFamily& InViewFamily);
diff --git a/Engine/Source/Runtime/Engine/Public/WorldPartition/WorldPartition.h b/Engine/Source/Runtime/Engine/Public/WorldPartition/WorldPartition.h
--- a/Engine/Source/Runtime/Engine/Public/WorldPartition/WorldPartition.h
+++ b/Engine/Source/Runtime/Engine/Public/WorldPartition/WorldPartition.h
@@ -75,13 +75,19 @@ class ENGINE_API UWorldPartition final : public UActorDescContainer
 	friend class FWorldPartitionEditorModule;
 	friend class FUnrealEdMisc;
 
-#if WITH_EDITOR
 public:
+#if WITH_EDITOR
 	static UWorldPartition* CreateOrRepairWorldPartition(AWorldSettings* WorldSettings, TSubclassOf<UWorldPartitionEditorHash> EditorHashClass = nullptr, TSubclassOf<UWorldPartitionRuntimeHash> RuntimeHashClass = nullptr);
 	
 	DECLARE_MULTICAST_DELEGATE_OneParam(FCancelWorldPartitionUpdateEditorCellsDelegate, UWorldPartition*);
 	FCancelWorldPartitionUpdateEditorCellsDelegate OnCancelWorldPartitionUpdateEditorCells;
+#endif
 
+	DECLARE_MULTICAST_DELEGATE_OneParam(FWorldPartitionInitializeDelegate, UWorldPartition*);
+	FWorldPartitionInitializeDelegate OnWorldPartitionInitialized;
+	FWorldPartitionInitializeDelegate OnWorldPartitionUninitialized;
+
+#if WITH_EDITOR
 	TArray<FName> GetUserLoadedEditorGridCells() const;
 private:
 
@@ -111,8 +117,6 @@ private:
 	//~ End UActorDescContainer Interface
 #endif
 
-	static void WorldPartitionOnLevelRemovedFromWorld(class ULevel* Level, UWorld* InWorld);
-
 public:
 	//~ Begin UObject Interface
 	virtual bool ResolveSubobject(const TCHAR* SubObjectPath, UObject*& OutObject, bool bLoadIfExists) override;
diff --git a/Engine/Source/Runtime/Engine/Public/WorldPartition/WorldPartitionSubsystem.h b/Engine/Source/Runtime/Engine/Public/WorldPartition/WorldPartitionSubsystem.h
--- a/Engine/Source/Runtime/Engine/Public/WorldPartition/WorldPartitionSubsystem.h
+++ b/Engine/Source/Runtime/Engine/Public/WorldPartition/WorldPartitionSubsystem.h
@@ -39,10 +39,6 @@ public:
 	virtual TStatId GetStatId() const override;
 	//~End FTickableGameObject
 
-	DECLARE_MULTICAST_DELEGATE_OneParam(FWorldPartitionRegisterDelegate, UWorldPartition*);
-	FWorldPartitionRegisterDelegate OnWorldPartitionRegistered;
-	FWorldPartitionRegisterDelegate OnWorldPartitionUnregistered;
-
 	UFUNCTION(BlueprintCallable, Category = Streaming)
 	bool IsStreamingCompleted(EWorldPartitionRuntimeCellState QueryState, const TArray<FWorldPartitionStreamingQuerySource>& QuerySources, bool bExactState) const;
 
