Cascade LWC changes:
* cpu sims are run with full double precision vectors
* gpu sims use a tile offset to translate non-localspace emitter positions to float precision vectors

diff --git a/Engine/Shaders/Private/ParticleGPUSpriteVertexFactory.ush b/Engine/Shaders/Private/ParticleGPUSpriteVertexFactory.ush
--- a/Engine/Shaders/Private/ParticleGPUSpriteVertexFactory.ush
+++ b/Engine/Shaders/Private/ParticleGPUSpriteVertexFactory.ush
@@ -245,7 +245,7 @@ FVertexFactoryIntermediates GetVertexFactoryIntermediates( FVertexFactoryInput I
 	// For clarity, store some information in local variables.
 	const float RelativeTime = PositionSample.w;
 	const float IsAlive = step( RelativeTime, 1.0f );
-	float3 ParticlePosition = PositionSample.xyz;
+	float3 ParticlePosition = LWCHackToFloat(MakeLWCVector3(EmitterDynamicUniforms.LWCTile, PositionSample.xyz));
 
 	// Put velocity in to world space.
 	float3 ParticleWorldVelocity = mul(VelocitySample.xyz, GetLocalToWorld3x3()).xyz;
diff --git a/Engine/Shaders/Private/ParticleSimulationShader.usf b/Engine/Shaders/Private/ParticleSimulationShader.usf
--- a/Engine/Shaders/Private/ParticleSimulationShader.usf
+++ b/Engine/Shaders/Private/ParticleSimulationShader.usf
@@ -103,6 +103,8 @@ float2 LocalToWorldScale;
 float DeltaSeconds;
 /** Number of iterations, each applying DeltaSeconds. */ 
 int	NumIterations;
+/** LWC tile offset, will be 0,0,0 for localspace emitters. */
+float3 LWCTile;
 
 /** Texture from which to read particle position. */
 Texture2D PositionTexture;
@@ -552,7 +554,7 @@ void PixelMain(
 	const float TimeScale = VelocitySample.w;
 
 	// Position.w holds the relative time of the particle.
-	float3 Position = PositionSample.xyz;
+	float3 Position = LWCHackToFloat(MakeLWCVector3(LWCTile, PositionSample.xyz));
 	float RelativeTime = PositionSample.w;
 
 	for (int IterationIndex = 0; IterationIndex < NumIterations; ++IterationIndex)
@@ -657,7 +659,7 @@ void PixelMain(
 	}
 
 	// Store the new position, time, and velocity for the particle.
-	OutPosition.xyz = Position + PositionOffsetAndAttractorStrength.xyz;
+	OutPosition.xyz = (Position + PositionOffsetAndAttractorStrength.xyz) - LWCTile * UE_LWC_RENDER_TILE_SIZE;
 	OutPosition.w = RelativeTime;
 	OutVelocity.xyz = Velocity;
 	OutVelocity.w = TimeScale;
diff --git a/Engine/Source/Runtime/Engine/Classes/Particles/ParticleModuleRequired.h b/Engine/Source/Runtime/Engine/Classes/Particles/ParticleModuleRequired.h
--- a/Engine/Source/Runtime/Engine/Classes/Particles/ParticleModuleRequired.h
+++ b/Engine/Source/Runtime/Engine/Classes/Particles/ParticleModuleRequired.h
@@ -200,6 +200,10 @@ class UParticleModuleRequired : public UParticleModule
 	UPROPERTY(EditAnywhere, Category=Emitter, meta=(DisplayName = "Remove HMD Roll"))
 	uint8 bRemoveHMDRoll:1;
 
+	/** If true, gpu simulation positions are offset to support double precision vectors. Cpu sims always support large world coordinates. */
+	UPROPERTY(EditAnywhere, Category=Emitter)
+	uint8 bSupportLargeWorldCoordinates:1;
+
 	/**
 	 *	If true, select the emitter duration from the range
 	 *		[EmitterDurationLow..EmitterDuration]
diff --git a/Engine/Source/Runtime/Engine/Classes/Particles/ParticleSystemComponent.h b/Engine/Source/Runtime/Engine/Classes/Particles/ParticleSystemComponent.h
--- a/Engine/Source/Runtime/Engine/Classes/Particles/ParticleSystemComponent.h
+++ b/Engine/Source/Runtime/Engine/Classes/Particles/ParticleSystemComponent.h
@@ -824,7 +824,8 @@ private:
 	FBoxSphereBounds AsyncBounds;
 	/** Cached copy of PartSysVelocity */
 	FVector AsyncPartSysVelocity;
-
+	/** LWC tile offset for gpu simulations */
+	FVector3f LWCTile = FVector3f::ZeroVector;
 public:
 
 	/** Called from game code when the significance required for a component changes. */
@@ -841,6 +842,8 @@ public:
 		bIsManagingSignificance = bManageSignificance;
 	}
 
+	FORCEINLINE const FVector3f& GetLWCTile() const { return LWCTile;  }
+	
 	/** Returns the approximate distance squared from this component to the passed location. */
 	float GetApproxDistanceSquared(FVector Point)const;
 	
diff --git a/Engine/Source/Runtime/Engine/Private/Particles/ParticleComponents.cpp b/Engine/Source/Runtime/Engine/Private/Particles/ParticleComponents.cpp
--- a/Engine/Source/Runtime/Engine/Private/Particles/ParticleComponents.cpp
+++ b/Engine/Source/Runtime/Engine/Private/Particles/ParticleComponents.cpp
@@ -1718,7 +1718,7 @@ void UParticleEmitter::CacheEmitterModuleInfo()
 	// Added FBaseColor...
 	// Linear color change
 	// Added Flags field	
-	static_assert(sizeof(FBaseParticle) == 128, "FBaseParticle size");
+	static_assert(sizeof(FBaseParticle) == 160, "FBaseParticle size");
 
 
 	bRequiresLoopNotification = false;
@@ -5965,6 +5965,9 @@ void UParticleSystemComponent::ActivateSystem(bool bFlagAsJustAttached)
 	bOldPositionValid = false;
 	OldPosition = FVector::ZeroVector;
 	PartSysVelocity = FVector::ZeroVector;
+	
+	// Set tile for LWC offset
+	LWCTile = FLargeWorldRenderScalar::GetTileFor(GetComponentLocation());
 
 	UWorld* World = GetWorld();
 	check(World);
diff --git a/Engine/Source/Runtime/Engine/Private/Particles/ParticleGpuSimulation.cpp b/Engine/Source/Runtime/Engine/Private/Particles/ParticleGpuSimulation.cpp
--- a/Engine/Source/Runtime/Engine/Private/Particles/ParticleGpuSimulation.cpp
+++ b/Engine/Source/Runtime/Engine/Private/Particles/ParticleGpuSimulation.cpp
@@ -385,6 +385,8 @@ public:
 
 	/** Frame index used to track double buffered resources on the GPU. */
 	int32 FrameIndex = 0;
+	/** LWC tile offset, will be 0,0,0 for localspace emitters. */
+	FVector3f LWCTile = FVector3f::ZeroVector;
 
 	/**
 	 * Initialize resources.
@@ -533,6 +535,7 @@ typedef TUniformBufferRef<FGPUSpriteEmitterUniformParameters> FGPUSpriteEmitterU
  * Uniform buffer to hold dynamic parameters for GPU particle sprite emitters.
  */
 BEGIN_GLOBAL_SHADER_PARAMETER_STRUCT( FGPUSpriteEmitterDynamicUniformParameters, )
+	SHADER_PARAMETER( FVector3f, LWCTile )
 	SHADER_PARAMETER( FVector2f, LocalToWorldScale )
 	SHADER_PARAMETER( float, EmitterInstRandom)
 	SHADER_PARAMETER( FVector4f, AxisLockRight )
@@ -687,7 +690,8 @@ public:
 	FRHITexture2D* VelocityTextureRHI;
 	/** Texture containint attributes for all particles. */
 	FRHITexture2D* AttributesTextureRHI;
-
+	/** LWC tile offset, will be 0,0,0 for localspace emitters. */
+	FVector3f LWCTile;
 
 	FGPUSpriteVertexFactory()
 		: FParticleVertexFactoryBase(PVFT_MAX, ERHIFeatureLevel::Num)
@@ -695,6 +699,7 @@ public:
 		, PositionTextureRHI(nullptr)
 		, VelocityTextureRHI(nullptr)
 		, AttributesTextureRHI(nullptr)
+		, LWCTile(FVector3f::ZeroVector)
 	{}
 
 	/**
@@ -865,6 +870,9 @@ struct FParticlePerFrameSimulationParameters
 	/** Amount of time by which to simulate particles. */
 	float DeltaSeconds;
 
+	/** LWC tile offset, will be 0,0,0 for localspace emitters. */
+	FVector3f LWCTile;
+
 	FParticlePerFrameSimulationParameters()
 		: PointAttractor(FVector::ZeroVector,0.0f)
 		, PositionOffsetAndAttractorStrength(FVector::ZeroVector,0.0f)
@@ -874,6 +882,7 @@ struct FParticlePerFrameSimulationParameters
 		, DeltaSecondsInVar(0.0f)
 		, NumIterationsInVar(0)
 		, DeltaSeconds(0.0f)
+		, LWCTile(FVector3f::ZeroVector)
 
 	{
 	}
@@ -903,6 +912,7 @@ public:
 		LocalToWorldScale.Bind(ParameterMap,TEXT("LocalToWorldScale"));
 		DeltaSeconds.Bind(ParameterMap,TEXT("DeltaSeconds"));
 		NumIterations.Bind(ParameterMap,TEXT("NumIterations"));
+		LWCTile.Bind(ParameterMap,TEXT("LWCTile"));
 	}
 
 	template <typename ShaderRHIParamRef>
@@ -918,6 +928,7 @@ public:
 		SetShaderValue(RHICmdList,ShaderRHI,LocalToWorldScale,Parameters.LocalToWorldScale);
 		SetShaderValue(RHICmdList,ShaderRHI,DeltaSeconds, bUseFixDT ? Parameters.DeltaSecondsInFix : Parameters.DeltaSecondsInVar);
 		SetShaderValue(RHICmdList,ShaderRHI,NumIterations, bUseFixDT ? Parameters.NumIterationsInFix : Parameters.NumIterationsInVar);
+		SetShaderValue(RHICmdList,ShaderRHI,LWCTile, Parameters.LWCTile);
 	}
 
 	LAYOUT_FIELD(FShaderParameter, PointAttractor);
@@ -925,6 +936,7 @@ public:
 	LAYOUT_FIELD(FShaderParameter, LocalToWorldScale);
 	LAYOUT_FIELD(FShaderParameter, DeltaSeconds);
 	LAYOUT_FIELD(FShaderParameter, NumIterations);
+	LAYOUT_FIELD(FShaderParameter, LWCTile);
 	
 };
 
@@ -935,6 +947,7 @@ FArchive& operator<<(FArchive& Ar, FParticlePerFrameSimulationShaderParameters&
 	Ar << PerFrameParameters.LocalToWorldScale;
 	Ar << PerFrameParameters.DeltaSeconds;
 	Ar << PerFrameParameters.NumIterations;
+	Ar << PerFrameParameters.LWCTile;
 	return Ar;
 }
 
@@ -2935,7 +2948,7 @@ public:
 				FVector2D ObjectNDCPosition;
 				FVector2D ObjectMacroUVScales;
 				Proxy->GetObjectPositionAndScale(*View,ObjectNDCPosition, ObjectMacroUVScales);
-				PerViewDynamicParameters.MacroUVParameters = FVector4f(ObjectNDCPosition.X, ObjectNDCPosition.Y, ObjectMacroUVScales.X, ObjectMacroUVScales.Y); 
+				PerViewDynamicParameters.MacroUVParameters = FVector4f(ObjectNDCPosition.X, ObjectNDCPosition.Y, ObjectMacroUVScales.X, ObjectMacroUVScales.Y);
 
 				if (bUseLocalSpace == false)
 				{
@@ -2988,6 +3001,7 @@ public:
 					VertexFactory.PositionTextureRHI = StateTextures.PositionTextureRHI;
 					VertexFactory.VelocityTextureRHI = StateTextures.VelocityTextureRHI;
 					VertexFactory.AttributesTextureRHI = ParticleSimulationResources->RenderAttributesTexture.TextureRHI;
+					VertexFactory.LWCTile = ParticleSimulationResources->LWCTile;
 
 					FMeshBatch& Mesh = Collector.AllocateMesh();
 					FMeshBatchElement& BatchElement = Mesh.Elements[0];
@@ -3229,7 +3243,13 @@ FGPUSpriteParticleEmitterInstance(FFXSystem* InFXSystem, FGPUSpriteEmitterInfo&
 		UParticleSystem *Template = Component->Template;
 
 		const bool bLocalSpace = EmitterInfo.RequiredModule->bUseLocalSpace;
-		const FMatrix ComponentToWorldMatrix = Component->GetComponentTransform().ToMatrixWithScale();
+		const bool bUseTileOffset = bLocalSpace == false && EmitterInfo.RequiredModule->bSupportLargeWorldCoordinates;
+		FTransform ComponentTransform = Component->GetComponentTransform();
+		if (bUseTileOffset)
+		{
+			ComponentTransform.AddToTranslation(FVector(Component->GetLWCTile()) * -FLargeWorldRenderScalar::GetTileSize());	
+		}
+		const FMatrix ComponentToWorldMatrix = ComponentTransform.ToMatrixWithScale();
 		const FMatrix ComponentToWorld = (bLocalSpace || EmitterInfo.LocalVectorField.bIgnoreComponentTransform) ? FMatrix::Identity : ComponentToWorldMatrix;
 
 		const FRotationMatrix VectorFieldTransform(LocalVectorFieldRotation);
@@ -3252,8 +3272,11 @@ FGPUSpriteParticleEmitterInstance(FFXSystem* InFXSystem, FGPUSpriteEmitterInfo&
 		DynamicData->bSelected = bSelected;
 		DynamicData->bUseLocalSpace = EmitterInfo.RequiredModule->bUseLocalSpace;
 
+		// Get LWC tile
+		DynamicData->EmitterDynamicParameters.LWCTile = bUseTileOffset ? Component->GetLWCTile() : FVector3f::ZeroVector;
+
 		// Account for LocalToWorld scaling
-		FVector ComponentScale = Component->GetComponentTransform().GetScale3D();
+		FVector ComponentScale = ComponentTransform.GetScale3D();
 		// Figure out if we need to replicate the X channel of size to Y.
 		const bool bSquare = (EmitterInfo.ScreenAlignment == PSA_Square)
 			|| (EmitterInfo.ScreenAlignment == PSA_FacingCameraPosition)
@@ -3399,12 +3422,13 @@ FGPUSpriteParticleEmitterInstance(FFXSystem* InFXSystem, FGPUSpriteEmitterInfo&
 			#endif
 
 			}
-
+			
 			FVector PointAttractorPosition = ComponentToWorld.TransformPosition(EmitterInfo.PointAttractorPosition);
 			DynamicData->PerFrameSimulationParameters.PointAttractor = FVector4f(PointAttractorPosition, EmitterInfo.PointAttractorRadiusSq);
 			DynamicData->PerFrameSimulationParameters.PositionOffsetAndAttractorStrength = FVector4f(PositionOffsetThisTick, PointAttractorStrength);
 			DynamicData->PerFrameSimulationParameters.LocalToWorldScale = DynamicData->EmitterDynamicParameters.LocalToWorldScale;
 			DynamicData->PerFrameSimulationParameters.DeltaSeconds = PendingDeltaSeconds; // This value is used when updating vector fields.
+			DynamicData->PerFrameSimulationParameters.LWCTile = bUseTileOffset ? Component->GetLWCTile() : FVector3f::ZeroVector;
 			Exchange(DynamicData->TilesToClear, TilesToClear);
 			Exchange(DynamicData->NewParticles, NewParticles);
 		}
@@ -3451,6 +3475,8 @@ FGPUSpriteParticleEmitterInstance(FFXSystem* InFXSystem, FGPUSpriteEmitterInfo&
 		EmitterInstRandom = RandomStream.GetFraction();
 
 		FParticleSimulationResources* ParticleSimulationResources = FXSystem->GetParticleSimulationResources();
+		const bool bUseTileOffset = EmitterInfo.RequiredModule->bUseLocalSpace == false && EmitterInfo.RequiredModule->bSupportLargeWorldCoordinates;
+		ParticleSimulationResources->LWCTile = bUseTileOffset ? Component->GetLWCTile() : FVector3f::ZeroVector;
 		const int32 MinTileCount = GetMinTileCount();
 		int32 NumAllocated = 0;
 		{
@@ -4155,7 +4181,15 @@ private:
 			float InterpFraction = (float)i / (float)SpawnInfo.Count;
 
 			NewParticle->Velocity = TempParticle->BaseVelocity;
-			NewParticle->Position = TempParticle->Location + InterpFraction * EmitterDelta + SpawnInfo.StartTime * NewParticle->Velocity + EmitterInfo.OrbitOffsetBase + EmitterInfo.OrbitOffsetRange * RandomOrbit;
+			FVector WSPosition = TempParticle->Location + InterpFraction * EmitterDelta + SpawnInfo.StartTime * NewParticle->Velocity + EmitterInfo.OrbitOffsetBase + EmitterInfo.OrbitOffsetRange * RandomOrbit;
+			if (RequiredModule->bUseLocalSpace == false && RequiredModule->bSupportLargeWorldCoordinates)
+			{
+				NewParticle->Position = WSPosition - FLargeWorldRenderScalar::GetTileSize() * FVector(Component->GetLWCTile());
+			}
+			else
+			{
+				NewParticle->Position = WSPosition;
+			}
 			NewParticle->RelativeTime = TempParticle->RelativeTime;
 			NewParticle->TimeScale = FMath::Max<float>(TempParticle->OneOverMaxLifetime, 0.001f);
 
@@ -4643,7 +4677,8 @@ void FFXSystem::PrepareGPUSimulation(FRHICommandListImmediate& RHICmdList)
 		// Setup render states.
 		RHICmdList.Transition({
 			FRHITransitionInfo(CurrentStateTextures.PositionTextureTargetRHI, ERHIAccess::SRVGraphics, ERHIAccess::RTV),
-			FRHITransitionInfo(CurrentStateTextures.VelocityTextureTargetRHI, ERHIAccess::SRVGraphics, ERHIAccess::RTV)});
+			FRHITransitionInfo(CurrentStateTextures.VelocityTextureTargetRHI, ERHIAccess::SRVGraphics, ERHIAccess::RTV)
+		});
 	}
 }
 
diff --git a/Engine/Source/Runtime/Engine/Private/Particles/ParticleModules.cpp b/Engine/Source/Runtime/Engine/Private/Particles/ParticleModules.cpp
--- a/Engine/Source/Runtime/Engine/Private/Particles/ParticleModules.cpp
+++ b/Engine/Source/Runtime/Engine/Private/Particles/ParticleModules.cpp
@@ -1074,6 +1074,7 @@ UParticleModuleRequired::UParticleModuleRequired(const FObjectInitializer& Objec
 	NormalsSphereCenter = FVector(0.0f, 0.0f, 100.0f);
 	NormalsCylinderDirection = FVector(0.0f, 0.0f, 1.0f);
 	bUseLegacyEmitterTime = true;
+	bSupportLargeWorldCoordinates = true;
 	UVFlippingMode = EParticleUVFlipMode::None;
 	BoundingMode = BVC_EightVertices;
 	AlphaThreshold = 0.1f;
diff --git a/Engine/Source/Runtime/Engine/Private/Particles/ParticleModules_Location.cpp b/Engine/Source/Runtime/Engine/Private/Particles/ParticleModules_Location.cpp
--- a/Engine/Source/Runtime/Engine/Private/Particles/ParticleModules_Location.cpp
+++ b/Engine/Source/Runtime/Engine/Private/Particles/ParticleModules_Location.cpp
@@ -219,7 +219,7 @@ void UParticleModuleLocationWorldOffset::SpawnEx(FParticleEmitterInstance* Owner
 		// We need to inverse transform the location so that the bUseLocalSpace transform uses the proper value
 		FMatrix InvMat = Owner->Component->GetComponentTransform().ToMatrixWithScale().InverseFast();
 		FVector StartLoc = StartLocation.GetValue(Owner->EmitterTime, Owner->Component, 0, InRandomStream);
-		Particle.Location += (FVector4f)InvMat.TransformVector(StartLoc);
+		Particle.Location += InvMat.TransformVector(StartLoc);
 	}
 	ensureMsgf(!Particle.Location.ContainsNaN(), TEXT("NaN in Particle Location. Template: %s, Component: %s"), Owner->Component ? *GetNameSafe(Owner->Component->Template) : TEXT("UNKNOWN"), *GetPathNameSafe(Owner->Component));
 }
@@ -1061,7 +1061,7 @@ void UParticleModuleLocationPrimitiveCylinder::SpawnEx(FParticleEmitterInstance*
 	vOffset[RadialIndex1]	+= vStartLoc[RadialIndex1];
 	vOffset[HeightIndex]	+= vStartLoc[HeightIndex];
 
-	Particle.Location += (FVector4f)Owner->EmitterToSimulation.TransformVector(vOffset);
+	Particle.Location += Owner->EmitterToSimulation.TransformVector(vOffset);
 
 	if (Velocity)
 	{
@@ -1297,7 +1297,7 @@ void UParticleModuleLocationPrimitiveSphere::SpawnEx(FParticleEmitterInstance* O
 	}
 
 	vOffset += vStartLoc;
-	Particle.Location += (FVector4f)Owner->EmitterToSimulation.TransformVector(vOffset);
+	Particle.Location += Owner->EmitterToSimulation.TransformVector(vOffset);
 
 	if (Velocity)
 	{
diff --git a/Engine/Source/Runtime/Engine/Private/Particles/ParticleSystemRender.cpp b/Engine/Source/Runtime/Engine/Private/Particles/ParticleSystemRender.cpp
--- a/Engine/Source/Runtime/Engine/Private/Particles/ParticleSystemRender.cpp
+++ b/Engine/Source/Runtime/Engine/Private/Particles/ParticleSystemRender.cpp
@@ -9,7 +9,6 @@
 #include "Misc/MemStack.h"
 #include "HAL/IConsoleManager.h"
 #include "EngineDefines.h"
-#include "EngineGlobals.h"
 #include "GameFramework/Actor.h"
 #include "RenderingThread.h"
 #include "RenderResource.h"
@@ -38,7 +37,6 @@
 #include "Particles/ParticleLODLevel.h"
 #include "Engine/StaticMesh.h"
 #include "UnrealEngine.h"
-#include "Engine/StaticMesh.h"
 
 DECLARE_CYCLE_STAT(TEXT("ParticleSystemSceneProxy Create GT"), STAT_FParticleSystemSceneProxy_Create, STATGROUP_Particles);
 DECLARE_CYCLE_STAT(TEXT("ParticleSystemSceneProxy GetMeshElements RT"), STAT_FParticleSystemSceneProxy_GetMeshElements, STATGROUP_Particles);
@@ -1698,15 +1696,15 @@ void FDynamicMeshEmitterData::GetParticlePrevTransform(
 
 void FDynamicMeshEmitterData::CalculateParticleTransform(
 	const FMatrix& ProxyLocalToWorld,
-	const FVector3f& ParticleLocation,
+	const FVector& ParticleLocation,
 		  float    ParticleRotation,
 	const FVector3f& ParticleVelocity,
 	const FVector3f& ParticleSize,
 	const FVector3f& ParticlePayloadInitialOrientation,
 	const FVector3f& ParticlePayloadRotation,
-	const FVector3f& ParticlePayloadCameraOffset,
+	const FVector& ParticlePayloadCameraOffset,
 	const FVector3f& ParticlePayloadOrbitOffset,
-	const FVector3f& ViewOrigin,
+	const FVector& ViewOrigin,
 	const FVector3f& ViewDirection,
 	FMatrix& OutTransformMat
 	) const
@@ -6025,7 +6023,7 @@ struct FAnimTrailParticleRenderData
 			check( InterpFactor >= 0.0f && InterpFactor <= 1.0f );
 			check( Particle && PrevParticle && Payload && PrevPayload );
 				
-			FVector3f PrevPrevLocation;
+			FVector PrevPrevLocation;
 			FVector3f PrevPrevDirection;		
 			float PrevPrevLength;
 			float PrevPrevTiledU;
@@ -6053,7 +6051,7 @@ struct FAnimTrailParticleRenderData
 				PrevPrevDynParamParticle = PrevParticle;
 			}
 
-			FVector3f NextLocation;
+			FVector NextLocation;
 			FVector3f NextDirection;
 			float NextLength;
 			float NextTiledU;
diff --git a/Engine/Source/Runtime/Engine/Private/Particles/ParticleTrail2EmitterInstance.cpp b/Engine/Source/Runtime/Engine/Private/Particles/ParticleTrail2EmitterInstance.cpp
--- a/Engine/Source/Runtime/Engine/Private/Particles/ParticleTrail2EmitterInstance.cpp
+++ b/Engine/Source/Runtime/Engine/Private/Particles/ParticleTrail2EmitterInstance.cpp
@@ -2959,7 +2959,7 @@ void FParticleAnimTrailEmitterInstance::RecalculateTangentAndInterpolationParam(
 {
 	check(CurrParticle);
 	check(PrevParticle || NextParticle);
-	FVector3f Tangent;
+	FVector Tangent;
 	float SegmentDistance;
 	
 	if( PrevParticle )
diff --git a/Engine/Source/Runtime/Engine/Public/ParticleHelper.h b/Engine/Source/Runtime/Engine/Public/ParticleHelper.h
--- a/Engine/Source/Runtime/Engine/Public/ParticleHelper.h
+++ b/Engine/Source/Runtime/Engine/Public/ParticleHelper.h
@@ -182,15 +182,9 @@ struct FStaticMeshLODResources;
 // Only used when required.
 struct FBaseParticle
 {
-	// LWC_TODO: Particle location/velocity kept single precision to maintain size, as there's size counting going on in the comments here (128 bytes, 2 L1 cache lines, vector/float pairs for SSE alignment?). Potential for double precision?
-
-	// 16 bytes
-	FVector3f		OldLocation;			// Last frame's location, used for collision
-	float			RelativeTime;			// Relative time, range is 0 (==spawn) to 1 (==death)
-
-	// 16 bytes
-	FVector3f		Location;				// Current location
-	float			OneOverMaxLifetime;		// Reciprocal of lifetime
+	// 48 bytes
+	FVector		OldLocation;			// Last frame's location, used for collision
+	FVector		Location;				// Current location
 
 	// 16 bytes
 	FVector3f		BaseVelocity;			// Velocity = BaseVelocity at the start of each frame.
@@ -213,6 +207,12 @@ struct FBaseParticle
 
 	// 16 bytes
 	FLinearColor	BaseColor;				// Base color of the particle
+
+	// 16 bytes
+	float			RelativeTime;			// Relative time, range is 0 (==spawn) to 1 (==death)
+	float			OneOverMaxLifetime;		// Reciprocal of lifetime
+	float			Placeholder0;
+	float			Placeholder1;
 };
 
 /*-----------------------------------------------------------------------------
@@ -1776,15 +1776,15 @@ struct FDynamicMeshEmitterData : public FDynamicSpriteEmitterDataBase
 
 	void CalculateParticleTransform(
 		const FMatrix& ProxyLocalToWorld,
-		const FVector3f& ParticleLocation,
+		const FVector& ParticleLocation,
 			  float    ParticleRotation,
 		const FVector3f& ParticleVelocity,
 		const FVector3f& ParticleSize,
 		const FVector3f& ParticlePayloadInitialOrientation,
 		const FVector3f& ParticlePayloadRotation,
-		const FVector3f& ParticlePayloadCameraOffset,
+		const FVector& ParticlePayloadCameraOffset,
 		const FVector3f& ParticlePayloadOrbitOffset,
-		const FVector3f& ViewOrigin,
+		const FVector& ViewOrigin,
 		const FVector3f& ViewDirection,
 		FMatrix& OutTransformMat
 		) const;
