Add extra profile tags

diff --git a/Engine/Source/Runtime/Engine/Private/GPUSkinCache.cpp b/Engine/Source/Runtime/Engine/Private/GPUSkinCache.cpp
--- a/Engine/Source/Runtime/Engine/Private/GPUSkinCache.cpp
+++ b/Engine/Source/Runtime/Engine/Private/GPUSkinCache.cpp
@@ -763,6 +763,8 @@ void FGPUSkinCache::Cleanup()
 
 void FGPUSkinCache::TransitionAllToReadable(FRHICommandList& RHICmdList)
 {
+	TRACE_CPUPROFILER_EVENT_SCOPE(FGPUSkinCache::TransitionAllToReadable);
+
 	if (BuffersToTransitionToRead.Num() > 0)
 	{
 		FMemMark Mark(FMemStack::Get());
@@ -1600,6 +1602,8 @@ void FGPUSkinCache::BeginBatchDispatch(FRHICommandListImmediate& RHICmdList)
 
 void FGPUSkinCache::EndBatchDispatch(FRHICommandListImmediate& RHICmdList)
 {
+	TRACE_CPUPROFILER_EVENT_SCOPE(FGPUSkinCache::EndBatchDispatch);
+
 	DoDispatch(RHICmdList);
 
 #if RHI_RAYTRACING
diff --git a/Engine/Source/Runtime/Engine/Private/LevelTick.cpp b/Engine/Source/Runtime/Engine/Private/LevelTick.cpp
--- a/Engine/Source/Runtime/Engine/Private/LevelTick.cpp
+++ b/Engine/Source/Runtime/Engine/Private/LevelTick.cpp
@@ -1014,6 +1014,7 @@ void EndSendEndOfFrameUpdatesDrawEvent(FSendAllEndOfFrameUpdates& SendAllEndOfFr
 		[GPUSkinCache = SendAllEndOfFrameUpdates.GPUSkinCache, ComputeTaskWorkers = SendAllEndOfFrameUpdates.ComputeTaskWorkers, FeatureLevel = SendAllEndOfFrameUpdates.FeatureLevel](FRHICommandListImmediate& RHICmdList)
 		{
 			SCOPED_GPU_STAT(RHICmdList, EndOfFrameUpdates);
+			TRACE_CPUPROFILER_EVENT_SCOPE(EndSendEndOfFrameUpdatesDrawEvent_RT);
 
 			if (GPUSkinCache != nullptr)
 			{
diff --git a/Engine/Source/Runtime/Engine/Private/PositionVertexBuffer.cpp b/Engine/Source/Runtime/Engine/Private/PositionVertexBuffer.cpp
--- a/Engine/Source/Runtime/Engine/Private/PositionVertexBuffer.cpp
+++ b/Engine/Source/Runtime/Engine/Private/PositionVertexBuffer.cpp
@@ -249,6 +249,8 @@ void FPositionVertexBuffer::CopyRHIForStreaming(const FPositionVertexBuffer& Oth
 
 void FPositionVertexBuffer::InitRHI()
 {
+	TRACE_CPUPROFILER_EVENT_SCOPE(FPositionVertexBuffer::InitRHI);
+
 	VertexBufferRHI = CreateRHIBuffer_RenderThread();
 	// we have decide to create the SRV based on GMaxRHIShaderPlatform because this is created once and shared between feature levels for editor preview.
 	// Also check to see whether cpu access has been activated on the vertex data
diff --git a/Engine/Source/Runtime/Engine/Private/Rendering/ColorVertexBuffer.cpp b/Engine/Source/Runtime/Engine/Private/Rendering/ColorVertexBuffer.cpp
--- a/Engine/Source/Runtime/Engine/Private/Rendering/ColorVertexBuffer.cpp
+++ b/Engine/Source/Runtime/Engine/Private/Rendering/ColorVertexBuffer.cpp
@@ -446,6 +446,7 @@ void FColorVertexBuffer::CopyRHIForStreaming(const FColorVertexBuffer& Other, bo
 
 void FColorVertexBuffer::InitRHI()
 {
+	TRACE_CPUPROFILER_EVENT_SCOPE(FColorVertexBuffer::InitRHI);
 	SCOPED_LOADTIMER(FColorVertexBuffer_InitRHI);
 
 	VertexBufferRHI = CreateRHIBuffer_RenderThread();
diff --git a/Engine/Source/Runtime/Engine/Private/Rendering/StaticMeshVertexBuffer.cpp b/Engine/Source/Runtime/Engine/Private/Rendering/StaticMeshVertexBuffer.cpp
--- a/Engine/Source/Runtime/Engine/Private/Rendering/StaticMeshVertexBuffer.cpp
+++ b/Engine/Source/Runtime/Engine/Private/Rendering/StaticMeshVertexBuffer.cpp
@@ -358,6 +358,7 @@ void FStaticMeshVertexBuffer::CopyRHIForStreaming(const FStaticMeshVertexBuffer&
 
 void FStaticMeshVertexBuffer::InitRHI()
 {
+	TRACE_CPUPROFILER_EVENT_SCOPE(FStaticMeshVertexBuffer::InitRHI);
 	SCOPED_LOADTIMER(FStaticMeshVertexBuffer_InitRHI);
 
 	TangentsVertexBuffer.VertexBufferRHI = CreateTangentsRHIBuffer_RenderThread();
diff --git a/Engine/Source/Runtime/Engine/Private/SkeletalRenderGPUSkin.cpp b/Engine/Source/Runtime/Engine/Private/SkeletalRenderGPUSkin.cpp
--- a/Engine/Source/Runtime/Engine/Private/SkeletalRenderGPUSkin.cpp
+++ b/Engine/Source/Runtime/Engine/Private/SkeletalRenderGPUSkin.cpp
@@ -509,6 +509,8 @@ FORCEINLINE bool IsDeferredSkeletalDynamicDataUpdateEnabled()
 
 void FSkeletalMeshObjectGPUSkin::UpdateDynamicData_RenderThread(FGPUSkinCache* GPUSkinCache, FRHICommandListImmediate& RHICmdList, FDynamicSkelMeshObjectDataGPUSkin* InDynamicData, FSceneInterface* Scene, bool bIsGameWorld, uint32 FrameNumberToPrepare, uint32 RevisionNumber)
 {
+	TRACE_CPUPROFILER_EVENT_SCOPE(GPUSkin::UpdateDynamicData_RT);
+
 	SCOPE_CYCLE_COUNTER(STAT_GPUSkinUpdateRTTime);
 	check(InDynamicData != nullptr);
 	CA_ASSUME(InDynamicData != nullptr);
diff --git a/Engine/Source/Runtime/MovieScene/Private/EntitySystem/MovieSceneEntitySystemRunner.cpp b/Engine/Source/Runtime/MovieScene/Private/EntitySystem/MovieSceneEntitySystemRunner.cpp
--- a/Engine/Source/Runtime/MovieScene/Private/EntitySystem/MovieSceneEntitySystemRunner.cpp
+++ b/Engine/Source/Runtime/MovieScene/Private/EntitySystem/MovieSceneEntitySystemRunner.cpp
@@ -185,6 +185,7 @@ void FMovieSceneEntitySystemRunner::DoFlushUpdateQueueOnce()
 	using namespace UE::MovieScene;
 
 	TRACE_COUNTER_INCREMENT(MovieSceneEntitySystemEvaluations);
+	TRACE_CPUPROFILER_EVENT_SCOPE(FMovieSceneEntitySystemRunner::DoFlushUpdateQueueOnce);
 
 	// Setup the completion task that we can wait on.
 	CompletionTask = TGraphTask<FNullGraphTask>::CreateTask(nullptr, ENamedThreads::GameThread)
