Removed broken assert

diff --git a/Engine/Source/Runtime/Engine/Private/Particles/ParticleComponents.cpp b/Engine/Source/Runtime/Engine/Private/Particles/ParticleComponents.cpp
--- a/Engine/Source/Runtime/Engine/Private/Particles/ParticleComponents.cpp
+++ b/Engine/Source/Runtime/Engine/Private/Particles/ParticleComponents.cpp
@@ -1717,9 +1717,7 @@ void UParticleEmitter::CacheEmitterModuleInfo()
 	// This assert makes sure that packing is as expected.
 	// Added FBaseColor...
 	// Linear color change
-	// Added Flags field	
-	static_assert(sizeof(FBaseParticle) == 160, "FBaseParticle size");
-
+	// Added Flags field
 
 	bRequiresLoopNotification = false;
 	bAxisLockEnabled = false;
