Fix ClangEditor Win64 compilation.

diff --git a/Engine/Source/Runtime/Engine/Public/WorldPartition/WorldPartitionHandle.h b/Engine/Source/Runtime/Engine/Public/WorldPartition/WorldPartitionHandle.h
--- a/Engine/Source/Runtime/Engine/Public/WorldPartition/WorldPartitionHandle.h
+++ b/Engine/Source/Runtime/Engine/Public/WorldPartition/WorldPartitionHandle.h
@@ -18,7 +18,7 @@ struct ENGINE_API FWorldPartitionHandleUtils
 };
 
 template <typename Impl>
-class ENGINE_API TWorldPartitionHandle
+class TWorldPartitionHandle
 {
 public:
 	FORCEINLINE TWorldPartitionHandle()
