Working on improving API for working with HLSL statements/scopes

diff --git a/Engine/Source/Runtime/Engine/Private/HLSLTree/HLSLTree.cpp b/Engine/Source/Runtime/Engine/Private/HLSLTree/HLSLTree.cpp
--- a/Engine/Source/Runtime/Engine/Private/HLSLTree/HLSLTree.cpp
+++ b/Engine/Source/Runtime/Engine/Private/HLSLTree/HLSLTree.cpp
@@ -33,7 +33,7 @@ public:
 
 	virtual void ComputeAnalyticDerivatives(FTree& Tree, FExpressionDerivatives& OutResult) const override;
 	virtual bool PrepareValue(FEmitContext& Context, const FRequestedType& RequestedType, FPrepareValueResult& OutResult) const override;
-	virtual void EmitValueShader(FEmitContext& Context, FEmitShaderScope& Scope, const FRequestedType& RequestedType, FEmitValueShaderResult& OutResult) const override;
+	virtual void EmitValueShader(FEmitContext& Context, FEmitScope& Scope, const FRequestedType& RequestedType, FEmitValueShaderResult& OutResult) const override;
 
 	TArray<EDerivativeCoordinate, TInlineAllocator<8>> DerivativeChain;
 	FName LocalName;
@@ -49,13 +49,12 @@ public:
 class FExpressionFunctionCall final : public FExpression
 {
 public:
-	FExpressionFunctionCall(FFunction* InFunction, TArrayView<FExpression*> InInputExpressions, int32 InOutputIndex) : Function(InFunction), InputExpressions(InInputExpressions), OutputIndex(InOutputIndex) {}
+	FExpressionFunctionCall(FFunction* InFunction, int32 InOutputIndex) : Function(InFunction), OutputIndex(InOutputIndex) {}
 
 	virtual bool PrepareValue(FEmitContext& Context, const FRequestedType& RequestedType, FPrepareValueResult& OutResult) const override;
-	virtual void EmitValueShader(FEmitContext& Context, FEmitShaderScope& Scope, const FRequestedType& RequestedType, FEmitValueShaderResult& OutResult) const override;
+	virtual void EmitValueShader(FEmitContext& Context, FEmitScope& Scope, const FRequestedType& RequestedType, FEmitValueShaderResult& OutResult) const override;
 
 	FFunction* Function;
-	TArray<FExpression*, TInlineAllocator<8>> InputExpressions;
 	int32 OutputIndex;
 };
 
@@ -125,11 +124,6 @@ EExpressionEvaluation CombineEvaluations(EExpressionEvaluation Lhs, EExpressionE
 	return EExpressionEvaluation::Preshader;
 }
 
-void FScope::Reset()
-{
-	State = EScopeState::Uninitialized;
-}
-
 FScope* FScope::FindSharedParent(FScope* Lhs, FScope* Rhs)
 {
 	FScope* Scope0 = Lhs;
@@ -203,7 +197,7 @@ bool FExpressionLocalPHI::PrepareValue(FEmitContext& Context, const FRequestedTy
 	for (int32 i = 0; i < NumValues; ++i)
 	{
 		// Ignore values in dead scopes
-		if (PrepareScope(Context, Scopes[i]))
+		if (Context.PrepareScope(Scopes[i]))
 		{
 			FExpression* ScopeExpression = Values[i];
 			if (!ForwardExpression)
@@ -232,17 +226,20 @@ bool FExpressionLocalPHI::PrepareValue(FEmitContext& Context, const FRequestedTy
 	{
 		for (int32 i = 0; i < NumValues; ++i)
 		{
-			if (TypePerValue[i].IsVoid() && PrepareScope(Context, Scopes[i]))
+			if (TypePerValue[i].IsVoid() && Context.PrepareScope(Scopes[i]))
 			{
-				const FPreparedType& ValueType = PrepareExpressionValue(Context, Values[i], RequestedType);
+				const FPreparedType& ValueType = Context.PrepareExpression(Values[i], RequestedType);
 				if (!ValueType.IsVoid())
 				{
 					TypePerValue[i] = ValueType;
-					CurrentType = MergePreparedTypes(CurrentType, ValueType);
-					if (CurrentType.IsVoid())
+					const FPreparedType MergedType = MergePreparedTypes(CurrentType, ValueType);
+					if (MergedType.IsVoid())
 					{
-						return Context.Errors->AddError(TEXT("Type mismatch"));
+						return Context.Errors->AddErrorf(TEXT("Mismatched types for local variable %s and %s"),
+							CurrentType.GetType().GetName(),
+							ValueType.GetType().GetName());
 					}
+					CurrentType = MergedType;
 					check(NumValidTypes < NumValues);
 					NumValidTypes++;
 				}
@@ -253,7 +250,10 @@ bool FExpressionLocalPHI::PrepareValue(FEmitContext& Context, const FRequestedTy
 	};
 
 	// First try to assign all the values we can
-	UpdateValueTypes();
+	if (!UpdateValueTypes())
+	{
+		return false;
+	}
 
 	// Assuming we have at least one value with a valid type, we use that to initialize our type
 	CurrentType.SetEvaluation(EExpressionEvaluation::Shader); // TODO - No support for preshader flow control
@@ -265,7 +265,10 @@ bool FExpressionLocalPHI::PrepareValue(FEmitContext& Context, const FRequestedTy
 	if (NumValidTypes < NumValues)
 	{
 		// Now try to assign remaining types that failed the first iteration 
-		UpdateValueTypes();
+		if (!UpdateValueTypes())
+		{
+			return false;
+		}
 		if (NumValidTypes < NumValues)
 		{
 			return Context.Errors->AddError(TEXT("Failed to compute all types for LocalPHI"));
@@ -275,7 +278,7 @@ bool FExpressionLocalPHI::PrepareValue(FEmitContext& Context, const FRequestedTy
 	return true;
 }
 
-void FExpressionLocalPHI::EmitValueShader(FEmitContext& Context, FEmitShaderScope& Scope, const FRequestedType& RequestedType, FEmitValueShaderResult& OutResult) const
+void FExpressionLocalPHI::EmitValueShader(FEmitContext& Context, FEmitScope& Scope, const FRequestedType& RequestedType, FEmitValueShaderResult& OutResult) const
 {
 	FEmitShaderExpression* const* PrevEmitExpression = Context.EmitLocalPHIMap.Find(this);
 	FEmitShaderExpression* EmitExpression = PrevEmitExpression ? *PrevEmitExpression : nullptr;
@@ -292,12 +295,12 @@ void FExpressionLocalPHI::EmitValueShader(FEmitContext& Context, FEmitShaderScop
 		Context.EmitLocalPHIMap.Add(this, EmitExpression);
 
 		// Find the outermost scope to declare our local variable
-		FEmitShaderScope* EmitDeclarationScope = &Scope;
-		FEmitShaderScope* EmitValueScopes[MaxNumPreviousScopes] = { nullptr };
+		FEmitScope* EmitDeclarationScope = &Scope;
+		FEmitScope* EmitValueScopes[MaxNumPreviousScopes] = { nullptr };
 		for (int32 i = 0; i < NumValues; ++i)
 		{
 			EmitValueScopes[i] = Context.AcquireEmitScope(Scopes[i]);
-			EmitDeclarationScope = FEmitShaderScope::FindSharedParent(EmitDeclarationScope, EmitValueScopes[i]);
+			EmitDeclarationScope = FEmitScope::FindSharedParent(EmitDeclarationScope, EmitValueScopes[i]);
 			if (!EmitDeclarationScope)
 			{
 				Context.Errors->AddError(TEXT("Invalid LocalPHI"));
@@ -308,7 +311,7 @@ void FExpressionLocalPHI::EmitValueShader(FEmitContext& Context, FEmitShaderScop
 		FEmitShaderStatement* EmitDeclaration = nullptr;
 		for (int32 i = 0; i < NumValues; ++i)
 		{
-			FEmitShaderScope* EmitValueScope = EmitValueScopes[i];
+			FEmitScope* EmitValueScope = EmitValueScopes[i];
 			if (EmitValueScope == EmitDeclarationScope)
 			{
 				FEmitShaderExpression* ShaderValue = Values[i]->GetValueShader(Context, *EmitValueScope, LocalType);
@@ -330,7 +333,7 @@ void FExpressionLocalPHI::EmitValueShader(FEmitContext& Context, FEmitShaderScop
 		int32 NumDependencies = 0;
 		for (int32 i = 0; i < NumValues; ++i)
 		{
-			FEmitShaderScope* EmitValueScope = EmitValueScopes[i];
+			FEmitScope* EmitValueScope = EmitValueScopes[i];
 			if (EmitValueScope != EmitDeclarationScope)
 			{
 				FEmitShaderExpression* ShaderValue = Values[i]->GetValueShader(Context, *EmitValueScope, LocalType);
@@ -350,29 +353,29 @@ void FExpressionLocalPHI::EmitValueShader(FEmitContext& Context, FEmitShaderScop
 
 bool FExpressionFunctionCall::PrepareValue(FEmitContext& Context, const FRequestedType& RequestedType, FPrepareValueResult& OutResult) const
 {
-	if (!PrepareScope(Context, Function->RootScope))
+	if (!Context.PrepareScopeWithParent(Function->RootScope, Function->CalledScope))
 	{
 		return false;
 	}
 
-	FPreparedType OutputType = PrepareExpressionValue(Context, Function->OutputExpressions[OutputIndex], RequestedType);
+	FPreparedType OutputType = Context.PrepareExpression(Function->OutputExpressions[OutputIndex], RequestedType);
 	return OutResult.SetType(Context, RequestedType, OutputType);
 }
 
-void FExpressionFunctionCall::EmitValueShader(FEmitContext& Context, FEmitShaderScope& Scope, const FRequestedType& RequestedType, FEmitValueShaderResult& OutResult) const
+void FExpressionFunctionCall::EmitValueShader(FEmitContext& Context, FEmitScope& Scope, const FRequestedType& RequestedType, FEmitValueShaderResult& OutResult) const
 {
 	FEmitShaderNode* const* PrevDependency = Context.EmitFunctionMap.Find(Function);
 	FEmitShaderNode* Dependency = PrevDependency ? *PrevDependency : nullptr;
 	if (!Dependency)
 	{
 		// Inject the function's root scope at scope where it's called
-		FEmitShaderScope* EmitCalledScope = Context.AcquireEmitScope(Function->CalledScope);
+		FEmitScope* EmitCalledScope = Context.AcquireEmitScope(Function->CalledScope);
 		Dependency = Context.EmitNextScope(*EmitCalledScope, Function->RootScope);
 		Context.EmitFunctionMap.Add(Function, Dependency);
 	}
 
 	FEmitShaderExpression* EmitFunctionOutput = Function->OutputExpressions[OutputIndex]->GetValueShader(Context, Scope, RequestedType);
-	OutResult.Code = Context.EmitInlineExpressionWithDependencies(Scope, MakeArrayView(&Dependency, 1), EmitFunctionOutput->Type, TEXT("%"), EmitFunctionOutput);
+	OutResult.Code = Context.EmitInlineExpressionWithDependency(Scope, Dependency, EmitFunctionOutput->Type, TEXT("%"), EmitFunctionOutput);
 }
 
 void FExpression::Reset()
@@ -381,76 +384,6 @@ void FExpression::Reset()
 	PrepareValueResult = FPrepareValueResult();
 }
 
-const FPreparedType& PrepareExpressionValue(FEmitContext& Context, FExpression* InExpression, const FRequestedType& RequestedType)
-{
-	static FPreparedType VoidType;
-	if (!InExpression)
-	{
-		return VoidType;
-	}
-
-	FOwnerScope OwnerScope(*Context.Errors, InExpression->GetOwner());
-	if (InExpression->bReentryFlag)
-	{
-		// Valid for this to be called reentrantly
-		// Code should ensure that the type is set before the reentrant call, otherwise type will not be valid here
-		// LocalPHI nodes rely on this to break loops
-		return InExpression->PrepareValueResult.GetPreparedType();
-	}
-
-	bool bNeedToUpdateType = false;
-	if (InExpression->CurrentRequestedType.RequestedComponents.Num() == 0)
-	{
-		InExpression->CurrentRequestedType = RequestedType;
-		bNeedToUpdateType = !RequestedType.IsVoid();
-	}
-	else if (InExpression->CurrentRequestedType.GetStructType() != RequestedType.GetStructType())
-	{
-		Context.Errors->AddError(TEXT("Type mismatch"));
-		return VoidType;
-	}
-	else
-	{
-		const int32 NumComponents = RequestedType.GetNumComponents();
-		InExpression->CurrentRequestedType.RequestedComponents.PadToNum(NumComponents, false);
-		for (int32 Index = 0; Index < NumComponents; ++Index)
-		{
-			const bool bPrevRequest = InExpression->CurrentRequestedType.IsComponentRequested(Index);
-			const bool bRequest = RequestedType.IsComponentRequested(Index);
-			if (!bPrevRequest && bRequest)
-			{
-				InExpression->CurrentRequestedType.SetComponentRequest(Index);
-				bNeedToUpdateType = true;
-			}
-		}
-	}
-
-	if (bNeedToUpdateType)
-	{
-		check(!InExpression->CurrentRequestedType.IsVoid());
-
-		bool bResult = false;
-		{
-			FExpressionReentryScope ReentryScope(InExpression);
-			bResult = InExpression->PrepareValue(Context, InExpression->CurrentRequestedType, InExpression->PrepareValueResult);
-		}
-
-		if (!bResult)
-		{
-			// If we failed to assign a valid type, reset the requested type as well
-			// This ensures we'll try to compute a type again the next time we're called
-			InExpression->CurrentRequestedType.Reset();
-			InExpression->PrepareValueResult.SetTypeVoid();
-		}
-		else
-		{
-			check(!InExpression->PrepareValueResult.GetPreparedType().IsVoid());
-		}
-	}
-
-	return InExpression->PrepareValueResult.GetPreparedType();
-}
-
 FRequestedType::FRequestedType(int32 NumComponents, bool bDefaultRequest)
 {
 	RequestedComponents.Init(bDefaultRequest, NumComponents);
@@ -839,7 +772,7 @@ bool FPrepareValueResult::SetForwardValue(FEmitContext& Context, const FRequeste
 {
 	if (InForwardValue != ForwardValue)
 	{
-		PreparedType = PrepareExpressionValue(Context, InForwardValue, RequestedType);
+		PreparedType = Context.PrepareExpression(InForwardValue, RequestedType);
 		ForwardValue = InForwardValue;
 	}
 	return (bool)InForwardValue;
@@ -850,7 +783,7 @@ void FExpression::ComputeAnalyticDerivatives(FTree& Tree, FExpressionDerivatives
 	// nop
 }
 
-void FExpression::EmitValueShader(FEmitContext& Context, FEmitShaderScope& Scope, const FRequestedType& RequestedType, FEmitValueShaderResult& OutResult) const
+void FExpression::EmitValueShader(FEmitContext& Context, FEmitScope& Scope, const FRequestedType& RequestedType, FEmitValueShaderResult& OutResult) const
 {
 	check(false);
 }
@@ -860,7 +793,7 @@ void FExpression::EmitValuePreshader(FEmitContext& Context, const FRequestedType
 	check(false);
 }
 
-FEmitShaderExpression* FExpression::GetValueShader(FEmitContext& Context, FEmitShaderScope& Scope, const FRequestedType& RequestedType, const Shader::FType& ResultType)
+FEmitShaderExpression* FExpression::GetValueShader(FEmitContext& Context, FEmitScope& Scope, const FRequestedType& RequestedType, const Shader::FType& ResultType)
 {
 	FOwnerScope OwnerScope(*Context.Errors, GetOwner());
 	if (PrepareValueResult.ForwardValue)
@@ -888,12 +821,12 @@ FEmitShaderExpression* FExpression::GetValueShader(FEmitContext& Context, FEmitS
 	return Value;
 }
 
-FEmitShaderExpression* FExpression::GetValueShader(FEmitContext& Context, FEmitShaderScope& Scope, const FRequestedType& RequestedType)
+FEmitShaderExpression* FExpression::GetValueShader(FEmitContext& Context, FEmitScope& Scope, const FRequestedType& RequestedType)
 {
 	return GetValueShader(Context, Scope, RequestedType, RequestedType.GetType());
 }
 
-FEmitShaderExpression* FExpression::GetValueShader(FEmitContext& Context, FEmitShaderScope& Scope)
+FEmitShaderExpression* FExpression::GetValueShader(FEmitContext& Context, FEmitScope& Scope)
 {
 	return GetValueShader(Context, Scope, GetRequestedType());
 }
@@ -978,31 +911,6 @@ void FScope::AddPreviousScope(FScope& Scope)
 	PreviousScope[NumPreviousScopes++] = &Scope;
 }
 
-bool PrepareScope(FEmitContext& Context, FScope* InScope)
-{
-	if (InScope && InScope->State == EScopeState::Uninitialized)
-	{
-		if (!InScope->ParentScope || PrepareScope(Context, InScope->ParentScope))
-		{
-			if (InScope->OwnerStatement)
-			{
-				FOwnerScope OwnerScope(*Context.Errors, InScope->OwnerStatement->GetOwner());
-				InScope->OwnerStatement->Prepare(Context);
-			}
-			else
-			{
-				InScope->State = EScopeState::Live;
-			}
-		}
-		else
-		{
-			InScope->State = EScopeState::Dead;
-		}
-	}
-
-	return InScope && InScope->State != EScopeState::Dead;
-}
-
 FTree* FTree::Create(FMemStackBase& Allocator)
 {
 	FTree* Tree = new(Allocator) FTree();
@@ -1027,32 +935,6 @@ void FTree::Destroy(FTree* Tree)
 	}
 }
 
-void FScope::MarkLive()
-{
-	if (State == EScopeState::Uninitialized)
-	{
-		State = EScopeState::Live;
-	}
-}
-
-void FScope::MarkLiveRecursive()
-{
-	return MarkLive();
-
-	FScope* Scope = this;
-	while (Scope && Scope->State == EScopeState::Uninitialized)
-	{
-		Scope->State = EScopeState::Live;
-		Scope = Scope->ParentScope;
-	}
-}
-
-void FScope::MarkDead()
-{
-	// TODO - mark child scopes as dead as well
-	State = EScopeState::Dead;
-}
-
 void FOwnerContext::PushOwner(UObject* Owner)
 {
 	OwnerStack.Add(Owner);
@@ -1116,26 +998,28 @@ bool FTree::Finalize()
 
 bool FTree::EmitShader(FEmitContext& Context, FStringBuilderBase& OutCode) const
 {
-	FEmitShaderScope* EmitRootScope = Context.AcquireEmitScope(RootScope);
-	
-	// Link all nodes to their proper scope
-	for (FEmitShaderNode* EmitNode : Context.EmitNodes)
+	FEmitScope* EmitRootScope = Context.InternalEmitScope(RootScope);
+	if (EmitRootScope)
 	{
-		FEmitShaderScope* EmitScope = EmitNode->Scope;
-		if (EmitScope)
+		// Link all nodes to their proper scope
+		for (FEmitShaderNode* EmitNode : Context.EmitNodes)
 		{
-			EmitNode->NextScopedNode = EmitScope->FirstNode;
-			EmitScope->FirstNode = EmitNode;
+			FEmitScope* EmitScope = EmitNode->Scope;
+			if (EmitScope)
+			{
+				EmitNode->NextScopedNode = EmitScope->FirstNode;
+				EmitScope->FirstNode = EmitNode;
+			}
 		}
-	}
 
-	{
-		FEmitShaderScopeStack Stack;
-		TStringBuilder<1024> ScopeCode;
-		Stack.Emplace(EmitRootScope, 1, ScopeCode);
-		EmitRootScope->EmitShaderCode(Stack);
-		check(Stack.Num() == 1);
-		OutCode.Append(ScopeCode.ToView());
+		{
+			FEmitShaderScopeStack Stack;
+			TStringBuilder<2048> ScopeCode;
+			Stack.Emplace(EmitRootScope, 1, ScopeCode);
+			EmitRootScope->EmitShaderCode(Stack);
+			check(Stack.Num() == 1);
+			OutCode.Append(ScopeCode.ToView());
+		}
 	}
 
 	Context.Finalize();
@@ -1208,7 +1092,7 @@ FExpression* FTree::AcquireLocal(FScope& Scope, const FName& LocalName)
 	return nullptr;
 }
 
-FExpression* FTree::NewFunctionCall(FScope& Scope, FFunction* Function, TArrayView<FExpression*> InputExpressions, int32 OutputIndex)
+FExpression* FTree::NewFunctionCall(FScope& Scope, FFunction* Function, int32 OutputIndex)
 {
 	FScope* CalledScope = &Scope;
 	if (Function->CalledScope)
@@ -1217,7 +1101,7 @@ FExpression* FTree::NewFunctionCall(FScope& Scope, FFunction* Function, TArrayVi
 		check(CalledScope);
 	}
 	Function->CalledScope = CalledScope;
-	return NewExpression<FExpressionFunctionCall>(Function, InputExpressions, OutputIndex);
+	return NewExpression<FExpressionFunctionCall>(Function, OutputIndex);
 }
 
 const FExpressionDerivatives& FTree::GetAnalyticDerivatives(FExpression* InExpression)
diff --git a/Engine/Source/Runtime/Engine/Private/HLSLTree/HLSLTreeCommon.cpp b/Engine/Source/Runtime/Engine/Private/HLSLTree/HLSLTreeCommon.cpp
--- a/Engine/Source/Runtime/Engine/Private/HLSLTree/HLSLTreeCommon.cpp
+++ b/Engine/Source/Runtime/Engine/Private/HLSLTree/HLSLTreeCommon.cpp
@@ -221,7 +221,7 @@ bool FExpressionExternalInput::PrepareValue(FEmitContext& Context, const FReques
 	return OutResult.SetType(Context, RequestedType, EExpressionEvaluation::Shader, InputDesc.Type);
 }
 
-void FExpressionExternalInput::EmitValueShader(FEmitContext& Context, FEmitShaderScope& Scope, const FRequestedType& RequestedType, FEmitValueShaderResult& OutResult) const
+void FExpressionExternalInput::EmitValueShader(FEmitContext& Context, FEmitScope& Scope, const FRequestedType& RequestedType, FEmitValueShaderResult& OutResult) const
 {
 	const int32 TypeIndex = (int32)InputType;
 	if (IsTexCoord(InputType))
@@ -274,14 +274,14 @@ void FExpressionExternalInput::EmitValueShader(FEmitContext& Context, FEmitShade
 
 bool FExpressionTextureSample::PrepareValue(FEmitContext& Context, const FRequestedType& RequestedType, FPrepareValueResult& OutResult) const
 {
-	PrepareExpressionValue(Context, TexCoordExpression, Shader::EValueType::Float2);
-	PrepareExpressionValue(Context, TexCoordDerivatives.ExpressionDdx, Shader::EValueType::Float2);
-	PrepareExpressionValue(Context, TexCoordDerivatives.ExpressionDdy, Shader::EValueType::Float2);
+	Context.PrepareExpression(TexCoordExpression, Shader::EValueType::Float2);
+	Context.PrepareExpression(TexCoordDerivatives.ExpressionDdx, Shader::EValueType::Float2);
+	Context.PrepareExpression(TexCoordDerivatives.ExpressionDdy, Shader::EValueType::Float2);
 
 	return OutResult.SetType(Context, RequestedType, EExpressionEvaluation::Shader, Shader::EValueType::Float4);
 }
 
-void FExpressionTextureSample::EmitValueShader(FEmitContext& Context, FEmitShaderScope& Scope, const FRequestedType& RequestedType, FEmitValueShaderResult& OutResult) const
+void FExpressionTextureSample::EmitValueShader(FEmitContext& Context, FEmitScope& Scope, const FRequestedType& RequestedType, FEmitValueShaderResult& OutResult) const
 {
 	const FTextureDescription& Desc = Declaration->Description;
 	const EMaterialValueType MaterialValueType = Desc.Texture->GetMaterialType();
@@ -440,7 +440,7 @@ bool FExpressionGetStructField::PrepareValue(FEmitContext& Context, const FReque
 	FRequestedType RequestedStructType(StructType, false);
 	RequestedStructType.SetField(Field, RequestedType);
 
-	const FPreparedType& StructPreparedType = PrepareExpressionValue(Context, StructExpression, RequestedStructType);
+	const FPreparedType& StructPreparedType = Context.PrepareExpression(StructExpression, RequestedStructType);
 	if (!StructPreparedType.IsVoid() && StructPreparedType.StructType != StructType)
 	{
 		return Context.Errors->AddErrorf(TEXT("Expected type %s"), StructType->Name);
@@ -449,7 +449,7 @@ bool FExpressionGetStructField::PrepareValue(FEmitContext& Context, const FReque
 	return OutResult.SetType(Context, RequestedType, StructPreparedType.GetFieldType(Field));
 }
 
-void FExpressionGetStructField::EmitValueShader(FEmitContext& Context, FEmitShaderScope& Scope, const FRequestedType& RequestedType, FEmitValueShaderResult& OutResult) const
+void FExpressionGetStructField::EmitValueShader(FEmitContext& Context, FEmitScope& Scope, const FRequestedType& RequestedType, FEmitValueShaderResult& OutResult) const
 {
 	FRequestedType RequestedStructType(StructType, false);
 	RequestedStructType.SetField(Field, RequestedType);
@@ -492,13 +492,13 @@ bool FExpressionSetStructField::PrepareValue(FEmitContext& Context, const FReque
 	FRequestedType RequestedStructType(RequestedType);
 	RequestedStructType.ClearFieldRequested(Field);
 
-	const FPreparedType& StructPreparedType = PrepareExpressionValue(Context, StructExpression, RequestedStructType);
+	const FPreparedType& StructPreparedType = Context.PrepareExpression(StructExpression, RequestedStructType);
 	if (!StructPreparedType.IsVoid() && StructPreparedType.StructType != StructType)
 	{
 		return Context.Errors->AddErrorf(TEXT("Expected type %s"), StructType->Name);
 	}
 
-	const FPreparedType& FieldPreparedType = PrepareExpressionValue(Context, FieldExpression, RequestedType.GetField(Field));
+	const FPreparedType& FieldPreparedType = Context.PrepareExpression(FieldExpression, RequestedType.GetField(Field));
 	if (FieldPreparedType.IsVoid())
 	{
 		return OutResult.SetForwardValue(Context, RequestedType, StructExpression);
@@ -515,7 +515,7 @@ bool FExpressionSetStructField::PrepareValue(FEmitContext& Context, const FReque
 	}
 }
 
-void FExpressionSetStructField::EmitValueShader(FEmitContext& Context, FEmitShaderScope& Scope, const FRequestedType& RequestedType, FEmitValueShaderResult& OutResult) const
+void FExpressionSetStructField::EmitValueShader(FEmitContext& Context, FEmitScope& Scope, const FRequestedType& RequestedType, FEmitValueShaderResult& OutResult) const
 {
 	FRequestedType RequestedStructType(RequestedType);
 	RequestedStructType.ClearFieldRequested(Field);
@@ -572,15 +572,15 @@ void FExpressionSetStructField::EmitValuePreshader(FEmitContext& Context, const
 
 bool FExpressionSelect::PrepareValue(FEmitContext& Context, const FRequestedType& RequestedType, FPrepareValueResult& OutResult) const
 {
-	const FPreparedType& ConditionType = PrepareExpressionValue(Context, ConditionExpression, 1);
+	const FPreparedType& ConditionType = Context.PrepareExpression(ConditionExpression, 1);
 	if (ConditionType.GetEvaluation(Shader::EValueType::Bool1) == EExpressionEvaluation::Constant)
 	{
 		const bool bCondition = ConditionExpression->GetValueConstant(Context, Shader::EValueType::Bool1).AsBoolScalar();
 		return OutResult.SetForwardValue(Context, RequestedType, bCondition ? TrueExpression : FalseExpression);
 	}
 
-	const FPreparedType& LhsType = PrepareExpressionValue(Context, FalseExpression, RequestedType);
-	const FPreparedType& RhsType = PrepareExpressionValue(Context, TrueExpression, RequestedType);
+	const FPreparedType& LhsType = Context.PrepareExpression(FalseExpression, RequestedType);
+	const FPreparedType& RhsType = Context.PrepareExpression(TrueExpression, RequestedType);
 	
 	if (LhsType.ValueComponentType != RhsType.ValueComponentType ||
 		LhsType.StructType != RhsType.StructType)
@@ -591,7 +591,7 @@ bool FExpressionSelect::PrepareValue(FEmitContext& Context, const FRequestedType
 	return OutResult.SetType(Context, RequestedType, MergePreparedTypes(LhsType, RhsType));
 }
 
-void FExpressionSelect::EmitValueShader(FEmitContext& Context, FEmitShaderScope& Scope, const FRequestedType& RequestedType, FEmitValueShaderResult& OutResult) const
+void FExpressionSelect::EmitValueShader(FEmitContext& Context, FEmitScope& Scope, const FRequestedType& RequestedType, FEmitValueShaderResult& OutResult) const
 {
 	FEmitShaderExpression* TrueValue = TrueExpression->GetValueShader(Context, Scope, RequestedType);
 	FEmitShaderExpression* FalseValue = FalseExpression->GetValueShader(Context, Scope, RequestedType);
@@ -640,7 +640,7 @@ void FExpressionUnaryOp::ComputeAnalyticDerivatives(FTree& Tree, FExpressionDeri
 
 bool FExpressionUnaryOp::PrepareValue(FEmitContext& Context, const FRequestedType& RequestedType, FPrepareValueResult& OutResult) const
 {
-	const FPreparedType& InputType = PrepareExpressionValue(Context, Input, RequestedType);
+	const FPreparedType& InputType = Context.PrepareExpression(Input, RequestedType);
 	if (InputType.IsVoid())
 	{
 		return false;
@@ -662,7 +662,7 @@ bool FExpressionUnaryOp::PrepareValue(FEmitContext& Context, const FRequestedTyp
 	return OutResult.SetType(Context, RequestedType, ResultType);
 }
 
-void FExpressionUnaryOp::EmitValueShader(FEmitContext& Context, FEmitShaderScope& Scope, const FRequestedType& RequestedType, FEmitValueShaderResult& OutResult) const
+void FExpressionUnaryOp::EmitValueShader(FEmitContext& Context, FEmitScope& Scope, const FRequestedType& RequestedType, FEmitValueShaderResult& OutResult) const
 {
 	FEmitShaderExpression* InputValue = Input->GetValueShader(Context, Scope, RequestedType);
 	Shader::FType ResultType = InputValue->Type;
@@ -765,8 +765,8 @@ void FExpressionBinaryOp::ComputeAnalyticDerivatives(FTree& Tree, FExpressionDer
 
 bool FExpressionBinaryOp::PrepareValue(FEmitContext& Context, const FRequestedType& RequestedType, FPrepareValueResult& OutResult) const
 {
-	const FPreparedType& LhsType = PrepareExpressionValue(Context, Lhs, RequestedType);
-	const FPreparedType& RhsType = PrepareExpressionValue(Context, Rhs, RequestedType);
+	const FPreparedType& LhsType = Context.PrepareExpression(Lhs, RequestedType);
+	const FPreparedType& RhsType = Context.PrepareExpression(Rhs, RequestedType);
 	if (LhsType.IsVoid() || RhsType.IsVoid())
 	{
 		return false;
@@ -808,7 +808,7 @@ FRequestedType GetBinaryOpInputType(const FRequestedType& RequestedType, EBinary
 }
 } // namespace Private
 
-void FExpressionBinaryOp::EmitValueShader(FEmitContext& Context, FEmitShaderScope& Scope, const FRequestedType& RequestedType, FEmitValueShaderResult& OutResult) const
+void FExpressionBinaryOp::EmitValueShader(FEmitContext& Context, FEmitScope& Scope, const FRequestedType& RequestedType, FEmitValueShaderResult& OutResult) const
 {
 	const FRequestedType InputType = Private::GetBinaryOpInputType(RequestedType, Op, Lhs->GetRequestedType(), Rhs->GetRequestedType());
 	const Shader::EValueComponentType InputComponentType = InputType.ValueComponentType;
@@ -899,7 +899,7 @@ void FExpressionSwizzle::ComputeAnalyticDerivatives(FTree& Tree, FExpressionDeri
 bool FExpressionSwizzle::PrepareValue(FEmitContext& Context, const FRequestedType& RequestedType, FPrepareValueResult& OutResult) const
 {
 	const FRequestedType RequestedInputType = Parameters.GetRequestedInputType(RequestedType);
-	const FPreparedType& InputType = PrepareExpressionValue(Context, Input, RequestedInputType);
+	const FPreparedType& InputType = Context.PrepareExpression(Input, RequestedInputType);
 
 	FPreparedType ResultType(InputType.ValueComponentType);
 	for (int32 ComponentIndex = 0; ComponentIndex < Parameters.NumComponents; ++ComponentIndex)
@@ -911,7 +911,7 @@ bool FExpressionSwizzle::PrepareValue(FEmitContext& Context, const FRequestedTyp
 	return OutResult.SetType(Context, RequestedType, ResultType);
 }
 
-void FExpressionSwizzle::EmitValueShader(FEmitContext& Context, FEmitShaderScope& Scope, const FRequestedType& RequestedType, FEmitValueShaderResult& OutResult) const
+void FExpressionSwizzle::EmitValueShader(FEmitContext& Context, FEmitScope& Scope, const FRequestedType& RequestedType, FEmitValueShaderResult& OutResult) const
 {
 	const FRequestedType RequestedInputType = Parameters.GetRequestedInputType(RequestedType);
 
@@ -968,7 +968,7 @@ void FExpressionAppend::ComputeAnalyticDerivatives(FTree& Tree, FExpressionDeriv
 
 bool FExpressionAppend::PrepareValue(FEmitContext& Context, const FRequestedType& RequestedType, FPrepareValueResult& OutResult) const
 {
-	const FPreparedType& LhsType = PrepareExpressionValue(Context, Lhs, RequestedType);
+	const FPreparedType& LhsType = Context.PrepareExpression(Lhs, RequestedType);
 	const int32 NumRequestedComponents = RequestedType.GetNumComponents();
 	const int32 NumLhsComponents = FMath::Min(LhsType.GetNumComponents(), NumRequestedComponents);
 
@@ -978,7 +978,7 @@ bool FExpressionAppend::PrepareValue(FEmitContext& Context, const FRequestedType
 		RhsRequestedType.SetComponentRequest(Index - NumLhsComponents, RequestedType.IsComponentRequested(Index));
 	}
 
-	const FPreparedType& RhsType = PrepareExpressionValue(Context, Rhs, RhsRequestedType);
+	const FPreparedType& RhsType = Context.PrepareExpression(Rhs, RhsRequestedType);
 	const int32 NumRhsComponents = FMath::Min(RhsType.GetNumComponents(), NumRequestedComponents - NumLhsComponents);
 
 	if (LhsType.ValueComponentType != RhsType.ValueComponentType)
@@ -1031,7 +1031,7 @@ FAppendTypes GetAppendTypes(const FRequestedType& RequestedType, const FRequeste
 }
 } // namespace Private
 
-void FExpressionAppend::EmitValueShader(FEmitContext& Context, FEmitShaderScope& Scope, const FRequestedType& RequestedType, FEmitValueShaderResult& OutResult) const
+void FExpressionAppend::EmitValueShader(FEmitContext& Context, FEmitScope& Scope, const FRequestedType& RequestedType, FEmitValueShaderResult& OutResult) const
 {
 	const Private::FAppendTypes Types = Private::GetAppendTypes(RequestedType, Lhs->GetRequestedType(), Rhs->GetRequestedType());
 	FEmitShaderExpression* LhsValue = Lhs->GetValueShader(Context, Scope, Types.LhsType);
@@ -1076,99 +1076,98 @@ bool FExpressionReflectionVector::PrepareValue(FEmitContext& Context, const FReq
 	return OutResult.SetType(Context, RequestedType, EExpressionEvaluation::Shader, Shader::EValueType::Float3);
 }
 
-void FExpressionReflectionVector::EmitValueShader(FEmitContext& Context, FEmitShaderScope& Scope, const FRequestedType& RequestedType, FEmitValueShaderResult& OutResult) const
+void FExpressionReflectionVector::EmitValueShader(FEmitContext& Context, FEmitScope& Scope, const FRequestedType& RequestedType, FEmitValueShaderResult& OutResult) const
 {
 	OutResult.Code = Context.EmitInlineExpression(Scope, Shader::EValueType::Float3, TEXT("Parameters.ReflectionVector"));
 }
 
-void FStatementBreak::Prepare(FEmitContext& Context) const
+bool FStatementBreak::Prepare(FEmitContext& Context) const
 {
+	return true;
 }
 
-void FStatementBreak::EmitShader(FEmitContext& Context, FEmitShaderScope& Scope) const
+void FStatementBreak::EmitShader(FEmitContext& Context, FEmitScope& Scope) const
 {
 	Context.EmitStatement(Scope, TEXT("break;"));
 }
 
-void FStatementReturn::Prepare(FEmitContext& Context) const
+bool FStatementReturn::Prepare(FEmitContext& Context) const
 {
+	return true;
 }
 
-void FStatementReturn::EmitShader(FEmitContext& Context, FEmitShaderScope& Scope) const
+void FStatementReturn::EmitShader(FEmitContext& Context, FEmitScope& Scope) const
 {
 	Context.EmitStatement(Scope, TEXT("return %;"), Expression->GetValueShader(Context, Scope));
 }
 
-void FStatementIf::Prepare(FEmitContext& Context) const
+bool FStatementIf::Prepare(FEmitContext& Context) const
 {
-	const FPreparedType& ConditionType = PrepareExpressionValue(Context, ConditionExpression, 1);
-	if (ConditionType.GetEvaluation(Shader::EValueType::Bool1) == EExpressionEvaluation::Constant)
+	const FPreparedType& ConditionType = Context.PrepareExpression(ConditionExpression, 1);
+	if (ConditionType.IsVoid())
+	{
+		return false;
+	}
+
+	const EExpressionEvaluation ConditionEvaluation = ConditionType.GetEvaluation(Shader::EValueType::Bool1);
+	check(ConditionEvaluation != EExpressionEvaluation::None);
+	if (ConditionEvaluation == EExpressionEvaluation::Constant)
 	{
 		const bool bCondition = ConditionExpression->GetValueConstant(Context, Shader::EValueType::Bool1).AsBoolScalar();
 		if (bCondition)
 		{
-			MarkScopeLive(ThenScope);
-			MarkScopeDead(ElseScope);
+			Context.PrepareScope(ThenScope);
+			Context.MarkScopeDead(ElseScope);
 		}
 		else
 		{
-			MarkScopeDead(ThenScope);
-			MarkScopeLive(ElseScope);
+			Context.MarkScopeDead(ThenScope);
+			Context.PrepareScope(ElseScope);
 		}
 	}
-	else if(!ConditionType.IsVoid())
+	else
 	{
-		MarkScopeLive(ThenScope);
-		MarkScopeLive(ElseScope);
+		Context.PrepareScope(ThenScope);
+		Context.PrepareScope(ElseScope);
 	}
+
+	return true;
 }
 
-void FStatementIf::EmitShader(FEmitContext& Context, FEmitShaderScope& Scope) const
+void FStatementIf::EmitShader(FEmitContext& Context, FEmitScope& Scope) const
 {
 	FEmitShaderNode* Dependency = nullptr;
-	if (ConditionExpression->GetEvaluation(Shader::EValueType::Bool1) == EExpressionEvaluation::Constant)
+	const EExpressionEvaluation ConditionEvaluation = ConditionExpression->GetEvaluation(Shader::EValueType::Bool1);
+	if (ConditionEvaluation == EExpressionEvaluation::Constant)
 	{
 		const bool bCondition = ConditionExpression->GetValueConstant(Context, Shader::EValueType::Bool1).AsBoolScalar();
-		if (bCondition && IsScopeLive(ThenScope))
+		if (bCondition)
 		{
 			Dependency = Context.EmitNextScope(Scope, ThenScope);
 		}
-		else if(!bCondition && IsScopeLive(ElseScope))
+		else if(!bCondition)
 		{
 			Dependency = Context.EmitNextScope(Scope, ElseScope);
 		}
 	}
-	else
+	else if(ConditionEvaluation != EExpressionEvaluation::None)
 	{
-		if (IsScopeLive(ThenScope) || IsScopeLive(ElseScope))
-		{
-			FEmitShaderExpression* ConditionValue = ConditionExpression->GetValueShader(Context, Scope, Shader::EValueType::Bool1);
-			if (IsScopeLive(ElseScope))
-			{
-				Dependency = Context.EmitNestedScopes(Scope, ThenScope, ElseScope, TEXT("if (%)"), TEXT("else"), ConditionValue);
-			}
-			else
-			{
-				Dependency = Context.EmitNestedScope(Scope, ThenScope, TEXT("if (%)"), ConditionValue);
-			}
-		}
+		FEmitShaderExpression* ConditionValue = ConditionExpression->GetValueShader(Context, Scope, Shader::EValueType::Bool1);
+		Dependency = Context.EmitNestedScopes(Scope, ThenScope, ElseScope, TEXT("if (%)"), TEXT("else"), ConditionValue);
 	}
 
 	Context.EmitNextScopeWithDependency(Scope, Dependency, NextScope);
 }
 
-void FStatementLoop::Prepare(FEmitContext& Context) const
+bool FStatementLoop::Prepare(FEmitContext& Context) const
 {
-	MarkScopeLive(LoopScope);
+	Context.PrepareScope(LoopScope);
+	return true;
 }
 
-void FStatementLoop::EmitShader(FEmitContext& Context, FEmitShaderScope& Scope) const
+void FStatementLoop::EmitShader(FEmitContext& Context, FEmitScope& Scope) const
 {
-	FEmitShaderNode* Dependency = nullptr;
-	if (IsScopeLive(LoopScope))
-	{
-		Dependency = Context.EmitNestedScope(Scope, LoopScope, TEXT("while (true)"));
-	}
+	FEmitShaderNode* Dependency = Context.EmitNestedScope(Scope, LoopScope, TEXT("while (true)"));
 	Context.EmitNextScopeWithDependency(Scope, Dependency, NextScope);
 }
 
diff --git a/Engine/Source/Runtime/Engine/Private/HLSLTree/HLSLTreeEmit.cpp b/Engine/Source/Runtime/Engine/Private/HLSLTree/HLSLTreeEmit.cpp
--- a/Engine/Source/Runtime/Engine/Private/HLSLTree/HLSLTreeEmit.cpp
+++ b/Engine/Source/Runtime/Engine/Private/HLSLTree/HLSLTreeEmit.cpp
@@ -14,14 +14,13 @@ namespace UE
 namespace HLSLTree
 {
 
-FEmitShaderNode::FEmitShaderNode(FEmitShaderScope& InScope, TArrayView<FEmitShaderNode*> InDependencies)
+FEmitShaderNode::FEmitShaderNode(FEmitScope& InScope, TArrayView<FEmitShaderNode*> InDependencies)
 	: Scope(&InScope)
 	, Dependencies(InDependencies)
 {
 	for (FEmitShaderNode* Dependency : InDependencies)
 	{
 		check(Dependency);
-		check(Dependency != (void*)0xcccccccccccccccc);
 	}
 }
 
@@ -42,7 +41,7 @@ void EmitShaderCode(FEmitShaderNode* EmitNode, FEmitShaderScopeStack& Stack)
 {
 	if (EmitNode && EmitNode->Scope)
 	{
-		const FEmitShaderScope* Scope = EmitNode->Scope;
+		const FEmitScope* Scope = EmitNode->Scope;
 		FEmitShaderScopeEntry EmitEntry;
 		for (int32 Index = Stack.Num() - 1; Index >= 0; --Index)
 		{
@@ -82,45 +81,46 @@ void FEmitShaderExpression::EmitShaderCode(FEmitShaderScopeStack& Stack, int32 I
 
 void FEmitShaderStatement::EmitShaderCode(FEmitShaderScopeStack& Stack, int32 Indent, FStringBuilderBase& OutString)
 {
-	TStringBuilder<1024> ScopeCode;
+	TStringBuilder<2048> ScopeCode;
 	for (int32 i = 0; i < 2; ++i)
 	{
+		bool bNeedToCloseScope = false;
+		int32 NestedScopeIndent = Indent;
+
 		if (Code[i].Len() > 0)
 		{
 			Private::WriteIndent(Indent, ScopeCode);
 			ScopeCode.Append(Code[i]);
 			ScopeCode.AppendChar(TEXT('\n'));
-		}
-
-		FEmitShaderScope* NestedScope = NestedScopes[i];
-		if (NestedScope)
-		{
-			bool bNeedToCloseScope = false;
-			int32 NextIndent = Indent;
 
+			// If ScopeFormat is set to 'Scoped', we need to emit an empty {}, even if our NestedScope is nullptr
 			if (ScopeFormat == EEmitScopeFormat::Scoped)
 			{
 				Private::WriteIndent(Indent, ScopeCode);
 				ScopeCode.Append(TEXT("{\n"));
 				bNeedToCloseScope = true;
-				NextIndent++;
+				NestedScopeIndent++;
 			}
+		}
 
-			Stack.Emplace(NestedScopes[i], NextIndent, ScopeCode);
+		FEmitScope* NestedScope = NestedScopes[i];
+		if (NestedScope)
+		{
+			Stack.Emplace(NestedScopes[i], NestedScopeIndent, ScopeCode);
 			NestedScope->EmitShaderCode(Stack);
 			Stack.Pop(false);
+		}
 
-			if (bNeedToCloseScope)
-			{
-				Private::WriteIndent(Indent, ScopeCode);
-				ScopeCode.Append(TEXT("}\n"));
-			}
+		if (bNeedToCloseScope)
+		{
+			Private::WriteIndent(Indent, ScopeCode);
+			ScopeCode.Append(TEXT("}\n"));
 		}
 	}
 	OutString.Append(ScopeCode.ToView());
 }
 
-void FEmitShaderScope::EmitShaderCode(FEmitShaderScopeStack& Stack)
+void FEmitScope::EmitShaderCode(FEmitShaderScopeStack& Stack)
 {
 	FEmitShaderNode* EmitNode = FirstNode;
 	while (EmitNode)
@@ -130,10 +130,10 @@ void FEmitShaderScope::EmitShaderCode(FEmitShaderScopeStack& Stack)
 	}
 }
 
-FEmitShaderScope* FEmitShaderScope::FindSharedParent(FEmitShaderScope* Lhs, FEmitShaderScope* Rhs)
+FEmitScope* FEmitScope::FindSharedParent(FEmitScope* Lhs, FEmitScope* Rhs)
 {
-	FEmitShaderScope* Scope0 = Lhs;
-	FEmitShaderScope* Scope1 = Rhs;
+	FEmitScope* Scope0 = Lhs;
+	FEmitScope* Scope1 = Rhs;
 	if (Scope1)
 	{
 		while (Scope0 != Scope1)
@@ -164,21 +164,174 @@ FEmitContext::~FEmitContext()
 {
 }
 
-const TCHAR* FEmitContext::AcquireLocalDeclarationCode()
+const FPreparedType& FEmitContext::PrepareExpression(FExpression* InExpression, const FRequestedType& RequestedType)
 {
-	return MemStack::AllocateStringf(*Allocator, TEXT("Local%d"), NumExpressionLocals++);
+	static FPreparedType VoidType;
+	if (!InExpression)
+	{
+		return VoidType;
+	}
+
+	FOwnerScope OwnerScope(*Errors, InExpression->GetOwner());
+	if (InExpression->bReentryFlag)
+	{
+		// Valid for this to be called reentrantly
+		// Code should ensure that the type is set before the reentrant call, otherwise type will not be valid here
+		// LocalPHI nodes rely on this to break loops
+		return InExpression->PrepareValueResult.GetPreparedType();
+	}
+
+	bool bNeedToUpdateType = false;
+	if (InExpression->CurrentRequestedType.RequestedComponents.Num() == 0)
+	{
+		InExpression->CurrentRequestedType = RequestedType;
+		bNeedToUpdateType = !RequestedType.IsVoid();
+	}
+	else if (InExpression->CurrentRequestedType.GetStructType() != RequestedType.GetStructType())
+	{
+		Errors->AddError(TEXT("Type mismatch"));
+		return VoidType;
+	}
+	else
+	{
+		const int32 NumComponents = RequestedType.GetNumComponents();
+		InExpression->CurrentRequestedType.RequestedComponents.PadToNum(NumComponents, false);
+		for (int32 Index = 0; Index < NumComponents; ++Index)
+		{
+			const bool bPrevRequest = InExpression->CurrentRequestedType.IsComponentRequested(Index);
+			const bool bRequest = RequestedType.IsComponentRequested(Index);
+			if (!bPrevRequest && bRequest)
+			{
+				InExpression->CurrentRequestedType.SetComponentRequest(Index);
+				bNeedToUpdateType = true;
+			}
+		}
+	}
+
+	if (bNeedToUpdateType)
+	{
+		check(!InExpression->CurrentRequestedType.IsVoid());
+
+		bool bResult = false;
+		{
+			FExpressionReentryScope ReentryScope(InExpression);
+			bResult = InExpression->PrepareValue(*this, InExpression->CurrentRequestedType, InExpression->PrepareValueResult);
+		}
+
+		if (!bResult)
+		{
+			// If we failed to assign a valid type, reset the requested type as well
+			// This ensures we'll try to compute a type again the next time we're called
+			InExpression->CurrentRequestedType.Reset();
+			InExpression->PrepareValueResult.SetTypeVoid();
+		}
+		else
+		{
+			check(!InExpression->PrepareValueResult.GetPreparedType().IsVoid());
+		}
+	}
+
+	return InExpression->PrepareValueResult.GetPreparedType();
+}
+
+FEmitScope* FEmitContext::InternalPrepareScope(FScope* Scope, FScope* ParentScope, bool bMarkDead)
+{
+	FEmitScope* EmitScope = nullptr;
+	if (Scope)
+	{
+		FEmitScope* EmitParentScope = InternalPrepareScope(ParentScope, ParentScope ? ParentScope->ParentScope : nullptr, false);
+		EmitScope = AcquireEmitScopeWithParent(Scope, EmitParentScope);
+
+		if (bMarkDead || (EmitParentScope && EmitParentScope->State == EEmitScopeState::Dead))
+		{
+			EmitScope->State = EEmitScopeState::Dead;
+		}
+		else if (EmitScope->State == EEmitScopeState::Uninitialized)
+		{
+			EmitScope->State = EEmitScopeState::Initializing;
+			bool bPrepareResult = true;
+			if (Scope->OwnerStatement)
+			{
+				bPrepareResult = Scope->OwnerStatement->Prepare(*this);
+			}
+			if (EmitScope->State == EEmitScopeState::Initializing)
+			{
+				// If Prepare() returns false, we're still uninitialized
+				EmitScope->State = bPrepareResult ? EEmitScopeState::Live : EEmitScopeState::Uninitialized;
+			}
+		}
+	}
+
+	return EmitScope;
+}
+
+bool FEmitContext::PrepareScope(FScope* Scope)
+{
+	FEmitScope* EmitScope = InternalPrepareScope(Scope, Scope ? Scope->ParentScope : nullptr, false);
+	return EmitScope && EmitScope->State != EEmitScopeState::Dead;
+}
+
+bool FEmitContext::PrepareScopeWithParent(FScope* Scope, FScope* ParentScope)
+{
+	FEmitScope* EmitScope = InternalPrepareScope(Scope, ParentScope, false);
+	return EmitScope && EmitScope->State != EEmitScopeState::Dead;
+}
+
+bool FEmitContext::MarkScopeDead(FScope* Scope)
+{
+	InternalPrepareScope(Scope, Scope ? Scope->ParentScope : nullptr, true);
+	return false;
+}
+
+FEmitScope* FEmitContext::AcquireEmitScopeWithParent(const FScope* Scope, FEmitScope* EmitParentScope)
+{
+	FEmitScope* EmitScope = nullptr;
+	if (Scope)
+	{
+		FEmitScope* const* PrevEmitScope = EmitScopeMap.Find(Scope);
+		EmitScope = PrevEmitScope ? *PrevEmitScope : nullptr;
+		if (!EmitScope)
+		{
+			EmitScope = new(*Allocator) FEmitScope();
+			EmitScope->ParentScope = EmitParentScope;
+			EmitScope->NestedLevel = EmitParentScope ? EmitParentScope->NestedLevel + 1 : 0;
+			EmitScopeMap.Add(Scope, EmitScope);
+		}
+		check(!EmitParentScope || EmitScope->ParentScope == EmitParentScope);
+	}
+	return EmitScope;
+}
+
+FEmitScope* FEmitContext::AcquireEmitScope(const FScope* Scope)
+{
+	FEmitScope* EmitParentScope = Scope ? AcquireEmitScope(Scope->ParentScope) : nullptr;
+	return AcquireEmitScopeWithParent(Scope, EmitParentScope);
+}
+
+FEmitScope* FEmitContext::InternalEmitScope(const FScope* Scope)
+{
+	FEmitScope* EmitScope = AcquireEmitScope(Scope);
+	if (EmitScope && EmitScope->State != EEmitScopeState::Dead)
+	{
+		if (Scope->ContainedStatement)
+		{
+			Scope->ContainedStatement->EmitShader(*this, *EmitScope);
+		}
+		return EmitScope;
+	}
+	return nullptr;
 }
 
 namespace Private
 {
-void MoveToScope(FEmitShaderNode* EmitNode, FEmitShaderScope& Scope)
+void MoveToScope(FEmitShaderNode* EmitNode, FEmitScope& Scope)
 {
 	if (EmitNode->Scope != &Scope)
 	{
-		FEmitShaderScope* NewScope = &Scope;
+		FEmitScope* NewScope = &Scope;
 		if (EmitNode->Scope)
 		{
-			NewScope = FEmitShaderScope::FindSharedParent(EmitNode->Scope, &Scope);
+			NewScope = FEmitScope::FindSharedParent(EmitNode->Scope, &Scope);
 			check(NewScope);
 		}
 
@@ -236,7 +389,7 @@ void InternalFormatStrings(FStringBuilderBase* OutString0, FStringBuilderBase* O
 
 } // namespace Private
 
-FEmitShaderExpression* FEmitContext::EmitExpressionInternal(FEmitShaderScope& Scope, TArrayView<FEmitShaderNode*> Dependencies, bool bInline, const Shader::FType& Type, FStringView Code)
+FEmitShaderExpression* FEmitContext::InternalEmitExpression(FEmitScope& Scope, TArrayView<FEmitShaderNode*> Dependencies, bool bInline, const Shader::FType& Type, FStringView Code)
 {
 	FEmitShaderExpression* ShaderValue = nullptr;
 
@@ -266,7 +419,7 @@ FEmitShaderExpression* FEmitContext::EmitExpressionInternal(FEmitShaderScope& Sc
 		}
 		else
 		{
-			ShaderValue->Reference = AcquireLocalDeclarationCode();
+			ShaderValue->Reference = MemStack::AllocateStringf(*Allocator, TEXT("Local%d"), NumExpressionLocals++);
 			ShaderValue->Value = MemStack::AllocateString(*Allocator, Code);
 		}
 		EmitExpressionMap.Add(ShaderHash, ShaderValue);
@@ -276,7 +429,7 @@ FEmitShaderExpression* FEmitContext::EmitExpressionInternal(FEmitShaderScope& Sc
 	return ShaderValue;
 }
 
-FEmitShaderStatement* FEmitContext::EmitStatementInternal(FEmitShaderScope& Scope, TArrayView<FEmitShaderNode*> Dependencies, EEmitScopeFormat ScopeFormat, FEmitShaderScope* NestedScope0, FEmitShaderScope* NestedScope1, FStringView Code0, FStringView Code1)
+FEmitShaderStatement* FEmitContext::InternalEmitStatement(FEmitScope& Scope, TArrayView<FEmitShaderNode*> Dependencies, EEmitScopeFormat ScopeFormat, FEmitScope* NestedScope0, FEmitScope* NestedScope1, FStringView Code0, FStringView Code1)
 {
 	FEmitShaderStatement* EmitStatement = new(*Allocator) FEmitShaderStatement(Scope, MemStack::AllocateArrayView(*Allocator, Dependencies));
 	EmitStatement->ScopeFormat = ScopeFormat;
@@ -289,46 +442,6 @@ FEmitShaderStatement* FEmitContext::EmitStatementInternal(FEmitShaderScope& Scop
 	return EmitStatement;
 }
 
-FEmitShaderScope* FEmitContext::AcquireEmitScope(FScope* Scope, FEmitShaderScope* OverrideParent)
-{
-	FEmitShaderScope* EmitScope = nullptr;
-	if (Scope)
-	{
-		FEmitShaderScope* const* PrevEmitScope = EmitScopeMap.Find(Scope);
-		EmitScope = PrevEmitScope ? *PrevEmitScope : nullptr;
-		if (!EmitScope)
-		{
-			FEmitShaderScope* ParentScope = OverrideParent;
-			if (!ParentScope)
-			{
-				ParentScope = AcquireEmitScope(Scope->ParentScope);
-			}
-
-			EmitScope = new(*Allocator) FEmitShaderScope();
-			EmitScope->ParentScope = ParentScope;
-			EmitScope->NestedLevel = ParentScope ? ParentScope->NestedLevel + 1 : 0;
-			EmitScopeMap.Add(Scope, EmitScope);
-			if (Scope->ContainedStatement)
-			{
-				Scope->ContainedStatement->EmitShader(*this, *EmitScope);
-			}
-		}
-		check(!OverrideParent || EmitScope->ParentScope == OverrideParent);
-	}
-	return EmitScope;
-}
-
-FEmitShaderScope* FEmitContext::FindEmitScope(FScope* Scope) const
-{
-	FEmitShaderScope* EmitScope = nullptr;
-	if (Scope)
-	{
-		FEmitShaderScope* const* PrevEmitScope = EmitScopeMap.Find(Scope);
-		EmitScope = PrevEmitScope ? *PrevEmitScope : nullptr;
-	}
-	return EmitScope;
-}
-
 namespace Private
 {
 void WriteMaterialUniformAccess(Shader::EValueComponentType ComponentType, uint32 NumComponents, uint32 UniformOffset, FStringBuilderBase& OutResult)
@@ -389,7 +502,7 @@ void WriteMaterialUniformAccess(Shader::EValueComponentType ComponentType, uint3
 }
 } // namespace Private
 
-FEmitShaderExpression* FEmitContext::EmitPreshaderOrConstant(FEmitShaderScope& Scope, const FRequestedType& RequestedType, FExpression* Expression)
+FEmitShaderExpression* FEmitContext::EmitPreshaderOrConstant(FEmitScope& Scope, const FRequestedType& RequestedType, FExpression* Expression)
 {
 	Shader::FPreshaderData LocalPreshader;
 	Expression->EmitValuePreshader(*this, RequestedType, LocalPreshader);
@@ -419,8 +532,6 @@ FEmitShaderExpression* FEmitContext::EmitPreshaderOrConstant(FEmitShaderScope& S
 	}
 
 	FMaterialUniformPreshaderHeader* PreshaderHeader = nullptr;
-	uint32 CurrentBoolUniformOffset = ~0u;
-	uint32 CurrentNumBoolComponents = 32u;
 
 	int32 ComponentIndex = 0;
 	for (int32 FieldIndex = 0; FieldIndex < Type.GetNumFlatFields(); ++FieldIndex)
@@ -588,18 +699,18 @@ FEmitShaderExpression* FEmitContext::EmitPreshaderOrConstant(FEmitShaderScope& S
 	}
 
 	const bool bInline = !Type.IsStruct(); // struct declarations can't be inline, due to HLSL syntax
-	FEmitShaderExpression* ShaderValue = EmitExpressionInternal(Scope, TArrayView<FEmitShaderNode*>(), bInline, Type, FormattedCode.ToView());
+	FEmitShaderExpression* ShaderValue = InternalEmitExpression(Scope, TArrayView<FEmitShaderNode*>(), bInline, Type, FormattedCode.ToView());
 	EmitPreshaderMap.Add(Hash, ShaderValue);
 
 	return ShaderValue;
 }
 
-FEmitShaderExpression* FEmitContext::EmitConstantZero(FEmitShaderScope& Scope, const Shader::FType& Type)
+FEmitShaderExpression* FEmitContext::EmitConstantZero(FEmitScope& Scope, const Shader::FType& Type)
 {
 	return EmitInlineExpression(Scope, Type, TEXT("((%)0)"), Type.GetName());
 }
 
-FEmitShaderExpression* FEmitContext::EmitCast(FEmitShaderScope& Scope, FEmitShaderExpression* ShaderValue, const Shader::FType& DestType)
+FEmitShaderExpression* FEmitContext::EmitCast(FEmitScope& Scope, FEmitShaderExpression* ShaderValue, const Shader::FType& DestType)
 {
 	check(ShaderValue);
 	check(!DestType.IsVoid());
diff --git a/Engine/Source/Runtime/Engine/Private/Materials/MaterialHLSLEmitter.cpp b/Engine/Source/Runtime/Engine/Private/Materials/MaterialHLSLEmitter.cpp
--- a/Engine/Source/Runtime/Engine/Private/Materials/MaterialHLSLEmitter.cpp
+++ b/Engine/Source/Runtime/Engine/Private/Materials/MaterialHLSLEmitter.cpp
@@ -732,15 +732,10 @@ bool MaterialEmitHLSL(const FMaterialCompileTargetParameters& InCompilerTarget,
 	const FStructField* NormalField = Generator.GetMaterialAttributesType()->FindFieldByName(*FMaterialAttributeDefinitionMap::GetAttributeName(MP_Normal));
 	check(NormalField);
 
-	if (!PrepareScope(EmitContext, Generator.GetResultStatement()->ParentScope))
-	{
-		return false;
-	}
-
 	// Prepare all fields *except* normal
 	FRequestedType RequestedMaterialAttributesType(Generator.GetMaterialAttributesType());
 	RequestedMaterialAttributesType.ClearFieldRequested(NormalField);
-	if (PrepareExpressionValue(EmitContext, Generator.GetResultExpression(), RequestedMaterialAttributesType).IsVoid())
+	if (EmitContext.PrepareExpression(Generator.GetResultExpression(), RequestedMaterialAttributesType).IsVoid())
 	{
 		return false;
 	}
@@ -753,7 +748,7 @@ bool MaterialEmitHLSL(const FMaterialCompileTargetParameters& InCompilerTarget,
 	{
 		// No access to material normal, can execute everything in phase0
 		RequestedMaterialAttributesType.SetFieldRequested(NormalField);
-		if (PrepareExpressionValue(EmitContext, Generator.GetResultExpression(), RequestedMaterialAttributesType).IsVoid())
+		if (EmitContext.PrepareExpression(Generator.GetResultExpression(), RequestedMaterialAttributesType).IsVoid())
 		{
 			return false;
 		}
@@ -767,15 +762,11 @@ bool MaterialEmitHLSL(const FMaterialCompileTargetParameters& InCompilerTarget,
 
 		// Reset state
 		HLSLTree->ResetNodes();
-		if (!PrepareScope(EmitContext, Generator.GetResultStatement()->ParentScope))
-		{
-			return false;
-		}
 
 		// Prepare code for just the normal
 		FRequestedType RequestedMaterialNormal(Generator.GetMaterialAttributesType(), false);
 		RequestedMaterialNormal.SetFieldRequested(NormalField);
-		if (PrepareExpressionValue(EmitContext, Generator.GetResultExpression(), RequestedMaterialNormal).IsVoid())
+		if (EmitContext.PrepareExpression(Generator.GetResultExpression(), RequestedMaterialNormal).IsVoid())
 		{
 			return false;
 		}
diff --git a/Engine/Source/Runtime/Engine/Private/Materials/MaterialHLSLGenerator.cpp b/Engine/Source/Runtime/Engine/Private/Materials/MaterialHLSLGenerator.cpp
--- a/Engine/Source/Runtime/Engine/Private/Materials/MaterialHLSLGenerator.cpp
+++ b/Engine/Source/Runtime/Engine/Private/Materials/MaterialHLSLGenerator.cpp
@@ -333,7 +333,7 @@ UE::HLSLTree::FScope* FMaterialHLSLGenerator::NewScope(UE::HLSLTree::FScope& Sco
 UE::HLSLTree::FScope* FMaterialHLSLGenerator::NewOwnedScope(UE::HLSLTree::FStatement& Owner)
 {
 	UE::HLSLTree::FScope* NewScope = HLSLTree->NewOwnedScope(Owner);
-	NewScope->AddPreviousScope(*Owner.ParentScope);
+	NewScope->AddPreviousScope(Owner.GetParentScope());
 	return NewScope;
 }
 
@@ -617,7 +617,7 @@ UE::HLSLTree::FExpression* FMaterialHLSLGenerator::GenerateFunctionCall(UE::HLSL
 		check(HLSLFunction->OutputExpressions.Num() == FunctionOutputs.Num());
 		if (HLSLFunction->OutputExpressions[OutputIndex])
 		{
-			Result = HLSLTree->NewFunctionCall(Scope, HLSLFunction, FunctionCall->ConnectedInputs, OutputIndex);
+			Result = HLSLTree->NewFunctionCall(Scope, HLSLFunction, OutputIndex);
 		}
 		else
 		{
diff --git a/Engine/Source/Runtime/Engine/Public/HLSLTree/HLSLTree.h b/Engine/Source/Runtime/Engine/Public/HLSLTree/HLSLTree.h
--- a/Engine/Source/Runtime/Engine/Public/HLSLTree/HLSLTree.h
+++ b/Engine/Source/Runtime/Engine/Public/HLSLTree/HLSLTree.h
@@ -39,7 +39,7 @@ class FExpressionLocalPHI;
 class FRequestedType;
 
 class FEmitContext;
-class FEmitShaderScope;
+class FEmitScope;
 class FEmitShaderExpression;
 
 static constexpr int32 MaxNumPreviousScopes = 2;
@@ -191,10 +191,17 @@ private:
 class FStatement : public FNode
 {
 public:
-	virtual void Prepare(FEmitContext& Context) const = 0;
-	virtual void EmitShader(FEmitContext& Context, FEmitShaderScope& Scope) const = 0;
+	FScope& GetParentScope() const { return *ParentScope; }
 
+protected:
+	virtual bool Prepare(FEmitContext& Context) const = 0;
+	virtual void EmitShader(FEmitContext& Context, FEmitScope& Scope) const = 0;
+
+private:
 	FScope* ParentScope = nullptr;
+
+	friend class FTree;
+	friend class FEmitContext;
 };
 
 /**
@@ -352,19 +359,17 @@ public:
 
 	virtual void Reset() override;
 
-	friend const FPreparedType& PrepareExpressionValue(FEmitContext& Context, FExpression* InExpression, const FRequestedType& RequestedType);
-
-	FEmitShaderExpression* GetValueShader(FEmitContext& Context, FEmitShaderScope& Scope, const FRequestedType& RequestedType, const Shader::FType& ResultType);
+	FEmitShaderExpression* GetValueShader(FEmitContext& Context, FEmitScope& Scope, const FRequestedType& RequestedType, const Shader::FType& ResultType);
 	void GetValuePreshader(FEmitContext& Context, const FRequestedType& RequestedType, Shader::FPreshaderData& OutPreshader);
 	Shader::FValue GetValueConstant(FEmitContext& Context, const FRequestedType& RequestedType);
 
-	FEmitShaderExpression* GetValueShader(FEmitContext& Context, FEmitShaderScope& Scope, const FRequestedType& RequestedType);
-	FEmitShaderExpression* GetValueShader(FEmitContext& Context, FEmitShaderScope& Scope);
+	FEmitShaderExpression* GetValueShader(FEmitContext& Context, FEmitScope& Scope, const FRequestedType& RequestedType);
+	FEmitShaderExpression* GetValueShader(FEmitContext& Context, FEmitScope& Scope);
 
 protected:
 	virtual void ComputeAnalyticDerivatives(FTree& Tree, FExpressionDerivatives& OutResult) const;
 	virtual bool PrepareValue(FEmitContext& Context, const FRequestedType& RequestedType, FPrepareValueResult& OutResult) const = 0;
-	virtual void EmitValueShader(FEmitContext& Context, FEmitShaderScope& Scope, const FRequestedType& RequestedType, FEmitValueShaderResult& OutResult) const;
+	virtual void EmitValueShader(FEmitContext& Context, FEmitScope& Scope, const FRequestedType& RequestedType, FEmitValueShaderResult& OutResult) const;
 	virtual void EmitValuePreshader(FEmitContext& Context, const FRequestedType& RequestedType, Shader::FPreshaderData& OutPreshader) const;
 
 private:
@@ -425,13 +430,6 @@ public:
 	FTextureDescription Description;
 };
 
-enum class EScopeState : uint8
-{
-	Uninitialized,
-	Live,
-	Dead,
-};
-
 /**
  * Represents an HLSL scope.  A scope contains a single statement, along with any expressions required by that statement
  */
@@ -440,11 +438,7 @@ class FScope final : public FNode
 public:
 	static FScope* FindSharedParent(FScope* Lhs, FScope* Rhs);
 
-	virtual void Reset() override;
-
 	inline FScope* GetParentScope() const { return ParentScope; }
-	inline bool IsLive() const { return State == EScopeState::Live; }
-	inline bool IsDead() const { return State == EScopeState::Dead; }
 
 	inline TArrayView<FScope*> GetPreviousScopes() const
 	{
@@ -456,11 +450,6 @@ public:
 
 	void AddPreviousScope(FScope& Scope);
 
-	friend bool PrepareScope(FEmitContext& Context, FScope* InScope);
-
-	void MarkLive();
-	void MarkLiveRecursive();
-	void MarkDead();
 
 private:
 	friend class FTree;
@@ -474,30 +463,8 @@ private:
 	TMap<FName, FExpression*> LocalMap;
 	int32 NumPreviousScopes = 0;
 	int32 NestedLevel = 0;
-	EScopeState State = EScopeState::Uninitialized;
 };
 
-inline bool IsScopeLive(const FScope* InScope)
-{
-	return (bool)InScope && InScope->IsLive();
-}
-
-inline void MarkScopeLive(FScope* InScope)
-{
-	if (InScope)
-	{
-		InScope->MarkLive();
-	}
-}
-
-inline void MarkScopeDead(FScope* InScope)
-{
-	if (InScope)
-	{
-		InScope->MarkDead();
-	}
-}
-
 namespace Private
 {
 uint64 GetNextTypeHash();
@@ -559,7 +526,7 @@ public:
 	void AssignLocal(FScope& Scope, const FName& LocalName, FExpression* Value);
 	FExpression* AcquireLocal(FScope& Scope, const FName& LocalName);
 
-	FExpression* NewFunctionCall(FScope& Scope, FFunction* Function, TArrayView<FExpression*> InputExpressions, int32 OutputIndex);
+	FExpression* NewFunctionCall(FScope& Scope, FFunction* Function, int32 OutputIndex);
 
 	const FExpressionDerivatives& GetAnalyticDerivatives(FExpression* InExpression);
 
@@ -608,9 +575,5 @@ private:
 	friend class FExpressionLocalPHI;
 };
 
-//Shader::EValueType RequestExpressionType(FEmitContext& Context, FExpression* InExpression, int8 InRequestedNumComponents); // friend of FExpression
-//EExpressionEvaluation PrepareExpressionValue(FEmitContext& Context, FExpression* InExpression); // friend of FExpression
-//void PrepareScopeValues(FEmitContext& Context, const FScope* InScope); // friend of FScope
-
 } // namespace HLSLTree
 } // namespace UE
diff --git a/Engine/Source/Runtime/Engine/Public/HLSLTree/HLSLTreeCommon.h b/Engine/Source/Runtime/Engine/Public/HLSLTree/HLSLTreeCommon.h
--- a/Engine/Source/Runtime/Engine/Public/HLSLTree/HLSLTreeCommon.h
+++ b/Engine/Source/Runtime/Engine/Public/HLSLTree/HLSLTreeCommon.h
@@ -123,7 +123,7 @@ public:
 
 	virtual void ComputeAnalyticDerivatives(FTree& Tree, FExpressionDerivatives& OutResult) const override;
 	virtual bool PrepareValue(FEmitContext& Context, const FRequestedType& RequestedType, FPrepareValueResult& OutResult) const override;
-	virtual void EmitValueShader(FEmitContext& Context, FEmitShaderScope& Scope, const FRequestedType& RequestedType, FEmitValueShaderResult& OutResult) const override;
+	virtual void EmitValueShader(FEmitContext& Context, FEmitScope& Scope, const FRequestedType& RequestedType, FEmitValueShaderResult& OutResult) const override;
 };
 
 class FExpressionTextureSample : public FExpression
@@ -144,7 +144,7 @@ public:
 	ETextureMipValueMode MipValueMode;
 
 	virtual bool PrepareValue(FEmitContext& Context, const FRequestedType& RequestedType, FPrepareValueResult& OutResult) const override;
-	virtual void EmitValueShader(FEmitContext& Context, FEmitShaderScope& Scope, const FRequestedType& RequestedType, FEmitValueShaderResult& OutResult) const override;
+	virtual void EmitValueShader(FEmitContext& Context, FEmitScope& Scope, const FRequestedType& RequestedType, FEmitValueShaderResult& OutResult) const override;
 };
 
 class FExpressionGetStructField : public FExpression
@@ -163,7 +163,7 @@ public:
 
 	virtual void ComputeAnalyticDerivatives(FTree& Tree, FExpressionDerivatives& OutResult) const override;
 	virtual bool PrepareValue(FEmitContext& Context, const FRequestedType& RequestedType, FPrepareValueResult& OutResult) const override;
-	virtual void EmitValueShader(FEmitContext& Context, FEmitShaderScope& Scope, const FRequestedType& RequestedType, FEmitValueShaderResult& OutResult) const override;
+	virtual void EmitValueShader(FEmitContext& Context, FEmitScope& Scope, const FRequestedType& RequestedType, FEmitValueShaderResult& OutResult) const override;
 	virtual void EmitValuePreshader(FEmitContext& Context, const FRequestedType& RequestedType, Shader::FPreshaderData& OutPreshader) const override;
 };
 
@@ -187,7 +187,7 @@ public:
 
 	virtual void ComputeAnalyticDerivatives(FTree& Tree, FExpressionDerivatives& OutResult) const override;
 	virtual bool PrepareValue(FEmitContext& Context, const FRequestedType& RequestedType, FPrepareValueResult& OutResult) const override;
-	virtual void EmitValueShader(FEmitContext& Context, FEmitShaderScope& Scope, const FRequestedType& RequestedType, FEmitValueShaderResult& OutResult) const override;
+	virtual void EmitValueShader(FEmitContext& Context, FEmitScope& Scope, const FRequestedType& RequestedType, FEmitValueShaderResult& OutResult) const override;
 	virtual void EmitValuePreshader(FEmitContext& Context, const FRequestedType& RequestedType, Shader::FPreshaderData& OutPreshader) const override;
 };
 
@@ -205,7 +205,7 @@ public:
 	FExpression* FalseExpression;
 
 	virtual bool PrepareValue(FEmitContext& Context, const FRequestedType& RequestedType, FPrepareValueResult& OutResult) const override;
-	virtual void EmitValueShader(FEmitContext& Context, FEmitShaderScope& Scope, const FRequestedType& RequestedType, FEmitValueShaderResult& OutResult) const override;
+	virtual void EmitValueShader(FEmitContext& Context, FEmitScope& Scope, const FRequestedType& RequestedType, FEmitValueShaderResult& OutResult) const override;
 	virtual void EmitValuePreshader(FEmitContext& Context, const FRequestedType& RequestedType, Shader::FPreshaderData& OutPreshader) const override;
 };
 
@@ -222,7 +222,7 @@ public:
 
 	virtual void ComputeAnalyticDerivatives(FTree& Tree, FExpressionDerivatives& OutResult) const override;
 	virtual bool PrepareValue(FEmitContext& Context, const FRequestedType& RequestedType, FPrepareValueResult& OutResult) const override;
-	virtual void EmitValueShader(FEmitContext& Context, FEmitShaderScope& Scope, const FRequestedType& RequestedType, FEmitValueShaderResult& OutResult) const override;
+	virtual void EmitValueShader(FEmitContext& Context, FEmitScope& Scope, const FRequestedType& RequestedType, FEmitValueShaderResult& OutResult) const override;
 	virtual void EmitValuePreshader(FEmitContext& Context, const FRequestedType& RequestedType, Shader::FPreshaderData& OutPreshader) const override;
 };
 
@@ -241,7 +241,7 @@ public:
 
 	virtual void ComputeAnalyticDerivatives(FTree& Tree, FExpressionDerivatives& OutResult) const override;
 	virtual bool PrepareValue(FEmitContext& Context, const FRequestedType& RequestedType, FPrepareValueResult& OutResult) const override;
-	virtual void EmitValueShader(FEmitContext& Context, FEmitShaderScope& Scope, const FRequestedType& RequestedType, FEmitValueShaderResult& OutResult) const override;
+	virtual void EmitValueShader(FEmitContext& Context, FEmitScope& Scope, const FRequestedType& RequestedType, FEmitValueShaderResult& OutResult) const override;
 	virtual void EmitValuePreshader(FEmitContext& Context, const FRequestedType& RequestedType, Shader::FPreshaderData& OutPreshader) const override;
 };
 
@@ -271,7 +271,7 @@ public:
 
 	virtual void ComputeAnalyticDerivatives(FTree& Tree, FExpressionDerivatives& OutResult) const override;
 	virtual bool PrepareValue(FEmitContext& Context, const FRequestedType& RequestedType, FPrepareValueResult& OutResult) const override;
-	virtual void EmitValueShader(FEmitContext& Context, FEmitShaderScope& Scope, const FRequestedType& RequestedType, FEmitValueShaderResult& OutResult) const override;
+	virtual void EmitValueShader(FEmitContext& Context, FEmitScope& Scope, const FRequestedType& RequestedType, FEmitValueShaderResult& OutResult) const override;
 	virtual void EmitValuePreshader(FEmitContext& Context, const FRequestedType& RequestedType, Shader::FPreshaderData& OutPreshader) const override;
 };
 
@@ -288,7 +288,7 @@ public:
 
 	virtual void ComputeAnalyticDerivatives(FTree& Tree, FExpressionDerivatives& OutResult) const override;
 	virtual bool PrepareValue(FEmitContext& Context, const FRequestedType& RequestedType, FPrepareValueResult& OutResult) const override;
-	virtual void EmitValueShader(FEmitContext& Context, FEmitShaderScope& Scope, const FRequestedType& RequestedType, FEmitValueShaderResult& OutResult) const override;
+	virtual void EmitValueShader(FEmitContext& Context, FEmitScope& Scope, const FRequestedType& RequestedType, FEmitValueShaderResult& OutResult) const override;
 	virtual void EmitValuePreshader(FEmitContext& Context, const FRequestedType& RequestedType, Shader::FPreshaderData& OutPreshader) const override;
 };
 
@@ -296,26 +296,23 @@ class FExpressionReflectionVector : public FExpression
 {
 public:
 	virtual bool PrepareValue(FEmitContext& Context, const FRequestedType& RequestedType, FPrepareValueResult& OutResult) const override;
-	virtual void EmitValueShader(FEmitContext& Context, FEmitShaderScope& Scope, const FRequestedType& RequestedType, FEmitValueShaderResult& OutResult) const override;
+	virtual void EmitValueShader(FEmitContext& Context, FEmitScope& Scope, const FRequestedType& RequestedType, FEmitValueShaderResult& OutResult) const override;
 };
 
 class FStatementReturn : public FStatement
 {
 public:
-	//static constexpr bool MarkScopeLiveRecursive = true;
 	FExpression* Expression;
 
-	virtual void Prepare(FEmitContext& Context) const override;
-	virtual void EmitShader(FEmitContext& Context, FEmitShaderScope& Scope) const override;
+	virtual bool Prepare(FEmitContext& Context) const override;
+	virtual void EmitShader(FEmitContext& Context, FEmitScope& Scope) const override;
 };
 
 class FStatementBreak : public FStatement
 {
 public:
-	//static constexpr bool MarkScopeLive = true;
-
-	virtual void Prepare(FEmitContext& Context) const override;
-	virtual void EmitShader(FEmitContext& Context, FEmitShaderScope& Scope) const override;
+	virtual bool Prepare(FEmitContext& Context) const override;
+	virtual void EmitShader(FEmitContext& Context, FEmitScope& Scope) const override;
 };
 
 class FStatementIf : public FStatement
@@ -326,8 +323,8 @@ public:
 	FScope* ElseScope;
 	FScope* NextScope;
 
-	virtual void Prepare(FEmitContext& Context) const override;
-	virtual void EmitShader(FEmitContext& Context, FEmitShaderScope& Scope) const override;
+	virtual bool Prepare(FEmitContext& Context) const override;
+	virtual void EmitShader(FEmitContext& Context, FEmitScope& Scope) const override;
 };
 
 class FStatementLoop : public FStatement
@@ -336,8 +333,8 @@ public:
 	FScope* LoopScope;
 	FScope* NextScope;
 
-	virtual void Prepare(FEmitContext& Context) const override;
-	virtual void EmitShader(FEmitContext& Context, FEmitShaderScope& Scope) const override;
+	virtual bool Prepare(FEmitContext& Context) const override;
+	virtual void EmitShader(FEmitContext& Context, FEmitScope& Scope) const override;
 };
 
 }
diff --git a/Engine/Source/Runtime/Engine/Public/HLSLTree/HLSLTreeEmit.h b/Engine/Source/Runtime/Engine/Public/HLSLTree/HLSLTreeEmit.h
--- a/Engine/Source/Runtime/Engine/Public/HLSLTree/HLSLTreeEmit.h
+++ b/Engine/Source/Runtime/Engine/Public/HLSLTree/HLSLTreeEmit.h
@@ -25,16 +25,17 @@ class FScope;
 class FExpression;
 class FFunction;
 class FRequestedType;
-class FEmitShaderScope;
+class FPreparedType;
+class FEmitScope;
 class FEmitShaderExpression;
 class FEmitShaderStatement;
 
 struct FEmitShaderScopeEntry
 {
 	FEmitShaderScopeEntry() = default;
-	FEmitShaderScopeEntry(FEmitShaderScope* InScope, int32 InIndent, FStringBuilderBase& InCode) : Scope(InScope), Code(&InCode), Indent(InIndent) {}
+	FEmitShaderScopeEntry(FEmitScope* InScope, int32 InIndent, FStringBuilderBase& InCode) : Scope(InScope), Code(&InCode), Indent(InIndent) {}
 
-	FEmitShaderScope* Scope = nullptr;
+	FEmitScope* Scope = nullptr;
 	FStringBuilderBase* Code = nullptr;
 	int32 Indent = 0;
 };
@@ -47,9 +48,9 @@ public:
 	virtual FEmitShaderExpression* AsExpression() { return nullptr; }
 	virtual FEmitShaderStatement* AsStatement() { return nullptr; }
 
-	FEmitShaderNode(FEmitShaderScope& InScope, TArrayView<FEmitShaderNode*> InDependencies);
+	FEmitShaderNode(FEmitScope& InScope, TArrayView<FEmitShaderNode*> InDependencies);
 
-	FEmitShaderScope* Scope = nullptr;
+	FEmitScope* Scope = nullptr;
 	FEmitShaderNode* NextScopedNode = nullptr;
 	TArrayView<FEmitShaderNode*> Dependencies;
 };
@@ -58,7 +59,7 @@ using FEmitShaderDependencies = TArray<FEmitShaderNode*, TInlineAllocator<8>>;
 class FEmitShaderExpression final : public FEmitShaderNode
 {
 public:
-	FEmitShaderExpression(FEmitShaderScope& InScope, TArrayView<FEmitShaderNode*> InDependencies, const Shader::FType& InType, FXxHash64 InHash)
+	FEmitShaderExpression(FEmitScope& InScope, TArrayView<FEmitShaderNode*> InDependencies, const Shader::FType& InType, FXxHash64 InHash)
 		: FEmitShaderNode(InScope, InDependencies)
 		, Type(InType)
 		, Hash(InHash)
@@ -84,17 +85,16 @@ enum class EEmitScopeFormat : uint8
 class FEmitShaderStatement final : public FEmitShaderNode
 {
 public:
-	FEmitShaderStatement(FEmitShaderScope& InScope, TArrayView<FEmitShaderNode*> InDependencies)
+	FEmitShaderStatement(FEmitScope& InScope, TArrayView<FEmitShaderNode*> InDependencies)
 		: FEmitShaderNode(InScope, InDependencies)
 	{}
 
 	virtual void EmitShaderCode(FEmitShaderScopeStack& Stack, int32 Indent, FStringBuilderBase& OutString) override;
 	virtual FEmitShaderStatement* AsStatement() override { return this; }
 
-	FEmitShaderScope* NestedScopes[2] = { nullptr };
+	FEmitScope* NestedScopes[2] = { nullptr };
 	FStringView Code[2];
 	EEmitScopeFormat ScopeFormat;
-	int8 NumEntries = 0;
 };
 
 enum class EFormatArgType : uint8
@@ -158,16 +158,25 @@ void FormatStrings(FStringBuilderBase& OutString0, FStringBuilderBase& OutString
 	Private::InternalFormatStrings(&OutString0, &OutString1, OutDependencies, Format0, Format1, ArgList);
 }
 
-class FEmitShaderScope
+enum class EEmitScopeState : uint8
+{
+	Uninitialized,
+	Initializing,
+	Live,
+	Dead,
+};
+
+class FEmitScope
 {
 public:
-	static FEmitShaderScope* FindSharedParent(FEmitShaderScope* Lhs, FEmitShaderScope* Rhs);
+	static FEmitScope* FindSharedParent(FEmitScope* Lhs, FEmitScope* Rhs);
 
 	void EmitShaderCode(FEmitShaderScopeStack& Stack);
 
-	FEmitShaderScope* ParentScope = nullptr;
+	FEmitScope* ParentScope = nullptr;
 	FEmitShaderNode* FirstNode;
 	int32 NestedLevel = 0;
+	EEmitScopeState State = EEmitScopeState::Uninitialized;
 };
 
 /** Tracks shared state while emitting HLSL code */
@@ -177,6 +186,17 @@ public:
 	explicit FEmitContext(FMemStackBase& InAllocator, FErrorHandlerInterface& InErrors, const Shader::FStructTypeRegistry& InTypeRegistry);
 	~FEmitContext();
 
+	const FPreparedType& PrepareExpression(FExpression* InExpression, const FRequestedType& RequestedType);
+
+	FEmitScope* InternalPrepareScope(FScope* Scope, FScope* ParentScope, bool bMarkDead);
+	bool PrepareScope(FScope* Scope);
+	bool PrepareScopeWithParent(FScope* Scope, FScope* ParentScope);
+	bool MarkScopeDead(FScope* Scope);
+
+	FEmitScope* AcquireEmitScopeWithParent(const FScope* Scope, FEmitScope* EmitParentScope);
+	FEmitScope* AcquireEmitScope(const FScope* Scope);
+	FEmitScope* InternalEmitScope(const FScope* Scope);
+
 	template<typename T>
 	static TArrayView<FEmitShaderNode*> MakeDependencies(T*& Dependency)
 	{
@@ -185,55 +205,52 @@ public:
 
 	void Finalize();
 
-	/** Get a unique local variable name */
-	const TCHAR* AcquireLocalDeclarationCode();
-
-	FEmitShaderExpression* EmitExpressionInternal(FEmitShaderScope& Scope, TArrayView<FEmitShaderNode*> Dependencies, bool bInline, const Shader::FType& Type, FStringView Code);
+	FEmitShaderExpression* InternalEmitExpression(FEmitScope& Scope, TArrayView<FEmitShaderNode*> Dependencies, bool bInline, const Shader::FType& Type, FStringView Code);
 
 	template<typename FormatType, typename... Types>
-	FEmitShaderExpression* EmitExpressionWithDependencies(FEmitShaderScope& Scope, TArrayView<FEmitShaderNode*> Dependencies, const Shader::FType& Type, const FormatType& Format, Types... Args)
+	FEmitShaderExpression* EmitExpressionWithDependencies(FEmitScope& Scope, TArrayView<FEmitShaderNode*> Dependencies, const Shader::FType& Type, const FormatType& Format, Types... Args)
 	{
 		TStringBuilder<2048> String;
 		FEmitShaderDependencies LocalDependencies(Dependencies);
 		FormatString(String, LocalDependencies, Format, Forward<Types>(Args)...);
-		return EmitExpressionInternal(Scope, LocalDependencies, false, Type, String.ToView());
+		return InternalEmitExpression(Scope, LocalDependencies, false, Type, String.ToView());
 	}
 
 	template<typename FormatType, typename... Types>
-	FEmitShaderExpression* EmitInlineExpressionWithDependencies(FEmitShaderScope& Scope, TArrayView<FEmitShaderNode*> Dependencies, const Shader::FType& Type, const FormatType& Format, Types... Args)
+	FEmitShaderExpression* EmitInlineExpressionWithDependencies(FEmitScope& Scope, TArrayView<FEmitShaderNode*> Dependencies, const Shader::FType& Type, const FormatType& Format, Types... Args)
 	{
 		TStringBuilder<2048> String;
 		FEmitShaderDependencies LocalDependencies(Dependencies);
 		FormatString(String, LocalDependencies, Format, Forward<Types>(Args)...);
-		return EmitExpressionInternal(Scope, LocalDependencies, true, Type, String.ToView());
+		return InternalEmitExpression(Scope, LocalDependencies, true, Type, String.ToView());
 	}
 
 	template<typename FormatType, typename... Types>
-	FEmitShaderExpression* EmitInlineExpressionWithDependency(FEmitShaderScope& Scope, FEmitShaderNode* Dependency, const Shader::FType& Type, const FormatType& Format, Types... Args)
+	FEmitShaderExpression* EmitInlineExpressionWithDependency(FEmitScope& Scope, FEmitShaderNode* Dependency, const Shader::FType& Type, const FormatType& Format, Types... Args)
 	{
 		return EmitInlineExpressionWithDependencies(Scope, MakeDependencies(Dependency), Type, Format, Forward<Types>(Args)...);
 	}
 
 	template<typename FormatType, typename... Types>
-	FEmitShaderExpression* EmitExpression(FEmitShaderScope& Scope, const Shader::FType& Type, const FormatType& Format, Types... Args)
+	FEmitShaderExpression* EmitExpression(FEmitScope& Scope, const Shader::FType& Type, const FormatType& Format, Types... Args)
 	{
 		return EmitExpressionWithDependencies(Scope, TArrayView<FEmitShaderNode*>(), Type, Format, Forward<Types>(Args)...);
 	}
 
 	template<typename FormatType, typename... Types>
-	FEmitShaderExpression* EmitInlineExpression(FEmitShaderScope& Scope, const Shader::FType& Type, const FormatType& Format, Types... Args)
+	FEmitShaderExpression* EmitInlineExpression(FEmitScope& Scope, const Shader::FType& Type, const FormatType& Format, Types... Args)
 	{
 		return EmitInlineExpressionWithDependencies(Scope, TArrayView<FEmitShaderNode*>(), Type, Format, Forward<Types>(Args)...);
 	}
 
-	FEmitShaderStatement* EmitStatementInternal(FEmitShaderScope& Scope, TArrayView<FEmitShaderNode*> Dependencies, EEmitScopeFormat ScopeFormat, FEmitShaderScope* NestedScope0, FEmitShaderScope* NestedScope1, FStringView Code0, FStringView Code1);
+	FEmitShaderStatement* InternalEmitStatement(FEmitScope& Scope, TArrayView<FEmitShaderNode*> Dependencies, EEmitScopeFormat ScopeFormat, FEmitScope* NestedScope0, FEmitScope* NestedScope1, FStringView Code0, FStringView Code1);
 
 	template<typename FormatType0, typename FormatType1, typename... Types>
-	FEmitShaderStatement* EmitFormatStatementInternal(FEmitShaderScope& Scope,
+	FEmitShaderStatement* EmitFormatStatementInternal(FEmitScope& Scope,
 		TArrayView<FEmitShaderNode*> Dependencies,
 		EEmitScopeFormat ScopeFormat,
-		FEmitShaderScope* NestedScope0,
-		FEmitShaderScope* NestedScope1,
+		FEmitScope* NestedScope0,
+		FEmitScope* NestedScope1,
 		const FormatType0& Format0,
 		const FormatType1& Format1,
 		Types... Args)
@@ -242,66 +259,93 @@ public:
 		TStringBuilder<1024> String1;
 		FEmitShaderDependencies LocalDependencies(Dependencies);
 		FormatStrings(String0, String1, LocalDependencies, Format0, Format1, Forward<Types>(Args)...);
-		return EmitStatementInternal(Scope, LocalDependencies, ScopeFormat, NestedScope0, NestedScope1, String0.ToView(), String1.ToView());
+		return InternalEmitStatement(Scope, LocalDependencies, ScopeFormat, NestedScope0, NestedScope1, String0.ToView(), String1.ToView());
 	}
 
 	template<typename FormatType, typename... Types>
-	FEmitShaderStatement* EmitStatementWithDependencies(FEmitShaderScope& Scope, TArrayView<FEmitShaderNode*> Dependencies, const FormatType& Format, Types... Args)
+	FEmitShaderStatement* EmitStatementWithDependencies(FEmitScope& Scope, TArrayView<FEmitShaderNode*> Dependencies, const FormatType& Format, Types... Args)
 	{
-		return EmitFormatStatementInternal(Scope, Dependencies, EEmitScopeFormat::None, nullptr, nullptr, Format, FStringView(), Forward<Types>(Args)...);
+		TStringBuilder<1024> String;
+		FEmitShaderDependencies LocalDependencies(Dependencies);
+		FormatString(String, LocalDependencies, Format, Forward<Types>(Args)...);
+		return InternalEmitStatement(Scope, LocalDependencies, EEmitScopeFormat::None, nullptr, nullptr, String.ToView(), FStringView());
 	}
 
 	template<typename FormatType, typename... Types>
-	FEmitShaderStatement* EmitStatementWithDependency(FEmitShaderScope& Scope, FEmitShaderNode* Dependency, const FormatType& Format, Types... Args)
+	FEmitShaderStatement* EmitStatementWithDependency(FEmitScope& Scope, FEmitShaderNode* Dependency, const FormatType& Format, Types... Args)
 	{
 		return EmitStatementWithDependencies(Scope, MakeDependencies(Dependency), Format, Forward<Types>(Args)...);
 	}
 
 	template<typename FormatType, typename... Types>
-	FEmitShaderStatement* EmitStatement(FEmitShaderScope& Scope, const FormatType& Format, Types... Args)
+	FEmitShaderStatement* EmitStatement(FEmitScope& Scope, const FormatType& Format, Types... Args)
 	{
 		return EmitStatementWithDependencies(Scope, TArrayView<FEmitShaderNode*>(), Format, Forward<Types>(Args)...);
 	}
 
-	FEmitShaderStatement* EmitNextScopeWithDependency(FEmitShaderScope& Scope, FEmitShaderNode* Dependency, FScope* NextScope)
+	FEmitShaderStatement* EmitNextScopeWithDependency(FEmitScope& Scope, FEmitShaderNode* Dependency, FScope* NextScope)
 	{
-		FEmitShaderScope* EmitScope = AcquireEmitScope(NextScope, &Scope);
-		return EmitStatementInternal(Scope, MakeDependencies(Dependency), EEmitScopeFormat::Unscoped, EmitScope, nullptr, FStringView(), FStringView());
+		FEmitScope* EmitScope = InternalEmitScope(NextScope);
+		if (EmitScope)
+		{
+			return InternalEmitStatement(Scope, MakeDependencies(Dependency), EEmitScopeFormat::Unscoped, EmitScope, nullptr, FStringView(), FStringView());
+		}
+		return nullptr;
 	}
 
-	FEmitShaderStatement* EmitNextScope(FEmitShaderScope& Scope, FScope* NextScope)
+	FEmitShaderStatement* EmitNextScope(FEmitScope& Scope, FScope* NextScope)
 	{
 		return EmitNextScopeWithDependency(Scope, nullptr, NextScope);
 	}
 
 	template<typename FormatType, typename... Types>
-	FEmitShaderStatement* EmitNestedScope(FEmitShaderScope& Scope, FScope* NestedScope, const FormatType& Format, Types... Args)
+	FEmitShaderStatement* EmitNestedScope(FEmitScope& Scope, FScope* NestedScope, const FormatType& Format, Types... Args)
 	{
-		FEmitShaderScope* EmitScope = AcquireEmitScope(NestedScope, &Scope);
-		return EmitFormatStatementInternal(Scope, TArrayView<FEmitShaderNode*>(), EEmitScopeFormat::Scoped, EmitScope, nullptr, Format, FStringView(), Forward<Types>(Args)...);
+		FEmitScope* EmitScope = InternalEmitScope(NestedScope);
+		if (EmitScope)
+		{
+			TStringBuilder<1024> String;
+			FEmitShaderDependencies LocalDependencies;
+			FormatString(String, LocalDependencies, Format, Forward<Types>(Args)...);
+			return InternalEmitStatement(Scope, LocalDependencies, EEmitScopeFormat::Scoped, EmitScope, nullptr, String.ToView(), FStringView());
+		}
+		return nullptr;
 	}
 
 	template<typename FormatType0, typename FormatType1, typename... Types>
-	FEmitShaderStatement* EmitNestedScopes(FEmitShaderScope& Scope, FScope* NestedScope0, FScope* NestedScope1, const FormatType0& Format0, const FormatType1& Format1, Types... Args)
+	FEmitShaderStatement* EmitNestedScopes(FEmitScope& Scope, FScope* NestedScope0, FScope* NestedScope1, const FormatType0& Format0, const FormatType1& Format1, Types... Args)
 	{
-		FEmitShaderScope* EmitScope0 = AcquireEmitScope(NestedScope0, &Scope);
-		FEmitShaderScope* EmitScope1 = AcquireEmitScope(NestedScope1, &Scope);
-		return EmitFormatStatementInternal(Scope, TArrayView<FEmitShaderNode*>(), EEmitScopeFormat::Scoped, EmitScope0, EmitScope1, Format0, Format1, Forward<Types>(Args)...);
+		FEmitScope* EmitScope0 = InternalEmitScope(NestedScope0);
+		FEmitScope* EmitScope1 = InternalEmitScope(NestedScope1);
+		if (EmitScope1)
+		{
+			TStringBuilder<1024> String0;
+			TStringBuilder<1024> String1;
+			FEmitShaderDependencies LocalDependencies;
+			FormatStrings(String0, String1, LocalDependencies, Format0, Format1, Forward<Types>(Args)...);
+			return InternalEmitStatement(Scope, LocalDependencies, EEmitScopeFormat::Scoped, EmitScope0, EmitScope1, String0.ToView(), String1.ToView());
+		}
+		else if (EmitScope0)
+		{
+			TStringBuilder<1024> String;
+			FEmitShaderDependencies LocalDependencies;
+			FormatString(String, LocalDependencies, Format0, Forward<Types>(Args)...);
+			return InternalEmitStatement(Scope, LocalDependencies, EEmitScopeFormat::Scoped, EmitScope0, nullptr, String.ToView(), FStringView());
+		}
+
+		return nullptr;
 	}
 
-	FEmitShaderScope* AcquireEmitScope(FScope* Scope, FEmitShaderScope* OverrideParent = nullptr);
-	FEmitShaderScope* FindEmitScope(FScope* Scope) const;
-
-	FEmitShaderExpression* EmitPreshaderOrConstant(FEmitShaderScope& Scope, const FRequestedType& RequestedType, FExpression* Expression);
-	FEmitShaderExpression* EmitConstantZero(FEmitShaderScope& Scope, const Shader::FType& Type);
-	FEmitShaderExpression* EmitCast(FEmitShaderScope& Scope, FEmitShaderExpression* ShaderValue, const Shader::FType& DestType);
+	FEmitShaderExpression* EmitPreshaderOrConstant(FEmitScope& Scope, const FRequestedType& RequestedType, FExpression* Expression);
+	FEmitShaderExpression* EmitConstantZero(FEmitScope& Scope, const Shader::FType& Type);
+	FEmitShaderExpression* EmitCast(FEmitScope& Scope, FEmitShaderExpression* ShaderValue, const Shader::FType& DestType);
 
 	FMemStackBase* Allocator = nullptr;
 	FErrorHandlerInterface* Errors = nullptr;
 	const Shader::FStructTypeRegistry* TypeRegistry = nullptr;
 
 	TArray<FEmitShaderNode*> EmitNodes;
-	TMap<const FScope*, FEmitShaderScope*> EmitScopeMap;
+	TMap<const FScope*, FEmitScope*> EmitScopeMap;
 	TMap<const FExpression*, FEmitShaderExpression*> EmitLocalPHIMap;
 	TMap<FXxHash64, FEmitShaderExpression*> EmitExpressionMap;
 	TMap<FXxHash64, FEmitShaderExpression*> EmitPreshaderMap;
@@ -313,6 +357,8 @@ public:
 	FMaterialCompilationOutput* MaterialCompilationOutput = nullptr;
 	TMap<Shader::FValue, uint32> DefaultUniformValues;
 	uint32 UniformPreshaderOffset = 0u;
+	uint32 CurrentBoolUniformOffset = ~0u;
+	uint32 CurrentNumBoolComponents = 32u;
 	bool bReadMaterialNormal = false;
 
 	int32 NumExpressionLocals = 0;
