Fix merge error wiping some partner code after a small refactor of the skeletalmesh cache code from cl 18583443
#rnx

diff --git a/Engine/Source/Runtime/Engine/Private/MorphTools.cpp b/Engine/Source/Runtime/Engine/Private/MorphTools.cpp
--- a/Engine/Source/Runtime/Engine/Private/MorphTools.cpp
+++ b/Engine/Source/Runtime/Engine/Private/MorphTools.cpp
@@ -10,6 +10,7 @@
 #include "Animation/MorphTarget.h"
 #include "Rendering/SkeletalMeshModel.h"
 #include "Rendering/SkeletalMeshLODModel.h"
+#include "UObject/GarbageCollection.h"
 
 /** compare based on base mesh source vertex indices */
 struct FCompareMorphTargetDeltas
@@ -168,13 +169,15 @@ TUniquePtr<FFinishBuildMorphTargetData> UMorphTarget::CreateFinishBuildMorphTarg
 
 void FFinishBuildMorphTargetData::ApplyEditorData(USkeletalMesh * SkeletalMesh) const
 {
-	check(IsInGameThread());
 	//Return if we do not need to apply data
 	if (!bApplyMorphTargetsData)
 	{
 		return;
 	}
 	
+	//Make sure we do not create new morph target during a gc
+	FGCScopeGuard GCScopeGuard;
+
 	FSkeletalMeshModel * SkelMeshModel = SkeletalMesh->GetImportedModel();
 	check(SkelMeshModel);
 	
@@ -215,13 +218,17 @@ void FFinishBuildMorphTargetData::ApplyEditorData(USkeletalMesh * SkeletalMesh)
 	//Rebuild the mapping and rehook the curve data
 	SkeletalMesh->InitMorphTargets();
 	
+	//Clear any async flags after the morphtargets have been set to the skeletalmesh
+	for (UMorphTarget* MorphTarget : SkeletalMesh->GetMorphTargets())
+	{
+		const EInternalObjectFlags AsyncFlags = EInternalObjectFlags::Async | EInternalObjectFlags::AsyncLoading;
+		MorphTarget->ClearInternalFlags(AsyncFlags);
+	}
+
 	for (UMorphTarget* ToDeleteMorphTarget : ToDeleteMorphTargets)
 	{
 		ToDeleteMorphTarget->BaseSkelMesh = nullptr;
 		ToDeleteMorphTarget->EmptyMorphLODModels();
-		
-		//Move the unused asset in the transient package and mark it pending kill
-		ToDeleteMorphTarget->Rename(nullptr, GetTransientPackage(), REN_ForceNoResetLoaders | REN_DoNotDirty | REN_DontCreateRedirectors | REN_NonTransactional);
 		ToDeleteMorphTarget->MarkAsGarbage();
 	}
 }
diff --git a/Engine/Source/Runtime/Engine/Private/SkeletalMeshRenderData.cpp b/Engine/Source/Runtime/Engine/Private/SkeletalMeshRenderData.cpp
--- a/Engine/Source/Runtime/Engine/Private/SkeletalMeshRenderData.cpp
+++ b/Engine/Source/Runtime/Engine/Private/SkeletalMeshRenderData.cpp
@@ -21,7 +21,6 @@
 #include "Serialization/LargeMemoryReader.h"
 #include "Serialization/LargeMemoryWriter.h"
 #include "Async/Async.h"
-#include "UObject/GarbageCollection.h"
 
 #if ENABLE_COOK_STATS
 namespace SkeletalMeshCookStats
@@ -158,69 +157,6 @@ namespace DDCUtils64Bit
 	}
 } //namespace DDCUtils64Bit
 
-namespace MorphTargetUtils
-{
-	void ApplyMorphTargetsEditorData(USkeletalMesh* SkeletalMesh, TMap<FName, TArray<FMorphTargetLODModel> >& MorphLODModelsPerTargetName)
-	{
-		//Make sure we do not create new morph target during a gc
-		FGCScopeGuard GCScopeGuard;
-
-		FSkeletalMeshModel* SkelMeshModel = SkeletalMesh->GetImportedModel();
-		check(SkelMeshModel);
-
-		TMap<FName, UMorphTarget*> ExistingMorphTargets;
-		ExistingMorphTargets.Reserve(SkeletalMesh->GetMorphTargets().Num());
-		for (UMorphTarget* MorphTarget : SkeletalMesh->GetMorphTargets())
-		{
-			ExistingMorphTargets.Add(MorphTarget->GetFName(), MorphTarget);
-		}
-		int32 MorphTargetNumber = MorphLODModelsPerTargetName.Num();
-		TArray<UMorphTarget*> ToDeleteMorphTargets;
-		ToDeleteMorphTargets.Append(SkeletalMesh->GetMorphTargets());
-		SkeletalMesh->GetMorphTargets().Empty();
-		//Rebuild the MorphTarget object
-		for (TPair<FName, TArray<FMorphTargetLODModel>> TargetNameAndMorphLODModels : MorphLODModelsPerTargetName)
-		{
-			FName MorphTargetName = TargetNameAndMorphLODModels.Key;
-			UMorphTarget* MorphTarget = ExistingMorphTargets.FindRef(MorphTargetName);
-			if (!MorphTarget)
-			{
-				MorphTarget = NewObject<UMorphTarget>(SkeletalMesh, MorphTargetName);
-				check(MorphTarget);
-			}
-			else
-			{
-				ToDeleteMorphTargets.Remove(MorphTarget);
-			}
-			MorphTarget->EmptyMorphLODModels();
-			SkeletalMesh->GetMorphTargets().Add(MorphTarget);
-			const TArray<FMorphTargetLODModel>& MorphTargetLODModels = TargetNameAndMorphLODModels.Value;
-			int32 MorphLODModelNumber = MorphTargetLODModels.Num();
-			MorphTarget->GetMorphLODModels().AddDefaulted(MorphLODModelNumber);
-			for (int32 MorphDataIndex = 0; MorphDataIndex < MorphLODModelNumber; ++MorphDataIndex)
-			{
-				MorphTarget->GetMorphLODModels()[MorphDataIndex] = MorphTargetLODModels[MorphDataIndex];
-			}
-		}
-		//Rebuild the mapping and rehook the curve data
-		SkeletalMesh->InitMorphTargets();
-
-		//Clear any async flags after the morphtargets have been set to the skeletalmesh
-		for (UMorphTarget* MorphTarget : SkeletalMesh->GetMorphTargets())
-		{
-			const EInternalObjectFlags AsyncFlags = EInternalObjectFlags::Async | EInternalObjectFlags::AsyncLoading;
-			MorphTarget->ClearInternalFlags(AsyncFlags);
-		}
-
-		//Make sure the old unused morph targets are clean up properly
-		for (UMorphTarget* ToDeleteMorphTarget : ToDeleteMorphTargets)
-		{
-			ToDeleteMorphTarget->BaseSkelMesh = nullptr;
-			ToDeleteMorphTarget->EmptyMorphLODModels();
-			ToDeleteMorphTarget->MarkAsGarbage();
-		}
-	}
-}
 //Serialize the LODInfo and append the result to the KeySuffix to build the LODInfo part of the DDC KEY
 //Note: this serializer is only used to build the mesh DDC key, no versioning is required
 static void SerializeLODInfoForDDC(USkeletalMesh* SkeletalMesh, FString& KeySuffix)
@@ -424,7 +360,7 @@ void FSkeletalMeshRenderData::Cache(const ITargetPlatform* TargetPlatform, USkel
 			FLargeMemoryReader Ar(DerivedData.GetData(), DerivedData.Num(), ELargeMemoryReaderFlags::Persistent);
 
 			//Helper structure to change the morph targets
-			TMap<FName, TArray<FMorphTargetLODModel> > MorphLODModelsPerTargetName;
+			TUniquePtr<FFinishBuildMorphTargetData> FinishBuildMorphTargetData;
 
 			//With skeletal mesh build refactor we serialize the LODModel data into the DDC
 			//We need to store those so we do not have to rerun the reduction to make them up to date
@@ -437,24 +373,17 @@ void FSkeletalMeshRenderData::Cache(const ITargetPlatform* TargetPlatform, USkel
 
 				//Get the morph target data, we put it in the compilation context to apply them in the game thread before the InitResources
 				{
-					int32 MorphTargetNumber = 0;
-					Ar << MorphTargetNumber;
-					MorphLODModelsPerTargetName.Reserve(MorphTargetNumber);
-					//We cannot serialize directly the UMorphTarget with a FMemoryArchive. This is not supported.
-					for (int32 MorphTargetIndex = 0; MorphTargetIndex < MorphTargetNumber; ++MorphTargetIndex)
+					if (Owner->GetMorphTargets().Num() > 0)
+					{
+						FinishBuildMorphTargetData = Owner->GetMorphTargets()[0]->CreateFinishBuildMorphTargetData();
+					}
+					else
 					{
-						FName MorphTargetName = NAME_None;
-						Ar << MorphTargetName;
-						TArray<FMorphTargetLODModel>& MorphTargetLODModels = MorphLODModelsPerTargetName.FindOrAdd(MorphTargetName);
-						int32 MorphLODModelNumber = 0;
-						Ar << MorphLODModelNumber;
-						MorphTargetLODModels.Empty(MorphLODModelNumber);
-						MorphTargetLODModels.AddDefaulted(MorphLODModelNumber);
-						for (int32 MorphDataIndex = 0; MorphDataIndex < MorphLODModelNumber; ++MorphDataIndex)
-						{
-							Ar << MorphTargetLODModels[MorphDataIndex];
-						}
+						// Create and initialize the FinishBuildInternalData, use the class default object to call the virtual function
+						FinishBuildMorphTargetData = UMorphTarget::StaticClass()->GetDefaultObject<UMorphTarget>()->CreateFinishBuildMorphTargetData();
 					}
+					check(FinishBuildMorphTargetData);
+					FinishBuildMorphTargetData->LoadFromMemoryArchive(Ar);
 				}
 
 				//Serialize the LODModel sections since they are dependent on the reduction
@@ -481,10 +410,7 @@ void FSkeletalMeshRenderData::Cache(const ITargetPlatform* TargetPlatform, USkel
 			}
 
 			//Apply the morphtargets change if any
-			if (MorphLODModelsPerTargetName.Num() > 0)
-			{
-				MorphTargetUtils::ApplyMorphTargetsEditorData(Owner, MorphLODModelsPerTargetName);
-			}
+			FinishBuildMorphTargetData->ApplyEditorData(Owner);
 
 			int32 T1 = FPlatformTime::Cycles();
 			UE_LOG(LogSkeletalMesh, Verbose, TEXT("Skeletal Mesh found in DDC [%fms] %s"), FPlatformTime::ToMilliseconds(T1 - T0), *Owner->GetPathName());
