Fix strata material conversion when shading model is driven by an expression.

diff --git a/Engine/Source/Runtime/Engine/Private/Materials/MaterialExpressions.cpp b/Engine/Source/Runtime/Engine/Private/Materials/MaterialExpressions.cpp
--- a/Engine/Source/Runtime/Engine/Private/Materials/MaterialExpressions.cpp
+++ b/Engine/Source/Runtime/Engine/Private/Materials/MaterialExpressions.cpp
@@ -20572,11 +20572,8 @@ void UMaterialExpressionStrataLegacyConversion::GatherStrataMaterialInfo(FStrata
 	if (ConvertedStrataMaterialInfo.HasShadingModel(SSM_VolumetricFogCloud))	{ StrataMaterialInfo.AddShadingModel(SSM_VolumetricFogCloud); }
 	if (ConvertedStrataMaterialInfo.HasShadingModel(SSM_Hair))					{ StrataMaterialInfo.AddShadingModel(SSM_Hair); }
 	if (ConvertedStrataMaterialInfo.HasShadingModel(SSM_SingleLayerWater))		{ StrataMaterialInfo.AddShadingModel(SSM_SingleLayerWater); }
-
-	if (SubsurfaceProfile)
-	{
-		StrataMaterialInfo.AddSubsurfaceProfile(SubsurfaceProfile);
-	}
+	if (SubsurfaceProfile)														{ StrataMaterialInfo.AddSubsurfaceProfile(SubsurfaceProfile); }
+	if (ConvertedStrataMaterialInfo.HasShadingModelFromExpression())			{ StrataMaterialInfo.SetShadingModelFromExpression(true); }
 }
 
 bool UMaterialExpressionStrataLegacyConversion::HasSSS() const
