Add Parameters.AbsoluteWorldPosition to the list of global expressions that will get modified in custom hlsl expressions to include LWCToFloat

diff --git a/Engine/Source/Runtime/Engine/Private/Materials/MaterialExpressions.cpp b/Engine/Source/Runtime/Engine/Private/Materials/MaterialExpressions.cpp
--- a/Engine/Source/Runtime/Engine/Private/Materials/MaterialExpressions.cpp
+++ b/Engine/Source/Runtime/Engine/Private/Materials/MaterialExpressions.cpp
@@ -12317,6 +12317,7 @@ void UMaterialExpressionCustom::Serialize(FStructuredArchive::FRecord Record)
 			TEXT("GetPrimitiveData(Parameters).LocalToWorld"),
 			TEXT("GetPrimitiveData(Parameters.PrimitiveId).WorldToLocal"),
 			TEXT("GetPrimitiveData(Parameters.PrimitiveId).LocalToWorld"),
+			TEXT("Parameters.AbsoluteWorldPosition"),
 		};
 
 		for (const TCHAR* Member : UniformMembers)
