Properly hook up instance client manager in placement mode.

This fixes duplication crash.

diff --git a/Engine/Source/Editor/AssetPlacementEdMode/Private/AssetPlacementPaletteItem.cpp b/Engine/Source/Editor/AssetPlacementEdMode/Private/AssetPlacementPaletteItem.cpp
--- a/Engine/Source/Editor/AssetPlacementEdMode/Private/AssetPlacementPaletteItem.cpp
+++ b/Engine/Source/Editor/AssetPlacementEdMode/Private/AssetPlacementPaletteItem.cpp
@@ -17,17 +17,18 @@
 #include "Engine/Blueprint.h"
 #include "Factories/AssetFactoryInterface.h"
 #include "AssetPlacementSettings.h"
+#include "Instances/InstancedPlacementClientInfo.h"
 
 #define LOCTEXT_NAMESPACE "AssetPlacementMode"
 
-FPaletteItemUIInfo::FPaletteItemUIInfo(const FAssetData& InAssetData, const FPaletteItem& InPaletteItem)
+FPaletteItemUIInfo::FPaletteItemUIInfo(const FAssetData& InAssetData, const UPlacementPaletteClient* InPaletteItem)
 	: AssetData(InAssetData)
-	, SettingsObject(InPaletteItem.SettingsObject.Get())
+	, SettingsObject(InPaletteItem ? InPaletteItem->SettingsObject.Get() : nullptr)
 {
 	check(AssetData.IsValid());
 }
 
-FAssetPlacementPaletteItemModel::FAssetPlacementPaletteItemModel(const FAssetData& InAssetData, const FPaletteItem& InPaletteItem, TSharedRef<SAssetPlacementPalette> InParentPalette, TSharedPtr<FAssetThumbnailPool> InThumbnailPool)
+FAssetPlacementPaletteItemModel::FAssetPlacementPaletteItemModel(const FAssetData& InAssetData, const UPlacementPaletteClient* InPaletteItem, TSharedRef<SAssetPlacementPalette> InParentPalette, TSharedPtr<FAssetThumbnailPool> InThumbnailPool)
 	: TypeInfo(MakeShared<FPaletteItemUIInfo>(InAssetData, InPaletteItem))
 	, AssetPalette(InParentPalette)
 {
diff --git a/Engine/Source/Editor/AssetPlacementEdMode/Private/AssetPlacementPaletteItem.h b/Engine/Source/Editor/AssetPlacementEdMode/Private/AssetPlacementPaletteItem.h
--- a/Engine/Source/Editor/AssetPlacementEdMode/Private/AssetPlacementPaletteItem.h
+++ b/Engine/Source/Editor/AssetPlacementEdMode/Private/AssetPlacementPaletteItem.h
@@ -14,8 +14,9 @@
 #include "AssetRegistry/AssetData.h"
 
 class FAssetThumbnailPool;
-struct FPaletteItem;
+class UPlacementPaletteClient;
 class IAssetTypeActions;
+class UInstancedPlacemenClientSettings;
 
 namespace PlacementPaletteConstants
 {
@@ -24,10 +25,10 @@ namespace PlacementPaletteConstants
 
 struct FPaletteItemUIInfo
 {
-	FPaletteItemUIInfo(const FAssetData& InAssetData, const FPaletteItem& InPaletteItem);
+	FPaletteItemUIInfo(const FAssetData& InAssetData, const UPlacementPaletteClient* InPaletteItem);
 
 	const FAssetData AssetData;
-	TWeakObjectPtr<UEditorFactorySettingsObject> SettingsObject;
+	TWeakObjectPtr<UInstancedPlacemenClientSettings> SettingsObject;
 };
 
 typedef TSharedPtr<FPaletteItemUIInfo> FAssetPlacementUIInfoPtr;
@@ -40,7 +41,7 @@ namespace AssetPlacementPaletteTreeColumns
 class FAssetPlacementPaletteItemModel : public TSharedFromThis<FAssetPlacementPaletteItemModel>
 {
 public:
-	FAssetPlacementPaletteItemModel(const FAssetData& InAssetData, const FPaletteItem& InPaletteItem, TSharedRef<class SAssetPlacementPalette> InParentPalette, TSharedPtr<class FAssetThumbnailPool> InThumbnailPool);
+	FAssetPlacementPaletteItemModel(const FAssetData& InAssetData, const UPlacementPaletteClient* InPaletteItem, TSharedRef<class SAssetPlacementPalette> InParentPalette, TSharedPtr<class FAssetThumbnailPool> InThumbnailPool);
 
 	FAssetPlacementUIInfoPtr GetTypeUIInfo() const;
 
diff --git a/Engine/Source/Editor/AssetPlacementEdMode/Private/AssetPlacementSettings.cpp b/Engine/Source/Editor/AssetPlacementEdMode/Private/AssetPlacementSettings.cpp
--- a/Engine/Source/Editor/AssetPlacementEdMode/Private/AssetPlacementSettings.cpp
+++ b/Engine/Source/Editor/AssetPlacementEdMode/Private/AssetPlacementSettings.cpp
@@ -11,6 +11,7 @@
 #include "Elements/Framework/TypedElementHandle.h"
 #include "Elements/Framework/TypedElementRegistry.h"
 #include "Elements/Interfaces/TypedElementAssetDataInterface.h"
+#include "Instances/InstancedPlacementClientInfo.h"
 
 bool UAssetPlacementSettings::CanEditChange(const FProperty* InProperty) const
 {
@@ -54,10 +55,10 @@ void UAssetPlacementSettings::SetPaletteAsset(UPlacementPaletteAsset* InPaletteA
 	LastActivePalettePath = FSoftObjectPath(InPaletteAsset);
 }
 
-FPaletteItem UAssetPlacementSettings::AddItemToActivePalette(const FAssetData& InAssetData)
+UPlacementPaletteClient* UAssetPlacementSettings::AddClientToActivePalette(const FAssetData& InAssetData)
 {
-	FPaletteItem NewPaletteItem;
-	if (!InAssetData.IsValid())
+	UPlacementPaletteClient* NewPaletteItem = nullptr;
+	if (!InAssetData.IsValid() || !GetActivePalette())
 	{
 		return NewPaletteItem;
 	}
@@ -67,23 +68,22 @@ FPaletteItem UAssetPlacementSettings::AddItemToActivePalette(const FAssetData& I
 		return NewPaletteItem;
 	}
 
-	if (!GetActivePaletteItems().FindByPredicate([InAssetData](const FPaletteItem& ItemIter) { return (ItemIter.AssetPath == InAssetData.ToSoftObjectPath()); }))
+	if (!GetActivePaletteItems().FindByPredicate([InAssetData](const UPlacementPaletteClient* ItemIter) { return ItemIter ? (ItemIter->AssetPath == InAssetData.ToSoftObjectPath()) : false; }))
 	{
 		if (UPlacementSubsystem* PlacementSubystem = GEditor->GetEditorSubsystem<UPlacementSubsystem>())
 		{
 			if (TScriptInterface<IAssetFactoryInterface> AssetFactory = PlacementSubystem->FindAssetFactoryFromAssetData(InAssetData))
 			{
-				NewPaletteItem.ItemGuid = FGuid::NewGuid();
-				NewPaletteItem.AssetPath = InAssetData.ToSoftObjectPath();
-				NewPaletteItem.AssetFactoryInterface = AssetFactory;
+				NewPaletteItem = NewObject<UPlacementPaletteClient>(GetMutableActivePalette());
+				NewPaletteItem->AssetPath = InAssetData.ToSoftObjectPath();
 				FPlacementOptions PlacementOptions;
 				PlacementOptions.InstancedPlacementGridGuid = GetActivePaletteGuid();
-				NewPaletteItem.SettingsObject = AssetFactory->FactorySettingsObjectForPlacement(InAssetData, PlacementOptions);
+				NewPaletteItem->SettingsObject = AssetFactory->FactorySettingsObjectForPlacement(InAssetData, PlacementOptions);
 				GetMutableActivePalette()->Modify();
-				if (NewPaletteItem.SettingsObject)
+				if (NewPaletteItem->SettingsObject)
 				{
 					// Set the outer of the new object to the current palette
-					NewPaletteItem.SettingsObject->Rename(nullptr, GetMutableActivePalette());
+					NewPaletteItem->SettingsObject->Rename(nullptr, GetMutableActivePalette());
 				}
 				GetMutableActivePalette()->PaletteItems.Add(NewPaletteItem);
 			}
@@ -93,7 +93,7 @@ FPaletteItem UAssetPlacementSettings::AddItemToActivePalette(const FAssetData& I
 	return NewPaletteItem;
 }
 
-TArrayView<const FPaletteItem> UAssetPlacementSettings::GetActivePaletteItems() const
+TArrayView<const TObjectPtr<UPlacementPaletteClient>> UAssetPlacementSettings::GetActivePaletteItems() const
 {
 	return MakeArrayView(GetActivePalette()->PaletteItems);
 }
@@ -151,22 +151,28 @@ bool UAssetPlacementSettings::DoesActivePaletteSupportElement(const FTypedElemen
 	if (TTypedElement<ITypedElementAssetDataInterface> AssetDataInterface = UTypedElementRegistry::GetInstance()->GetElement<ITypedElementAssetDataInterface>(InElementToCheck))
 	{
 		TArray<FAssetData> ReferencedAssetDatas = AssetDataInterface.GetAllReferencedAssetDatas();
-		for (const FPaletteItem& Item : GetActivePaletteItems())
+		for (const UPlacementPaletteClient* Item : GetActivePaletteItems())
 		{
-			if (ReferencedAssetDatas.FindByPredicate([&Item](const FAssetData& ReferencedAssetData) { return (ReferencedAssetData.ToSoftObjectPath() == Item.AssetPath); }))
+			if (!Item)
+			{
+				continue;
+			}
+
+			if (ReferencedAssetDatas.FindByPredicate([Item](const FAssetData& ReferencedAssetData) { return (ReferencedAssetData.ToSoftObjectPath() == Item->AssetPath); }))
 			{
 				return true;
 			}
 
 			// The current implementation of the asset data interface for actors requires that individual actors report on assets contained within their components.
-			// Not all actors do this reliably, so additionally check the supplied factory for a match. 
-			if (!Item.AssetFactoryInterface)
+			// Not all elements (legacy via actors) do this reliably, so additionally check the supplied factory for a match. 
+			if (!Item->FactoryInterfaceClass.Get())
 			{
 				continue;
 			}
 
-			FAssetData FoundAssetDataFromFactory = Item.AssetFactoryInterface->GetAssetDataFromElementHandle(InElementToCheck);
-			if (FoundAssetDataFromFactory.ToSoftObjectPath() == Item.AssetPath)
+			TScriptInterface<IAssetFactoryInterface> FactoryInterface = GEditor->GetEditorSubsystem<UPlacementSubsystem>()->GetAssetFactoryFromFactoryClass(Item->FactoryInterfaceClass);
+			FAssetData FoundAssetDataFromFactory = FactoryInterface->GetAssetDataFromElementHandle(InElementToCheck);
+			if (FoundAssetDataFromFactory.ToSoftObjectPath() == Item->AssetPath)
 			{
 				return true;
 			}
diff --git a/Engine/Source/Editor/AssetPlacementEdMode/Private/AssetPlacementSettings.h b/Engine/Source/Editor/AssetPlacementEdMode/Private/AssetPlacementSettings.h
--- a/Engine/Source/Editor/AssetPlacementEdMode/Private/AssetPlacementSettings.h
+++ b/Engine/Source/Editor/AssetPlacementEdMode/Private/AssetPlacementSettings.h
@@ -9,7 +9,7 @@
 #include "AssetPlacementSettings.generated.h"
 
 class UPlacementPaletteAsset;
-struct FPaletteItem;
+class UPlacementPaletteClient;
 
 UCLASS(config = EditorPerProjectUserSettings)
 class UAssetPlacementSettings : public UObject
@@ -123,8 +123,8 @@ public:
 	virtual bool CanEditChange(const FProperty* InProperty) const override;
 
 	void SetPaletteAsset(UPlacementPaletteAsset* InPaletteAsset);
-	FPaletteItem AddItemToActivePalette(const FAssetData& InAssetData);
-	TArrayView<const FPaletteItem> GetActivePaletteItems() const;
+	UPlacementPaletteClient* AddClientToActivePalette(const FAssetData& InAssetData);
+	TArrayView<const TObjectPtr<UPlacementPaletteClient>> GetActivePaletteItems() const;
 	const FGuid GetActivePaletteGuid() const;
 	void ClearActivePaletteItems();
 	void SaveActivePalette();
diff --git a/Engine/Source/Editor/AssetPlacementEdMode/Private/Modes/PlacementModeSubsystem.cpp b/Engine/Source/Editor/AssetPlacementEdMode/Private/Modes/PlacementModeSubsystem.cpp
--- a/Engine/Source/Editor/AssetPlacementEdMode/Private/Modes/PlacementModeSubsystem.cpp
+++ b/Engine/Source/Editor/AssetPlacementEdMode/Private/Modes/PlacementModeSubsystem.cpp
@@ -3,20 +3,21 @@
 #include "Modes/PlacementModeSubsystem.h"
 
 #include "AssetPlacementSettings.h"
-
+#include "Subsystems/PlacementSubsystem.h"
 
 #include "Editor.h"
 
 void UPlacementModeSubsystem::Initialize(FSubsystemCollectionBase& Collection)
 {
-	ModeSettings = NewObject<UAssetPlacementSettings>(this);
-	ModeSettings->LoadSettings();
+	Collection.InitializeDependency<UPlacementSubsystem>();
 
+	GEditor->GetEditorSubsystem<UPlacementSubsystem>()->OnPlacementFactoriesRegistered().AddUObject(this, &UPlacementModeSubsystem::LoadSettings);
 	FCoreDelegates::OnEnginePreExit.AddUObject(this, &UPlacementModeSubsystem::SaveSettings);
 }
 
 void UPlacementModeSubsystem::Deinitialize()
 {
+	GEditor->GetEditorSubsystem<UPlacementSubsystem>()->OnPlacementFactoriesRegistered().RemoveAll(this);
 	ModeSettings = nullptr;
 }
 
@@ -37,3 +38,11 @@ void UPlacementModeSubsystem::SaveSettings() const
 		ModeSettings->SaveSettings();
 	}
 }
+
+void UPlacementModeSubsystem::LoadSettings()
+{
+	ModeSettings = NewObject<UAssetPlacementSettings>(this);
+	ModeSettings->LoadSettings();
+
+	GEditor->GetEditorSubsystem<UPlacementSubsystem>()->OnPlacementFactoriesRegistered().RemoveAll(this);
+}
diff --git a/Engine/Source/Editor/AssetPlacementEdMode/Private/Modes/PlacementModeSubsystem.h b/Engine/Source/Editor/AssetPlacementEdMode/Private/Modes/PlacementModeSubsystem.h
--- a/Engine/Source/Editor/AssetPlacementEdMode/Private/Modes/PlacementModeSubsystem.h
+++ b/Engine/Source/Editor/AssetPlacementEdMode/Private/Modes/PlacementModeSubsystem.h
@@ -29,4 +29,5 @@ protected:
 	TObjectPtr<UAssetPlacementSettings> ModeSettings;
 
 	void SaveSettings() const;
+	void LoadSettings();
 };
diff --git a/Engine/Source/Editor/AssetPlacementEdMode/Private/PlacementPaletteAsset.h b/Engine/Source/Editor/AssetPlacementEdMode/Private/PlacementPaletteAsset.h
--- a/Engine/Source/Editor/AssetPlacementEdMode/Private/PlacementPaletteAsset.h
+++ b/Engine/Source/Editor/AssetPlacementEdMode/Private/PlacementPaletteAsset.h
@@ -17,7 +17,7 @@ class UPlacementPaletteAsset : public UObject
 
 public:
 	UPROPERTY(EditAnywhere, Category=Items)
-	TArray<FPaletteItem> PaletteItems;
+	TArray<TObjectPtr<UPlacementPaletteClient>> PaletteItems;
 
 	UPROPERTY()
 	FGuid GridGuid;
diff --git a/Engine/Source/Editor/AssetPlacementEdMode/Private/PlacementPaletteItem.cpp b/Engine/Source/Editor/AssetPlacementEdMode/Private/PlacementPaletteItem.cpp
new file mode 100644
--- /dev/null
+++ b/Engine/Source/Editor/AssetPlacementEdMode/Private/PlacementPaletteItem.cpp
@@ -0,0 +1,20 @@
+// Copyright Epic Games, Inc. All Rights Reserved.
+
+#include "PlacementPaletteItem.h"
+
+#include "Subsystems/PlacementSubsystem.h"
+
+void UPlacementPaletteClient::PostEditChangeProperty(struct FPropertyChangedEvent& PropertyChangedEvent)
+{
+	Super::PostEditChangeProperty(PropertyChangedEvent);
+}
+
+void UPlacementPaletteClient::PostLoad()
+{
+	Super::PostLoad();
+}
+
+void UPlacementPaletteClient::Serialize(FArchive& Ar)
+{
+	Super::Serialize(Ar);
+}
diff --git a/Engine/Source/Editor/AssetPlacementEdMode/Private/PlacementPaletteItem.h b/Engine/Source/Editor/AssetPlacementEdMode/Private/PlacementPaletteItem.h
--- a/Engine/Source/Editor/AssetPlacementEdMode/Private/PlacementPaletteItem.h
+++ b/Engine/Source/Editor/AssetPlacementEdMode/Private/PlacementPaletteItem.h
@@ -3,26 +3,32 @@
 #pragma once
 
 #include "CoreMinimal.h"
+#include "ISMPartition/ISMPartitionInstanceManager.h"
 
 #include "PlacementPaletteItem.generated.h"
 
-class IAssetFactoryInterface;
-class UEditorFactorySettingsObject;
+class UAssetFactoryInterface;
+class UInstancedPlacemenClientSettings;
 
-USTRUCT()
-struct FPaletteItem
+UCLASS(NotPlaceable)
+class UPlacementPaletteClient : public UObject
 {
 	GENERATED_BODY()
 
+public:
 	UPROPERTY(EditAnywhere, Category = Settings)
 	FSoftObjectPath AssetPath;
 
-	UPROPERTY(EditAnywhere, Category = Settings, AdvancedDisplay)
-	TScriptInterface<IAssetFactoryInterface> AssetFactoryInterface;
+	UPROPERTY()
+	FGuid ClientGuid;
+
+	UPROPERTY(EditAnywhere, Category = Settings)
+	TObjectPtr<UInstancedPlacemenClientSettings> SettingsObject = nullptr;
 
 	UPROPERTY()
-	FGuid ItemGuid;
+	TSubclassOf<UAssetFactoryInterface> FactoryInterfaceClass;
 
-	UPROPERTY(EditAnywhere, Instanced, Category = Settings)
-	TObjectPtr<UEditorFactorySettingsObject> SettingsObject = nullptr;
+	virtual void PostEditChangeProperty(struct FPropertyChangedEvent& PropertyChangedEvent) override;
+	virtual void PostLoad() override;
+	virtual void Serialize(FArchive& Ar) override;
 };
diff --git a/Engine/Source/Editor/AssetPlacementEdMode/Private/SAssetPlacementPalette.cpp b/Engine/Source/Editor/AssetPlacementEdMode/Private/SAssetPlacementPalette.cpp
--- a/Engine/Source/Editor/AssetPlacementEdMode/Private/SAssetPlacementPalette.cpp
+++ b/Engine/Source/Editor/AssetPlacementEdMode/Private/SAssetPlacementPalette.cpp
@@ -39,6 +39,7 @@
 #include "AssetRegistry/AssetRegistryModule.h"
 #include "PlacementPaletteItem.h"
 #include "Factories/AssetFactoryInterface.h"
+#include "Instances/EditorPlacementSettings.h"
 
 #define LOCTEXT_NAMESPACE "AssetPlacementMode"
 
@@ -254,8 +255,8 @@ void SAssetPlacementPalette::AddPlacementType(const FAssetData& AssetData)
 	// Try to add the item to the mode's palette
 	if (UPlacementModeSubsystem* PlacementModeSubsystem = GEditor->GetEditorSubsystem<UPlacementModeSubsystem>())
 	{
-		FPaletteItem PlacementInfo = PlacementModeSubsystem->GetMutableModeSettingsObject()->AddItemToActivePalette(AssetData);
-		if (!PlacementInfo.AssetPath.IsValid())
+		UPlacementPaletteClient* PlacementInfo = PlacementModeSubsystem->GetMutableModeSettingsObject()->AddClientToActivePalette(AssetData);
+		if (!PlacementInfo || !PlacementInfo->AssetPath.IsValid())
 		{
 			return;
 		}
@@ -477,7 +478,7 @@ void SAssetPlacementPalette::OnSetPaletteAsset(const FAssetData& InAssetData)
 	}
 }
 
-void SAssetPlacementPalette::SetPaletteItems(TArrayView<const FPaletteItem> InPaletteItems)
+void SAssetPlacementPalette::SetPaletteItems(TArrayView<const TObjectPtr<UPlacementPaletteClient>> InPaletteItems)
 {
 	PaletteItems.Empty();
 	FilteredItems.Empty();
@@ -488,13 +489,16 @@ void SAssetPlacementPalette::SetPaletteItems(TArrayView<const FPaletteItem> InPa
 
 	FAssetRegistryModule& AssetRegistry = FModuleManager::Get().LoadModuleChecked<FAssetRegistryModule>(TEXT("AssetRegistry"));
 	TArray<FSoftObjectPath> AssetsToLoad;
-	for (const FPaletteItem& PaletteItem : InPaletteItems)
+	for (const UPlacementPaletteClient* PaletteItem : InPaletteItems)
 	{
-		FAssetData AssetData = AssetRegistry.Get().GetAssetByObjectPath(PaletteItem.AssetPath.GetAssetPathName());
-		if (AssetData.IsValid())
+		if (PaletteItem)
 		{
-			PaletteItems.Add(MakeShared<FAssetPlacementPaletteItemModel>(AssetData, PaletteItem, SharedThis(this), ThumbnailPool));
-			AssetsToLoad.Emplace(PaletteItem.AssetPath);
+			FAssetData AssetData = AssetRegistry.Get().GetAssetByObjectPath(PaletteItem->AssetPath.GetAssetPathName());
+			if (AssetData.IsValid())
+			{
+				PaletteItems.Add(MakeShared<FAssetPlacementPaletteItemModel>(AssetData, PaletteItem, SharedThis(this), ThumbnailPool));
+				AssetsToLoad.Emplace(PaletteItem->AssetPath);
+			}
 		}
 	}
 
@@ -725,7 +729,7 @@ void SAssetPlacementPalette::OnPaletteSelectionChanged(FPlacementPaletteItemMode
 		TArray<UObject*> DetailsObjects;
 		for (const FPlacementPaletteItemModelPtr& SelectedItem : GetActiveViewWidget()->GetSelectedItems())
 		{
-			if (UEditorFactorySettingsObject* SelectedItemSettingsObject = SelectedItem->GetTypeUIInfo()->SettingsObject.Get())
+			if (UInstancedPlacemenClientSettings* SelectedItemSettingsObject = SelectedItem->GetTypeUIInfo()->SettingsObject.Get())
 			{
 				DetailsObjects.Add(SelectedItemSettingsObject);
 			}
diff --git a/Engine/Source/Editor/AssetPlacementEdMode/Private/SAssetPlacementPalette.h b/Engine/Source/Editor/AssetPlacementEdMode/Private/SAssetPlacementPalette.h
--- a/Engine/Source/Editor/AssetPlacementEdMode/Private/SAssetPlacementPalette.h
+++ b/Engine/Source/Editor/AssetPlacementEdMode/Private/SAssetPlacementPalette.h
@@ -27,6 +27,7 @@ struct FAssetData;
 class UAssetPlacementSettings;
 class IPropertyHandle;
 class SSearchBox;
+class UPlacementPaletteClient;
 
 typedef TSharedPtr<FAssetPlacementPaletteItemModel> FPlacementPaletteItemModelPtr;
 typedef STreeView<FPlacementPaletteItemModelPtr> SPlacementTypeTreeView;
@@ -100,7 +101,7 @@ private:
 	void OnContentBrowserSelectionChanged(const TArray<FAssetData>& NewSelectedAssets, bool bIsPrimaryBrowser);
 	void SetupContentBrowserMirroring(bool bInMirrorContentBrowser);
 	void OnSetPaletteAsset(const FAssetData& InAssetData);
-	void SetPaletteItems(TArrayView<const FPaletteItem> InPaletteItems);
+	void SetPaletteItems(TArrayView<const TObjectPtr<UPlacementPaletteClient>> InPaletteItems);
 	FString GetPalettePath() const;
 
 	/** Sets the view mode of the palette */
diff --git a/Engine/Source/Editor/AssetPlacementEdMode/Private/Tools/PlacementPlaceSingleTool.cpp b/Engine/Source/Editor/AssetPlacementEdMode/Private/Tools/PlacementPlaceSingleTool.cpp
--- a/Engine/Source/Editor/AssetPlacementEdMode/Private/Tools/PlacementPlaceSingleTool.cpp
+++ b/Engine/Source/Editor/AssetPlacementEdMode/Private/Tools/PlacementPlaceSingleTool.cpp
@@ -237,21 +237,23 @@ void UPlacementModePlaceSingleTool::GeneratePlacementData(const FInputDeviceRay&
 	const UAssetPlacementSettings* PlacementSettings = PlacementSubsystem ? PlacementSubsystem->GetModeSettingsObject() : nullptr;
 	if (PlacementSettings)
 	{
-		TArrayView<const FPaletteItem> CurrentPaletteItems = PlacementSettings->GetActivePaletteItems();
+		TArrayView<const TObjectPtr<UPlacementPaletteClient>> CurrentPaletteItems = PlacementSettings->GetActivePaletteItems();
 		if (CurrentPaletteItems.Num())
 		{
 			int32 ItemIndex = FMath::RandHelper(CurrentPaletteItems.Num());
-			const FPaletteItem& ItemToPlace = CurrentPaletteItems[ItemIndex];
-			
-			FTransform TransformToUpdate(GenerateRandomRotation(PlacementSettings), LastBrushStamp.WorldPosition, GenerateRandomScale(PlacementSettings));
-
-			PlacementInfo = MakeUnique<FAssetPlacementInfo>();
-			PlacementInfo->AssetToPlace = ItemToPlace.AssetPath.TryLoad();
-			PlacementInfo->FactoryOverride = ItemToPlace.AssetFactoryInterface;
-			PlacementInfo->ItemGuid = ItemToPlace.ItemGuid;
-			PlacementInfo->FinalizedTransform = FinalizeTransform(TransformToUpdate, LastBrushStamp.WorldNormal, PlacementSettings);
-			PlacementInfo->PreferredLevel = GEditor->GetEditorWorldContext().World()->GetCurrentLevel();
-			PlacementInfo->SettingsObject = ItemToPlace.SettingsObject;
+			const UPlacementPaletteClient* ItemToPlace = CurrentPaletteItems[ItemIndex];
+			if (ItemToPlace)
+			{
+				FTransform TransformToUpdate(GenerateRandomRotation(PlacementSettings), LastBrushStamp.WorldPosition, GenerateRandomScale(PlacementSettings));
+
+				PlacementInfo = MakeUnique<FAssetPlacementInfo>();
+				PlacementInfo->AssetToPlace = ItemToPlace->AssetPath.TryLoad();
+				PlacementInfo->FactoryOverride = ItemToPlace->FactoryInterfaceClass;
+				PlacementInfo->ItemGuid = ItemToPlace->ClientGuid;
+				PlacementInfo->FinalizedTransform = FinalizeTransform(TransformToUpdate, LastBrushStamp.WorldNormal, PlacementSettings);
+				PlacementInfo->PreferredLevel = GEditor->GetEditorWorldContext().World()->GetCurrentLevel();
+				PlacementInfo->SettingsObject = ItemToPlace->SettingsObject;
+			}
 		}
 	}
 }
diff --git a/Engine/Source/Editor/AssetPlacementEdMode/Private/Tools/PlacementPlaceTool.cpp b/Engine/Source/Editor/AssetPlacementEdMode/Private/Tools/PlacementPlaceTool.cpp
--- a/Engine/Source/Editor/AssetPlacementEdMode/Private/Tools/PlacementPlaceTool.cpp
+++ b/Engine/Source/Editor/AssetPlacementEdMode/Private/Tools/PlacementPlaceTool.cpp
@@ -11,6 +11,8 @@
 #include "Modes/PlacementModeSubsystem.h"
 #include "PlacementPaletteItem.h"
 
+#include "Instances/InstancedPlacementClientInfo.h"
+
 constexpr TCHAR UPlacementModePlacementTool::ToolName[];
 
 UPlacementBrushToolBase* UPlacementModePlacementToolBuilder::FactoryToolInstance(UObject* Outer) const
@@ -43,7 +45,7 @@ void UPlacementModePlacementTool::OnTick(float DeltaTime)
 	const UAssetPlacementSettings* PlacementSettings = PlacementModeSubsystem ? PlacementModeSubsystem->GetModeSettingsObject() : nullptr;
 	if (BrushProperties && PlacementSettings)
 	{
-		TArrayView<const FPaletteItem> PaletteItems = MakeArrayView(PlacementSettings->GetActivePaletteItems());
+		TArrayView<const TObjectPtr<UPlacementPaletteClient>> PaletteItems = PlacementSettings->GetActivePaletteItems();
 		if (PaletteItems.Num())
 		{
 			FVector BrushLocation = LastBrushStamp.WorldPosition;
@@ -92,17 +94,18 @@ void UPlacementModePlacementTool::OnTick(float DeltaTime)
 					PotentialInstanceHash.InsertInstance(SpawnLocation, PotentialIdx);
 
 					int32 ItemIndex = FMath::RandHelper(PaletteItems.Num());
-					const FPaletteItem& ItemToPlace = PaletteItems[ItemIndex];
-
-					FAssetPlacementInfo NewInfo;
-					NewInfo.AssetToPlace = ItemToPlace.AssetPath.TryLoad();
-					NewInfo.FactoryOverride = ItemToPlace.AssetFactoryInterface;
-					NewInfo.ItemGuid = ItemToPlace.ItemGuid;
-					NewInfo.PreferredLevel = GEditor->GetEditorWorldContext().World()->GetCurrentLevel();
-					NewInfo.FinalizedTransform = GenerateTransformFromHitLocationAndNormal(SpawnLocation, SpawnNormal);
-					NewInfo.SettingsObject = ItemToPlace.SettingsObject;
-
-					DesiredInfoSpawnLocations.Emplace(NewInfo);
+					if (const UPlacementPaletteClient* ItemToPlace = PaletteItems[ItemIndex])
+					{
+						FAssetPlacementInfo NewInfo;
+						NewInfo.AssetToPlace = ItemToPlace->AssetPath.TryLoad();
+						NewInfo.FactoryOverride = ItemToPlace->FactoryInterfaceClass;
+						NewInfo.ItemGuid = ItemToPlace->ClientGuid;
+						NewInfo.PreferredLevel = GEditor->GetEditorWorldContext().World()->GetCurrentLevel();
+						NewInfo.FinalizedTransform = GenerateTransformFromHitLocationAndNormal(SpawnLocation, SpawnNormal);
+						NewInfo.SettingsObject = ItemToPlace->SettingsObject;
+
+						DesiredInfoSpawnLocations.Emplace(NewInfo);
+					}
 				}
 			}
 
diff --git a/Engine/Source/Editor/EditorFramework/Private/Subsystems/PlacementSubsystem.cpp b/Engine/Source/Editor/EditorFramework/Private/Subsystems/PlacementSubsystem.cpp
--- a/Engine/Source/Editor/EditorFramework/Private/Subsystems/PlacementSubsystem.cpp
+++ b/Engine/Source/Editor/EditorFramework/Private/Subsystems/PlacementSubsystem.cpp
@@ -37,7 +37,7 @@ TArray<FTypedElementHandle> UPlacementSubsystem::PlaceAssets(TArrayView<const FA
 	for (const FAssetPlacementInfo& PlacementInfo : InPlacementInfos)
 	{
 		const FAssetData& AssetData = PlacementInfo.AssetToPlace;
-		TScriptInterface<IAssetFactoryInterface> FactoryInterface = PlacementInfo.FactoryOverride;
+		TScriptInterface<IAssetFactoryInterface> FactoryInterface = GetAssetFactoryFromFactoryClass(PlacementInfo.FactoryOverride);
 		if (!FactoryInterface)
 		{
 			FactoryInterface = FindAssetFactoryFromAssetData(AssetData);
@@ -96,6 +96,11 @@ bool UPlacementSubsystem::IsCreatingPreviewElements() const
 	return bIsCreatingPreviewElements;
 }
 
+FSimpleMulticastDelegate& UPlacementSubsystem::OnPlacementFactoriesRegistered()
+{
+	return PlacementFactoriesRegistered;
+}
+
 void UPlacementSubsystem::RegisterPlacementFactories()
 {
 	for (TObjectIterator<UClass> ObjectIt; ObjectIt; ++ObjectIt)
@@ -110,6 +115,24 @@ void UPlacementSubsystem::RegisterPlacementFactories()
 			}
 		}
 	}
+
+	PlacementFactoriesRegistered.Broadcast();
+}
+
+TScriptInterface<IAssetFactoryInterface> UPlacementSubsystem::GetAssetFactoryFromFactoryClass(TSubclassOf<UAssetFactoryInterface> InFactoryInterfaceClass) const
+{
+	if (InFactoryInterfaceClass)
+	{
+		for (const TScriptInterface<IAssetFactoryInterface>& AssetFactory : AssetFactories)
+		{
+			if (AssetFactory.GetObject()->IsA(InFactoryInterfaceClass))
+			{
+				return AssetFactory;
+			}
+		}
+	}
+
+	return nullptr;
 }
 
 void UPlacementSubsystem::UnregisterPlacementFactories()
diff --git a/Engine/Source/Editor/EditorFramework/Public/Factories/AssetFactoryInterface.h b/Engine/Source/Editor/EditorFramework/Public/Factories/AssetFactoryInterface.h
--- a/Engine/Source/Editor/EditorFramework/Public/Factories/AssetFactoryInterface.h
+++ b/Engine/Source/Editor/EditorFramework/Public/Factories/AssetFactoryInterface.h
@@ -6,17 +6,14 @@
 #include "Elements/Framework/TypedElementHandle.h"
 #include "UObject/Interface.h"
 
+#include "Instances/InstancedPlacementClientInfo.h"
+
 #include "AssetFactoryInterface.generated.h"
 
 class IAssetFactoryInterface;
 struct FPlacementOptions;
 struct FAssetPlacementInfo;
-
-UCLASS(Abstract, hideCategories=Object, editinlinenew, collapsecategories, MinimalAPI)
-class UEditorFactorySettingsObject : public UObject
-{
-	GENERATED_BODY()
-};
+class UInstancedPlacemenClientSettings;
 
 UINTERFACE(MinimalAPI)
 class UAssetFactoryInterface : public UInterface
@@ -84,5 +81,5 @@ public:
 	 /**
 	  * Returns the settings object which this factory will use to customize placement settings, based on the given placement information.
 	  */
-	 virtual UEditorFactorySettingsObject* FactorySettingsObjectForPlacement(const FAssetData& InAssetData, const FPlacementOptions& InPlacementOptions) = 0;
+	 virtual UInstancedPlacemenClientSettings* FactorySettingsObjectForPlacement(const FAssetData& InAssetData, const FPlacementOptions& InPlacementOptions) = 0;
 };
diff --git a/Engine/Source/Editor/EditorFramework/Public/Subsystems/PlacementSubsystem.h b/Engine/Source/Editor/EditorFramework/Public/Subsystems/PlacementSubsystem.h
--- a/Engine/Source/Editor/EditorFramework/Public/Subsystems/PlacementSubsystem.h
+++ b/Engine/Source/Editor/EditorFramework/Public/Subsystems/PlacementSubsystem.h
@@ -9,8 +9,9 @@
 
 #include "PlacementSubsystem.generated.h"
 
+class UAssetFactoryInterface;
+class UInstancedPlacemenClientSettings;
 class IAssetFactoryInterface;
-class UEditorFactorySettingsObject;
 
 USTRUCT()
 struct EDITORFRAMEWORK_API FAssetPlacementInfo
@@ -35,7 +36,7 @@ struct EDITORFRAMEWORK_API FAssetPlacementInfo
 
 	// If set, will use the given factory to place the asset, instead of allowing the placement subsystem to determine which factory to use.
 	UPROPERTY()
-	TScriptInterface<IAssetFactoryInterface> FactoryOverride;
+	TSubclassOf<UAssetFactoryInterface> FactoryOverride;
 
 	/**
 	 * The Guid which corresponds to the item that should be placed.
@@ -46,7 +47,7 @@ struct EDITORFRAMEWORK_API FAssetPlacementInfo
 	FGuid ItemGuid;
 
 	UPROPERTY()
-	TObjectPtr<UEditorFactorySettingsObject> SettingsObject = nullptr;
+	TObjectPtr<UInstancedPlacemenClientSettings> SettingsObject = nullptr;
 };
 
 USTRUCT()
@@ -105,6 +106,10 @@ public:
 	 */
 	bool IsCreatingPreviewElements() const;
 
+	FSimpleMulticastDelegate& OnPlacementFactoriesRegistered();
+
+	TScriptInterface<IAssetFactoryInterface> GetAssetFactoryFromFactoryClass(TSubclassOf<UAssetFactoryInterface> InFactoryInterfaceClass) const;
+
 private:
 	void RegisterPlacementFactories();
 	void UnregisterPlacementFactories();
@@ -112,5 +117,7 @@ private:
 	UPROPERTY()
 	TArray<TScriptInterface<IAssetFactoryInterface>> AssetFactories;
 
+	FSimpleMulticastDelegate PlacementFactoriesRegistered;
+
 	bool bIsCreatingPreviewElements = false;
 };
diff --git a/Engine/Source/Editor/UnrealEd/Classes/ActorFactories/ActorFactory.h b/Engine/Source/Editor/UnrealEd/Classes/ActorFactories/ActorFactory.h
--- a/Engine/Source/Editor/UnrealEd/Classes/ActorFactories/ActorFactory.h
+++ b/Engine/Source/Editor/UnrealEd/Classes/ActorFactories/ActorFactory.h
@@ -19,24 +19,7 @@ struct FActorSpawnParameters;
 struct FAssetData;
 class UBlueprint;
 class ULevel;
-
-UENUM()
-enum class EActorPlacementType
-{
-	FullActor,	// Places a full actor, one per instance
-	//InstancedStaticMesh,	// Places the actor as a set of instanced static meshes. All other components will be ignored.
-	//LightweightInstance,	// Places the actor as a lightweight instanced static mesh, which can be hydrated at runtime.
-};
-
-UCLASS()
-class UActorFactoryPlacementSettings : public UEditorFactorySettingsObject
-{
-	GENERATED_BODY()
-
-public:
-	UPROPERTY(VisibleAnywhere, Category="Factory Settings")
-	EActorPlacementType ActorPlacementType = EActorPlacementType::FullActor;
-};
+class UInstancedPlacemenClientSettings;
 
 UCLASS(collapsecategories, hidecategories=Object, editinlinenew, config=Editor, abstract, transient)
 class UNREALED_API UActorFactory : public UObject, public IAssetFactoryInterface
@@ -108,7 +91,7 @@ class UNREALED_API UActorFactory : public UObject, public IAssetFactoryInterface
 	virtual FAssetData GetAssetDataFromElementHandle(const FTypedElementHandle& InHandle) override;
 	virtual void BeginPlacement(const FPlacementOptions& InPlacementOptions) override;
 	virtual void EndPlacement(TArrayView<const FTypedElementHandle> InPlacedElements, const FPlacementOptions& InPlacementOptions) override;
-	virtual UEditorFactorySettingsObject* FactorySettingsObjectForPlacement(const FAssetData& InAssetData, const FPlacementOptions& InPlacementOptions) override;
+	virtual UInstancedPlacemenClientSettings* FactorySettingsObjectForPlacement(const FAssetData& InAssetData, const FPlacementOptions& InPlacementOptions) override;
 	// End IAssetFactoryInterface Interface
 
 protected:
diff --git a/Engine/Source/Editor/UnrealEd/Private/AssetSelection.cpp b/Engine/Source/Editor/UnrealEd/Private/AssetSelection.cpp
--- a/Engine/Source/Editor/UnrealEd/Private/AssetSelection.cpp
+++ b/Engine/Source/Editor/UnrealEd/Private/AssetSelection.cpp
@@ -616,7 +616,7 @@ static AActor* PrivateAddActor( UObject* Asset, UActorFactory* Factory, bool Sel
 			PlacementInfo.PreferredLevel = DesiredLevel;
 			PlacementInfo.NameOverride = Name;
 			PlacementInfo.FinalizedTransform = ActorTransform;
-			PlacementInfo.FactoryOverride = Factory;
+			PlacementInfo.FactoryOverride = Factory->GetClass();
 
 			FPlacementOptions PlacementOptions;
 			PlacementOptions.bIsCreatingPreviewElements = FLevelEditorViewportClient::IsDroppingPreviewActor();
diff --git a/Engine/Source/Editor/UnrealEd/Private/Factories/ActorFactory.cpp b/Engine/Source/Editor/UnrealEd/Private/Factories/ActorFactory.cpp
--- a/Engine/Source/Editor/UnrealEd/Private/Factories/ActorFactory.cpp
+++ b/Engine/Source/Editor/UnrealEd/Private/Factories/ActorFactory.cpp
@@ -383,10 +383,9 @@ void UActorFactory::EndPlacement(TArrayView<const FTypedElementHandle> InPlacedE
 {
 }
 
-UEditorFactorySettingsObject* UActorFactory::FactorySettingsObjectForPlacement(const FAssetData& InAssetData, const FPlacementOptions& InPlacementOptions)
+UInstancedPlacemenClientSettings* UActorFactory::FactorySettingsObjectForPlacement(const FAssetData& InAssetData, const FPlacementOptions& InPlacementOptions)
 {
-	UActorFactoryPlacementSettings* FactorySettingsObject = NewObject<UActorFactoryPlacementSettings>(this);
-	return FactorySettingsObject;
+	return nullptr;
 }
 
 AActor* UActorFactory::CreateActor( UObject* Asset, ULevel* InLevel, const FTransform& SpawnTransform, EObjectFlags InObjectFlags, const FName Name)
@@ -1163,7 +1162,7 @@ bool UActorFactorySkeletalMesh::CanCreateActorFrom( const FAssetData& AssetData,
 		SkeletalMeshData = AssetData;
 	}
 
-	if ( !SkeletalMeshData.IsValid() && (AssetData.GetClass()->IsChildOf(UAnimBlueprint::StaticClass()) || AssetData.GetClass()->IsChildOf(UAnimBlueprintGeneratedClass::StaticClass())))
+	if ( !SkeletalMeshData.IsValid() && AssetData.GetClass()->IsChildOf( UAnimBlueprint::StaticClass() ) )
 	{
 		const FString TargetSkeletonPath = AssetData.GetTagValueRef<FString>( "TargetSkeleton" );
 		if ( TargetSkeletonPath.IsEmpty() )
@@ -1250,26 +1249,14 @@ bool UActorFactorySkeletalMesh::CanCreateActorFrom( const FAssetData& AssetData,
 
 USkeletalMesh* UActorFactorySkeletalMesh::GetSkeletalMeshFromAsset( UObject* Asset )
 {
-	USkeletalMesh* SkeletalMesh = Cast<USkeletalMesh>( Asset );
+	USkeletalMesh*SkeletalMesh = Cast<USkeletalMesh>( Asset );
+	UAnimBlueprint* AnimBlueprint = Cast<UAnimBlueprint>( Asset );
 	USkeleton* Skeleton = Cast<USkeleton>( Asset );
 
-	if(SkeletalMesh == nullptr)
+	if( SkeletalMesh == nullptr && AnimBlueprint != nullptr && AnimBlueprint->TargetSkeleton )
 	{
 		// base it on preview skeletal mesh, just to have something
-		if(UAnimBlueprint* AnimBlueprint = Cast<UAnimBlueprint>(Asset))
-		{
-			if(AnimBlueprint->TargetSkeleton)
-			{
-				SkeletalMesh = AnimBlueprint->TargetSkeleton->GetPreviewMesh(true);
-			}
-		}
-		if(UAnimBlueprintGeneratedClass* AnimBlueprintGeneratedClass = Cast<UAnimBlueprintGeneratedClass>(Asset))
-		{
-			if(AnimBlueprintGeneratedClass->TargetSkeleton)
-			{
-				SkeletalMesh = AnimBlueprintGeneratedClass->TargetSkeleton->GetPreviewMesh(true);
-			}
-		}
+		SkeletalMesh = AnimBlueprint->TargetSkeleton->GetPreviewMesh(true);
 	}
 
 	if( SkeletalMesh == nullptr && Skeleton != nullptr )
diff --git a/Engine/Source/Editor/UnrealEd/Private/Factories/EditorStaticMeshFactory.cpp b/Engine/Source/Editor/UnrealEd/Private/Factories/EditorStaticMeshFactory.cpp
--- a/Engine/Source/Editor/UnrealEd/Private/Factories/EditorStaticMeshFactory.cpp
+++ b/Engine/Source/Editor/UnrealEd/Private/Factories/EditorStaticMeshFactory.cpp
@@ -6,7 +6,11 @@
 #include "Elements/Framework/TypedElementRegistry.h"
 #include "Elements/Interfaces/TypedElementAssetDataInterface.h"
 #include "Elements/Interfaces/TypedElementObjectInterface.h"
+#include "Elements/Interfaces/TypedElementWorldInterface.h"
 #include "Instances/InstancedPlacementPartitionActor.h"
+#include "Instances/EditorPlacementSettings.h"
+
+#include "LevelEditorSubsystem.h"
 
 #include "ActorPartition/ActorPartitionSubsystem.h"
 
@@ -14,14 +18,34 @@
 #include "Engine/StaticMesh.h"
 #include "Engine/StaticMeshActor.h"
 #include "UObject/Class.h"
+#include "Editor.h"
 
 #include "Subsystems/PlacementSubsystem.h"
 
-void UEditorStaticMeshFactoryPlacementSettings::PostEditChangeProperty(struct FPropertyChangedEvent& PropertyChangedEvent)
+bool UEditorStaticMeshFactory::PrePlaceAsset(FAssetPlacementInfo& InPlacementInfo, const FPlacementOptions& InPlacementOptions)
 {
-	Super::PostEditChangeProperty(PropertyChangedEvent);
+	if (!Super::PrePlaceAsset(InPlacementInfo, InPlacementOptions))
+	{
+		return false;
+	}
 
-	StaticMeshComponentDescriptor.ComputeHash();
+	if (!InPlacementInfo.PreferredLevel.IsValid())
+	{
+		return false;
+	}
+
+	if (!InPlacementInfo.SettingsObject)
+	{
+		return false;
+	}
+
+	// Make a good known client GUID out of the placed asset's package if one was not given to us.
+	if (!InPlacementInfo.ItemGuid.IsValid())
+	{
+		InPlacementInfo.ItemGuid = InPlacementInfo.AssetToPlace.GetAsset()->GetPackage()->GetPersistentGuid();
+	}
+
+	return true;
 }
 
 TArray<FTypedElementHandle> UEditorStaticMeshFactory::PlaceAsset(const FAssetPlacementInfo& InPlacementInfo, const FPlacementOptions& InPlacementOptions)
@@ -32,71 +56,72 @@ TArray<FTypedElementHandle> UEditorStaticMeshFactory::PlaceAsset(const FAssetPla
 		return Super::PlaceAsset(InPlacementInfo, InPlacementOptions);
 	}
 
-	TArray<FTypedElementHandle> PlacedInstanceHandles;
-	if (InPlacementInfo.PreferredLevel.IsValid())
+	if (UActorPartitionSubsystem* PartitionSubsystem = UWorld::GetSubsystem<UActorPartitionSubsystem>(InPlacementInfo.PreferredLevel->GetWorld()))
 	{
-		if (!InPlacementInfo.SettingsObject)
+		// Create or find the placement partition actor
+		auto OnActorCreated = [InPlacementOptions](APartitionActor* CreatedPartitionActor)
 		{
-			return TArray<FTypedElementHandle>();
+			if (AInstancedPlacementPartitionActor* ElementPartitionActor = Cast<AInstancedPlacementPartitionActor>(CreatedPartitionActor))
+			{
+				ElementPartitionActor->SetGridGuid(InPlacementOptions.InstancedPlacementGridGuid);
+			}
+		};
+
+		constexpr bool bCreatePartitionActorIfMissing = true;
+		FActorPartitionGetParams PartitionActorFindParams(
+			AInstancedPlacementPartitionActor::StaticClass(),
+			bCreatePartitionActorIfMissing,
+			InPlacementInfo.PreferredLevel.Get(),
+			InPlacementInfo.FinalizedTransform.GetLocation(),
+			0,
+			InPlacementOptions.InstancedPlacementGridGuid,
+			true,
+			OnActorCreated
+		);
+		AInstancedPlacementPartitionActor* PlacedElementsActor = Cast<AInstancedPlacementPartitionActor>(PartitionSubsystem->GetActor(PartitionActorFindParams));
+
+		auto RegisterISMDefinitionFunc = [&InPlacementInfo](AInstancedPlacementPartitionActor* PartitionActor, TSortedMap<int32, TArray<FTransform>>& ISMDefinition)
+		{
+			FISMComponentDescriptor ComponentDescriptor = InPlacementInfo.SettingsObject->InstancedComponentSettings;
+			ComponentDescriptor.StaticMesh = Cast<UStaticMesh>(InPlacementInfo.AssetToPlace.GetAsset());
+			ComponentDescriptor.ComputeHash();
+
+			int32 DescriptorIndex = PartitionActor->RegisterISMComponentDescriptor(ComponentDescriptor);
+			ISMDefinition.Emplace(DescriptorIndex, TArray<FTransform>({ FTransform() }));
+		};
+
+		FString ClientDisplayName = InPlacementInfo.NameOverride.ToString();
+		if (ClientDisplayName.IsEmpty())
+		{
+			ClientDisplayName = InPlacementInfo.AssetToPlace.GetFullName();
 		}
 
-		FISMComponentDescriptor ComponentDescriptor = CastChecked<UEditorStaticMeshFactoryPlacementSettings>(InPlacementInfo.SettingsObject)->StaticMeshComponentDescriptor;
-		if (!ComponentDescriptor.StaticMesh)
+		// Create an info or find it based on the given client guid
+		FClientPlacementInfo* PlacementInfo = PlacedElementsActor->PreAddClientInstances(InPlacementInfo.ItemGuid, ClientDisplayName, RegisterISMDefinitionFunc);
+		if (!PlacementInfo || !PlacementInfo->IsInitialized())
 		{
 			return TArray<FTypedElementHandle>();
 		}
 
-		// Make sure the component descriptor's hash matches its current settings before we place.
-		ComponentDescriptor.ComputeHash();
+		ModifiedPartitionActors.Add(PlacedElementsActor);
 
-		if (UActorPartitionSubsystem* PartitionSubsystem = UWorld::GetSubsystem<UActorPartitionSubsystem>(InPlacementInfo.PreferredLevel->GetWorld()))
+		FPlacementInstance InstanceToPlace;
+		InstanceToPlace.SetInstanceWorldTransform(InPlacementInfo.FinalizedTransform);
+		// todo: z offset and align to normal
+		TArray<FTypedElementHandle> PlacedInstanceHandles;
+		TArray<FSMInstanceId> PlacedInstances = PlacementInfo->AddInstances(MakeArrayView({ InstanceToPlace }));
+		for (const FSMInstanceId PlacedInstanceID : PlacedInstances)
 		{
-			// Create or find the placement partition actor
-			auto OnActorCreated = [InPlacementOptions](APartitionActor* CreatedPartitionActor)
-			{
-				if (AInstancedPlacementPartitionActor* ElementPartitionActor = Cast<AInstancedPlacementPartitionActor>(CreatedPartitionActor))
-				{
-					ElementPartitionActor->SetGridGuid(InPlacementOptions.InstancedPlacementGridGuid);
-				}
-			};
-
-			// Make a good known client GUID out of the placed asset's package if one was not given to us.
-			FGuid ItemGuidToUse = InPlacementInfo.ItemGuid;
-			if (!ItemGuidToUse.IsValid())
-			{
-				ItemGuidToUse = InPlacementInfo.AssetToPlace.GetAsset()->GetPackage()->GetPersistentGuid();
-			}
-
-			constexpr bool bCreatePartitionActorIfMissing = true;
-			FActorPartitionGetParams PartitionActorFindParams(
-				AInstancedPlacementPartitionActor::StaticClass(),
-				bCreatePartitionActorIfMissing, InPlacementInfo.PreferredLevel.Get(),
-				InPlacementInfo.FinalizedTransform.GetLocation(),
-				0,
-				InPlacementOptions.InstancedPlacementGridGuid,
-				true,
-				OnActorCreated
-			);
-			AInstancedPlacementPartitionActor* PlacedElementsActor = Cast<AInstancedPlacementPartitionActor>(PartitionSubsystem->GetActor(PartitionActorFindParams));
-
-			FISMClientHandle ClientHandle = PlacedElementsActor->RegisterClient(ItemGuidToUse);
-			TSortedMap<int32, TArray<FTransform>> InstanceMap;
-			InstanceMap.Emplace(PlacedElementsActor->RegisterISMComponentDescriptor(ComponentDescriptor), TArray({ FTransform() }));
-			ModifiedPartitionActors.Add(PlacedElementsActor);
-			PlacedElementsActor->BeginUpdate();
-
-			TArray<FSMInstanceId> PlacedInstances = PlacedElementsActor->AddISMInstance(ClientHandle, InPlacementInfo.FinalizedTransform, InstanceMap);
-			for (const FSMInstanceId PlacedInstanceID : PlacedInstances)
+			if (FTypedElementHandle PlacedHandle = UEngineElementsLibrary::AcquireEditorSMInstanceElementHandle(PlacedInstanceID))
 			{
-				if (FTypedElementHandle PlacedHandle = UEngineElementsLibrary::AcquireEditorSMInstanceElementHandle(PlacedInstanceID))
-				{
-					PlacedInstanceHandles.Emplace(PlacedHandle);
-				}
+				PlacedInstanceHandles.Emplace(PlacedHandle);
 			}
 		}
+
+		return PlacedInstanceHandles;
 	}
 
-	return PlacedInstanceHandles;
+	return TArray<FTypedElementHandle>();
 }
 
 FAssetData UEditorStaticMeshFactory::GetAssetDataFromElementHandle(const FTypedElementHandle& InHandle)
@@ -137,18 +162,18 @@ FAssetData UEditorStaticMeshFactory::GetAssetDataFromElementHandle(const FTypedE
 	return Super::GetAssetDataFromElementHandle(InHandle);
 }
 
-UEditorFactorySettingsObject* UEditorStaticMeshFactory::FactorySettingsObjectForPlacement(const FAssetData& InAssetData, const FPlacementOptions& InPlacementOptions)
+UInstancedPlacemenClientSettings* UEditorStaticMeshFactory::FactorySettingsObjectForPlacement(const FAssetData& InAssetData, const FPlacementOptions& InPlacementOptions)
 {
 	if (!ShouldPlaceInstancedStaticMeshes(InPlacementOptions))
 	{
 		return Super::FactorySettingsObjectForPlacement(InAssetData, InPlacementOptions);
 	}
 
-	UEditorStaticMeshFactoryPlacementSettings* PlacementSettingsObject = NewObject<UEditorStaticMeshFactoryPlacementSettings>(this);
+	UEditorInstancedPlacementSettings* PlacementSettingsObject = NewObject<UEditorInstancedPlacementSettings>(this);
 	if (PlacementSettingsObject)
 	{
 		UObject* AssetToPlaceAsObject = InAssetData.GetAsset();
-		FISMComponentDescriptor& ComponentDescriptor = PlacementSettingsObject->StaticMeshComponentDescriptor;
+		FISMComponentDescriptor& ComponentDescriptor = PlacementSettingsObject->InstancedComponentSettings;
 		if (UStaticMesh* StaticMeshObject = Cast<UStaticMesh>(AssetToPlaceAsObject))
 		{
 			// If this is a Nanite mesh, prefer to use ISM over HISM, as HISM duplicates many features/bookkeeping that Nanite already handles for us.
@@ -180,11 +205,11 @@ bool UEditorStaticMeshFactory::ShouldPlaceInstancedStaticMeshes(const FPlacement
 
 void UEditorStaticMeshFactory::EndPlacement(TArrayView<const FTypedElementHandle> InPlacedElements, const FPlacementOptions& InPlacementOptions)
 {
-	for (TWeakObjectPtr<AISMPartitionActor> ISMPartitionActor : ModifiedPartitionActors)
+	for (TWeakObjectPtr<AInstancedPlacementPartitionActor> ISMPartitionActor : ModifiedPartitionActors)
 	{
 		if (ISMPartitionActor.IsValid())
 		{
-			ISMPartitionActor->EndUpdate();
+			ISMPartitionActor->PostAddClientInstances();
 		}
 	}
 
diff --git a/Engine/Source/Editor/UnrealEd/Private/Factories/EditorStaticMeshFactory.h b/Engine/Source/Editor/UnrealEd/Private/Factories/EditorStaticMeshFactory.h
--- a/Engine/Source/Editor/UnrealEd/Private/Factories/EditorStaticMeshFactory.h
+++ b/Engine/Source/Editor/UnrealEd/Private/Factories/EditorStaticMeshFactory.h
@@ -5,22 +5,12 @@
 #include "Factories/AssetFactoryInterface.h"
 #include "ActorFactories/ActorFactoryStaticMesh.h"
 #include "ISMPartition/ISMComponentDescriptor.h"
+#include "Instances/InstancedPlacementClientInfo.h"
 
 #include "EditorStaticMeshFactory.generated.h"
 
-class AISMPartitionActor;
-
-UCLASS()
-class UEditorStaticMeshFactoryPlacementSettings : public UEditorFactorySettingsObject
-{
-	GENERATED_BODY()
-
-public:
-	UPROPERTY(EditAnywhere, Category = "ISM Component Settings", meta=(ShowOnlyInnerProperties))
-	FISMComponentDescriptor StaticMeshComponentDescriptor;
-
-	virtual void PostEditChangeProperty(struct FPropertyChangedEvent& PropertyChangedEvent) override;
-};
+class AInstancedPlacementPartitionActor;
+class UInstancedPlacemenClientSettings;
 
 UCLASS(Transient)
 class UEditorStaticMeshFactory : public UActorFactoryStaticMesh
@@ -29,13 +19,15 @@ class UEditorStaticMeshFactory : public UActorFactoryStaticMesh
 
 public:
 	// Begin IAssetFactoryInterface
+	virtual bool PrePlaceAsset(FAssetPlacementInfo& InPlacementInfo, const FPlacementOptions& InPlacementOptions) override;
 	virtual TArray<FTypedElementHandle> PlaceAsset(const FAssetPlacementInfo& InPlacementInfo, const FPlacementOptions& InPlacementOptions) override;
 	virtual FAssetData GetAssetDataFromElementHandle(const FTypedElementHandle& InHandle) override;
 	virtual void EndPlacement(TArrayView<const FTypedElementHandle> InPlacedElements, const FPlacementOptions& InPlacementOptions) override;
-	virtual UEditorFactorySettingsObject* FactorySettingsObjectForPlacement(const FAssetData& InAssetData, const FPlacementOptions& InPlacementOptions) override;
+	virtual UInstancedPlacemenClientSettings* FactorySettingsObjectForPlacement(const FAssetData& InAssetData, const FPlacementOptions& InPlacementOptions) override;
 	// End IAssetFactoryInterface
 
 protected:
 	bool ShouldPlaceInstancedStaticMeshes(const FPlacementOptions& InPlacementOptions) const;
-	TSet<TWeakObjectPtr<AISMPartitionActor>> ModifiedPartitionActors;
+
+	TSet<TWeakObjectPtr<AInstancedPlacementPartitionActor>> ModifiedPartitionActors;
 };
\ No newline at end of file
diff --git a/Engine/Source/Editor/UnrealEd/Private/Instances/EditorPlacementSettings.cpp b/Engine/Source/Editor/UnrealEd/Private/Instances/EditorPlacementSettings.cpp
new file mode 100644
--- /dev/null
+++ b/Engine/Source/Editor/UnrealEd/Private/Instances/EditorPlacementSettings.cpp
@@ -0,0 +1,110 @@
+// Copyright Epic Games, Inc. All Rights Reserved.
+
+#include "Instances/EditorPlacementSettings.h"
+
+#include "Elements/Framework/EngineElementsLibrary.h"
+#include "Elements/Framework/TypedElementRegistry.h"
+#include "Elements/Interfaces/TypedElementAssetDataInterface.h"
+#include "Elements/Interfaces/TypedElementObjectInterface.h"
+#include "Elements/Interfaces/TypedElementWorldInterface.h"
+#include "Instances/InstancedPlacementPartitionActor.h"
+
+#include "Components/InstancedStaticMeshComponent.h"
+
+#include "Editor.h"
+#include "Subsystems/PlacementSubsystem.h"
+#include "LevelEditorSubsystem.h"
+
+void UEditorInstancedPlacementSettings::PostEditChangeProperty(struct FPropertyChangedEvent& PropertyChangedEvent)
+{
+	Super::PostEditChangeProperty(PropertyChangedEvent);
+}
+
+void UEditorInstancedPlacementSettings::PostLoad()
+{
+	Super::PostLoad();
+	
+	// Register Delegates if necessary
+}
+
+void UEditorInstancedPlacementSettings::Serialize(FArchive& Ar)
+{
+	Super::Serialize(Ar);
+	if (Ar.IsLoading())
+	{
+		// Register Delegates if necessary
+	}
+}
+
+void UEditorInstancedPlacementSettings::RegisterISMDescriptors(AInstancedPlacementPartitionActor* ParentPartitionActor, TSortedMap<int32, TArray<FTransform>>& ISMDefinition) const
+{
+	ISMDefinition.Empty();
+	if (!ParentPartitionActor)
+	{
+		return;
+	}
+
+	UPlacementSubsystem* PlacementSubsystem = GEditor->GetEditorSubsystem<UPlacementSubsystem>();
+	if (!PlacementSubsystem)
+	{
+		return;
+	}
+
+	TArray<FTypedElementHandle> SpawnedTemporaryElements;// = PlacementSubsystem->SpawnPreviewElementsFromAssetData(ObjectPath.TryLoad());
+	for (FTypedElementHandle& PreviewElement : SpawnedTemporaryElements)
+	{
+		AActor* PreviewActor = nullptr;
+		if (TTypedElement<ITypedElementObjectInterface> ObjectInterface = UTypedElementRegistry::GetInstance()->GetElement<ITypedElementObjectInterface>(PreviewElement))
+		{
+			PreviewActor = ObjectInterface.GetObjectAs<AActor>();
+		}
+
+		if (!PreviewActor)
+		{
+			return;
+		}
+
+		FTransform ActorTransform = PreviewActor->GetActorTransform();
+		TArray<UStaticMeshComponent*> StaticMeshComponents;
+		PreviewActor->GetComponents<UStaticMeshComponent>(StaticMeshComponents);
+
+		for (UStaticMeshComponent* StaticMeshComponent : StaticMeshComponents)
+		{
+			if (UStaticMesh* StaticMesh = StaticMeshComponent->GetStaticMesh())
+			{
+				FISMComponentDescriptor Descriptor = InstancedComponentSettings;
+				Descriptor.StaticMesh = StaticMesh;
+				Descriptor.ComputeHash();
+				int32 DescriptorIndex = ParentPartitionActor->RegisterISMComponentDescriptor(Descriptor);
+
+				TArray<FTransform>& Transforms = ISMDefinition.FindOrAdd(DescriptorIndex);
+				if (UInstancedStaticMeshComponent* ISMComponent = Cast<UInstancedStaticMeshComponent>(StaticMeshComponent))
+				{
+					for (int32 InstanceIndex = 0; InstanceIndex < ISMComponent->GetInstanceCount(); ++InstanceIndex)
+					{
+						FTransform InstanceTransform;
+						constexpr bool bUseWorldSpace = true;
+						if (ensure(ISMComponent->GetInstanceTransform(InstanceIndex, InstanceTransform, bUseWorldSpace)))
+						{
+							FTransform LocalTransform = InstanceTransform.GetRelativeTransform(ActorTransform);
+							Transforms.Add(LocalTransform);
+						}
+					}
+				}
+				else
+				{
+					FTransform LocalTransform = StaticMeshComponent->GetComponentTransform().GetRelativeTransform(ActorTransform);
+					Transforms.Add(LocalTransform);
+				}
+			}
+		}
+
+		if (TTypedElement<ITypedElementWorldInterface> WorldInterfaceElement = UTypedElementRegistry::GetInstance()->GetElement<ITypedElementWorldInterface>(PreviewElement))
+		{
+			if (ULevelEditorSubsystem* LevelEditorSubsystem = GEditor->GetEditorSubsystem<ULevelEditorSubsystem>())
+			{
+				WorldInterfaceElement.DeleteElement(WorldInterfaceElement.GetOwnerWorld(), LevelEditorSubsystem->GetSelectionSet(), FTypedElementDeletionOptions());
+			}
+		}
+	}
+}
diff --git a/Engine/Source/Editor/UnrealEd/Public/Instances/EditorPlacementSettings.h b/Engine/Source/Editor/UnrealEd/Public/Instances/EditorPlacementSettings.h
new file mode 100644
--- /dev/null
+++ b/Engine/Source/Editor/UnrealEd/Public/Instances/EditorPlacementSettings.h
@@ -0,0 +1,25 @@
+// Copyright Epic Games, Inc. All Rights Reserved.
+
+#pragma once
+
+#include "Instances/InstancedPlacementClientInfo.h"
+
+#include "EditorPlacementSettings.generated.h"
+
+class AInstancedPlacementPartitionActor;
+struct FISMComponentDescriptor;
+class FArchive;
+struct FPropertyChangedEvent;
+
+UCLASS()
+class UEditorInstancedPlacementSettings : public UInstancedPlacemenClientSettings
+{
+	GENERATED_BODY()
+
+public:
+	virtual void PostEditChangeProperty(FPropertyChangedEvent& PropertyChangedEvent) override;
+	virtual void PostLoad() override;
+	virtual void Serialize(FArchive& Ar) override;
+
+	virtual void RegisterISMDescriptors(AInstancedPlacementPartitionActor* ParentPartitionActor, TSortedMap<int32, TArray<FTransform>>& ISMDefinition) const override;
+};
\ No newline at end of file
diff --git a/Engine/Source/Runtime/Engine/Private/Instances/InstancedPlacementClientInfo.cpp b/Engine/Source/Runtime/Engine/Private/Instances/InstancedPlacementClientInfo.cpp
new file mode 100644
--- /dev/null
+++ b/Engine/Source/Runtime/Engine/Private/Instances/InstancedPlacementClientInfo.cpp
@@ -0,0 +1,459 @@
+// Copyright Epic Games, Inc. All Rights Reserved.
+
+#include "Instances/InstancedPlacementClientInfo.h"
+
+#if WITH_EDITORONLY_DATA
+#include "Instances/InstancedPlacementPartitionActor.h"
+#include "Elements/Framework/EngineElementsLibrary.h"
+
+FArchive& operator<<(FArchive& Ar, FPlacementInstance& Instance)
+{
+	Ar << Instance.ZOffset;
+	Ar << Instance.Flags;
+	Ar << Instance.WorldTransform;
+	Ar << Instance.PreAlignRotation;
+
+	return Ar;
+}
+
+FTransform FPlacementInstance::GetInstanceWorldTransform() const
+{
+	return WorldTransform;
+}
+
+void FPlacementInstance::SetInstanceWorldTransform(const FTransform& Transform)
+{
+	WorldTransform = Transform;
+}
+
+void FPlacementInstance::AlignToNormal(const FVector& InNormal, int32 AlignMaxAngle)
+{
+	Flags |= (uint32)EPlacementInstanceFlags::AlignToNormal;
+
+	FRotator AlignRotation = InNormal.Rotation();
+	// Static meshes are authored along the vertical axis rather than the X axis, so we add 90 degrees to the static mesh's Pitch.
+	AlignRotation.Pitch -= 90.f;
+	// Clamp its value inside +/- one rotation
+	AlignRotation.Pitch = FRotator::NormalizeAxis(AlignRotation.Pitch);
+
+	// limit the maximum pitch angle if it's > 0.
+	if (AlignMaxAngle > 0)
+	{
+		int32 MaxPitch = AlignMaxAngle;
+		if (AlignRotation.Pitch > MaxPitch)
+		{
+			AlignRotation.Pitch = MaxPitch;
+		}
+		else if (AlignRotation.Pitch < -MaxPitch)
+		{
+			AlignRotation.Pitch = -MaxPitch;
+		}
+	}
+
+	PreAlignRotation = WorldTransform.GetRotation();
+	WorldTransform.SetRotation(FQuat(AlignRotation) * PreAlignRotation);
+}
+
+#if WITH_EDITOR
+void UInstancedPlacemenClientSettings::PostEditChangeProperty(FPropertyChangedEvent& PropertyChangedEvent)
+{
+	Super::PostEditChangeProperty(PropertyChangedEvent);
+
+	UpdateGuid = FGuid::NewGuid();
+
+	bool bSourceChanged = false;
+	if (const FProperty* Property = PropertyChangedEvent.Property)
+	{
+		if (Property->GetFName() == GET_MEMBER_NAME_CHECKED(UInstancedPlacemenClientSettings, ObjectPath))
+		{
+			bSourceChanged = true;
+		}
+	}
+	uint32 OriginalHash = InstancedComponentSettings.Hash;
+	InstancedComponentSettings.ComputeHash();
+}
+
+void UInstancedPlacemenClientSettings::RegisterISMDescriptors(AInstancedPlacementPartitionActor* ParentPartitionActor, TSortedMap<int32, TArray<FTransform>>& ISMDefinition) const
+{
+}
+#endif	// WITH_EDITOR
+
+FClientPlacementInfo::FClientPlacementInfo()
+{
+}
+
+FArchive& operator<<(FArchive& Ar, FClientPlacementInfo& PlacementInfo)
+{
+	if (Ar.IsTransacting())
+	{
+		PlacementInfo.Instances.BulkSerialize(Ar, true);
+	}
+	else
+	{
+		Ar << PlacementInfo.Instances;
+	}
+
+	Ar << PlacementInfo.ClientGuid;
+	Ar << PlacementInfo.UpdateGuid;
+	Ar << PlacementInfo.ISMDefinition;
+	PlacementInfo.ClientHandle.Serialize(Ar);
+
+	return Ar;
+}
+
+bool FClientPlacementInfo::Initialize(const FGuid& InClientGuid, const FString& InClientDisplayName, AInstancedPlacementPartitionActor* InParentPartitionActor, FClientDescriptorFunc RegisterDescriptorFunc)
+{
+	check(!IsInitialized());
+
+	if (!InParentPartitionActor || !InClientGuid.IsValid())
+	{
+		return false;
+	}
+
+	// Register the descriptors and generate the ISM definition from the callback provided
+	RegisterDescriptorFunc(InParentPartitionActor, ISMDefinition);
+	if (ISMDefinition.IsEmpty())
+	{
+		return false;
+	}
+
+	ParentPartitionActor = InParentPartitionActor;
+	ClientGuid = InClientGuid;
+
+	ClientDisplayName = InClientDisplayName;
+	if (ClientDisplayName.IsEmpty())
+	{
+		ClientDisplayName = ParentPartitionActor->GetName();
+	}
+
+	ClientHandle = ParentPartitionActor->RegisterClient(InClientGuid);
+	ParentPartitionActor->RegisterClientInstanceManager(ClientHandle, this);
+	//RegisterDelegates();
+
+	return true;
+}
+
+void FClientPlacementInfo::Uninitialize()
+{
+	check(IsInitialized());
+	//UnregisterDelegates();
+	ParentPartitionActor->UnregisterClient(ClientHandle);
+	ISMDefinition.Empty();
+}
+
+void FClientPlacementInfo::PostLoad(AInstancedPlacementPartitionActor* InParentPartitionActor)
+{
+#if WITH_EDITORONLY_DATA
+	ParentPartitionActor = InParentPartitionActor;
+
+	for (int32 Index = 0; Index < Instances.Num(); ++Index)
+	{
+		InstanceHash.InsertInstance(Instances[Index].GetInstanceWorldTransform().GetLocation(), Index);
+	}
+#endif
+
+#if WITH_EDITOR
+	if (GIsEditor && IsInitialized())
+	{
+		//RegisterDelegates();
+	}
+#endif
+}
+void FClientPlacementInfo::PostSerialize(FArchive& Ar, AInstancedPlacementPartitionActor* InParentPartitionActor)
+{
+#if WITH_EDITORONLY_DATA
+	ParentPartitionActor = InParentPartitionActor;
+#endif
+
+#if WITH_EDITOR
+	if (GIsEditor && IsInitialized() && Ar.IsLoading())
+	{
+		ParentPartitionActor->RegisterClientInstanceManager(ClientHandle, this);
+	}
+#endif	// WITH_EDITOR
+}
+bool FClientPlacementInfo::IsInitialized() const
+{
+	return ClientHandle.IsValid();
+}
+
+void FClientPlacementInfo::PreEditUndo()
+{
+	InstanceHash.CheckInstanceCount(Instances.Num());
+}
+
+void FClientPlacementInfo::PostEditUndo()
+{
+	// Regenerate instance hash
+	// We regenerate it here instead of saving to transaction buffer to speed up modify operations
+	// And also because not all edits of the hash get correctly saved with the frequency of modify operations we have
+	InstanceHash.Empty();
+	for (int32 InstanceIdx = 0; InstanceIdx < Instances.Num(); InstanceIdx++)
+	{
+		InstanceHash.InsertInstance(Instances[InstanceIdx].GetInstanceWorldTransform().GetLocation(), InstanceIdx);
+	}
+
+	InstanceHash.CheckInstanceCount(Instances.Num());
+}
+
+TArray<FSMInstanceId> FClientPlacementInfo::AddInstances(TArrayView<const FPlacementInstance> InstancesToAdd)
+{
+	check(IsInitialized());
+
+	ParentPartitionActor->Modify();
+	ParentPartitionActor->BeginUpdate();
+
+	Instances.Reserve(Instances.Num() + InstancesToAdd.Num());
+
+	TArray<FSMInstanceId> PlacedInstances;
+	PlacedInstances.Reserve(InstancesToAdd.Num());
+	for (const FPlacementInstance& Instance : InstancesToAdd)
+	{
+		PlacedInstances.Append(ParentPartitionActor->AddISMInstance(ClientHandle, Instance.GetInstanceWorldTransform(), ISMDefinition));
+		int32 InstanceIndex = Instances.Add(Instance);
+		InstanceHash.InsertInstance(Instance.GetInstanceWorldTransform().GetLocation(), InstanceIndex);
+	}
+
+	ParentPartitionActor->EndUpdate();
+
+	return PlacedInstances;
+}
+
+void FClientPlacementInfo::RemoveInstancesFromPartitionActor(TArrayView<const FISMClientInstanceId> InstanceIds, bool bUpdateHISMTrees, TFunctionRef<void(int32)> RemoveFn)
+{
+	if (InstanceIds.IsEmpty())
+	{
+		return;
+	}
+
+	check(IsInitialized());
+
+	TSet<int32> InstancesToRemove;
+	InstancesToRemove.Reserve(InstanceIds.Num());
+	for (const FISMClientInstanceId& InstanceId : InstanceIds)
+	{
+		InstancesToRemove.Add(InstanceId.Index);
+	}
+
+	ParentPartitionActor->Modify();
+	ParentPartitionActor->BeginUpdate();
+
+	while (InstancesToRemove.Num() > 0)
+	{
+		// Get an item from the set for processing
+		auto It = InstancesToRemove.CreateConstIterator();
+		int32 InstanceIndex = *It;
+		int32 InstanceIndexToRemove = InstanceIndex;
+
+		FPlacementInstance& Instance = Instances[InstanceIndex];
+
+		// remove from hash
+		InstanceHash.RemoveInstance(Instance.GetInstanceWorldTransform().GetLocation(), InstanceIndex);
+
+		// remove this instance from the parent partition actor
+		RemoveFn(InstanceIndex);
+
+		// remove from instances array
+		Instances.RemoveAtSwap(InstanceIndex, 1, false);
+
+		// update hashes for swapped instance
+		if (InstanceIndex != Instances.Num() && Instances.Num() > 0)
+		{
+			// Instance hash
+			FPlacementInstance& SwappedInstance = Instances[InstanceIndex];
+			InstanceHash.SwapInstance(SwappedInstance.GetInstanceWorldTransform().GetLocation(), Instances.Num(), InstanceIndex);
+
+			// Removal list
+			if (InstancesToRemove.Contains(Instances.Num()))
+			{
+				// The item from the end of the array that we swapped in to InstanceIndex is also on the list to remove.
+				// Remove the item at the end of the array and leave InstanceIndex in the removal list.
+				InstanceIndexToRemove = Instances.Num();
+			}
+		}
+
+		// Remove the removed item from the removal list
+		InstancesToRemove.Remove(InstanceIndexToRemove);
+	}
+
+	ParentPartitionActor->EndUpdate();
+
+	Instances.Shrink();
+	if (bUpdateHISMTrees)
+	{
+		ParentPartitionActor->UpdateHISMTrees(true, true);
+	}
+}
+
+FText FClientPlacementInfo::GetISMPartitionInstanceDisplayName(const FISMClientInstanceId& InstanceId) const
+{
+#if WITH_EDITOR
+	return FText::Format(NSLOCTEXT("PlacementClientInfo", "DisplayNameFmt", "{0} - Instance {1}"), FText::FromString(ClientDisplayName), InstanceId.Index);
+#else
+	return FText();
+#endif
+}
+
+FText FClientPlacementInfo::GetISMPartitionInstanceTooltip(const FISMClientInstanceId& InstanceId) const
+{
+#if WITH_EDITOR
+	const FText OwnerDisplayPath = FText::FromString(ParentPartitionActor->GetPathName(ParentPartitionActor->GetWorld())); // stops the path at the level of the world the object is in
+	return FText::Format(NSLOCTEXT("PlacementClientInfo", "TooltipFmt", "Instance {0} on {1}"), InstanceId.Index, OwnerDisplayPath);
+#else
+	return FText();
+#endif
+}
+
+bool FClientPlacementInfo::CanEditISMPartitionInstance(const FISMClientInstanceId& InstanceId) const
+{
+#if WITH_EDITOR
+	return true;
+#else
+	return false;
+#endif
+}
+
+bool FClientPlacementInfo::CanMoveISMPartitionInstance(const FISMClientInstanceId& InstanceId, const ETypedElementWorldType InWorldType) const
+{
+#if WITH_EDITOR
+	return true;
+#else
+	return false;
+#endif
+}
+
+bool FClientPlacementInfo::GetISMPartitionInstanceTransform(const FISMClientInstanceId& InstanceId, FTransform& OutInstanceTransform, bool bWorldSpace) const
+{
+#if WITH_EDITOR
+	check(InstanceId.Handle == ClientHandle);
+	OutInstanceTransform = Instances[InstanceId.Index].GetInstanceWorldTransform();
+	if (!bWorldSpace)
+	{
+		OutInstanceTransform = OutInstanceTransform.GetRelativeTransform(ParentPartitionActor->GetActorTransform());
+	}
+	return true;
+#else
+	return false;
+#endif
+}
+
+bool FClientPlacementInfo::SetISMPartitionInstanceTransform(const FISMClientInstanceId& InstanceId, const FTransform& InstanceTransform, bool bWorldSpace, bool bTeleport)
+{
+#if WITH_EDITOR
+	ParentPartitionActor->Modify();
+
+	check(InstanceId.Handle == ClientHandle);
+	bool bWasMoving = MovingInstances.Contains(InstanceId);
+	if (!bWasMoving)
+	{
+		NotifyISMPartitionInstanceMovementStarted(InstanceId);
+	}
+
+	FTransform WorldTransform = bWorldSpace ? InstanceTransform : InstanceTransform * ParentPartitionActor->GetActorTransform();
+	Instances[InstanceId.Index].SetInstanceWorldTransform(WorldTransform);
+	ParentPartitionActor->SetISMInstanceTransform(ClientHandle, InstanceId.Index, WorldTransform, bTeleport, ISMDefinition);
+
+	if (!bWasMoving)
+	{
+		NotifyISMPartitionInstanceMovementEnded(InstanceId);
+	}
+
+	return true;
+#else
+	return false;
+#endif
+}
+
+void FClientPlacementInfo::NotifyISMPartitionInstanceMovementStarted(const FISMClientInstanceId& InstanceId)
+{
+#if WITH_EDITOR
+	check(InstanceId.Handle == ClientHandle);
+	MovingInstances.Add(InstanceId);
+	InstanceHash.RemoveInstance(Instances[InstanceId.Index].GetInstanceWorldTransform().GetLocation(), InstanceId.Index);
+#endif
+}
+
+void FClientPlacementInfo::NotifyISMPartitionInstanceMovementOngoing(const FISMClientInstanceId& InstanceId)
+{
+#if WITH_EDITOR
+	check(InstanceId.Handle == ClientHandle);
+	NotifyISMPartitionInstanceMovementEnded(InstanceId);
+	NotifyISMPartitionInstanceMovementStarted(InstanceId);
+#endif
+}
+
+void FClientPlacementInfo::NotifyISMPartitionInstanceMovementEnded(const FISMClientInstanceId& InstanceId)
+{
+#if WITH_EDITOR
+	check(InstanceId.Handle == ClientHandle);
+	FTransform InstanceWorldTransform = Instances[InstanceId.Index].GetInstanceWorldTransform();
+	constexpr bool bTeleport = true;
+	ParentPartitionActor->SetISMInstanceTransform(ClientHandle, InstanceId.Index, InstanceWorldTransform, bTeleport, ISMDefinition);
+	InstanceHash.InsertInstance(Instances[InstanceId.Index].GetInstanceWorldTransform().GetLocation(), InstanceId.Index);
+	MovingInstances.Remove(InstanceId);
+#endif
+}
+
+void FClientPlacementInfo::NotifyISMPartitionInstanceSelectionChanged(const FISMClientInstanceId& InstanceId, const bool bIsSelected)
+{
+#if WITH_EDITOR
+	check(IsInitialized());
+	check(InstanceId.Handle == ClientHandle);
+
+	ParentPartitionActor->Modify(false);
+	ParentPartitionActor->SelectISMInstances(ClientHandle, bIsSelected, { InstanceId.Index });
+#endif
+}
+
+bool FClientPlacementInfo::DeleteISMPartitionInstances(TArrayView<const FISMClientInstanceId> InstanceIds)
+{
+#if WITH_EDITOR
+	auto RemoveFn = [this](int32 Index)
+	{
+		bool bOutIsEmpty = false;
+		ParentPartitionActor->RemoveISMInstance(ClientHandle, Index, &bOutIsEmpty);
+		if (bOutIsEmpty)
+		{
+			Uninitialize();
+		}
+	};
+
+	constexpr bool bUpdateHISMTrees = true;
+	RemoveInstancesFromPartitionActor(InstanceIds, bUpdateHISMTrees, RemoveFn);
+	return true;
+#else
+	return false;
+#endif
+}
+
+bool FClientPlacementInfo::DuplicateISMPartitionInstances(TArrayView<const FISMClientInstanceId> InstanceIds, TArray<FISMClientInstanceId>& OutNewInstanceIds)
+{
+#if WITH_EDITOR
+	if (InstanceIds.Num() == 0)
+	{
+		return false;
+	}
+
+	TArray<FPlacementInstance> InstancesToDuplicate;
+	InstancesToDuplicate.Reserve(InstanceIds.Num());
+	for (const FISMClientInstanceId& InstanceId : InstanceIds)
+	{
+		InstancesToDuplicate.Add(Instances[InstanceId.Index]);
+	}
+
+	TArray<FSMInstanceId> AddedInstances = AddInstances(InstancesToDuplicate);
+	ParentPartitionActor->UpdateHISMTrees(true, true);
+
+	OutNewInstanceIds.Reset(InstanceIds.Num());
+	for (const FSMInstanceId& NewInstanceIds : AddedInstances)
+	{
+		OutNewInstanceIds.Emplace(FISMClientInstanceId{ ClientHandle, NewInstanceIds.InstanceIndex });
+	}
+
+	return OutNewInstanceIds.Num() > 0;
+#else
+	return false;
+#endif
+}
+
+#endif	// WITH_EDITORONLY_DATA
diff --git a/Engine/Source/Runtime/Engine/Private/Instances/InstancedPlacementPartitionActor.cpp b/Engine/Source/Runtime/Engine/Private/Instances/InstancedPlacementPartitionActor.cpp
--- a/Engine/Source/Runtime/Engine/Private/Instances/InstancedPlacementPartitionActor.cpp
+++ b/Engine/Source/Runtime/Engine/Private/Instances/InstancedPlacementPartitionActor.cpp
@@ -2,7 +2,6 @@
 
 #include "Instances/InstancedPlacementPartitionActor.h"
 
-#include "Components/InstancedStaticMeshComponent.h"
 #include "Engine/World.h"
 #include "GameFramework/WorldSettings.h"
 
@@ -13,6 +12,58 @@ AInstancedPlacementPartitionActor::AInstancedPlacementPartitionActor(const FObje
 	SetActorEnableCollision(true);
 }
 
+void AInstancedPlacementPartitionActor::PostLoad()
+{
+	Super::PostLoad();
+
+#if WITH_EDITORONLY_DATA
+	for (auto& Pair : PlacedClientInfo)
+	{
+		FClientPlacementInfo& ClientInfo = *Pair.Value;
+		ClientInfo.PostLoad(this);
+	}
+#endif	// WITH_EDITORONLY_DATA
+}
+
+#if WITH_EDITOR
+void AInstancedPlacementPartitionActor::PreEditUndo()
+{
+	Super::PreEditUndo();
+
+	for (auto& Pair : PlacedClientInfo)
+	{
+		FClientPlacementInfo& ClientInfo = *Pair.Value;
+		ClientInfo.PreEditUndo();
+	}
+}
+
+void AInstancedPlacementPartitionActor::PostEditUndo()
+{
+	Super::PostEditUndo();
+
+	for (auto& Pair : PlacedClientInfo)
+	{
+		FClientPlacementInfo& ClientInfo = *Pair.Value;
+		ClientInfo.PostEditUndo();
+	}
+}
+#endif
+
+void AInstancedPlacementPartitionActor::Serialize(FArchive& Ar)
+{
+	Super::Serialize(Ar);
+
+#if WITH_EDITORONLY_DATA
+	Ar << PlacedClientInfo;
+
+	for (auto& Pair : PlacedClientInfo)
+	{
+		FClientPlacementInfo& ClientInfo = *Pair.Value;
+		ClientInfo.PostSerialize(Ar, this);
+	}
+#endif	// WITH_EDITORONLY_DATA
+}
+
 #if WITH_EDITOR	
 uint32 AInstancedPlacementPartitionActor::GetDefaultGridSize(UWorld* InWorld) const
 {
@@ -28,23 +79,36 @@ void AInstancedPlacementPartitionActor::SetGridGuid(const FGuid& InGuid)
 {
 	PlacementGridGuid = InGuid;
 }
-#endif
 
-ISMInstanceManager* AInstancedPlacementPartitionActor::GetSMInstanceManager(const FSMInstanceId& InstanceId)
+FClientPlacementInfo* AInstancedPlacementPartitionActor::PreAddClientInstances(const FGuid& ClientGuid, const FString& InClientDisplayString, FClientDescriptorFunc RegisterDefinitionFunc)
 {
-	if (ISMInstanceManager* ParentSMInstanceManager = Super::GetSMInstanceManager(InstanceId))
+	BeginUpdate();
+
+	if (TUniqueObj<FClientPlacementInfo>* FoundClientInfo = PlacedClientInfo.Find(ClientGuid))
 	{
-		return ParentSMInstanceManager;
+		if (!(*FoundClientInfo)->IsInitialized())
+		{
+			if (!(*FoundClientInfo)->Initialize(ClientGuid, InClientDisplayString, this, RegisterDefinitionFunc))
+			{
+				return nullptr;
+			}
+		}
+
+		return &FoundClientInfo->Get();
 	}
 
-#if WITH_EDITOR
-	// For now, assume that if we didn't have some manager registered, it is safe to edit the ISM directly.
-	// This should be removed after palette items are set up properly in the placement API
-	if (IsISMComponent(InstanceId.ISMComponent))
+	TUniqueObj<FClientPlacementInfo>& CreatedInfo = PlacedClientInfo.Add(ClientGuid);
+	if (CreatedInfo->Initialize(ClientGuid, InClientDisplayString, this, RegisterDefinitionFunc))
 	{
-		return InstanceId.ISMComponent;
+		return &(*CreatedInfo);
 	}
-#endif
 
+	PlacedClientInfo.Remove(ClientGuid);
 	return nullptr;
 }
+
+void AInstancedPlacementPartitionActor::PostAddClientInstances()
+{
+	EndUpdate();
+}
+#endif	// WITH_EDITOR
diff --git a/Engine/Source/Runtime/Engine/Public/ISMPartition/ISMComponentDescriptor.h b/Engine/Source/Runtime/Engine/Public/ISMPartition/ISMComponentDescriptor.h
--- a/Engine/Source/Runtime/Engine/Public/ISMPartition/ISMComponentDescriptor.h
+++ b/Engine/Source/Runtime/Engine/Public/ISMPartition/ISMComponentDescriptor.h
@@ -55,7 +55,7 @@ public:
 	UPROPERTY(EditAnywhere, Category = "Component Settings")
 	TSubclassOf<UInstancedStaticMeshComponent> ComponentClass;
 	
-	UPROPERTY(VisibleAnywhere, Category = "Component Settings")
+	UPROPERTY(EditAnywhere, Category = "Component Settings")
 	TObjectPtr<UStaticMesh> StaticMesh = nullptr;
 	
 	UPROPERTY(EditAnywhere, Category = "Component Settings")
diff --git a/Engine/Source/Runtime/Engine/Public/Instances/InstancedPlacementClientInfo.h b/Engine/Source/Runtime/Engine/Public/Instances/InstancedPlacementClientInfo.h
new file mode 100644
--- /dev/null
+++ b/Engine/Source/Runtime/Engine/Public/Instances/InstancedPlacementClientInfo.h
@@ -0,0 +1,135 @@
+// Copyright Epic Games, Inc. All Rights Reserved.
+
+#pragma once
+
+#include "CoreMinimal.h"
+#include "ISMPartition/ISMComponentDescriptor.h"
+
+#if WITH_EDITORONLY_DATA
+#include "Containers/SortedMap.h"
+#include "ISMPartition/ISMPartitionClient.h"
+#include "ISMPartition/ISMPartitionInstanceManager.h"
+#include "Instances/InstancedPlacementHash.h"
+#include "Elements/SMInstance/SMInstanceElementId.h"
+#endif //WITH_EDITORONLY_DATA
+
+#include "InstancedPlacementClientInfo.generated.h"
+
+#if WITH_EDITORONLY_DATA
+struct FPropertyChangedEvent;
+class AInstancedPlacementPartitionActor;
+
+enum class EPlacementInstanceFlags : uint32
+{
+	None = 0,
+	AlignToNormal = 1 << 0,
+	NoRandomYaw = 1 << 1,
+	Readjusted = 1 << 2,
+};
+
+/**
+ * Editor information about where and how an instance was placed
+ */
+struct ENGINE_API FPlacementInstance
+{
+public:
+	friend FArchive& operator<<(FArchive& Ar, FPlacementInstance& Instance);
+	FTransform GetInstanceWorldTransform() const;
+	void SetInstanceWorldTransform(const FTransform& Transform);
+	void AlignToNormal(const FVector& InNormal, int32 AlignMaxAngle = 0);
+
+	double ZOffset = 0.0;
+	uint32 Flags = (uint32)EPlacementInstanceFlags::None;
+
+protected:
+	FTransform WorldTransform = FTransform(FQuat::Identity, FVector::ZeroVector, FVector3d::OneVector);
+	FQuat PreAlignRotation = FQuat::Identity;
+};
+#endif //WITH_EDITORONLY_DATA
+
+// Settings which can be shared across partition actors
+UCLASS(Abstract, hideCategories=Object, editinlinenew, collapsecategories)
+class ENGINE_API UInstancedPlacemenClientSettings : public UObject
+{
+	GENERATED_BODY()
+
+public:
+#if WITH_EDITORONLY_DATA
+	UPROPERTY()
+	FGuid UpdateGuid;
+
+	UPROPERTY(EditAnywhere, Category = "Placement")
+	FSoftObjectPath ObjectPath;
+
+	UPROPERTY(EditAnywhere, Category = "Component Settings", meta = (ShowOnlyInnerProperties))
+	FISMComponentDescriptor InstancedComponentSettings;
+#endif //WITH_EDITORONLY_DATA
+
+#if WITH_EDITOR
+	virtual void PostEditChangeProperty(FPropertyChangedEvent& PropertyChangedEvent) override;
+	virtual void RegisterISMDescriptors(AInstancedPlacementPartitionActor* ParentPartitionActor, TSortedMap<int32, TArray<FTransform>>& ISMDefinition) const;
+#endif	// WITH_EDITOR
+};
+
+#if WITH_EDITORONLY_DATA
+// Settings which are unique per partition actor
+struct ENGINE_API FClientPlacementInfo : public IISMPartitionInstanceManager
+{
+	FClientPlacementInfo();
+
+	friend FArchive& operator<<(FArchive& Ar, FClientPlacementInfo& PlacementInfo);
+
+	// Allows detection of client updates while a partition wasn't loaded
+	FGuid UpdateGuid;
+
+	// The client guid we can regenerate the client handle from
+	FGuid ClientGuid;
+
+	// Editor-only placed instances
+	TArray<FPlacementInstance> Instances;
+
+	// The client we map to in the partition actor
+	FISMClientHandle ClientHandle;
+
+	// Definition of the set of ISM descriptors and their local space transforms which make up this client
+	TSortedMap<int32, TArray<FTransform>> ISMDefinition;
+
+	// Transient, editor-only locality hash of instances
+	FInstancedPlacementHash InstanceHash;
+
+	// The display name of this client
+	FString ClientDisplayName;
+
+	TWeakObjectPtr<AInstancedPlacementPartitionActor> ParentPartitionActor;
+
+	TSet<FISMClientInstanceId> MovingInstances;
+
+	using FClientDescriptorFunc = TFunctionRef<void(AInstancedPlacementPartitionActor*, TSortedMap<int32, TArray<FTransform>>&)>;
+	bool Initialize(const FGuid& InClientGuid, const FString& InClientDisplayName, AInstancedPlacementPartitionActor* InParentPartitionActor, FClientDescriptorFunc RegisterDescriptorFunc);
+	void Uninitialize();
+	void PostLoad(AInstancedPlacementPartitionActor* InParentPartitionActor);
+	void PostSerialize(FArchive& Ar, AInstancedPlacementPartitionActor* InParentPartitionActor);
+	bool IsInitialized() const;
+	void PreEditUndo();
+	void PostEditUndo();
+	TArray<FSMInstanceId> AddInstances(TArrayView<const FPlacementInstance> InWorldTransforms);
+
+private:
+	void RemoveInstancesFromPartitionActor(TArrayView<const FISMClientInstanceId> InstanceIds, bool bUpdateHISMTrees, TFunctionRef<void(int32)> RemoveFn);
+
+	//~ IISMPartitionInstanceManager interface
+	virtual FText GetISMPartitionInstanceDisplayName(const FISMClientInstanceId& InstanceId) const override;
+	virtual FText GetISMPartitionInstanceTooltip(const FISMClientInstanceId& InstanceId) const override;
+	virtual bool CanEditISMPartitionInstance(const FISMClientInstanceId& InstanceId) const override;
+	virtual bool CanMoveISMPartitionInstance(const FISMClientInstanceId& InstanceId, const ETypedElementWorldType InWorldType) const override;
+	virtual bool GetISMPartitionInstanceTransform(const FISMClientInstanceId& InstanceId, FTransform& OutInstanceTransform, bool bWorldSpace = false) const override;
+	virtual bool SetISMPartitionInstanceTransform(const FISMClientInstanceId& InstanceId, const FTransform& InstanceTransform, bool bWorldSpace = false, bool bTeleport = false) override;
+	virtual void NotifyISMPartitionInstanceMovementStarted(const FISMClientInstanceId& InstanceId) override;
+	virtual void NotifyISMPartitionInstanceMovementOngoing(const FISMClientInstanceId& InstanceId) override;
+	virtual void NotifyISMPartitionInstanceMovementEnded(const FISMClientInstanceId& InstanceId) override;
+	virtual void NotifyISMPartitionInstanceSelectionChanged(const FISMClientInstanceId& InstanceId, const bool bIsSelected) override;
+	virtual bool DeleteISMPartitionInstances(TArrayView<const FISMClientInstanceId> InstanceIds) override;
+	virtual bool DuplicateISMPartitionInstances(TArrayView<const FISMClientInstanceId> InstanceIds, TArray<FISMClientInstanceId>& OutNewInstanceIds) override;
+};
+
+#endif	// WITH_EDITORONLY_DATA
\ No newline at end of file
diff --git a/Engine/Source/Runtime/Engine/Public/Instances/InstancedPlacementHash.h b/Engine/Source/Runtime/Engine/Public/Instances/InstancedPlacementHash.h
new file mode 100644
--- /dev/null
+++ b/Engine/Source/Runtime/Engine/Public/Instances/InstancedPlacementHash.h
@@ -0,0 +1,159 @@
+// Copyright Epic Games, Inc. All Rights Reserved.
+
+#pragma once
+
+#include "CoreMinimal.h"
+
+#if WITH_EDITORONLY_DATA
+#include "Misc/HashBuilder.h"
+
+/**
+ * FPartitionInstanceHash
+ * Used for editor locality of placed instances in a partition actor to control density and avoid overlapping instances
+ */
+struct FInstancedPlacementHash
+{
+private:
+	static constexpr int32 MaxHashCellBits = 9;	// 512x512x512 grid
+	struct FKey
+	{
+		int64 X;
+		int64 Y;
+		int64 Z;
+
+		FKey() = default;
+
+		FKey(int64 InX, int64 InY, int64 InZ)
+			: X(InX), Y(InY), Z(InZ) {}
+
+		bool operator==(const FKey& Other) const
+		{
+			return (X == Other.X) && (Y == Other.Y) && (Z == Other.Z);
+		}
+
+		friend uint32 GetTypeHash(const FKey& Key)
+		{
+			FHashBuilder HashBuilder;
+			HashBuilder << Key.X << Key.Y << Key.Z;
+			return HashBuilder.GetHash();
+		}
+
+		friend FArchive& operator<<(FArchive& Ar, FKey& Key)
+		{
+			Ar << Key.X;
+			Ar << Key.Y;
+			Ar << Key.Z;
+			return Ar;
+		}
+	};
+
+	const int32 HashCellBits;
+	TMap<FKey, TSet<int32>> CellMap;
+
+	FKey MakeKey(const FVector& Location) const
+	{
+		return FKey(FMath::FloorToInt64(Location.X) >> HashCellBits, FMath::FloorToInt64(Location.Y) >> HashCellBits, FMath::FloorToInt64(Location.Z) >> HashCellBits);
+	}
+
+	FVector MakeLocation(FKey CellKey) const
+	{
+		return FVector(CellKey.X << HashCellBits, CellKey.Y << HashCellBits, CellKey.Z << HashCellBits);
+	}
+
+public:
+	FInstancedPlacementHash(int32 InHashCellBits = MaxHashCellBits)
+		: HashCellBits(InHashCellBits)
+	{}
+
+	~FInstancedPlacementHash() = default;
+
+	void InsertInstance(const FVector& InstanceLocation, int32 InstanceIndex)
+	{
+		FKey Key = MakeKey(InstanceLocation);
+
+		CellMap.FindOrAdd(Key).Add(InstanceIndex);
+	}
+
+	void RemoveInstance(const FVector& InstanceLocation, int32 InstanceIndex, bool bChecked = true)
+	{
+		FKey Key = MakeKey(InstanceLocation);
+
+		if (bChecked)
+		{
+			int32 RemoveCount = CellMap.FindChecked(Key).Remove(InstanceIndex);
+			check(RemoveCount == 1);
+		}
+		else if (TSet<int32>* Value = CellMap.Find(Key))
+		{
+			Value->Remove(InstanceIndex);
+		}
+	}
+
+	void SwapInstance(const FVector& InstanceLocation, int32 OldIndex, int32 NewIndex, bool bCheckedRemoval = true)
+	{
+		RemoveInstance(InstanceLocation, OldIndex, bCheckedRemoval);
+		InsertInstance(InstanceLocation, NewIndex);
+	}
+
+	void GetInstancesOverlappingBox(const FBox& InBox, TArray<int32>& OutInstanceIndices) const
+	{
+		FKey MinKey = MakeKey(InBox.Min);
+		FKey MaxKey = MakeKey(InBox.Max);
+
+		for (int32 z = MinKey.Z; z <= MaxKey.Z; ++z)
+		{
+			for (int32 y = MinKey.Y; y <= MaxKey.Y; y++)
+			{
+				for (int32 x = MinKey.X; x <= MaxKey.X; x++)
+				{
+					auto* SetPtr = CellMap.Find(FKey(x, y, z));
+					if (SetPtr)
+					{
+						OutInstanceIndices.Append(SetPtr->Array());
+					}
+				}
+			}
+		}
+	}
+
+	TArray<int32> GetInstancesOverlappingBox(const FBox& InBox) const
+	{
+		TArray<int32> Result;
+		GetInstancesOverlappingBox(InBox, Result);
+		return Result;
+	}
+
+	void CheckInstanceCount(int32 InCount) const
+	{
+		int32 HashCount = 0;
+		for (const auto& Pair : CellMap)
+		{
+			HashCount += Pair.Value.Num();
+		}
+
+		check(HashCount == InCount);
+	}
+
+	FBox GetBounds() const
+	{
+		FBox HashBounds(ForceInit);
+		for (const auto& Pair : CellMap)
+		{
+			HashBounds += MakeLocation(Pair.Key);
+		}
+
+		return HashBounds;
+	}
+
+	void Empty()
+	{
+		CellMap.Empty();
+	}
+
+	friend FArchive& operator<<(FArchive& Ar, FInstancedPlacementHash& Hash)
+	{
+		Ar << Hash.CellMap;
+		return Ar;
+	}
+};
+#endif
diff --git a/Engine/Source/Runtime/Engine/Public/Instances/InstancedPlacementPartitionActor.h b/Engine/Source/Runtime/Engine/Public/Instances/InstancedPlacementPartitionActor.h
--- a/Engine/Source/Runtime/Engine/Public/Instances/InstancedPlacementPartitionActor.h
+++ b/Engine/Source/Runtime/Engine/Public/Instances/InstancedPlacementPartitionActor.h
@@ -4,9 +4,12 @@
 
 #include "CoreMinimal.h"
 #include "ISMPartition/ISMPartitionActor.h"
-#include "Components/InstancedStaticMeshComponent.h"
+#include "Instances/InstancedPlacementClientInfo.h"
+
 #include "InstancedPlacementPartitionActor.generated.h"
 
+class UInstancedPlacementClientSettings;
+
 /**
  * The base class used by any editor placement of instanced objects, which holds any relevant runtime data for the placed instances.
  */
@@ -16,17 +19,29 @@ class ENGINE_API AInstancedPlacementPartitionActor : public AISMPartitionActor
 	GENERATED_UCLASS_BODY()
 
 public:
+	virtual void Serialize(FArchive& Ar) override;
+	virtual void PostLoad() override;
+
 #if WITH_EDITOR	
+	virtual void PreEditUndo() override;
+	virtual void PostEditUndo() override;
 	virtual uint32 GetDefaultGridSize(UWorld* InWorld) const override;
 	virtual FGuid GetGridGuid() const override;
 
 	void SetGridGuid(const FGuid& InGuid);
-#endif
 
-	virtual ISMInstanceManager* GetSMInstanceManager(const FSMInstanceId& InstanceId) override;
+	using FClientDescriptorFunc = FClientPlacementInfo::FClientDescriptorFunc;
+	FClientPlacementInfo* PreAddClientInstances(const FGuid& ClientGuid, const FString& InClientDisplayString, FClientDescriptorFunc RegisterDefinitionFunc);
+	void PostAddClientInstances();
+	void NotifySettingsObjectChanged(UInstancedPlacementClientSettings* InSettingsObject);
+#endif
 
 protected:
 #if WITH_EDITORONLY_DATA
+	UPROPERTY()
 	FGuid PlacementGridGuid;
+
+	// Placed Client info by corresponding client Guid
+	TMap<FGuid, TUniqueObj<FClientPlacementInfo>> PlacedClientInfo;
 #endif
 };
diff --git a/Engine/Source/Runtime/Foliage/Public/InstancedFoliage.h b/Engine/Source/Runtime/Foliage/Public/InstancedFoliage.h
--- a/Engine/Source/Runtime/Foliage/Public/InstancedFoliage.h
+++ b/Engine/Source/Runtime/Foliage/Public/InstancedFoliage.h
@@ -7,9 +7,9 @@ InstancedFoliage.h: Instanced foliage type definitions.
 
 #include "CoreMinimal.h"
 #include "Misc/Guid.h"
-#include "Misc/HashBuilder.h"
 #include "Containers/ArrayView.h"
 #include "FoliageInstanceBase.h"
+#include "Instances/InstancedPlacementHash.h"
 
 class AInstancedFoliageActor;
 class UActorComponent;
@@ -20,7 +20,10 @@ class UFoliageType_InstancedStaticMesh;
 class UPrimitiveComponent;
 class UStaticMesh;
 struct FSMInstanceId;
-struct FFoliageInstanceHash;
+
+#if WITH_EDITORONLY_DATA
+using FFoliageInstanceHash = FInstancedPlacementHash;
+#endif
 
 DECLARE_LOG_CATEGORY_EXTERN(LogInstancedFoliage, Log, All);
 
@@ -442,151 +445,6 @@ struct FFoliageInstanceId
 	int32 Index = INDEX_NONE;
 };
 
-#if WITH_EDITORONLY_DATA
-//
-// FFoliageInstanceHash
-//
-
-#define FOLIAGE_HASH_CELL_BITS 9	// 512x512x512 grid
-
-struct FFoliageInstanceHash
-{
-private:
-	struct FKey
-	{
-		int64 X;
-		int64 Y;
-		int64 Z;
-
-		FKey() = default;
-
-		FKey(int64 InX, int64 InY, int64 InZ)
-			: X(InX), Y(InY), Z(InZ) {}
-
-		bool operator==(const FKey& Other) const
-		{
-			return (X == Other.X) && (Y == Other.Y) && (Z == Other.Z);
-		}
-
-		friend uint32 GetTypeHash(const FKey& Key)
-		{
-			FHashBuilder HashBuilder;
-			HashBuilder << Key.X << Key.Y << Key.Z;
-			return HashBuilder.GetHash();
-		}
-				
-		friend FArchive& operator<<(FArchive& Ar, FKey& Key)
-		{
-			Ar << Key.X;
-			Ar << Key.Y;
-			Ar << Key.Z;
-			return Ar;
-		}
-	};
-
-	const int32 HashCellBits;
-	TMap<FKey, TSet<int32>> CellMap;
-		
-	FKey MakeKey(const FVector& Location) const
-	{
-		return FKey(FMath::FloorToInt64(Location.X) >> HashCellBits, FMath::FloorToInt64(Location.Y) >> HashCellBits, FMath::FloorToInt64(Location.Z) >> HashCellBits);
-	}
-
-	FVector MakeLocation(FKey CellKey) const
-	{
-		return FVector(CellKey.X << HashCellBits, CellKey.Y << HashCellBits, CellKey.Z << HashCellBits);
-	}
-
-public:
-	FFoliageInstanceHash(int32 InHashCellBits = FOLIAGE_HASH_CELL_BITS)
-		: HashCellBits(InHashCellBits)
-	{}
-
-	void InsertInstance(const FVector& InstanceLocation, int32 InstanceIndex)
-	{
-		FKey Key = MakeKey(InstanceLocation);
-
-		CellMap.FindOrAdd(Key).Add(InstanceIndex);
-	}
-
-	void RemoveInstance(const FVector& InstanceLocation, int32 InstanceIndex, bool bChecked = true)
-	{
-		FKey Key = MakeKey(InstanceLocation);
-
-		if (bChecked)
-		{
-			int32 RemoveCount = CellMap.FindChecked(Key).Remove(InstanceIndex);
-			check(RemoveCount == 1);
-		}
-		else if(TSet<int32>* Value = CellMap.Find(Key))
-		{
-			Value->Remove(InstanceIndex);
-		}
-	}
-
-	void GetInstancesOverlappingBox(const FBox& InBox, TArray<int32>& OutInstanceIndices) const
-	{
-		FKey MinKey = MakeKey(InBox.Min);
-		FKey MaxKey = MakeKey(InBox.Max);
-
-		for (int32 z = MinKey.Z; z <= MaxKey.Z; ++z)
-		{
-			for (int32 y = MinKey.Y; y <= MaxKey.Y; y++)
-			{
-				for (int32 x = MinKey.X; x <= MaxKey.X; x++)
-				{
-					auto* SetPtr = CellMap.Find(FKey(x, y, z));
-					if (SetPtr)
-					{
-						OutInstanceIndices.Append(SetPtr->Array());
-					}
-				}
-			}
-		}
-	}
-
-	TArray<int32> GetInstancesOverlappingBox(const FBox& InBox) const
-	{
-		TArray<int32> Result;
-		GetInstancesOverlappingBox(InBox, Result);
-		return Result;
-	}
-
-	void CheckInstanceCount(int32 InCount) const
-	{
-		int32 HashCount = 0;
-		for (const auto& Pair : CellMap)
-		{
-			HashCount += Pair.Value.Num();
-		}
-
-		check(HashCount == InCount);
-	}
-
-	FBox GetBounds() const
-	{
-		FBox HashBounds(ForceInit);
-		for (const auto& Pair : CellMap)
-		{
-			HashBounds += MakeLocation(Pair.Key);
-		}
-
-		return HashBounds;
-	}
-
-	void Empty()
-	{
-		CellMap.Empty();
-	}
-
-	friend FArchive& operator<<(FArchive& Ar, FFoliageInstanceHash& Hash)
-	{
-		Ar << Hash.CellMap;
-		return Ar;
-	}
-};
-#endif
-
 /** This is kind of a hack, but is needed right now for backwards compat of code. We use it to describe the placement mode (procedural vs manual)*/
 namespace EFoliagePlacementMode
 {
