Add primitive drawing utilities for torus, rectangle, and rectangle mesh.

diff --git a/Engine/Source/Runtime/Engine/Private/PrimitiveDrawingUtils.cpp b/Engine/Source/Runtime/Engine/Private/PrimitiveDrawingUtils.cpp
--- a/Engine/Source/Runtime/Engine/Private/PrimitiveDrawingUtils.cpp
+++ b/Engine/Source/Runtime/Engine/Private/PrimitiveDrawingUtils.cpp
@@ -734,6 +734,113 @@ void DrawCylinder(class FPrimitiveDrawInterface* PDI, const FVector& Start, cons
 
 }
 
+void DrawTorus(FPrimitiveDrawInterface* PDI, const FMatrix& Transform, const FVector& XAxis, const FVector& YAxis, 
+			   float OuterRadius, float InnerRadius, int32 OuterSegments, int32 InnerSegments, const FMaterialRenderProxy* MaterialRenderProxy, 
+			   uint8 DepthPriority, bool bPartial, float Angle, bool bEndCaps)
+{
+	if (OuterSegments < 3)
+	{
+		OuterSegments = 3;
+	}
+
+	if (InnerSegments < 3)
+	{
+		InnerSegments = 3;
+	}
+
+	const float	AngleDelta = (bPartial ? Angle : 2.0f * PI) / OuterSegments;
+	FVector	LastVertex = XAxis * OuterRadius;
+
+	int32 OuterSlices = bPartial ? (bEndCaps ? OuterSegments + 3 : OuterSegments + 1) : OuterSegments;
+	const float InnerAngleDelta = 2.0f * PI / InnerSegments;
+
+	FVector ZAxis = XAxis ^ YAxis;
+	ZAxis.Normalize();
+
+	FVector2D TC = FVector2D(0.0f, 0.0f);
+	float TCStepX = 1.0f / InnerSegments;
+	float TCStepY = 1.0f / OuterSegments;
+
+	FDynamicMeshBuilder MeshBuilder(PDI->View->GetFeatureLevel());
+
+	for (int32 OuterIndex = 0; OuterIndex < OuterSlices; OuterIndex++)
+	{
+		float OuterAngle = AngleDelta * OuterIndex;
+
+		FVector UnitDir = (XAxis * FMath::Cos(OuterAngle) + YAxis * FMath::Sin(OuterAngle));
+		FVector InnerCenter = UnitDir * OuterRadius;
+
+		TC.X = 0.0f;
+		TC.Y = TCStepY * OuterIndex;
+
+		bool bIsEndCapVertex = (bPartial && bEndCaps && (OuterIndex == 0 || OuterIndex == OuterSlices - 1));
+
+		for (int32 InnerIndex = 0; InnerIndex < InnerSegments; InnerIndex++)
+		{
+			float InnerAngle = InnerAngleDelta * (InnerIndex + 1);
+
+			FVector Vertex = InnerCenter;
+			FVector Dir = UnitDir * FMath::Cos(InnerAngle) + ZAxis * FMath::Sin(InnerAngle);
+
+			if (!bIsEndCapVertex)
+			{
+				Vertex += Dir * InnerRadius;
+			}
+
+			FDynamicMeshVertex MeshVertex;
+			MeshVertex.Position = Vertex;
+			MeshVertex.TextureCoordinate[0] = TC;
+			MeshVertex.TextureCoordinate[0].X += TCStepX * InnerIndex;
+
+			FVector TangentX = -UnitDir * FMath::Sin(InnerAngle) + ZAxis * FMath::Cos(InnerAngle);
+			FVector TangentY = Dir ^ TangentX;
+			FVector TangentZ = Dir;
+			MeshVertex.SetTangents(TangentX, TangentY, TangentZ);
+
+			MeshBuilder.AddVertex(MeshVertex); //Add bottom vertex
+		}
+	}
+
+	// Add sides.
+	int32 NumVertices = OuterSlices * InnerSegments;
+	int32 LoopSegments = bPartial ? OuterSlices - 1 : OuterSlices;
+
+	for (int32 OuterIndex = 0; OuterIndex < LoopSegments; OuterIndex++)
+	{
+		int32 BaseVertexIndex = OuterIndex * InnerSegments;
+
+		bool bIsEndCap = (bPartial && bEndCaps && (OuterIndex == 0 || OuterIndex == OuterSlices - 1));
+
+		for (int32 InnerIndex = 0; InnerIndex < InnerSegments; InnerIndex++)
+		{
+			int32 V0 = BaseVertexIndex + InnerIndex;
+			int32 V1 = BaseVertexIndex + ((InnerIndex + 1) % InnerSegments);
+			int32 V2 = (V0 + InnerSegments) % NumVertices;
+			int32 V3 = (V1 + InnerSegments) % NumVertices;
+
+			if (bIsEndCap)
+			{
+				if (OuterIndex == 0)
+				{
+					MeshBuilder.AddTriangle(V0, V2, V3);
+				}
+				else
+				{
+					MeshBuilder.AddTriangle(V0, V1, V2);
+				}
+			}
+			else
+			{
+				MeshBuilder.AddTriangle(V0, V2, V1);
+				MeshBuilder.AddTriangle(V2, V3, V1);
+			}
+		}
+	}
+
+	MeshBuilder.Draw(PDI, Transform, MaterialRenderProxy, DepthPriority, 0.f);
+}
+
+
 void DrawDisc(class FPrimitiveDrawInterface* PDI,const FVector& Base,const FVector& XAxis,const FVector& YAxis,FColor Color,float Radius,int32 NumSides,const FMaterialRenderProxy* MaterialRenderProxy, uint8 DepthPriority)
 {
 	check (NumSides >= 3);
@@ -938,6 +1045,80 @@ void DrawArc(FPrimitiveDrawInterface* PDI, const FVector Base, const FVector X,
 		CurrentAngle += AngleStep;
 	}
 }
+void DrawRectangleMesh(FPrimitiveDrawInterface* PDI, const FVector& Center, const FVector& XAxis, const FVector& YAxis, 
+					   FColor Color, float Width, float Height, const FMaterialRenderProxy* MaterialRenderProxy, uint8 DepthPriority)
+{
+	FVector XOffset = XAxis * Width * 0.5f;
+	FVector YOffset = YAxis * Height * 0.5f;
+
+	// Calculate verts for a face lying in plane defined by the X and Y vectors
+	FVector Positions[4] =
+	{
+		Center - XOffset - YOffset,
+		Center + XOffset - YOffset,
+		Center + XOffset + YOffset,
+		Center - XOffset + YOffset
+	};
+
+	FVector2D UVs[4] =
+	{
+		FVector2D(0.0f,0.0f),
+		FVector2D(0.0f,1.0f),
+		FVector2D(1.0f,1.0f),
+		FVector2D(1.0f,0.0f),
+	};
+
+	FDynamicMeshBuilder MeshBuilder(PDI->View->GetFeatureLevel());
+
+	int32 VertexIndices[4];
+	for (int32 VertexIndex = 0; VertexIndex < 4; VertexIndex++)
+	{
+		VertexIndices[VertexIndex] = MeshBuilder.AddVertex(
+			Positions[VertexIndex],
+			UVs[VertexIndex],
+			FVector(1.0f, 0.0f, 0.0f),
+			FVector(0.0f, 1.0f, 0.0f),
+			FVector(0.0f, 0.0f, 1.0f),
+			Color
+		);
+	}
+
+	MeshBuilder.AddTriangle(VertexIndices[0], VertexIndices[1], VertexIndices[2]);
+	MeshBuilder.AddTriangle(VertexIndices[0], VertexIndices[2], VertexIndices[3]);
+
+	MeshBuilder.AddTriangle(VertexIndices[0], VertexIndices[2], VertexIndices[1]);
+	MeshBuilder.AddTriangle(VertexIndices[0], VertexIndices[3], VertexIndices[2]);
+
+	MeshBuilder.Draw(PDI, FMatrix::Identity, MaterialRenderProxy, DepthPriority, 0.f);
+}
+
+void DrawRectangle(class FPrimitiveDrawInterface* PDI, const FVector& Center, const FVector& XAxis, const FVector& YAxis, FColor Color, float Width, float Height, uint8 DepthPriority, float Thickness, float DepthBias, bool bScreenSpace)
+{
+	FVector XOffset = XAxis * Width * 0.5f;
+	FVector YOffset = YAxis * Height * 0.5f;
+
+	// Calculate verts for a face lying in plane defined by the X and Y vectors
+	FVector Positions[4] =
+	{
+		Center - XOffset - YOffset,
+		Center + XOffset - YOffset,
+		Center + XOffset + YOffset,
+		Center - XOffset + YOffset
+	};
+
+	FVector2D UVs[4] =
+	{
+		FVector2D(0.0f,0.0f),
+		FVector2D(0.0f,1.0f),
+		FVector2D(1.0f,1.0f),
+		FVector2D(1.0f,0.0f),
+	};
+
+	PDI->DrawLine(Positions[0], Positions[1], Color, DepthPriority, Thickness, DepthBias, bScreenSpace);
+	PDI->DrawLine(Positions[1], Positions[2], Color, DepthPriority, Thickness, DepthBias, bScreenSpace);
+	PDI->DrawLine(Positions[2], Positions[3], Color, DepthPriority, Thickness, DepthBias, bScreenSpace);
+	PDI->DrawLine(Positions[3], Positions[0], Color, DepthPriority, Thickness, DepthBias, bScreenSpace);
+}
 
 void DrawWireSphere(class FPrimitiveDrawInterface* PDI, const FVector& Base, const FLinearColor& Color, float Radius, int32 NumSides, uint8 DepthPriority, float Thickness, float DepthBias, bool bScreenSpace)
 {
diff --git a/Engine/Source/Runtime/Engine/Public/SceneManagement.h b/Engine/Source/Runtime/Engine/Public/SceneManagement.h
--- a/Engine/Source/Runtime/Engine/Public/SceneManagement.h
+++ b/Engine/Source/Runtime/Engine/Public/SceneManagement.h
@@ -2650,6 +2650,28 @@ extern ENGINE_API void GetCapsuleMesh(const FVector& Origin, const FVector& XAxi
 									const FMaterialRenderProxy* MaterialRenderProxy, uint8 DepthPriority, bool bDisableBackfaceCulling, int32 ViewIndex, FMeshElementCollector& Collector, HHitProxy* HitProxy = NULL);
 
 
+/**
+ * Draws a torus using triangles.
+ *
+ * @param	PDI						Draw interface.
+ * @param	Transform				Generic transform to apply (ex. a local-to-world transform).
+ * @param	XAxis					Normalized X alignment axis.
+ * @param	YAxis					Normalized Y alignment axis.
+ * @param	Color					Color of the circle.
+ * @param	OuterRadius				Radius of the torus center-line. Viewed from above, the outside of the torus has a radius 
+ *                                  of OuterRadius + InnerRadius and the hole of the torus has a radius of OuterRadius - InnerRadius.
+ * @param	InnerRadius				Radius of the torus's cylinder.
+ * @param	OuterSegments			Numbers of segment divisions for outer circle.
+ * @param	InnerSegments			Numbers of segment divisions for inner circle.
+ * @param	MaterialRenderProxy		Material to use for render
+ * @param	DepthPriority			Depth priority for the circle.
+ * @param	bPartial				Whether full or partial torus should be rendered.
+ * @param	Angle					If partial, angle in radians of the arc clockwise beginning at the XAxis.
+ * @param	bEndCaps				If partial, whether the ends should be capped with triangles.
+ */
+extern ENGINE_API void DrawTorus(FPrimitiveDrawInterface* PDI, const FMatrix& Transform, const FVector& XAxis, const FVector& YAxis, 
+								 float OuterRadius, float InnerRadius, int32 OuterSegments, int32 InnerSegments, const FMaterialRenderProxy* MaterialRenderProxy, uint8 DepthPriority, bool bPartial, float Angle, bool bEndCaps);
+
 /**
  * Draws a circle using triangles.
  *
@@ -2665,6 +2687,21 @@ extern ENGINE_API void GetCapsuleMesh(const FVector& Origin, const FVector& XAxi
  */
 extern ENGINE_API void DrawDisc(class FPrimitiveDrawInterface* PDI,const FVector& Base,const FVector& XAxis,const FVector& YAxis,FColor Color,float Radius,int32 NumSides, const FMaterialRenderProxy* MaterialRenderProxy, uint8 DepthPriority);
 
+/**
+ * Draws a rectangle using triangles.
+ *
+ * @param	PDI						Draw interface.
+ * @param	Center					Center of the rectangle.
+ * @param	XAxis					Normalized X alignment axis.
+ * @param	YAxis					Normalized Y alignment axis.
+ * @param	Color					Color of the circle.
+ * @param	Width					Width of rectangle along the X dimension.
+ * @param	Height					Height of rectangle along the Y dimension.
+ * @param	MaterialRenderProxy		Material to use for render
+ * @param	DepthPriority			Depth priority for the rectangle.
+ */
+extern ENGINE_API void DrawRectangleMesh(FPrimitiveDrawInterface* PDI, const FVector& Center, const FVector& XAxis, const FVector& YAxis, 
+										 FColor Color, float Width, float Height, const FMaterialRenderProxy* MaterialRenderProxy, uint8 DepthPriority);
 
 /**
  * Draws a flat arrow with an outline.
@@ -2741,6 +2778,23 @@ extern ENGINE_API void DrawCircle(class FPrimitiveDrawInterface* PDI, const FVec
  */
 extern ENGINE_API void DrawArc(FPrimitiveDrawInterface* PDI, const FVector Base, const FVector X, const FVector Y, const float MinAngle, const float MaxAngle, const float Radius, const int32 Sections, const FLinearColor& Color, uint8 DepthPriority);
 
+/**
+ * Draws a rectangle using lines.
+ *
+ * @param	PDI						Draw interface.
+ * @param	Center					Center of the rectangle.
+ * @param	XAxis					Normalized X alignment axis.
+ * @param	YAxis					Normalized Y alignment axis.
+ * @param	Color					Color of the circle.
+ * @param	Width					Width of rectangle along the X dimension.
+ * @param	Height					Height of rectangle along the Y dimension.
+ * @param	MaterialRenderProxy		Material to use for render
+ * @param	DepthPriority			Depth priority for the rectangle.
+ * @param	Thickness				Thickness of the lines comprising the rectangle.
+ */
+extern ENGINE_API void DrawRectangle(FPrimitiveDrawInterface* PDI, const FVector& Center, const FVector& XAxis, const FVector& YAxis, 
+									 FColor Color, float Width, float Height, uint8 DepthPriority, float Thickness = 0.0f, float DepthBias = 0.0f, bool bScreenSpace = false);
+
 /**
  * Draws a sphere using circles.
  *
