Fix cook crash, missing a pointer check

diff --git a/Engine/Source/Runtime/Engine/Private/SkeletalMeshRenderData.cpp b/Engine/Source/Runtime/Engine/Private/SkeletalMeshRenderData.cpp
--- a/Engine/Source/Runtime/Engine/Private/SkeletalMeshRenderData.cpp
+++ b/Engine/Source/Runtime/Engine/Private/SkeletalMeshRenderData.cpp
@@ -410,7 +410,10 @@ void FSkeletalMeshRenderData::Cache(const ITargetPlatform* TargetPlatform, USkel
 			}
 
 			//Apply the morphtargets change if any
-			FinishBuildMorphTargetData->ApplyEditorData(Owner);
+			if (FinishBuildMorphTargetData.IsValid())
+			{
+				FinishBuildMorphTargetData->ApplyEditorData(Owner);
+			}
 
 			int32 T1 = FPlatformTime::Cycles();
 			UE_LOG(LogSkeletalMesh, Verbose, TEXT("Skeletal Mesh found in DDC [%fms] %s"), FPlatformTime::ToMilliseconds(T1 - T0), *Owner->GetPathName());
