* Add support for property names to be serialized instead of property handles in certain situations in RepLayout serialization. This can be activated with the FReplicationFlags::bSerializePropertyNames flag.
* Wrapped the call of ReplicateSingleActor for all the actors for a given connection in its own function so that it can be overridden.
* FObjectReplicator::ReplicateProperties now has an override for passing in a custom FNetBitWriter implementation.
* Added a replication flag for custom subobject serialization on UActorChannel for a class that derives from it.

diff --git a/Engine/Plugins/Runtime/ReplicationGraph/Source/Private/ReplicationGraph.cpp b/Engine/Plugins/Runtime/ReplicationGraph/Source/Private/ReplicationGraph.cpp
--- a/Engine/Plugins/Runtime/ReplicationGraph/Source/Private/ReplicationGraph.cpp
+++ b/Engine/Plugins/Runtime/ReplicationGraph/Source/Private/ReplicationGraph.cpp
@@ -1399,46 +1399,13 @@ void UReplicationGraph::ReplicateActorListsForConnections_Default(UNetReplicatio
 			SortingArray->Sort();
 		}
 	}
-
-	// --------------------------------------------------------------------------------------------------------------
-	// REPLICATE Actors For Connection
-	// --------------------------------------------------------------------------------------------------------------
+	
 	{
 		QUICK_SCOPE_CYCLE_COUNTER(NET_ReplicateActors_ReplicateActorsForConnection);
-
-		for (int32 ActorIdx = 0; ActorIdx < PrioritizedReplicationList.Items.Num(); ++ActorIdx)
-		{
-			const FPrioritizedRepList::FItem& RepItem = PrioritizedReplicationList.Items[ActorIdx];
-
-			AActor* Actor = RepItem.Actor;
-			FConnectionReplicationActorInfo& ActorInfo = *RepItem.ConnectionData;
-
-			// Always skip if we've already replicated this frame. This happens if an actor is in more than one replication list
-			if (ActorInfo.LastRepFrameNum == FrameNum)
-			{
-				INC_DWORD_STAT_BY(STAT_NetRepActorListDupes, 1);
-				continue;
-			}
-
-			FGlobalActorReplicationInfo& GlobalActorInfo = *RepItem.GlobalData;
-
-			int64 BitsWritten = ReplicateSingleActor(Actor, ActorInfo, GlobalActorInfo, ConnectionActorInfoMap, *ConnectionManager, FrameNum);
-
-			// --------------------------------------------------
-			//	Update Packet Budget Tracking
-			// --------------------------------------------------
-
-			if (IsConnectionReady(NetConnection) == false)
-			{
-				// We've exceeded the budget for this category of replication list.
-				RG_QUICK_SCOPE_CYCLE_COUNTER(NET_ReplicateActors_PartialStarvedActorList);
-				HandleStarvedActorList(PrioritizedReplicationList, ActorIdx + 1, ConnectionActorInfoMap, FrameNum);
-				NotifyConnectionSaturated(*ConnectionManager);
-				break;
-			}
-		}
+		ReplicateActorsForConnection(NetConnection, ConnectionActorInfoMap, ConnectionManager, FrameNum);
 	}
 
+
 	// Broadcast the list we just handled. This is intended to be for debugging/logging features.
 	ConnectionManager->OnPostReplicatePrioritizeLists.Broadcast(ConnectionManager, &PrioritizedReplicationList);
 
@@ -1454,6 +1421,43 @@ void UReplicationGraph::ReplicateActorListsForConnections_Default(UNetReplicatio
 	}
 }
 
+void UReplicationGraph::ReplicateActorsForConnection(UNetConnection* NetConnection, FPerConnectionActorInfoMap& ConnectionActorInfoMap, UNetReplicationGraphConnection* ConnectionManager, const uint32 FrameNum)
+{
+	QUICK_SCOPE_CYCLE_COUNTER(NET_ReplicateActors_ReplicateActorsForConnection);
+
+	for (int32 ActorIdx = 0; ActorIdx < PrioritizedReplicationList.Items.Num(); ++ActorIdx)
+	{
+		const FPrioritizedRepList::FItem& RepItem = PrioritizedReplicationList.Items[ActorIdx];
+
+		AActor* Actor = RepItem.Actor;
+		FConnectionReplicationActorInfo& ActorInfo = *RepItem.ConnectionData;
+
+		// Always skip if we've already replicated this frame. This happens if an actor is in more than one replication list
+		if (ActorInfo.LastRepFrameNum == FrameNum)
+		{
+			INC_DWORD_STAT_BY(STAT_NetRepActorListDupes, 1);
+			continue;
+		}
+
+		FGlobalActorReplicationInfo& GlobalActorInfo = *RepItem.GlobalData;
+
+		int64 BitsWritten = ReplicateSingleActor(Actor, ActorInfo, GlobalActorInfo, ConnectionActorInfoMap, *ConnectionManager, FrameNum);
+
+		// --------------------------------------------------
+		//	Update Packet Budget Tracking
+		// --------------------------------------------------
+
+		if (IsConnectionReady(NetConnection) == false)
+		{
+			// We've exceeded the budget for this category of replication list.
+			RG_QUICK_SCOPE_CYCLE_COUNTER(NET_ReplicateActors_PartialStarvedActorList);
+			HandleStarvedActorList(PrioritizedReplicationList, ActorIdx + 1, ConnectionActorInfoMap, FrameNum);
+			NotifyConnectionSaturated(*ConnectionManager);
+			break;
+		}
+	}
+}
+
 struct FScopedQueuedBits
 {
 	FScopedQueuedBits(int32& InQueuedBits, int32& InTotalBits) : QueuedBits(InQueuedBits), TotalBits(InTotalBits) { }
diff --git a/Engine/Plugins/Runtime/ReplicationGraph/Source/Public/ReplicationGraph.h b/Engine/Plugins/Runtime/ReplicationGraph/Source/Public/ReplicationGraph.h
--- a/Engine/Plugins/Runtime/ReplicationGraph/Source/Public/ReplicationGraph.h
+++ b/Engine/Plugins/Runtime/ReplicationGraph/Source/Public/ReplicationGraph.h
@@ -995,6 +995,8 @@ public:
 
 	// --------------------------------------------------------------
 
+	virtual void ReplicateActorsForConnection(UNetConnection* NetConnection, FPerConnectionActorInfoMap& ConnectionActorInfoMap, UNetReplicationGraphConnection* ConnectionManager, const uint32 FrameNum);
+
 	int64 ReplicateSingleActor(AActor* Actor, FConnectionReplicationActorInfo& ActorInfo, FGlobalActorReplicationInfo& GlobalActorInfo, FPerConnectionActorInfoMap& ConnectionActorInfoMap, UNetReplicationGraphConnection& ConnectionManager, const uint32 FrameNum);
 	int64 ReplicateSingleActor_FastShared(AActor* Actor, FConnectionReplicationActorInfo& ConnectionData, FGlobalActorReplicationInfo& GlobalActorInfo, UNetReplicationGraphConnection& ConnectionManager, const uint32 FrameNum);
 
diff --git a/Engine/Source/Runtime/Engine/Classes/Engine/ActorChannel.h b/Engine/Source/Runtime/Engine/Classes/Engine/ActorChannel.h
--- a/Engine/Source/Runtime/Engine/Classes/Engine/ActorChannel.h
+++ b/Engine/Source/Runtime/Engine/Classes/Engine/ActorChannel.h
@@ -267,6 +267,9 @@ public:
 
 	/** Replicates given subobject on this actor channel */
 	bool ReplicateSubobject(UObject *Obj, FOutBunch &Bunch, const FReplicationFlags &RepFlags);
+	
+	/** Custom implementation for ReplicateSubobject when RepFlags.bUseCustomSubobjectReplication is true */
+	virtual bool ReplicateSubobjectCustom(UObject* Obj, FOutBunch& Bunch, const FReplicationFlags& RepFlags) { return true;  }
 
 	/** utility template for replicating list of replicated subobjects */
 	template<typename Type>
diff --git a/Engine/Source/Runtime/Engine/Classes/Engine/EngineTypes.h b/Engine/Source/Runtime/Engine/Classes/Engine/EngineTypes.h
--- a/Engine/Source/Runtime/Engine/Classes/Engine/EngineTypes.h
+++ b/Engine/Source/Runtime/Engine/Classes/Engine/EngineTypes.h
@@ -3896,6 +3896,11 @@ struct FReplicationFlags
 			uint32 bRolesOnly:1;
 			/** True if we should force all properties dirty on initial replication. */
 			uint32 bForceInitialDirty:1;
+			/** True if we should serialize property names instead of handles. */
+			uint32 bSerializePropertyNames : 1;
+			/** True if a subclass of UActorChannel needs custom subobject replication */
+			uint32 bUseCustomSubobjectReplication : 1;
+
 		};
 
 		uint32	Value;
diff --git a/Engine/Source/Runtime/Engine/Private/DataChannel.cpp b/Engine/Source/Runtime/Engine/Private/DataChannel.cpp
--- a/Engine/Source/Runtime/Engine/Private/DataChannel.cpp
+++ b/Engine/Source/Runtime/Engine/Private/DataChannel.cpp
@@ -4361,6 +4361,11 @@ bool UActorChannel::ReplicateSubobject(UObject *Obj, FOutBunch &Bunch, const FRe
 		return false;
 	}
 
+	if (RepFlags.bUseCustomSubobjectReplication)
+	{
+		return ReplicateSubobjectCustom(Obj, Bunch, RepFlags);
+	}
+
 	bool bFoundInvalidReplicator = false;
 	TSharedRef<FObjectReplicator>* FoundReplicator = FindReplicator(Obj, &bFoundInvalidReplicator);
 	const bool bFoundReplicator = (FoundReplicator != nullptr);
diff --git a/Engine/Source/Runtime/Engine/Private/DataReplication.cpp b/Engine/Source/Runtime/Engine/Private/DataReplication.cpp
--- a/Engine/Source/Runtime/Engine/Private/DataReplication.cpp
+++ b/Engine/Source/Runtime/Engine/Private/DataReplication.cpp
@@ -1716,8 +1716,19 @@ bool FObjectReplicator::CanSkipUpdate(FReplicationFlags RepFlags)
 	return false;
 }
 
+bool FObjectReplicator::ReplicateProperties(FOutBunch& Bunch, FReplicationFlags RepFlags, FNetBitWriter& Writer)
+{
+	return ReplicateProperties_r(Bunch, RepFlags, Writer);
+}
+
+bool FObjectReplicator::ReplicateProperties(FOutBunch& Bunch, FReplicationFlags RepFlags)
+{
+	FNetBitWriter Writer(Bunch.PackageMap, 8192);
+	return ReplicateProperties_r(Bunch, RepFlags, Writer);
+}
+
 /** Replicates properties to the Bunch. Returns true if it wrote anything */
-bool FObjectReplicator::ReplicateProperties( FOutBunch & Bunch, FReplicationFlags RepFlags )
+bool FObjectReplicator::ReplicateProperties_r( FOutBunch & Bunch, FReplicationFlags RepFlags, FNetBitWriter& Writer)
 {
 	UObject* Object = GetObject();
 
@@ -1741,8 +1752,6 @@ bool FObjectReplicator::ReplicateProperties( FOutBunch & Bunch, FReplicationFlag
 
 	UNetConnection* OwningChannelConnection = OwningChannel->Connection;
 
-	FNetBitWriter Writer( Bunch.PackageMap, 8192 );
-
 #if UE_NET_TRACE_ENABLED
 	// Create trace collector if tracing is enabled for the target bunch
 	SetTraceCollector(Writer, GetTraceCollector(Bunch) ? UE_NET_TRACE_CREATE_COLLECTOR(ENetTraceVerbosity::Trace) : nullptr);
diff --git a/Engine/Source/Runtime/Engine/Private/RepLayout.cpp b/Engine/Source/Runtime/Engine/Private/RepLayout.cpp
--- a/Engine/Source/Runtime/Engine/Private/RepLayout.cpp
+++ b/Engine/Source/Runtime/Engine/Private/RepLayout.cpp
@@ -1820,6 +1820,32 @@ static FORCEINLINE void WritePropertyHandle(
 	NETWORK_PROFILER(GNetworkProfiler.TrackWritePropertyHandle(Writer.GetNumBits() - NumStartingBits, nullptr));
 }
 
+static void WritePropertyName(
+	FNetBitWriter& Writer,
+	const FName& PropertyName,
+	bool bDoChecksum)
+{
+	const int NumStartingBits = Writer.GetNumBits();
+
+	UE_NET_TRACE_SCOPE(PropertyName, Writer, GetTraceCollector(Writer), ENetTraceVerbosity::Trace);
+
+	FName LocalPropertyName = PropertyName;
+	Writer << LocalPropertyName;
+
+	UE_LOG(LogRepProperties, VeryVerbose, TEXT("WritePropertyName: Name=%s"), *PropertyName.ToString());
+
+#ifdef ENABLE_PROPERTY_CHECKSUMS
+	if (bDoChecksum)
+	{
+		SerializeGenericChecksum(Writer);
+	}
+#endif
+
+	// TODO: Write a separate network profiler function for tracking that the name is being written instead
+	// of the handle.
+	NETWORK_PROFILER(GNetworkProfiler.TrackWritePropertyHandle(Writer.GetNumBits() - NumStartingBits, nullptr));
+}
+
 bool FRepLayout::ReplicateProperties(
 	FSendingRepState* RESTRICT RepState,
 	FRepChangelistState* RESTRICT RepChangelistState,
@@ -2034,7 +2060,7 @@ bool FRepLayout::ReplicateProperties(
 	}
 	else if (Changed.Num() > 0)
 	{
-		SendProperties(RepState, ChangeTracker, Data, ObjectClass, Writer, Changed, RepChangelistState->SharedSerialization);
+		SendProperties(RepState, ChangeTracker, Data, ObjectClass, Writer, Changed, RepChangelistState->SharedSerialization, RepFlags.bSerializePropertyNames ? ESerializePropertyType::Name : ESerializePropertyType::Handle);
 	}
 
 	// See if something actually sent (this may be false due to conditional checks inside the send properties function
@@ -2629,7 +2655,8 @@ void FRepLayout::SendProperties_r(
 	FRepHandleIterator& HandleIterator,
 	const FConstRepObjectDataBuffer SourceData,
 	const int32 ArrayDepth,
-	const FRepSerializationSharedInfo* const RESTRICT SharedInfo) const
+	const FRepSerializationSharedInfo* const RESTRICT SharedInfo,
+	const ESerializePropertyType SerializePropertyType) const
 {
 	const bool bDoSharedSerialization = SharedInfo && !!GNetSharedSerializedData;
 
@@ -2671,7 +2698,7 @@ void FRepLayout::SendProperties_r(
 			check(ArrayHandleIterator.ArrayElementSize> 0);
 			check(ArrayHandleIterator.NumHandlesPerElement> 0);
 
-			SendProperties_r(RepState, Writer, bDoChecksum, ArrayHandleIterator, ArrayData, ArrayDepth + 1, SharedInfo);
+			SendProperties_r(RepState, Writer, bDoChecksum, ArrayHandleIterator, ArrayData, ArrayDepth + 1, SharedInfo, SerializePropertyType);
 
 			check(HandleIterator.ChangelistIterator.ChangedIndex - OldChangedIndex == ArrayChangedCount);				// Make sure we read correct amount
 			check(HandleIterator.ChangelistIterator.Changed[HandleIterator.ChangelistIterator.ChangedIndex] == 0);	// Make sure we are at the end
@@ -2750,7 +2777,19 @@ void FRepLayout::SendProperties_r(
 		else
 		{
 			GNumSharedSerializationMiss++;
-			WritePropertyHandle(Writer, HandleIterator.Handle, bDoChecksum);
+
+			if (SerializePropertyType == ESerializePropertyType::Handle)
+			{
+				WritePropertyHandle(Writer, HandleIterator.Handle, bDoChecksum);
+			}
+			else if (SerializePropertyType == ESerializePropertyType::Name)
+			{
+				WritePropertyName(Writer, Cmd.Property->GetFName(), bDoChecksum);
+			}
+			else
+			{
+				UE_LOG(LogRep, Error, TEXT("Unsupported ESerializePropertyType encountered"));
+			}
 
 			UE_NET_TRACE_DYNAMIC_NAME_SCOPE(Cmd.Property->GetFName(), Writer, GetTraceCollector(Writer), ENetTraceVerbosity::Trace);
 
@@ -2788,7 +2827,8 @@ void FRepLayout::SendProperties(
 	UClass* ObjectClass,
 	FNetBitWriter& Writer,
 	TArray<uint16>& Changed,
-	const FRepSerializationSharedInfo& SharedInfo) const
+	const FRepSerializationSharedInfo& SharedInfo,
+	const ESerializePropertyType SerializePropertyType) const
 {
 	SCOPE_CYCLE_COUNTER(STAT_NetReplicateDynamicPropSendTime);
 
@@ -2818,7 +2858,7 @@ void FRepLayout::SendProperties(
 	FChangelistIterator ChangelistIterator(Changed, 0);
 	FRepHandleIterator HandleIterator(Owner, ChangelistIterator, Cmds, BaseHandleToCmdIndex, 0, 1, 0, Cmds.Num() - 1);
 
-	SendProperties_r(RepState, Writer, bDoChecksum, HandleIterator, Data, 0, &SharedInfo);
+	SendProperties_r(RepState, Writer, bDoChecksum, HandleIterator, Data, 0, &SharedInfo, SerializePropertyType);
 
 	if (NumBits != Writer.GetNumBits())
 	{
@@ -7783,7 +7823,8 @@ ERepLayoutResult FRepLayout::DeltaSerializeFastArrayProperty(FFastArrayDeltaSeri
 						HandleIterator,
 						ElementData,
 						/*ArrayDepth=*/ 1,
-						/*SharedInfo=*/ nullptr);
+						/*SharedInfo=*/ nullptr,
+						ESerializePropertyType::Handle);
 
 					WritePropertyHandle(Writer, 0, false);
 				}
diff --git a/Engine/Source/Runtime/Engine/Public/Net/DataReplication.h b/Engine/Source/Runtime/Engine/Public/Net/DataReplication.h
--- a/Engine/Source/Runtime/Engine/Public/Net/DataReplication.h
+++ b/Engine/Source/Runtime/Engine/Public/Net/DataReplication.h
@@ -156,6 +156,9 @@ public:
 	/** Writes dirty properties to bunch */
 	void ReplicateCustomDeltaProperties(FNetBitWriter& Bunch, FReplicationFlags RepFlags);
 	bool ReplicateProperties(FOutBunch& Bunch, FReplicationFlags RepFlags);
+	bool ReplicateProperties(FOutBunch& Bunch, FReplicationFlags RepFlags, FNetBitWriter& Writer);
+	bool ReplicateProperties_r(FOutBunch& Bunch, FReplicationFlags RepFlags, FNetBitWriter& Writer);
+
 	void PostSendBunch(FPacketIdRange& PacketRange, uint8 bReliable);
 
 	/** Updates the custom delta state for a replay delta checkpoint */
diff --git a/Engine/Source/Runtime/Engine/Public/Net/RepLayout.h b/Engine/Source/Runtime/Engine/Public/Net/RepLayout.h
--- a/Engine/Source/Runtime/Engine/Public/Net/RepLayout.h
+++ b/Engine/Source/Runtime/Engine/Public/Net/RepLayout.h
@@ -48,6 +48,12 @@ enum class EReceivePropertiesFlags : uint32
 
 ENUM_CLASS_FLAGS(EReceivePropertiesFlags);
 
+enum class ESerializePropertyType : uint8
+{
+	Handle,	// Properties are seralized using handles in SendProperties_r
+	Name	// Properties are seralized using their names in SendProperties_r
+};
+
 enum class ERepDataBufferType
 {
 	ObjectBuffer,	//! Indicates this buffer is a full object's memory.
@@ -1619,7 +1625,8 @@ private:
 		UClass* ObjectClass,
 		FNetBitWriter& Writer,
 		TArray<uint16>& Changed,
-		const FRepSerializationSharedInfo& SharedInfo) const;
+		const FRepSerializationSharedInfo& SharedInfo,
+		const ESerializePropertyType SerializePropertyType) const;
 
 	/**
 	 * Clamps a changelist so that it conforms to the current size of either an array, or arrays within structs/arrays.
@@ -1677,7 +1684,8 @@ private:
 		FRepHandleIterator& HandleIterator,
 		const FConstRepObjectDataBuffer SourceData,
 		const int32	 ArrayDepth,
-		const FRepSerializationSharedInfo* const RESTRICT SharedInfo) const;
+		const FRepSerializationSharedInfo* const RESTRICT SharedInfo,
+		const ESerializePropertyType SerializePropertyType) const;
 
 	void BuildSharedSerialization(
 		const FConstRepObjectDataBuffer Data,
