Allow an actor type to override its default label

diff --git a/Engine/Source/Runtime/Engine/Classes/GameFramework/Actor.h b/Engine/Source/Runtime/Engine/Classes/GameFramework/Actor.h
--- a/Engine/Source/Runtime/Engine/Classes/GameFramework/Actor.h
+++ b/Engine/Source/Runtime/Engine/Classes/GameFramework/Actor.h
@@ -2291,6 +2291,10 @@ public:
 	 */
 	virtual void CheckForDeprecated();
 
+	/** Returns this actor's default label (does not include any numeric suffix).  Actor labels are only available in development builds. */
+	UFUNCTION(BlueprintCallable, Category = "Editor Scripting | Actor Editing", meta = (KeyWords = "Display Name"))
+	virtual FString GetDefaultActorLabel() const;
+
 	/** Returns this actor's current label.  Actor labels are only available in development builds. */
 	UFUNCTION(BlueprintCallable, Category = "Editor Scripting | Actor Editing", meta = (KeyWords = "Display Name"))
 	const FString& GetActorLabel(bool bCreateIfNone = true) const;
diff --git a/Engine/Source/Runtime/Engine/Private/ActorEditor.cpp b/Engine/Source/Runtime/Engine/Private/ActorEditor.cpp
--- a/Engine/Source/Runtime/Engine/Private/ActorEditor.cpp
+++ b/Engine/Source/Runtime/Engine/Private/ActorEditor.cpp
@@ -874,6 +874,21 @@ TUniquePtr<class FWorldPartitionActorDesc> AActor::StaticCreateClassActorDesc(co
 	return CastChecked<AActor>(ActorClass->GetDefaultObject())->CreateClassActorDesc();
 }
 
+FString AActor::GetDefaultActorLabel() const
+{
+	UClass* ActorClass = GetClass();
+
+	FString DefaultActorLabel = ActorClass->GetName();
+
+	// Strip off the ugly "_C" suffix for Blueprint class actor instances
+	if (Cast<UBlueprint>(ActorClass->ClassGeneratedBy))
+	{
+		DefaultActorLabel.RemoveFromEnd(TEXT("_C"), ESearchCase::CaseSensitive);
+	}
+
+	return DefaultActorLabel;
+}
+
 const FString& AActor::GetActorLabel(bool bCreateIfNone) const
 {
 	// If the label string is empty then we'll use the default actor label (usually the actor's class name.)
@@ -888,16 +903,7 @@ const FString& AActor::GetActorLabel(bool bCreateIfNone) const
 
 	if( ActorLabel.IsEmpty() && bCreateIfNone )
 	{
-		// Get the class
-		UClass* ActorClass = GetClass();
-
-		FString DefaultActorLabel = ActorClass->GetName();
-
-		// Strip off the ugly "_C" suffix for Blueprint class actor instances
-		if (Cast<UBlueprint>(ActorClass->ClassGeneratedBy))
-		{
-			DefaultActorLabel.RemoveFromEnd(TEXT("_C"), ESearchCase::CaseSensitive);
-		}
+		FString DefaultActorLabel = GetDefaultActorLabel();
 
 		// We want the actor's label to be initially unique, if possible, so we'll use the number of the
 		// actor's FName when creating the initially.  It doesn't actually *need* to be unique, this is just
