Exposed MontageSync follow to blueprint

diff --git a/Engine/Source/Runtime/Engine/Classes/Animation/AnimInstance.h b/Engine/Source/Runtime/Engine/Classes/Animation/AnimInstance.h
--- a/Engine/Source/Runtime/Engine/Classes/Animation/AnimInstance.h
+++ b/Engine/Source/Runtime/Engine/Classes/Animation/AnimInstance.h
@@ -729,6 +729,22 @@ public:
 	UFUNCTION(BlueprintPure, Category = "Montage")
 	float Montage_GetPlayRate(const UAnimMontage* Montage) const;
 
+	/*********************************************************************************************
+	* AnimMontage sync. See notes in AnimMontage.h
+	********************************************************************************************* */
+
+	/** Synchronize a montage to another anim instance's montage. Both montages must be playing already
+	* @param MontageFollower : The montage that will follow the leader in OtherAnimInstance
+	* @param OtherAnimInstance	: The other anim instance we want to synchronize to. Can be set to self
+	* @param MontageLeader	: The montage we want to follow in the other anim instance */
+	UFUNCTION(BlueprintCallable, Category = "Montage")
+	void MontageSync_Follow(const UAnimMontage* MontageFollower, const UAnimInstance* OtherAnimInstance, const UAnimMontage* MontageLeader);
+
+	/** Stop following the montage's leader in this anim instance
+	* @param MontageFollower : The montage we want to stop synchronizing */
+	UFUNCTION(BlueprintCallable, Category = "Montage")
+	void MontageSync_StopFollowing(const UAnimMontage* MontageFollower);
+
 	/** Returns true if any montage is playing currently. Doesn't mean it's active though, it could be blending out. */
 	UFUNCTION(BlueprintPure, Category = "Montage")
 	bool IsAnyMontagePlaying() const;
diff --git a/Engine/Source/Runtime/Engine/Private/Animation/AnimInstance.cpp b/Engine/Source/Runtime/Engine/Private/Animation/AnimInstance.cpp
--- a/Engine/Source/Runtime/Engine/Private/Animation/AnimInstance.cpp
+++ b/Engine/Source/Runtime/Engine/Private/Animation/AnimInstance.cpp
@@ -2589,6 +2589,39 @@ float UAnimInstance::Montage_GetPlayRate(const UAnimMontage* Montage) const
 	return 0.f;
 }
 
+void UAnimInstance::MontageSync_Follow(const UAnimMontage* MontageFollower, const UAnimInstance* OtherAnimInstance, const UAnimMontage* MontageLeader)
+{
+	if (!MontageFollower || !OtherAnimInstance || !MontageLeader)
+	{
+		return;
+	}
+
+	FAnimMontageInstance* FollowerMontageInstance = GetActiveInstanceForMontage(MontageFollower);
+	FAnimMontageInstance* LeaderMontageInstance = OtherAnimInstance->GetActiveInstanceForMontage(MontageLeader);
+	if (!FollowerMontageInstance || !LeaderMontageInstance)
+	{
+		return;
+	}
+
+	FollowerMontageInstance->MontageSync_Follow(LeaderMontageInstance);
+}
+
+void UAnimInstance::MontageSync_StopFollowing(const UAnimMontage* MontageFollower)
+{
+	if (!MontageFollower)
+	{
+		return;
+	}
+
+	FAnimMontageInstance* FollowerMontageInstance = GetActiveInstanceForMontage(MontageFollower);
+	if (!FollowerMontageInstance)
+	{
+		return;
+	}
+
+	FollowerMontageInstance->MontageSync_StopFollowing();
+}
+
 int32 UAnimInstance::Montage_GetNextSectionID(UAnimMontage const* const Montage, int32 const& CurrentSectionID) const
 {
 	if (Montage)
