Fix a bug of mobile alpha to coverage when the material is masked but write every pixel.

diff --git a/Engine/Source/Runtime/Engine/Private/Materials/MaterialShared.cpp b/Engine/Source/Runtime/Engine/Private/Materials/MaterialShared.cpp
--- a/Engine/Source/Runtime/Engine/Private/Materials/MaterialShared.cpp
+++ b/Engine/Source/Runtime/Engine/Private/Materials/MaterialShared.cpp
@@ -1502,7 +1502,7 @@ EMaterialFloatPrecisionMode FMaterialResource::GetMaterialFloatPrecisionMode() c
 
 bool FMaterialResource::IsUsingAlphaToCoverage() const
 {
-	return Material->bUseAlphaToCoverage && Material->MaterialDomain == EMaterialDomain::MD_Surface && Material->GetBlendMode() == EBlendMode::BLEND_Masked;
+	return Material->bUseAlphaToCoverage && Material->MaterialDomain == EMaterialDomain::MD_Surface && Material->GetBlendMode() == EBlendMode::BLEND_Masked && !WritesEveryPixel();
 }
 
 bool FMaterialResource::IsUsingPreintegratedGFForSimpleIBL() const
