Add initial strata conversion for postprocess domain material.

diff --git a/Engine/Shaders/Private/PostProcessMaterialShaders.usf b/Engine/Shaders/Private/PostProcessMaterialShaders.usf
--- a/Engine/Shaders/Private/PostProcessMaterialShaders.usf
+++ b/Engine/Shaders/Private/PostProcessMaterialShaders.usf
@@ -202,13 +202,24 @@ void MainPS(
 	// fill out other related material parameters
 	CalcMaterialParametersPost(Parameters, PixelMaterialInputs, SvPosition, true);
 
+	float Alpha = 1.0f;
+#if STRATA_ENABLED  
+	// Unlit forces a single BSDF
+	const FStrataBSDF UnlitBSDF = PixelMaterialInputs.FrontMaterial.Layers[0].BSDFs[0];
+
 	// Grab emissive colour as output
-#if MATERIAL_OUTPUT_OPACITY_AS_ALPHA
-	const float Alpha = GetMaterialOpacity(PixelMaterialInputs);
+  #if MATERIAL_OUTPUT_OPACITY_AS_ALPHA
+	Alpha = UnlitBSDF.Coverage;
+  #endif
+	const float3 Emissive = BSDF_GETEMISSIVE(UnlitBSDF);
 #else
-	const float Alpha = 1.0f;
+	// Grab emissive colour as output
+  #if MATERIAL_OUTPUT_OPACITY_AS_ALPHA
+	Alpha = GetMaterialOpacity(PixelMaterialInputs);
+  #endif
+	const float3 Emissive = GetMaterialEmissive(PixelMaterialInputs);
 #endif
-	OutColor = float4(GetMaterialEmissive(PixelMaterialInputs), Alpha );
+	OutColor = float4(Emissive, Alpha );
 
 #if !MATERIALBLENDING_MODULATE && POST_PROCESS_MATERIAL_BEFORE_TONEMAP
 	OutColor.xyz *= View.PreExposure;
diff --git a/Engine/Source/Runtime/Engine/Private/Materials/Material.cpp b/Engine/Source/Runtime/Engine/Private/Materials/Material.cpp
--- a/Engine/Source/Runtime/Engine/Private/Materials/Material.cpp
+++ b/Engine/Source/Runtime/Engine/Private/Materials/Material.cpp
@@ -2940,6 +2940,30 @@ void UMaterial::ConvertMaterialToStrataMaterial()
 			FrontMaterial.Connect(0, VolBSDF);
 			bInvalidateShader = true;
 		}
+		else if (MaterialDomain == MD_PostProcess)
+		{
+			// Some post-process material don't have their shading mode set correctly to Unlit. Since only Unlit is supported, forcing it here.
+			ShadingModel = MSM_Unlit;
+			ShadingModels.ClearShadingModels();
+			ShadingModels.AddShadingModel(MSM_Unlit);
+
+			// Only Emissive & Opacity are valid input for PostProcess material
+			UMaterialExpressionStrataLegacyConversion* ConvertNode = NewObject<UMaterialExpressionStrataLegacyConversion>(this);
+			MoveConnectionTo(EmissiveColor, ConvertNode, 5);
+			MoveConnectionTo(Opacity, ConvertNode, 11);
+
+			// Add constant for the Unlit shading model
+			UMaterialExpressionConstant* ShadingModelNode = NewObject<UMaterialExpressionConstant>(this);
+			ShadingModelNode->SetParameterName(FName(TEXT("ConstantShadingModel")));
+			ShadingModelNode->R = ShadingModel;
+			ConvertNode->ShadingModel.Connect(0, ShadingModelNode);
+
+			AddStrataShadingModelFromMaterialShadingModel(ConvertNode->ConvertedStrataMaterialInfo, ShadingModels);
+			check(ConvertNode->ConvertedStrataMaterialInfo.CountShadingModels() == 1);
+
+			FrontMaterial.Connect(0, ConvertNode);
+			bInvalidateShader = true;
+		}
 	}
 
 	if (bInvalidateShader)
@@ -3943,7 +3967,10 @@ void UMaterial::RebuildShadingModelField()
 			// Now derive some properties from the material
 			if (StrataMaterialInfo.HasOnlyShadingModel(SSM_Unlit))
 			{
-				MaterialDomain = EMaterialDomain::MD_Surface;
+				if (MaterialDomain != EMaterialDomain::MD_Surface && MaterialDomain != EMaterialDomain::MD_PostProcess)
+				{
+					MaterialDomain = EMaterialDomain::MD_Surface;
+				}
 				ShadingModel = MSM_Unlit;
 				if (BlendMode != EBlendMode::BLEND_Opaque && BlendMode != EBlendMode::BLEND_Masked)
 				{
@@ -5261,9 +5288,19 @@ static bool IsPropertyActive_Internal(EMaterialProperty InProperty,
 	bool bUsesShadingModelFromMaterialExpression,
 	bool bIsTranslucencyWritingVelocity)
 {
+	static const auto CVar = IConsoleManager::Get().FindTConsoleVariableDataInt(TEXT("r.Strata"));
+	const bool bStrataEnabled = CVar && CVar->GetValueOnAnyThread() > 0;
+
 	if (Domain == MD_PostProcess)
 	{
-		return InProperty == MP_EmissiveColor || (bBlendableOutputAlpha && InProperty == MP_Opacity);
+		if (bStrataEnabled)
+		{
+			return InProperty == MP_FrontMaterial;
+		}
+		else
+		{
+			return InProperty == MP_EmissiveColor || (bBlendableOutputAlpha && InProperty == MP_Opacity);
+		}
 	}
 	else if (Domain == MD_LightFunction)
 	{
@@ -5417,8 +5454,6 @@ static bool IsPropertyActive_Internal(EMaterialProperty InProperty,
 		break;
 	case MP_FrontMaterial:
 		{
-			static const auto CVar = IConsoleManager::Get().FindTConsoleVariableDataInt(TEXT("r.Strata"));
-			const bool bStrataEnabled = CVar && CVar->GetValueOnAnyThread() > 0;
 			Active = bStrataEnabled;
 			break;
 		}
diff --git a/Engine/Source/Runtime/Engine/Private/Materials/MaterialExpressions.cpp b/Engine/Source/Runtime/Engine/Private/Materials/MaterialExpressions.cpp
--- a/Engine/Source/Runtime/Engine/Private/Materials/MaterialExpressions.cpp
+++ b/Engine/Source/Runtime/Engine/Private/Materials/MaterialExpressions.cpp
@@ -20178,6 +20178,20 @@ static int32 StrataBlendNormal(class FMaterialCompiler* Compiler, int32 NormalCo
 
 #if WITH_EDITOR
 
+// Optionnaly cast CodeChunk type to non-LWC type. 
+// Input can be built of WorldPosition data, which would force the derived data to have LWC type 
+// creating issues, as Strata functions' inputs don't support LWC
+static int32 CastToNonLWCType(class FMaterialCompiler* Compiler, int32 CodeChunk)
+{
+	EMaterialValueType Type = Compiler->GetType(CodeChunk);
+	if (IsLWCType(Type))
+	{
+		Type = MakeNonLWCType(Type);
+		CodeChunk = Compiler->ValidCast(CodeChunk, Type);
+	}
+	return CodeChunk;
+}
+
 // The compilation of an expression can sometimes lead to a INDEX_NONE code chunk when editing material graphs 
 // or when the node is inside a material function, linked to an input pin of the material function and that input is not plugged in to anything.
 // But for normals or tangents, Strata absolutely need a valid code chunk to de-duplicate when stored in memory. 
@@ -20193,6 +20207,10 @@ static int32 CompileWithDefaultCodeChunk(class FMaterialCompiler* Compiler, FExp
 	{
 		*bDefaultIsUsed |= CodeChunk == INDEX_NONE;
 	}
+	else
+	{
+		CodeChunk = CastToNonLWCType(Compiler, CodeChunk);
+	}
 	return CodeChunk == INDEX_NONE ? DefaultCodeChunk : CodeChunk;
 }
 static int32 CompileWithDefaultFloat1(class FMaterialCompiler* Compiler, FExpressionInput& Input, float X, bool* bDefaultIsUsed = nullptr)
@@ -20206,6 +20224,10 @@ static int32 CompileWithDefaultFloat1(class FMaterialCompiler* Compiler, FExpres
 	{
 		*bDefaultIsUsed |= CodeChunk == INDEX_NONE;
 	}
+	else
+	{
+		CodeChunk = CastToNonLWCType(Compiler, CodeChunk);
+	}
 	return CodeChunk == INDEX_NONE ? Compiler->Constant(X) : CodeChunk;
 }
 static int32 CompileWithDefaultFloat3(class FMaterialCompiler* Compiler, FExpressionInput& Input, float X, float Y, float Z, bool* bDefaultIsUsed = nullptr)
@@ -20219,6 +20241,10 @@ static int32 CompileWithDefaultFloat3(class FMaterialCompiler* Compiler, FExpres
 	{
 		*bDefaultIsUsed |= CodeChunk == INDEX_NONE;
 	}
+	else
+	{
+		CodeChunk = CastToNonLWCType(Compiler, CodeChunk);
+	}
 	return CodeChunk == INDEX_NONE ? Compiler->Constant3(X, Y, Z) : CodeChunk;
 }
 static int32 CompileWithDefaultNormalWS(class FMaterialCompiler* Compiler, FExpressionInput& Input, bool* bDefaultIsUsed = nullptr)
diff --git a/Engine/Source/Runtime/Renderer/Private/PostProcess/PostProcessMaterial.cpp b/Engine/Source/Runtime/Renderer/Private/PostProcess/PostProcessMaterial.cpp
--- a/Engine/Source/Runtime/Renderer/Private/PostProcess/PostProcessMaterial.cpp
+++ b/Engine/Source/Runtime/Renderer/Private/PostProcess/PostProcessMaterial.cpp
@@ -168,6 +168,8 @@ public:
 		{
 			OutEnvironment.SetDefine(TEXT("POST_PROCESS_MATERIAL_BEFORE_TONEMAP"), (Parameters.MaterialParameters.BlendableLocation != BL_AfterTonemapping) ? 1 : 0);
 		}
+
+		OutEnvironment.SetDefine(TEXT("STRATA_ENABLED"), Strata::IsStrataEnabled() ? 1u : 0u);
 	}
 
 protected:
