GetLinkedAnimLayerInstanceByGroup now works correctly in packaged builds

Change to using the concrete class rather than the interface class, as interfaces dont actually hold their group info in non-editor builds due to it being filled from metadata, and all other compilation mechanisms (e.g. CopyTermDefaultsToDefaultObject) not being called for interface BP compilation.

diff --git a/Engine/Source/Runtime/Engine/Private/Animation/AnimInstance.cpp b/Engine/Source/Runtime/Engine/Private/Animation/AnimInstance.cpp
--- a/Engine/Source/Runtime/Engine/Private/Animation/AnimInstance.cpp
+++ b/Engine/Source/Runtime/Engine/Private/Animation/AnimInstance.cpp
@@ -3173,7 +3173,11 @@ UAnimInstance* UAnimInstance::GetLinkedAnimLayerInstanceByGroup(FName InGroup) c
 			const FAnimNode_LinkedAnimLayer* Layer = LayerNodeProperty->ContainerPtrToValuePtr<FAnimNode_LinkedAnimLayer>(this);
 
 			UClass* ClassForGroups;
-			if(UClass* InterfaceClass = Layer->Interface.Get())
+			if (UClass* InstanceClass = Layer->InstanceClass.Get())
+			{
+				ClassForGroups = InstanceClass;
+			}
+			else if(UClass* InterfaceClass = Layer->Interface.Get())
 			{
 				ClassForGroups = InterfaceClass;
 			}
