Fixed BlendSpace filters being allocated even when they are not in use

Saves ~2.6MB peak in synthetic tests

diff --git a/Engine/Source/Runtime/Engine/Classes/Animation/AnimationAsset.h b/Engine/Source/Runtime/Engine/Classes/Animation/AnimationAsset.h
--- a/Engine/Source/Runtime/Engine/Classes/Animation/AnimationAsset.h
+++ b/Engine/Source/Runtime/Engine/Classes/Animation/AnimationAsset.h
@@ -195,15 +195,23 @@ struct FBlendFilter
 {
 	GENERATED_USTRUCT_BODY()
 
-	FFIRFilterTimeBased FilterPerAxis[3];
+	TArray<FFIRFilterTimeBased> FilterPerAxis;
 
 	FBlendFilter()
 	{
 	}
-
+	
+	bool IsValid() const
+	{
+		return FilterPerAxis.Num() > 0;
+	}
+	
 	FVector GetFilterLastOutput() const
 	{
-		return FVector (FilterPerAxis[0].LastOutput, FilterPerAxis[1].LastOutput, FilterPerAxis[2].LastOutput);
+		return FVector(
+			FilterPerAxis.Num() > 0 ? FilterPerAxis[0].LastOutput : 0.0f,
+			FilterPerAxis.Num() > 1 ? FilterPerAxis[1].LastOutput : 0.0f,
+			FilterPerAxis.Num() > 2 ? FilterPerAxis[2].LastOutput : 0.0f);
 	}
 };
 
diff --git a/Engine/Source/Runtime/Engine/Private/Animation/BlendSpace.cpp b/Engine/Source/Runtime/Engine/Private/Animation/BlendSpace.cpp
--- a/Engine/Source/Runtime/Engine/Private/Animation/BlendSpace.cpp
+++ b/Engine/Source/Runtime/Engine/Private/Animation/BlendSpace.cpp
@@ -1056,24 +1056,40 @@ void UBlendSpace::InitializeFilter(FBlendFilter* Filter) const
 {
 	if (Filter)
 	{
-		for (int FilterIndex = 0 ; FilterIndex != 3 ; ++FilterIndex)
+		Filter->FilterPerAxis.Empty();
+
+		int32 NumActiveAxes = 0;
+		for (int32 FilterIndex = 0; FilterIndex != 3; ++FilterIndex)
 		{
-			Filter->FilterPerAxis[FilterIndex].Initialize(InterpolationParam[FilterIndex].InterpolationTime,
-			                                              InterpolationParam[FilterIndex].InterpolationType,
-			                                              InterpolationParam[FilterIndex].DampingRatio,
-			                                              BlendParameters[FilterIndex].Min,
-			                                              BlendParameters[FilterIndex].Max,
-			                                              InterpolationParam[FilterIndex].MaxSpeed,
-			                                              !BlendParameters[FilterIndex].bWrapInput);
+			if(InterpolationParam[FilterIndex].InterpolationTime > 0.0f)
+			{
+				NumActiveAxes = FMath::Max(NumActiveAxes, FilterIndex + 1);
+			}
+		}
+
+		if(NumActiveAxes > 0)
+		{
+			Filter->FilterPerAxis.SetNum(NumActiveAxes);
+			
+			for (int32 FilterIndex = 0; FilterIndex < Filter->FilterPerAxis.Num(); ++FilterIndex)
+			{
+				Filter->FilterPerAxis[FilterIndex].Initialize(InterpolationParam[FilterIndex].InterpolationTime,
+															  InterpolationParam[FilterIndex].InterpolationType,
+															  InterpolationParam[FilterIndex].DampingRatio,
+															  BlendParameters[FilterIndex].Min,
+															  BlendParameters[FilterIndex].Max,
+															  InterpolationParam[FilterIndex].MaxSpeed,
+															  !BlendParameters[FilterIndex].bWrapInput);
+			}
 		}
 	}
 }
 
 void UBlendSpace::UpdateFilterParams(FBlendFilter* Filter) const
 {
-	if (Filter)
+	if (Filter && Filter->IsValid())
 	{
-		for (int FilterIndex = 0 ; FilterIndex != 3 ; ++FilterIndex)
+		for (int32 FilterIndex = 0; FilterIndex < Filter->FilterPerAxis.Num(); ++FilterIndex)
 		{
 			Filter->FilterPerAxis[FilterIndex].SetParams(InterpolationParam[FilterIndex].DampingRatio,
 			                                             BlendParameters[FilterIndex].Min,
@@ -1769,30 +1785,36 @@ bool UBlendSpace::InterpolateWeightOfSampleData(float DeltaTime, const TArray<FB
 
 FVector UBlendSpace::FilterInput(FBlendFilter* Filter, const FVector& BlendInput, float DeltaTime) const
 {
-#if WITH_EDITOR
-	// Check 
-	for (int32 AxisIndex = 0; AxisIndex < 3; ++AxisIndex)
+	FVector FilteredBlendInput = BlendInput;
+	if(Filter)
 	{
-		if (Filter->FilterPerAxis[AxisIndex].NeedsUpdate(InterpolationParam[AxisIndex].InterpolationType,
-		                                                 InterpolationParam[AxisIndex].InterpolationTime))
+#if WITH_EDITOR
+		// Check
+		for (int32 AxisIndex = 0; AxisIndex < Filter->FilterPerAxis.Num(); ++AxisIndex)
 		{
-			InitializeFilter(Filter);
-			break;
+			if (Filter->FilterPerAxis[AxisIndex].NeedsUpdate(InterpolationParam[AxisIndex].InterpolationType,
+															 InterpolationParam[AxisIndex].InterpolationTime))
+			{
+				InitializeFilter(Filter);
+				break;
+			}
 		}
-	}
-	// Note that if we expose the damping ratio etc as pins, this should be called outside of the editor too.
-	UpdateFilterParams(Filter);
+		// Note that if we expose the damping ratio etc as pins, this should be called outside of the editor too.
+		UpdateFilterParams(Filter);
 #endif
-	FVector FilteredBlendInput;
-	for (int32 AxisIndex = 0; AxisIndex < 3; ++AxisIndex)
-	{
-		if (BlendParameters[AxisIndex].bWrapInput)
+		if(Filter->IsValid())
 		{
-			Filter->FilterPerAxis[AxisIndex].WrapToValue(
-				BlendInput[AxisIndex], BlendParameters[AxisIndex].Max - BlendParameters[AxisIndex].Min);
+			for (int32 AxisIndex = 0; AxisIndex < Filter->FilterPerAxis.Num(); ++AxisIndex)
+			{
+				if (BlendParameters[AxisIndex].bWrapInput)
+				{
+					Filter->FilterPerAxis[AxisIndex].WrapToValue(
+						BlendInput[AxisIndex], BlendParameters[AxisIndex].Max - BlendParameters[AxisIndex].Min);
+				}
+				FilteredBlendInput[AxisIndex] = Filter->FilterPerAxis[AxisIndex].UpdateAndGetFilteredData(
+					BlendInput[AxisIndex], DeltaTime);
+			}
 		}
-		FilteredBlendInput[AxisIndex] = Filter->FilterPerAxis[AxisIndex].UpdateAndGetFilteredData(
-			BlendInput[AxisIndex], DeltaTime);
 	}
 	return FilteredBlendInput;
 }
