LWC: Increase WORLD_MAX to a more reasonable value. Note: This is NOT final!

diff --git a/Engine/Source/Programs/UnrealLightmass/Private/Lighting/Lighting.h b/Engine/Source/Programs/UnrealLightmass/Private/Lighting/Lighting.h
--- a/Engine/Source/Programs/UnrealLightmass/Private/Lighting/Lighting.h
+++ b/Engine/Source/Programs/UnrealLightmass/Private/Lighting/Lighting.h
@@ -49,6 +49,7 @@ namespace Lightmass
 	extern double GStartupTime;
 }
 
+// LWC_TODO: Lightmass definition of WORLD_MAX is no longer in sync with EngineDefines.h. Is this acceptable?
 #define WORLD_MAX			2097152.0			/* Maximum size of the world */
 #define HALF_WORLD_MAX		( WORLD_MAX*0.5f )	/* Half the maximum size of the world */
 
diff --git a/Engine/Source/Runtime/Core/Private/Misc/LargeWorldRenderPosition.cpp b/Engine/Source/Runtime/Core/Private/Misc/LargeWorldRenderPosition.cpp
--- a/Engine/Source/Runtime/Core/Private/Misc/LargeWorldRenderPosition.cpp
+++ b/Engine/Source/Runtime/Core/Private/Misc/LargeWorldRenderPosition.cpp
@@ -1,10 +1,15 @@
 // Copyright Epic Games, Inc. All Rights Reserved.
 
 #include "Misc/LargeWorldRenderPosition.h"
+#include "EngineDefines.h"
 
 #if !UE_LARGE_WORLD_COORDINATES_DISABLED
 
-#define UE_LWC_RENDER_TILE_SIZE (262144.0)
+#define UE_LWC_RENDER_TILE_SIZE_MIN (262144.0)	// LWC_TODO: Make this smaller?
+#define UE_LWC_RENDER_TILE_SIZE (FMath::Max(UE_LWC_RENDER_TILE_SIZE_MIN,  HALF_WORLD_MAX / 16777216.0))	// LWC_TODO: Tweak render tile divisor to maximise range.
+static_assert(UE_LWC_RENDER_TILE_SIZE <= 1048576.0, "Current WORLD_MAX is too large and is likely to adversely affect world space coordinate precision within shaders!");	// LWC_TODO: Set this to something reasonable.
+
+//#define UE_LWC_RENDER_TILE_SIZE (262144.0)
 //#define UE_LWC_RENDER_TILE_SIZE (1048576.0)
 //#define UE_LWC_RENDER_TILE_SIZE (2097152.0)
 
diff --git a/Engine/Source/Runtime/Engine/Public/EngineDefines.h b/Engine/Source/Runtime/Engine/Public/EngineDefines.h
--- a/Engine/Source/Runtime/Engine/Public/EngineDefines.h
+++ b/Engine/Source/Runtime/Engine/Public/EngineDefines.h
@@ -28,9 +28,23 @@
 /*-----------------------------------------------------------------------------
 	Size of the world.
 -----------------------------------------------------------------------------*/
+#define UE_OLD_WORLD_MAX			2097152.0				/* UE4 world size */
+
+#ifndef UE_USE_UE4_WORLD_MAX
+#define UE_USE_UE4_WORLD_MAX		(UE_LARGE_WORLD_COORDINATES_DISABLED)			// Force UE4 WORLD_MAX for converted UE4 titles that explicitly rely on it.
+#endif
+
+#if UE_USE_UE4_WORLD_MAX
+#define WORLD_MAX					(UE_OLD_WORLD_MAX)		/* Maximum size of the world */
+#else
+// Note: Modifying WORLD_MAX affects UE_LWC_RENDER_TILE_SIZE in Engine\Source\Runtime\Core\Private\Misc\LargeWorldRenderPosition.cpp and may introduce precision issues in shaders using world coordinates.
+#define WORLD_MAX					8796093022208.0			/* Maximum size of the world  - Approx 87,960,930.2 km across, or 43,980,465.1 km from the origin */
+#endif
 
-#define WORLD_MAX					2097152.0				/* Maximum size of the world */
 #define HALF_WORLD_MAX				(WORLD_MAX * 0.5)		/* Half the maximum size of the world */
 #define HALF_WORLD_MAX1				(HALF_WORLD_MAX - 1.0)	/* Half the maximum size of the world minus one */
 
+
+
+
 #define DEFAULT_ORTHOZOOM			10000.0					/* Default 2D viewport zoom */
