WorldPartition: Add Package name to log when failing to load actor

diff --git a/Engine/Source/Runtime/Engine/Private/WorldPartition/WorldPartitionActorDesc.cpp b/Engine/Source/Runtime/Engine/Private/WorldPartition/WorldPartitionActorDesc.cpp
--- a/Engine/Source/Runtime/Engine/Private/WorldPartition/WorldPartitionActorDesc.cpp
+++ b/Engine/Source/Runtime/Engine/Private/WorldPartition/WorldPartitionActorDesc.cpp
@@ -327,7 +327,7 @@ AActor* FWorldPartitionActorDesc::Load() const
 			}
 			else
 			{
-				UE_LOG(LogWorldPartition, Warning, TEXT("Can't load actor %s"), *GetActorName().ToString());
+				UE_LOG(LogWorldPartition, Warning, TEXT("Can't load actor '%s' from package '%s'"), *GetActorName().ToString(), *ActorPackage.ToString());
 			}
 		}
 	}
