Network Prediction
-Add PreTickFlush to world so that NP has final chance to marshal networked data over

#rnx

diff --git a/Engine/Source/Runtime/Engine/Classes/Engine/World.h b/Engine/Source/Runtime/Engine/Classes/Engine/World.h
--- a/Engine/Source/Runtime/Engine/Classes/Engine/World.h
+++ b/Engine/Source/Runtime/Engine/Classes/Engine/World.h
@@ -1579,6 +1579,9 @@ private:
 	/** Event to gather up all net drivers and call PostTickDispatch at once */
 	FOnTickFlushEvent PostTickDispatchEvent;
 
+	/** Event called prior to calling TickFlush */
+	FOnNetTickEvent PreTickFlushEvent;
+
 	/** Event to gather up all net drivers and call TickFlush at once */
 	FOnNetTickEvent TickFlushEvent;
 	
@@ -1595,6 +1598,11 @@ private:
 	{
 		PostTickDispatchEvent.Broadcast();
 	}
+	/** PreTickFlush */
+	void BroadcastPreTickFlush(float DeltaTime)
+	{
+		PreTickFlushEvent.Broadcast(DeltaTime);
+	}
 	/** All registered net drivers TickFlush() */
 	void BroadcastTickFlush(float DeltaTime)
 	{
@@ -3021,6 +3029,7 @@ public:
 	/** Network Tick events */
 	FOnNetTickEvent& OnTickDispatch() { return TickDispatchEvent; }
 	FOnTickFlushEvent& OnPostTickDispatch() { return PostTickDispatchEvent; }	
+	FOnNetTickEvent& OnPreTickFlush() { return PreTickFlushEvent; }
 	FOnNetTickEvent& OnTickFlush() { return TickFlushEvent; }
 	FOnTickFlushEvent& OnPostTickFlush() { return PostTickFlushEvent; }
 
diff --git a/Engine/Source/Runtime/Engine/Private/LevelTick.cpp b/Engine/Source/Runtime/Engine/Private/LevelTick.cpp
--- a/Engine/Source/Runtime/Engine/Private/LevelTick.cpp
+++ b/Engine/Source/Runtime/Engine/Private/LevelTick.cpp
@@ -1710,6 +1710,7 @@ void UWorld::Tick( ELevelTick TickType, float DeltaSeconds )
 		TRACE_CPUPROFILER_EVENT_SCOPE(NetBroadcastTickTime);
 		SCOPE_CYCLE_COUNTER(STAT_NetBroadcastTickTime);
 		LLM_SCOPE(ELLMTag::Networking);
+		BroadcastPreTickFlush(RealDeltaSeconds);
 		BroadcastTickFlush(RealDeltaSeconds); // note: undilated time is being used here
 	}
 	
