Fix Level Instance Edit Mode not exiting properly

diff --git a/Engine/Source/Editor/LevelInstanceEditor/Private/LevelInstanceEditorMode.cpp b/Engine/Source/Editor/LevelInstanceEditor/Private/LevelInstanceEditorMode.cpp
--- a/Engine/Source/Editor/LevelInstanceEditor/Private/LevelInstanceEditorMode.cpp
+++ b/Engine/Source/Editor/LevelInstanceEditor/Private/LevelInstanceEditorMode.cpp
@@ -7,9 +7,12 @@
 #include "Engine/World.h"
 #include "LevelInstance/LevelInstanceSubsystem.h"
 #include "LevelInstance/LevelInstanceActor.h"
+#include "LevelInstance/ILevelInstanceEditorModule.h"
 #include "LevelEditorViewport.h"
 #include "LevelEditorActions.h"
 #include "EditorModeManager.h"
+#include "Framework/Application/SlateApplication.h"
+#include "Modules/ModuleManager.h"
 
 #define LOCTEXT_NAMESPACE "LevelInstanceEditorMode"
 
@@ -32,10 +35,7 @@ ULevelInstanceEditorMode::~ULevelInstanceEditorMode()
 
 void ULevelInstanceEditorMode::OnPreBeginPIE(bool bSimulate)
 {
-	if (GLevelEditorModeTools().IsModeActive(ULevelInstanceEditorMode::EM_LevelInstanceEditorModeId))
-	{
-		GLevelEditorModeTools().DeactivateMode(ULevelInstanceEditorMode::EM_LevelInstanceEditorModeId);
-	}
+	ExitModeCommand();
 }
 
 void ULevelInstanceEditorMode::UpdateEngineShowFlags()
@@ -125,9 +125,15 @@ bool ULevelInstanceEditorMode::IsSelectionDisallowed(AActor* InActor, bool bInSe
 
 void ULevelInstanceEditorMode::ExitModeCommand()
 {	
-	if (FEditorModeTools* Manager = GetModeManager())
+	// Ignore command when any modal window is open
+	if (FSlateApplication::IsInitialized() && FSlateApplication::Get().GetActiveModalWindow().IsValid())
+	{
+		return;
+	}
+
+	if (ILevelInstanceEditorModule* EditorModule = FModuleManager::GetModulePtr<ILevelInstanceEditorModule>("LevelInstanceEditor"))
 	{
-		Manager->DeactivateMode(EM_LevelInstanceEditorModeId);
+		EditorModule->BroadcastTryExitEditorMode();
 	}
 }
 
diff --git a/Engine/Source/Editor/LevelInstanceEditor/Private/LevelInstanceEditorModule.cpp b/Engine/Source/Editor/LevelInstanceEditor/Private/LevelInstanceEditorModule.cpp
--- a/Engine/Source/Editor/LevelInstanceEditor/Private/LevelInstanceEditorModule.cpp
+++ b/Engine/Source/Editor/LevelInstanceEditor/Private/LevelInstanceEditorModule.cpp
@@ -657,6 +657,11 @@ void FLevelInstanceEditorModule::OnEditorModeIDChanged(const FEditorModeID& InMo
 	}
 }
 
+void FLevelInstanceEditorModule::BroadcastTryExitEditorMode() 
+{
+	TryExitEditorModeEvent.Broadcast();
+}
+
 void FLevelInstanceEditorModule::ActivateEditorMode()
 {
 	if (!GLevelEditorModeTools().IsModeActive(ULevelInstanceEditorMode::EM_LevelInstanceEditorModeId))
diff --git a/Engine/Source/Editor/LevelInstanceEditor/Private/LevelInstanceEditorModule.h b/Engine/Source/Editor/LevelInstanceEditor/Private/LevelInstanceEditorModule.h
--- a/Engine/Source/Editor/LevelInstanceEditor/Private/LevelInstanceEditorModule.h
+++ b/Engine/Source/Editor/LevelInstanceEditor/Private/LevelInstanceEditorModule.h
@@ -30,9 +30,14 @@ public:
 	virtual void ActivateEditorMode() override;
 	virtual void DeactivateEditorMode() override;
 		
+	virtual void BroadcastTryExitEditorMode() override;
+
 	DECLARE_DERIVED_EVENT(FLevelInstanceEditorModule, ILevelInstanceEditorModule::FExitEditorModeEvent, FExitEditorModeEvent);
 	virtual FExitEditorModeEvent& OnExitEditorMode() override { return ExitEditorModeEvent; }
 
+	DECLARE_DERIVED_EVENT(FLevelInstanceEditorModule, ILevelInstanceEditorModule::FTryExitEditorModeEvent, FTryExitEditorModeEvent);
+	virtual FTryExitEditorModeEvent& OnTryExitEditorMode() override { return TryExitEditorModeEvent; }
+
 private:
 	void OnEditorModeIDChanged(const FEditorModeID& InModeID, bool bIsEnteringMode);
 	void OnLevelActorDeleted(AActor* Actor);
@@ -41,4 +46,5 @@ private:
 	void ExtendContextMenu();
 
 	FExitEditorModeEvent ExitEditorModeEvent;
+	FTryExitEditorModeEvent TryExitEditorModeEvent;
 };
diff --git a/Engine/Source/Runtime/Engine/Private/LevelInstance/LevelInstanceSubsystem.cpp b/Engine/Source/Runtime/Engine/Private/LevelInstance/LevelInstanceSubsystem.cpp
--- a/Engine/Source/Runtime/Engine/Private/LevelInstance/LevelInstanceSubsystem.cpp
+++ b/Engine/Source/Runtime/Engine/Private/LevelInstance/LevelInstanceSubsystem.cpp
@@ -54,7 +54,8 @@ DEFINE_LOG_CATEGORY(LogLevelInstance);
 ULevelInstanceSubsystem::ULevelInstanceSubsystem()
 	: UWorldSubsystem()
 #if WITH_EDITOR
-	, bIsCreatingNewLevelStreamingLevelInstanceEditor(false)
+	, bIsCreatingLevelInstance(false)
+	, bIsCommittingLevelInstance(false)
 #endif
 {
 
@@ -73,19 +74,6 @@ void ULevelInstanceSubsystem::Initialize(FSubsystemCollectionBase& Collection)
 	if (GEditor)
 	{
 		ILevelInstanceEditorModule& EditorModule = FModuleManager::LoadModuleChecked<ILevelInstanceEditorModule>("LevelInstanceEditor");
-		EditorModule.OnExitEditorMode().AddUObject(this, &ULevelInstanceSubsystem::OnExitEditorMode);
-	}
-#endif
-}
-
-void ULevelInstanceSubsystem::Deinitialize()
-{
-	Super::Deinitialize();
-
-#if WITH_EDITOR
-	if (ILevelInstanceEditorModule* EditorModule = FModuleManager::GetModulePtr<ILevelInstanceEditorModule>("LevelInstanceEditor"))
-	{
-		EditorModule->OnExitEditorMode().RemoveAll(this);
 	}
 #endif
 }
@@ -265,11 +253,39 @@ void ULevelInstanceSubsystem::RegisterLoadedLevelStreamingLevelInstance(ULevelSt
 #if WITH_EDITOR
 void ULevelInstanceSubsystem::RegisterLoadedLevelStreamingLevelInstanceEditor(ULevelStreamingLevelInstanceEditor* LevelStreaming)
 {
-	if (!bIsCreatingNewLevelStreamingLevelInstanceEditor)
+	if (!bIsCreatingLevelInstance)
 	{
 		check(!LevelInstanceEdit.IsValid());
 		ALevelInstance* LevelInstanceActor = LevelStreaming->GetLevelInstanceActor();
 		LevelInstanceEdit = MakeUnique<FLevelInstanceEdit>(LevelStreaming, LevelInstanceActor->GetLevelInstanceID());
+
+		if (ILevelInstanceEditorModule* EditorModule = FModuleManager::GetModulePtr<ILevelInstanceEditorModule>("LevelInstanceEditor"))
+		{
+			EditorModule->OnExitEditorMode().AddUObject(this, &ULevelInstanceSubsystem::OnExitEditorMode);
+			EditorModule->OnTryExitEditorMode().AddUObject(this, &ULevelInstanceSubsystem::OnTryExitEditorMode);
+		}
+	}
+}
+
+void ULevelInstanceSubsystem::ResetEdit(TUniquePtr<FLevelInstanceEdit>& InLevelInstanceEdit)
+{
+	if (InLevelInstanceEdit)
+	{
+		if (InLevelInstanceEdit == LevelInstanceEdit)
+		{
+			if (ILevelInstanceEditorModule* EditorModule = FModuleManager::GetModulePtr<ILevelInstanceEditorModule>("LevelInstanceEditor"))
+			{
+				EditorModule->OnExitEditorMode().RemoveAll(this);
+				EditorModule->OnTryExitEditorMode().RemoveAll(this);
+			}
+		}
+		else
+		{
+			// Only case supported where we are using a tmp FLevelInstanceEdit
+			check(bIsCreatingLevelInstance);
+		}
+
+		InLevelInstanceEdit.Reset();
 	}
 }
 #endif
@@ -409,8 +425,28 @@ void ULevelInstanceSubsystem::Tick()
 
 void ULevelInstanceSubsystem::OnExitEditorMode()
 {
-	if (LevelInstanceEdit && !bCreatingLevelInstance)
+	OnExitEditorModeInternal(/*bForceExit=*/true);
+}
+
+void ULevelInstanceSubsystem::OnTryExitEditorMode()
+{	
+	if (OnExitEditorModeInternal(/*bForceExit=*/false))
+	{
+		ILevelInstanceEditorModule& EditorModule = FModuleManager::GetModuleChecked<ILevelInstanceEditorModule>("LevelInstanceEditor");
+		EditorModule.DeactivateEditorMode();
+	}
+}
+
+bool ULevelInstanceSubsystem::OnExitEditorModeInternal(bool bForceExit)
+{
+	if (bIsCommittingLevelInstance || bIsCreatingLevelInstance)
 	{
+		return false;
+	}
+
+	if (LevelInstanceEdit)
+	{
+		TGuardValue<bool> CommitScope(bIsCommittingLevelInstance, true);
 		ALevelInstance* LevelInstance = GetEditingLevelInstance();
 
 		bool bDiscard = false;
@@ -418,11 +454,20 @@ void ULevelInstanceSubsystem::OnExitEditorMode()
 		if (bIsDirty && CanDiscardLevelInstance(LevelInstance))
 		{
 			FText Title = LOCTEXT("CommitOrDiscardChangesTitle", "Save changes?");
-			bDiscard = (FMessageDialog::Open(EAppMsgType::YesNo, LOCTEXT("CommitOrDiscardChangesMsg", "Unsaved Level changes will get discarded. Do you want to save them now?"), &Title) == EAppReturnType::No);
+			// if bForceExit we can't cancel the exiting of the mode so the user needs to decide between saving or discarding
+			EAppReturnType::Type Ret = FMessageDialog::Open(bForceExit ? EAppMsgType::YesNo : EAppMsgType::YesNoCancel, LOCTEXT("CommitOrDiscardChangesMsg", "Unsaved Level changes will get discarded. Do you want to save them now?"), &Title);
+			if (Ret == EAppReturnType::Cancel)
+			{
+				return false;
+			}
+
+			bDiscard = (Ret == EAppReturnType::No);
 		}
 
-		CommitLevelInstanceInternal(LevelInstanceEdit, bDiscard);
+		return CommitLevelInstanceInternal(LevelInstanceEdit, bDiscard, /*bDiscardOnFailure=*/bForceExit);
 	}
+
+	return false;
 }
 
 bool ULevelInstanceSubsystem::CanPackAllLoadedActors() const
@@ -743,7 +788,8 @@ bool ULevelInstanceSubsystem::MoveActorsTo(ALevelInstance* LevelInstanceActor, c
 
 ALevelInstance* ULevelInstanceSubsystem::CreateLevelInstanceFrom(const TArray<AActor*>& ActorsToMove, const FNewLevelInstanceParams& CreationParams)
 {
-	TGuardValue<bool> CreateLevelInstanceGuard(bCreatingLevelInstance, true);
+	check(!bIsCreatingLevelInstance);
+	TGuardValue<bool> CreateLevelInstanceGuard(bIsCreatingLevelInstance, true);
 	ULevel* CurrentLevel = GetWorld()->GetCurrentLevel();
 		
 	if (ActorsToMove.Num() == 0)
@@ -819,8 +865,6 @@ ALevelInstance* ULevelInstanceSubsystem::CreateLevelInstanceFrom(const TArray<AA
 
 	ULevelStreamingLevelInstanceEditor* LevelStreaming = nullptr;
 	{
-		check(!bIsCreatingNewLevelStreamingLevelInstanceEditor);
-		TGuardValue<bool> IsCreatingNewLevelStreamingLevelInstanceEditor(bIsCreatingNewLevelStreamingLevelInstanceEditor, true);
 		LevelStreaming = StaticCast<ULevelStreamingLevelInstanceEditor*>(EditorLevelUtils::CreateNewStreamingLevelForWorld(
 		*GetWorld(), ULevelStreamingLevelInstanceEditor::StaticClass(), CreationParams.UseExternalActors(), LevelFilename, &ActorsToMove, CreationParams.TemplateWorld));
 	}
@@ -917,6 +961,7 @@ ALevelInstance* ULevelInstanceSubsystem::CreateLevelInstanceFrom(const TArray<AA
 	}
 	
 	// New level instance
+	const FLevelInstanceID NewLevelInstanceID = NewLevelInstanceActor->GetLevelInstanceID();
 	TUniquePtr<FLevelInstanceEdit> TempLevelInstanceEdit = MakeUnique<FLevelInstanceEdit>(LevelStreaming, NewLevelInstanceActor->GetLevelInstanceID());
 	// Force mark it as changed
 	TempLevelInstanceEdit->MarkCommittedChanges();
@@ -929,7 +974,8 @@ ALevelInstance* ULevelInstanceSubsystem::CreateLevelInstanceFrom(const TArray<AA
 		DirtyPackages.Add(NewLevelInstanceActor->GetPackage()->GetFName());
 	}
 			
-	ALevelInstance* CommittedLevelInstance = CommitLevelInstanceInternal(TempLevelInstanceEdit, /*bDiscardEdits*/false, &DirtyPackages);
+	bool bCommitted = CommitLevelInstanceInternal(TempLevelInstanceEdit, /*bDiscardEdits=*/false, /*bDiscardOnFailure=*/true, &DirtyPackages);
+	check(bCommitted);
 	check(!TempLevelInstanceEdit);
 
 	// EditorLevelUtils::CreateNewStreamingLevelForWorld deactivates all modes. Re-activate if needed
@@ -939,7 +985,7 @@ ALevelInstance* ULevelInstanceSubsystem::CreateLevelInstanceFrom(const TArray<AA
 		EditorModule.ActivateEditorMode();
 	}
 
-	return CommittedLevelInstance;
+	return GetLevelInstance(NewLevelInstanceID);
 }
 
 bool ULevelInstanceSubsystem::BreakLevelInstance(ALevelInstance* LevelInstanceActor, uint32 Levels /* = 1 */, TArray<AActor*>* OutMovedActors /* = nullptr */)
@@ -1230,11 +1276,6 @@ bool ULevelInstanceSubsystem::CanMoveActorToLevel(const AActor* Actor, FText* Ou
 	return true;
 }
 
-void ULevelInstanceSubsystem::DiscardEdits()
-{
-	LevelInstanceEdit.Reset();
-}
-
 void ULevelInstanceSubsystem::OnActorDeleted(AActor* Actor)
 {
 	if (ALevelInstance* LevelInstanceActor = Cast<ALevelInstance>(Actor))
@@ -1281,7 +1322,7 @@ void ULevelInstanceSubsystem::OnActorDeleted(AActor* Actor)
 				if (const FLevelInstanceEdit* ChildLevelInstanceEdit = GetLevelInstanceEdit(ChildLevelInstanceActor))
 				{
 					check(!IsLevelInstanceEditDirty(ChildLevelInstanceEdit));
-					LevelInstanceEdit.Reset();
+					ResetEdit(LevelInstanceEdit);
 					return false;
 				}
 				return true;
@@ -1395,16 +1436,6 @@ FLevelInstanceID ULevelInstanceSubsystem::FLevelInstanceEdit::GetLevelInstanceID
 	return LevelStreaming ? LevelStreaming->GetLevelInstanceID() : FLevelInstanceID();
 }
 
-ULevel* ULevelInstanceSubsystem::GetLevelInstanceEditLevel(const ALevelInstance* LevelInstanceActor) const
-{
-	const FLevelInstanceEdit* CurrentEdit = GetLevelInstanceEdit(LevelInstanceActor);
-	if (UWorld* EditingWorld = CurrentEdit ? CurrentEdit->GetEditWorld() : nullptr)
-	{
-		return EditingWorld->PersistentLevel;
-	}
-	return nullptr;
-}
-
 const ULevelInstanceSubsystem::FLevelInstanceEdit* ULevelInstanceSubsystem::GetLevelInstanceEdit(const ALevelInstance* LevelInstanceActor) const
 {
 	if (LevelInstanceEdit && LevelInstanceEdit->GetLevelInstanceID() == LevelInstanceActor->GetLevelInstanceID())
@@ -1662,20 +1693,24 @@ bool ULevelInstanceSubsystem::EditLevelInstanceInternal(ALevelInstance* LevelIns
 	return true;
 }
 
-ALevelInstance* ULevelInstanceSubsystem::CommitLevelInstance(ALevelInstance* LevelInstanceActor, bool bDiscardEdits, TSet<FName>* DirtyPackages)
+bool ULevelInstanceSubsystem::CommitLevelInstance(ALevelInstance* LevelInstanceActor, bool bDiscardEdits, TSet<FName>* DirtyPackages)
 {
 	check(LevelInstanceEdit.Get() == GetLevelInstanceEdit(LevelInstanceActor));
 	check(CanCommitLevelInstance(LevelInstanceActor));
-	ALevelInstance* CommittedLevelInstance = CommitLevelInstanceInternal(LevelInstanceEdit, bDiscardEdits, DirtyPackages);
+	if (CommitLevelInstanceInternal(LevelInstanceEdit, bDiscardEdits, /*bDiscardOnFailure=*/false, DirtyPackages))
+	{
+		ILevelInstanceEditorModule& EditorModule = FModuleManager::GetModuleChecked<ILevelInstanceEditorModule>("LevelInstanceEditor");
+		EditorModule.DeactivateEditorMode();
 
-	ILevelInstanceEditorModule& EditorModule = FModuleManager::GetModuleChecked<ILevelInstanceEditorModule>("LevelInstanceEditor");
-	EditorModule.DeactivateEditorMode();
+		return true;
+	}
 
-	return CommittedLevelInstance;
+	return false;
 }
 
-ALevelInstance* ULevelInstanceSubsystem::CommitLevelInstanceInternal(TUniquePtr<FLevelInstanceEdit>& InLevelInstanceEdit, bool bDiscardEdits, TSet<FName>* DirtyPackages)
+bool ULevelInstanceSubsystem::CommitLevelInstanceInternal(TUniquePtr<FLevelInstanceEdit>& InLevelInstanceEdit, bool bDiscardEdits, bool bDiscardOnFailure, TSet<FName>* DirtyPackages)
 {
+	TGuardValue<bool> CommitScope(bIsCommittingLevelInstance, true);
 	ALevelInstance* LevelInstanceActor = GetLevelInstance(InLevelInstanceEdit->GetLevelInstanceID());
 	check(InLevelInstanceEdit);
 	UWorld* EditingWorld = InLevelInstanceEdit->GetEditWorld();
@@ -1704,6 +1739,7 @@ ALevelInstance* ULevelInstanceSubsystem::CommitLevelInstanceInternal(TUniquePtr<
 		PackagesToSave.Append(*DirtyPackages);
 	}
 		
+	// Did some change get saved outside of the commit (regular saving in editor while editing)
 	bool bChangesCommitted = InLevelInstanceEdit->HasCommittedChanges();
 	if (PackagesToSave.Num() > 0 && !bDiscardEdits)
 	{
@@ -1714,19 +1750,23 @@ ALevelInstance* ULevelInstanceSubsystem::CommitLevelInstanceInternal(TUniquePtr<
 		const bool bNotifyNoPackagesSaved = false;
 		const bool bCanBeDeclined = true;
 
-		if (!FEditorFileUtils::SaveDirtyPackages(bPromptUserToSave, bSaveMapPackages, bSaveContentPackages, bFastSave, bNotifyNoPackagesSaved, bCanBeDeclined, nullptr,
-			[&PackagesToSave,EditingWorld](UPackage* DirtyPackage)
+		bool bSaveSucceeded = FEditorFileUtils::SaveDirtyPackages(bPromptUserToSave, bSaveMapPackages, bSaveContentPackages, bFastSave, bNotifyNoPackagesSaved, bCanBeDeclined, nullptr,
+			[&PackagesToSave, EditingWorld](UPackage* DirtyPackage)
 			{
 				if (PackagesToSave.Contains(DirtyPackage->GetFName()))
 				{
 					return false;
 				}
 				return ShouldIgnoreDirtyPackage(DirtyPackage, EditingWorld);
-			}))
+			});
+				
+		if (!bSaveSucceeded && !bDiscardOnFailure)
 		{
-			return LevelInstanceActor;
+			return false;
 		}
-		bChangesCommitted = true;
+
+		// Consider changes committed is was already set to true because of outside saves or if the save succeeded.
+		bChangesCommitted |= bSaveSucceeded;
 	}
 
 	FScopedSlowTask SlowTask(0, LOCTEXT("EndEditLevelInstance", "Unloading Level..."), !GetWorld()->IsGameWorld());
@@ -1737,7 +1777,7 @@ ALevelInstance* ULevelInstanceSubsystem::CommitLevelInstanceInternal(TUniquePtr<
 	const FString EditPackage = LevelInstanceActor->GetWorldAssetPackage();
 
 	// Remove from streaming level...
-	InLevelInstanceEdit.Reset();
+	ResetEdit(InLevelInstanceEdit);
 
 	if (bChangesCommitted)
 	{
@@ -1807,7 +1847,7 @@ ALevelInstance* ULevelInstanceSubsystem::CommitLevelInstanceInternal(TUniquePtr<
 
 	GEngine->BroadcastLevelActorListChanged();
 
-	return GetLevelInstance(LevelInstanceID);
+	return true;
 }
 
 void ULevelInstanceSubsystem::SaveLevelInstanceAs(ALevelInstance* LevelInstanceActor)
@@ -1833,7 +1873,7 @@ void ULevelInstanceSubsystem::SaveLevelInstanceAs(ALevelInstance* LevelInstanceA
 
 	UWorld* SavedWorld = StaticCast<UWorld*>(OutObjects[0]);
 	// Discard edits and unload streaming level
-	DiscardEdits();
+	ResetEdit(LevelInstanceEdit);
 	
 	LevelInstanceActor->SetWorldAsset(SavedWorld);
 
diff --git a/Engine/Source/Runtime/Engine/Public/LevelInstance/ILevelInstanceEditorModule.h b/Engine/Source/Runtime/Engine/Public/LevelInstance/ILevelInstanceEditorModule.h
--- a/Engine/Source/Runtime/Engine/Public/LevelInstance/ILevelInstanceEditorModule.h
+++ b/Engine/Source/Runtime/Engine/Public/LevelInstance/ILevelInstanceEditorModule.h
@@ -15,7 +15,12 @@ public:
 	virtual void ActivateEditorMode() = 0;
 	virtual void DeactivateEditorMode() = 0;
 
+	virtual void BroadcastTryExitEditorMode() = 0;
+
 	/** Broadcasts before exiting mode */
 	DECLARE_EVENT(ILevelInstanceEditorModule, FExitEditorModeEvent);
 	virtual FExitEditorModeEvent& OnExitEditorMode() = 0;
+
+	DECLARE_EVENT(ILevelInstanceEditorModule, FTryExitEditorModeEvent);
+	virtual FTryExitEditorModeEvent& OnTryExitEditorMode() = 0;
 };
\ No newline at end of file
diff --git a/Engine/Source/Runtime/Engine/Public/LevelInstance/LevelInstanceSubsystem.h b/Engine/Source/Runtime/Engine/Public/LevelInstance/LevelInstanceSubsystem.h
--- a/Engine/Source/Runtime/Engine/Public/LevelInstance/LevelInstanceSubsystem.h
+++ b/Engine/Source/Runtime/Engine/Public/LevelInstance/LevelInstanceSubsystem.h
@@ -32,7 +32,6 @@ public:
 		
 	//~ Begin USubsystem Interface.
 	virtual void Initialize(FSubsystemCollectionBase& Collection) override;
-	virtual void Deinitialize() override;
 	virtual bool DoesSupportWorldType(EWorldType::Type WorldType) const override;
 	//~ End USubsystem Interface.
 	
@@ -51,6 +50,8 @@ public:
 #if WITH_EDITOR
 	void Tick();
 	void OnExitEditorMode();
+	void OnTryExitEditorMode();
+	bool OnExitEditorModeInternal(bool bForceExit);
 	void PackAllLoadedActors();
 	bool CanPackAllLoadedActors() const;
 
@@ -59,11 +60,10 @@ public:
 	bool CanCommitLevelInstance(const ALevelInstance* LevelInstanceActor, FText* OutReason = nullptr) const;
 	bool CanDiscardLevelInstance(const ALevelInstance* LevelInstanceActor, FText* OutReason = nullptr) const;
 	void EditLevelInstance(ALevelInstance* LevelInstanceActor, TWeakObjectPtr<AActor> ContextActorPtr = nullptr);
-	ALevelInstance* CommitLevelInstance(ALevelInstance* LevelInstanceActor, bool bDiscardEdits = false, TSet<FName>* DirtyPackages = nullptr);
+	bool CommitLevelInstance(ALevelInstance* LevelInstanceActor, bool bDiscardEdits = false, TSet<FName>* DirtyPackages = nullptr);
 	void SaveLevelInstanceAs(ALevelInstance* LevelInstanceActor);
 	bool IsEditingLevelInstanceDirty(const ALevelInstance* LevelInstanceActor) const;
 	bool IsEditingLevelInstance(const ALevelInstance* LevelInstanceActor) const { return GetLevelInstanceEdit(LevelInstanceActor) != nullptr; }
-	ULevel* GetLevelInstanceEditLevel(const ALevelInstance* LevelInstanceActor) const;
 	
 	bool GetLevelInstanceBounds(const ALevelInstance* LevelInstanceActor, FBox& OutBounds) const;
 	static bool GetLevelInstanceBoundsFromPackage(const FTransform& InstanceTransform, FName LevelPackage, FBox& OutBounds);
@@ -86,7 +86,6 @@ public:
 	bool BreakLevelInstance(ALevelInstance* LevelInstanceActor, uint32 Levels = 1, TArray<AActor*>* OutMovedActors = nullptr);
 
 	bool CanMoveActorToLevel(const AActor* Actor, FText* OutReason = nullptr) const;
-	void DiscardEdits();
 	void OnActorDeleted(AActor* Actor);
 	ULevel* GetLevelInstanceLevel(const ALevelInstance* LevelInstanceActor) const;
 
@@ -143,8 +142,9 @@ private:
 		void MarkCommittedChanges();
 	};
 
+	void ResetEdit(TUniquePtr<FLevelInstanceEdit>& InLevelInstanceEdit);
 	bool EditLevelInstanceInternal(ALevelInstance* LevelInstanceActor, TWeakObjectPtr<AActor> ContextActorPtr, bool bRecursive);
-	ALevelInstance* CommitLevelInstanceInternal(TUniquePtr<FLevelInstanceEdit>& InLevelInstanceEdit, bool bDiscardEdits = false, TSet<FName>* DirtyPackages = nullptr);
+	bool CommitLevelInstanceInternal(TUniquePtr<FLevelInstanceEdit>& InLevelInstanceEdit, bool bDiscardEdits = false, bool bDiscardOnFailure = false, TSet<FName>* DirtyPackages = nullptr);
 	
 	const FLevelInstanceEdit* GetLevelInstanceEdit(const ALevelInstance* LevelInstanceActor) const;
 	bool IsLevelInstanceEditDirty(const FLevelInstanceEdit* LevelInstanceEdit) const;
@@ -167,7 +167,8 @@ private:
 	friend ULevelStreamingLevelInstanceEditor;
 
 #if WITH_EDITOR
-	bool bIsCreatingNewLevelStreamingLevelInstanceEditor;
+	bool bIsCreatingLevelInstance;
+	bool bIsCommittingLevelInstance;
 #endif
 
 	struct FLevelInstance
@@ -185,7 +186,6 @@ private:
 	TUniquePtr<FLevelsToRemoveScope> LevelsToRemoveScope;
 
 	TUniquePtr<FLevelInstanceEdit> LevelInstanceEdit;
-	bool bCreatingLevelInstance = false;
 
 	TMap<FLevelInstanceID, int32> ChildEdits;
 #endif
