PR #8727: Prevent rare crash in AudioDevice by adjusting nullptr check. (Contributed by [REDACTED_PERSON])

diff --git a/Engine/Source/Runtime/Engine/Private/AudioDevice.cpp b/Engine/Source/Runtime/Engine/Private/AudioDevice.cpp
--- a/Engine/Source/Runtime/Engine/Private/AudioDevice.cpp
+++ b/Engine/Source/Runtime/Engine/Private/AudioDevice.cpp
@@ -4134,7 +4134,7 @@ void FAudioDevice::StopSources(TArray<FWaveInstance*>& WaveInstances, int32 Firs
 				// Have to check it again here, because:
 				// - Source->NotifyPlaybackData() does not handle it if it is nullptr
 				// - SetPauseManually might set it to nullptr if the sound is unpaused and is stopping
-				if (WaveInstance != nullptr)
+				if (Source->WaveInstance != nullptr)
 				{
 					// Need to update the source still so that it gets any volume settings applied to
 					// otherwise the source may play at a very quiet volume and not actually set to 0.0
