Fixed volumetric cloud property documentation.

diff --git a/Engine/Source/Runtime/Engine/Classes/Components/VolumetricCloudComponent.h b/Engine/Source/Runtime/Engine/Classes/Components/VolumetricCloudComponent.h
--- a/Engine/Source/Runtime/Engine/Classes/Components/VolumetricCloudComponent.h
+++ b/Engine/Source/Runtime/Engine/Classes/Components/VolumetricCloudComponent.h
@@ -29,7 +29,7 @@ class UVolumetricCloudComponent : public USceneComponent
 	UPROPERTY(EditAnywhere, BlueprintReadOnly, interp, Category = "Layer", meta = (UIMin = 0.0f, UIMax = 20.0f, SliderExponent = 2.0))
 	float LayerBottomAltitude;
 
-	/** The altitude at which the cloud layer ends. (kilometers above the ground) */
+	/** The height of the the cloud layer. (kilometers above the layer bottom altitude) */
 	UPROPERTY(EditAnywhere, BlueprintReadOnly, interp, Category = "Layer", meta = (UIMin = 0.1f, UIMax = 20.0f, ClampMin = 0.1, SliderExponent = 2.0))
 	float LayerHeight;
 
