Backward compat for IsLooping() to fix some issues.

diff --git a/Engine/Plugins/Runtime/Metasound/Source/MetasoundEngine/Private/MetasoundSource.cpp b/Engine/Plugins/Runtime/Metasound/Source/MetasoundEngine/Private/MetasoundSource.cpp
--- a/Engine/Plugins/Runtime/Metasound/Source/MetasoundEngine/Private/MetasoundSource.cpp
+++ b/Engine/Plugins/Runtime/Metasound/Source/MetasoundEngine/Private/MetasoundSource.cpp
@@ -724,6 +724,11 @@ bool UMetaSoundSource::IsParameterValid(const FAudioParameter& InParameter, cons
 	return bIsValid;
 }
 
+bool UMetaSoundSource::IsLooping()
+{
+	return !IsOneShot();
+}
+
 bool UMetaSoundSource::IsOneShot()
 {
 	using namespace Metasound::Frontend;
diff --git a/Engine/Plugins/Runtime/Metasound/Source/MetasoundEngine/Public/MetasoundSource.h b/Engine/Plugins/Runtime/Metasound/Source/MetasoundEngine/Public/MetasoundSource.h
--- a/Engine/Plugins/Runtime/Metasound/Source/MetasoundEngine/Public/MetasoundSource.h
+++ b/Engine/Plugins/Runtime/Metasound/Source/MetasoundEngine/Public/MetasoundSource.h
@@ -168,6 +168,7 @@ public:
 	virtual ISoundGeneratorPtr CreateSoundGenerator(const FSoundGeneratorInitParams& InParams) override;
 	virtual TUniquePtr<Audio::IParameterTransmitter> CreateParameterTransmitter(Audio::FParameterTransmitterInitParams&& InParams) const override;
 	virtual bool IsParameterValid(const FAudioParameter& InParameter) const override;
+	virtual bool IsLooping() override;
 	virtual bool IsOneShot() override;
 
 protected:
diff --git a/Engine/Source/Runtime/Engine/Classes/Sound/SoundBase.h b/Engine/Source/Runtime/Engine/Classes/Sound/SoundBase.h
--- a/Engine/Source/Runtime/Engine/Classes/Sound/SoundBase.h
+++ b/Engine/Source/Runtime/Engine/Classes/Sound/SoundBase.h
@@ -253,7 +253,7 @@ public:
 	virtual class UCurveTable* GetCurveData() const { return nullptr; }
 
 	/** Returns whether or not this sound is looping. TODO: Deprecate this to only use IsOneshot() in a MetaSound world. */
-	bool IsLooping();
+	virtual bool IsLooping();
 
 	/** Query if it's one shot. One shot is defined as a sound which is intended to have a fixed duration. */
 	virtual bool IsOneShot();
diff --git a/Engine/Source/Runtime/Engine/Private/Components/AudioComponent.cpp b/Engine/Source/Runtime/Engine/Private/Components/AudioComponent.cpp
--- a/Engine/Source/Runtime/Engine/Private/Components/AudioComponent.cpp
+++ b/Engine/Source/Runtime/Engine/Private/Components/AudioComponent.cpp
@@ -603,7 +603,7 @@ void UAudioComponent::PlayInternal(const PlayInternalRequestData& InPlayRequestD
 	USoundBase* SoundToPlay = InSoundOverride ? InSoundOverride : Sound.Get();
 	if (SoundToPlay)
 	{
-		UE_LOG(LogAudio, Verbose, TEXT("%g: Playing AudioComponent : '%s' with Sound: '%s'"), AudioTimeSeconds, *GetFullName(), *SoundToPlay->GetName());
+		UE_LOG(LogAudio, Verbose, TEXT("%g: Playing AudioComponent : '%s' with Sound: '%s'. OneShot?: '%s'"), AudioTimeSeconds, *GetFullName(), *SoundToPlay->GetName(), !SoundToPlay->IsLooping() ? TEXT("YES") : TEXT("NO"));
 	}
 
 	// Reset our fading out flag in case this is a reused audio component and we are replaying after previously fading out
