Disable an ensure about duplicate primary asset ids in the editor and improve warning text

diff --git a/Engine/Source/Runtime/Engine/Private/AssetManager.cpp b/Engine/Source/Runtime/Engine/Private/AssetManager.cpp
--- a/Engine/Source/Runtime/Engine/Private/AssetManager.cpp
+++ b/Engine/Source/Runtime/Engine/Private/AssetManager.cpp
@@ -1014,12 +1014,9 @@ bool UAssetManager::TryUpdateCachedAssetData(const FPrimaryAssetId& PrimaryAsset
 
 		if (OldData && OldData->AssetPtr.ToSoftObjectPath() != NewAssetPath)
 		{
-			UE_LOG(LogAssetManager, Warning, TEXT("Found Duplicate PrimaryAssetID %s, this must be resolved before saving. Path %s is replacing path %s"), *PrimaryAssetId.ToString(), *OldData->AssetPtr.ToString(), *NewAssetPath.ToString());
-			// Don't ensure for editor only types, this will not cause an actual game problem
-			if (!bAllowDuplicates && !TypeData.Info.bIsEditorOnly)
-			{
-				ensureMsgf(!OldData, TEXT("Found Duplicate PrimaryAssetID %s! Path %s is replacing path %s"), *PrimaryAssetId.ToString(), *OldData->AssetPtr.ToString(), *NewAssetPath.ToString());
-			}
+			UE_LOG(LogAssetManager, Warning, 
+				TEXT("Found duplicate PrimaryAssetID %s, path %s conflicts with existing path %s. Two different primary assets can not have the same type and name."), 
+				*PrimaryAssetId.ToString(), *OldData->AssetPtr.ToString(), *NewAssetPath.ToString());
 
 #if WITH_EDITOR
 			if (GIsEditor)
@@ -1027,7 +1024,7 @@ bool UAssetManager::TryUpdateCachedAssetData(const FPrimaryAssetId& PrimaryAsset
 				const int MaxNotificationsPerFrame = 5;
 				if (NumberOfSpawnedNotifications++ < MaxNotificationsPerFrame)
 				{
-					FNotificationInfo Info(FText::Format(LOCTEXT("DuplicateAssetId", "Duplicate Asset ID {0} used by {1} and {2}, you must delete or rename one!"),
+					FNotificationInfo Info(FText::Format(LOCTEXT("DuplicateAssetId", "Duplicate Asset ID {0} used by both {1} and {2}, rename one to avoid a conflict."),
 						FText::FromString(PrimaryAssetId.ToString()), FText::FromString(OldData->AssetPtr.ToSoftObjectPath().GetLongPackageName()), FText::FromString(NewAssetPath.GetLongPackageName())));
 					Info.ExpireDuration = 30.0f;
 
@@ -1038,7 +1035,13 @@ bool UAssetManager::TryUpdateCachedAssetData(const FPrimaryAssetId& PrimaryAsset
 					}
 				}
 			}
+			else
 #endif
+			// Don't ensure for editor only types as they will not cause an actual game problem
+			if (!bAllowDuplicates && !TypeData.Info.bIsEditorOnly)
+			{
+				ensureMsgf(!OldData, TEXT("Found Duplicate PrimaryAssetID %s! Path %s conflicts with existing path %s"), *PrimaryAssetId.ToString(), *OldData->AssetPtr.ToString(), *NewAssetPath.ToString());
+			}
 		}
 
 		FPrimaryAssetData& NameData = TypeData.AssetMap.FindOrAdd(PrimaryAssetId.PrimaryAssetName);
