Change a crash to an ensure when a client tries to do a server travel, it will disconnect right afterwards

diff --git a/Engine/Source/Runtime/Engine/Private/UnrealEngine.cpp b/Engine/Source/Runtime/Engine/Private/UnrealEngine.cpp
--- a/Engine/Source/Runtime/Engine/Private/UnrealEngine.cpp
+++ b/Engine/Source/Runtime/Engine/Private/UnrealEngine.cpp
@@ -13558,10 +13558,10 @@ void UEngine::TickWorldTravel(FWorldContext& Context, float DeltaSeconds)
 			if (Ret != EBrowseReturnVal::Success )
 			{
 				UE_LOG(LogLoad, Warning, TEXT("UEngine::TickWorldTravel failed to Handle server travel to URL: %s. Error: %s"), *NextURL, *Error);
-				check(Ret != EBrowseReturnVal::Pending); // server travel should never create a pending net game
+				ensureMsgf(Ret != EBrowseReturnVal::Pending, TEXT("Server travel should never create a pending net game"));
 
-														 // Failed to load a new map
-				if (Context.World() != NULL)
+				// Failed to load a new map
+				if (Context.World() != nullptr)
 				{
 					// If we didn't change worlds, clear out NextURL so we don't do this again next frame.
 					Context.World()->NextURL = TEXT("");
