Removing stat audio

diff --git a/Engine/Source/Runtime/Engine/Classes/Engine/Engine.h b/Engine/Source/Runtime/Engine/Classes/Engine/Engine.h
--- a/Engine/Source/Runtime/Engine/Classes/Engine/Engine.h
+++ b/Engine/Source/Runtime/Engine/Classes/Engine/Engine.h
@@ -3532,12 +3532,6 @@ private:
 	bool ToggleStatUnitGraph(UWorld* World, FCommonViewportClient* ViewportClient, const TCHAR* Stream = nullptr);
 	bool ToggleStatUnitTime(UWorld* World, FCommonViewportClient* ViewportClient, const TCHAR* Stream = nullptr);
 	bool ToggleStatRaw(UWorld* World, FCommonViewportClient* ViewportClient, const TCHAR* Stream = nullptr);
-	bool ToggleStatSoundWaves(UWorld* World, FCommonViewportClient* ViewportClient, const TCHAR* Stream = nullptr);
-	bool ToggleStatSoundCues(UWorld* World, FCommonViewportClient* ViewportClient, const TCHAR* Stream = nullptr);
-	bool ToggleStatSounds(UWorld* World, FCommonViewportClient* ViewportClient, const TCHAR* Stream = nullptr);
-	bool ToggleStatAudioStreaming(UWorld* World, FCommonViewportClient* ViewportClient, const TCHAR* Stream = nullptr);
-	bool ToggleStatSoundMixes(UWorld* World, FCommonViewportClient* ViewportClient, const TCHAR* Stream = nullptr);
-	bool ToggleStatSoundModulators(UWorld* World, FCommonViewportClient* ViewportClient, const TCHAR* Stream = nullptr);
 	bool ToggleStatParticlePerf(UWorld* World, FCommonViewportClient* ViewportClient, const TCHAR* Stream = nullptr);
 #endif
 
diff --git a/Engine/Source/Runtime/Engine/Private/AudioDebug.cpp b/Engine/Source/Runtime/Engine/Private/AudioDebug.cpp
--- a/Engine/Source/Runtime/Engine/Private/AudioDebug.cpp
+++ b/Engine/Source/Runtime/Engine/Private/AudioDebug.cpp
@@ -223,8 +223,6 @@ namespace Audio
 		return FColor::White;
 	}
 
-	static bool bAllowUsingDeprecatedDebugStats = true;
-
 	// Whether or not respective stat data is active for any audio device (set on game thread)
 	static bool bDebugWavesForAllViewsEnabled = false;
 	static bool bDebugCuesForAllViewsEnabled = false;
@@ -602,8 +600,6 @@ namespace Audio
 
 	void DebugSoundObject(const TArray<FString>& Args, UWorld* InWorld, const uint32 InStatToEnable, bool& bAllEnabled)
 	{
-		bAllowUsingDeprecatedDebugStats = false;
-
 		if (Args.Num() > 0)
 		{
 			bool bAllViews = false;
@@ -633,7 +629,7 @@ namespace Audio
 			return false;
 		}
 
-		if (bEnablementBool || bAllowUsingDeprecatedDebugStats)
+		if (bEnablementBool)
 		{
 			return true;
 		}
@@ -775,11 +771,6 @@ namespace Audio
 	{
 		WorldRegisteredWithDeviceHandle = FAudioDeviceWorldDelegates::OnWorldRegisteredToAudioDevice.AddLambda([this](const UWorld* InWorld, FDeviceId InDeviceId)
 		{
-			if (bAllowUsingDeprecatedDebugStats)
-			{
-				return;
-			}
-
 			uint32 StatsToSet = 0;
 			uint32 StatsToClear = 0;
 			if (bDebugSoundsForAllViewsEnabled)
@@ -1317,10 +1308,9 @@ namespace Audio
 						{
 							ViewportClient->AddRealtimeOverride(true, SystemDisplayName);
 						}
-						else if (!bAllowUsingDeprecatedDebugStats)
-						{
-							ViewportClient->RemoveRealtimeOverride(SystemDisplayName);
-						}
+
+						ViewportClient->RemoveRealtimeOverride(SystemDisplayName);
+
 						return DrawDebugStatsInternal(World, *ViewportClient->Viewport, InCanvas, InY);
 					}
 				}
@@ -1344,7 +1334,7 @@ namespace Audio
 	int32 FAudioDebugger::DrawDebugStatsInternal(UWorld& World, FViewport& Viewport, FCanvas* InCanvas, int32 InY)
 	{
 		FCanvas* Canvas = InCanvas ? InCanvas : Viewport.GetDebugCanvas();
-		if (!Canvas || bAllowUsingDeprecatedDebugStats)
+		if (!Canvas)
 		{
 			return InY;
 		}
@@ -2165,27 +2155,6 @@ namespace Audio
 		AudioDeviceStats_AudioThread.Remove(AudioDevice.DeviceID);
 	}
 
-	bool FAudioDebugger::ToggleStats(UWorld* World, const uint32 StatToToggle)
-	{
-		if (!GEngine)
-		{
-			return false;
-		}
-
-		FAudioDevice* AudioDevice = GetWorldAudio(World);
-		if (!AudioDevice)
-		{
-			return false;
-		}
-
-		if (FAudioDeviceManager* DeviceManager = GEngine->GetAudioDeviceManager())
-		{
-			DeviceManager->GetDebugger().ToggleStats(AudioDevice->DeviceID, StatToToggle);
-		}
-
-		return true;
-	}
-
 	void FAudioDebugger::ClearStats(const uint32 StatsToClear, UWorld* InWorld)
 	{
 		if (!GEngine)
@@ -2276,60 +2245,6 @@ namespace Audio
 		Stats.RequestedStats |= StatsToSet;
 	}
 
-	void FAudioDebugger::ToggleStats(FDeviceId DeviceId, const uint32 StatsToToggle)
-	{
-		if (!IsInAudioThread())
-		{
-			DECLARE_CYCLE_STAT(TEXT("FAudioThreadTask.ToggleStats"), STAT_AudioToggleStats, STATGROUP_TaskGraphTasks);
-
-			FAudioThread::RunCommandOnAudioThread([this, DeviceId, StatsToToggle]()
-			{
-				ToggleStats(DeviceId, StatsToToggle);
-			}, GET_STATID(STAT_AudioToggleStats));
-			return;
-		}
-
-		FAudioStats_AudioThread& Stats = AudioDeviceStats_AudioThread.FindOrAdd(DeviceId);
-		Stats.RequestedStats ^= StatsToToggle;
-	}
-
-	bool FAudioDebugger::ToggleStatCues(UWorld* World, FCommonViewportClient* ViewportClient, const TCHAR* Stream)
-	{
-		return ToggleStats(World, FAudioStats::SoundCues);
-	}
-
-		bool FAudioDebugger::ToggleStatMixes(UWorld* World, FCommonViewportClient* ViewportClient, const TCHAR* Stream)
-	{
-		return ToggleStats(World, FAudioStats::SoundMixes);
-	}
-
-	bool FAudioDebugger::ToggleStatModulators(UWorld* World, FCommonViewportClient* ViewportClient, const TCHAR* Stream)
-	{
-		if (!GEngine)
-		{
-			return false;
-		}
-
-		FAudioDevice* AudioDevice = GetWorldAudio(World);
-		if (!AudioDevice)
-		{
-			return false;
-		}
-
-		if (AudioDevice->IsModulationPluginEnabled())
-		{
-			if (IAudioModulationManager* Modulation = AudioDevice->ModulationInterface.Get())
-			{
-				if (!Modulation->OnToggleStat(ViewportClient, Stream))
-				{
-					return false;
-				}
-			}
-		}
-
-		return true;
-	}
-
 	bool FAudioDebugger::PostStatModulatorHelp(UWorld* World, FCommonViewportClient* ViewportClient, const TCHAR* Stream)
 	{
 		// Ignore if all Viewports are closed.
@@ -2355,80 +2270,6 @@ namespace Audio
 		return true;
 	}
 
-	bool FAudioDebugger::ToggleStatSounds(UWorld* World, FCommonViewportClient* ViewportClient, const TCHAR* Stream)
-	{
-		// Ignore if all Viewports are closed.
-		if (!ViewportClient)
-		{
-			return false;
-		}
-
-		if (!ToggleStats(World, FAudioStats::Sounds))
-		{
-			return false;
-		}
-
-		const bool bHelp = Stream ? FCString::Stristr(Stream, TEXT("?")) != nullptr : false;
-		if (bHelp)
-		{
-			GLog->Logf(TEXT("Shows all active sounds. Displays value sorted by when sort is set"));
-			GLog->Logf(TEXT("  stat sounds sort=class|distance|name|priority|time|volume|waves"));
-			GLog->Logf(TEXT("      distance - sort list by distance to player"));
-			GLog->Logf(TEXT("      class - sort by sound class name"));
-			GLog->Logf(TEXT("      name - sort by cue pathname (default)"));
-			GLog->Logf(TEXT("      time - sort by time played back"));
-			GLog->Logf(TEXT("      waves - sort by waves' num"));
-			GLog->Logf(TEXT("      priority - sort by playback priority"));
-			GLog->Logf(TEXT("      volume - sort by volume"));
-			GLog->Logf(TEXT("  stat sounds -debug - enables debugging mode like showing sound radius sphere and names, but only for sounds with enabled 'Debug' property"));
-			GLog->Logf(TEXT(""));
-			GLog->Logf(TEXT("Ex. stat sounds sort=class -debug"));
-			GLog->Logf(TEXT(" This will show only sounds that have 'Debug' property set to true and will sort them by and display their sound class"));
-		}
-
-		uint8 ShowSounds = 0;
-		FAudioStats::EDisplaySort DisplaySort = FAudioStats::EDisplaySort::Name;
-		if (Stream)
-		{
-			const bool bHide = FParse::Command(&Stream, TEXT("off"));
-			if (!bHide)
-			{
-				const bool bDebug = FParse::Param(Stream, TEXT("debug"));
-				if (bDebug)
-				{
-					ShowSounds |= static_cast<uint8>(FAudioStats::EDisplayFlags::Debug);
-				}
-
-				const bool bLongNames = FParse::Param(Stream, TEXT("longnames"));
-				if (bLongNames)
-				{
-					ShowSounds |= static_cast<uint8>(FAudioStats::EDisplayFlags::Long_Names);
-				}
-
-				FString SortStr;
-				if (FParse::Value(Stream, TEXT("sort="), SortStr))
-				{
-					AudioDebugSoundSortCVarCVar = SortStr;
-				}
-			}
-		}
-
-		FAudioDevice* AudioDevice = GetWorldAudio(World);
-		check(AudioDevice);
-		FAudioStats& Stats = AudioDeviceStats.FindOrAdd(AudioDevice->DeviceID);
-		Stats.DisplayFlags = ShowSounds;
-		Stats.DisplaySort = DisplaySort;
-
-		ResolveDesiredStats(ViewportClient);
-
-		return true;
-	}
-
-	bool FAudioDebugger::ToggleStatWaves(UWorld* World, FCommonViewportClient* ViewportClient, const TCHAR* Stream)
-	{
-		return ToggleStats(World, FAudioStats::SoundWaves);
-	}
-
 	void FAudioDebugger::SendUpdateResultsToGameThread(const FAudioDevice& AudioDevice, const int32 FirstActiveIndex)
 	{
 		check(IsInAudioThread());
diff --git a/Engine/Source/Runtime/Engine/Private/UnrealEngine.cpp b/Engine/Source/Runtime/Engine/Private/UnrealEngine.cpp
--- a/Engine/Source/Runtime/Engine/Private/UnrealEngine.cpp
+++ b/Engine/Source/Runtime/Engine/Private/UnrealEngine.cpp
@@ -2049,16 +2049,7 @@ void UEngine::Init(IEngineLoop* InEngineLoop)
 
 	EngineStats.Add(FEngineStatFuncs(TEXT("STAT_ColorList"), TEXT("STATCAT_Engine"), FText::GetEmpty(), FEngineStatRender::CreateUObject(this, &UEngine::RenderStatColorList), FEngineStatToggle()));
 	EngineStats.Add(FEngineStatFuncs(TEXT("STAT_Levels"), TEXT("STATCAT_Engine"), FText::GetEmpty(), FEngineStatRender::CreateUObject(this, &UEngine::RenderStatLevels), FEngineStatToggle()));
-#if !UE_BUILD_SHIPPING
-	EngineStats.Add(FEngineStatFuncs(TEXT("STAT_SoundMixes"), TEXT("STATCAT_Engine"), FText::GetEmpty(), FEngineStatRender::CreateUObject(this, &UEngine::RenderStatSoundMixes), FEngineStatToggle::CreateUObject(this, &UEngine::ToggleStatSoundMixes)));
-	EngineStats.Add(FEngineStatFuncs(TEXT("STAT_SoundModulators"), TEXT("STATCAT_Engine"), FText::GetEmpty(), FEngineStatRender::CreateUObject(this, &UEngine::RenderStatSoundModulators), FEngineStatToggle::CreateUObject(this, &UEngine::ToggleStatSoundModulators)));
-	EngineStats.Add(FEngineStatFuncs(TEXT("STAT_SoundModulatorsHelp"), TEXT("STATCAT_Engine"), FText::GetEmpty(), FEngineStatRender(), FEngineStatToggle::CreateUObject(this, &UEngine::PostStatSoundModulatorHelp)));
-	EngineStats.Add(FEngineStatFuncs(TEXT("STAT_AudioStreaming"), TEXT("STATCAT_Engine"), FText::GetEmpty(), FEngineStatRender::CreateUObject(this, &UEngine::RenderStatAudioStreaming), FEngineStatToggle::CreateUObject(this, &UEngine::ToggleStatAudioStreaming)));
-	EngineStats.Add(FEngineStatFuncs(TEXT("STAT_SoundReverb"), TEXT("STATCAT_Engine"), FText::GetEmpty(), FEngineStatRender::CreateUObject(this, &UEngine::RenderStatSoundReverb), FEngineStatToggle()));
-	EngineStats.Add(FEngineStatFuncs(TEXT("STAT_Sounds"), TEXT("STATCAT_Engine"), FText::GetEmpty(), FEngineStatRender::CreateUObject(this, &UEngine::RenderStatSounds), FEngineStatToggle::CreateUObject(this, &UEngine::ToggleStatSounds)));
-	EngineStats.Add(FEngineStatFuncs(TEXT("STAT_SoundCues"), TEXT("STATCAT_Engine"), FText::GetEmpty(), FEngineStatRender::CreateUObject(this, &UEngine::RenderStatSoundCues), FEngineStatToggle::CreateUObject(this, &UEngine::ToggleStatSoundCues)));
-	EngineStats.Add(FEngineStatFuncs(TEXT("STAT_SoundWaves"), TEXT("STATCAT_Engine"), FText::GetEmpty(), FEngineStatRender::CreateUObject(this, &UEngine::RenderStatSoundWaves), FEngineStatToggle::CreateUObject(this, &UEngine::ToggleStatSoundWaves)));
-#endif // !UE_BUILD_SHIPPING
+
 	/* @todo physx fix this once we have convexelem drawing again
 	EngineStats.Add(FEngineStatFuncs(TEXT("STAT_LevelMap"), TEXT("STATCAT_Engine"), FText::GetEmpty(), &UEngine::RenderStatLevelMap, NULL));
 	*/
@@ -17360,160 +17351,6 @@ static FAutoConsoleCommand PakFileTestCmd(
 	FConsoleCommandWithArgsDelegate::CreateStatic(&PakFileTest)
 );
 
-// REVERB
-int32 UEngine::RenderStatSoundReverb(UWorld* World, FViewport* Viewport, FCanvas* Canvas, int32 X, int32 Y, const FVector* ViewLocation, const FRotator* ViewRotation)
-{
-#if ENABLE_AUDIO_DEBUG
-	const int32 FontHeight = GetStatsFont()->GetMaxCharHeight() + 2.0f;
-	Canvas->DrawShadowedString(X, Y, TEXT("'stat soundreverb' is deprecated. Please use au.Debug.SoundReverb"), GetStatsFont(), FLinearColor::Yellow);
-	Y += FontHeight;
-
-	return Audio::FAudioDebugger::RenderStatReverb(World, Viewport, Canvas, X, Y);
-#else // !ENABLE_AUDIO_DEBUG
-	return Y;
-#endif // !ENABLE_AUDIO_DEBUG
-}
-
-// SOUNDMIXES
-int32 UEngine::RenderStatSoundMixes(UWorld* World, FViewport* Viewport, FCanvas* Canvas, int32 X, int32 Y, const FVector* ViewLocation, const FRotator* ViewRotation)
-{
-#if ENABLE_AUDIO_DEBUG
-	const int32 FontHeight = GetStatsFont()->GetMaxCharHeight() + 2.0f;
-	Canvas->DrawShadowedString(X, Y, TEXT("'stat soundmixes' is deprecated. Use au.Debug.SoundMixes"), GetStatsFont(), FLinearColor::Yellow);
-	Y += FontHeight;
-
-	return Audio::FAudioDebugger::RenderStatMixes(World, Viewport, Canvas, X, Y);
-#else // !ENABLE_AUDIO_DEBUG
-	return Y;
-#endif // !ENABLE_AUDIO_DEBUG
-}
-
-// SOUNDMODULATION
-int32 UEngine::RenderStatSoundModulators(UWorld* World, FViewport* Viewport, FCanvas* Canvas, int32 X, int32 Y, const FVector* ViewLocation, const FRotator* ViewRotation)
-{
-#if ENABLE_AUDIO_DEBUG
-	const int32 FontHeight = GetStatsFont()->GetMaxCharHeight() + 2.0f;
-	Canvas->DrawShadowedString(X, Y, TEXT("'stat soundmodulators' is deprecated. Use au.Debug.SoundModulators"), GetStatsFont(), FLinearColor::Yellow);
-	Y += FontHeight;
-
-	return Audio::FAudioDebugger::RenderStatModulators(World, Viewport, Canvas, X, Y, ViewLocation, ViewRotation);
-#else // !ENABLE_AUDIO_DEBUG
-	return Y;
-#endif // !ENABLE_AUDIO_DEBUG
-}
-
-// SOUNDWAVES
-bool UEngine::ToggleStatSoundWaves(UWorld* World, FCommonViewportClient* ViewportClient, const TCHAR* Stream)
-{
-#if ENABLE_AUDIO_DEBUG
-	return Audio::FAudioDebugger::ToggleStatWaves(World, ViewportClient, Stream);
-#else // !ENABLE_AUDIO_DEBUG
-	return false;
-#endif // !ENABLE_AUDIO_DEBUG
-}
-
-// SOUNDCUES
-bool UEngine::ToggleStatSoundCues(UWorld* World, FCommonViewportClient* ViewportClient, const TCHAR* Stream)
-{
-#if ENABLE_AUDIO_DEBUG
-	return Audio::FAudioDebugger::ToggleStatCues(World, ViewportClient, Stream);
-#else // !ENABLE_AUDIO_DEBUG
-	return false;
-#endif // !ENABLE_AUDIO_DEBUG
-}
-
-bool UEngine::ToggleStatAudioStreaming(UWorld* World, FCommonViewportClient* ViewportClient, const TCHAR* Stream /*= nullptr*/)
-{
-	// Noop for now:
-	return true;
-}
-
-// SOUNDMIXES
-bool UEngine::ToggleStatSoundMixes(UWorld* World, FCommonViewportClient* ViewportClient, const TCHAR* Stream)
-{
-#if ENABLE_AUDIO_DEBUG
-	return Audio::FAudioDebugger::ToggleStatMixes(World, ViewportClient, Stream);
-#else // !ENABLE_AUDIO_DEBUG
-	return false;
-#endif // !ENABLE_AUDIO_DEBUG
-}
-
-// SOUNDMODULATORS
-bool UEngine::PostStatSoundModulatorHelp(UWorld* World, FCommonViewportClient* ViewportClient, const TCHAR* Stream)
-{
-#if ENABLE_AUDIO_DEBUG
-	return Audio::FAudioDebugger::PostStatModulatorHelp(World, ViewportClient, Stream);
-#else // !ENABLE_AUDIO_DEBUG
-	return false;
-#endif // !ENABLE_AUDIO_DEBUG
-}
-
-// SOUNDMODULATORS
-bool UEngine::ToggleStatSoundModulators(UWorld* World, FCommonViewportClient* ViewportClient, const TCHAR* Stream)
-{
-#if ENABLE_AUDIO_DEBUG
-	return Audio::FAudioDebugger::ToggleStatModulators(World, ViewportClient, Stream);
-#else // !ENABLE_AUDIO_DEBUG
-	return false;
-#endif // !ENABLE_AUDIO_DEBUG
-}
-
-// SOUNDWAVES
-int32 UEngine::RenderStatSoundWaves(UWorld* World, FViewport* Viewport, FCanvas* Canvas, int32 X, int32 Y, const FVector* ViewLocation, const FRotator* ViewRotation)
-{
-#if ENABLE_AUDIO_DEBUG
-	const int32 FontHeight = GetStatsFont()->GetMaxCharHeight() + 2.0f;
-	Canvas->DrawShadowedString(X, Y, TEXT("'stat soundwaves' is deprecated. Use au.Debug.SoundWaves"), GetStatsFont(), FLinearColor::Yellow);
-	Y += FontHeight;
-
-	return Audio::FAudioDebugger::RenderStatWaves(World, Viewport, Canvas, X, Y);
-#else // !ENABLE_AUDIO_DEBUG
-	return Y;
-#endif // !ENABLE_AUDIO_DEBUG
-}
-
-// AUDIOSTREAMING
-int32 UEngine::RenderStatAudioStreaming(UWorld* World, FViewport* Viewport, FCanvas* Canvas, int32 X, int32 Y, const FVector* ViewLocation /*= nullptr*/, const FRotator* ViewRotation /*= nullptr*/)
-{
-	return IStreamingManager::Get().GetAudioStreamingManager().RenderStatAudioStreaming(World, Viewport, Canvas, X, Y, ViewLocation, ViewRotation);
-}
-
-// SOUNDCUES
-int32 UEngine::RenderStatSoundCues(UWorld* World, FViewport* Viewport, FCanvas* Canvas, int32 X, int32 Y, const FVector* ViewLocation, const FRotator* ViewRotation)
-{
-#if ENABLE_AUDIO_DEBUG
-	const int32 FontHeight = GetStatsFont()->GetMaxCharHeight() + 2.0f;
-	Canvas->DrawShadowedString(X, Y, TEXT("'stat soundcues' is deprecated. Use au.Debug.SoundCues"), GetStatsFont(), FLinearColor::Yellow);
-	Y += FontHeight;
-
-	return Audio::FAudioDebugger::RenderStatCues(World, Viewport, Canvas, X, Y);
-#else // !ENABLE_AUDIO_DEBUG
-	return Y;
-#endif // !ENABLE_AUDIO_DEBUG
-}
-
-// SOUNDS
-bool UEngine::ToggleStatSounds(UWorld* World, FCommonViewportClient* ViewportClient, const TCHAR* Stream)
-{
-#if ENABLE_AUDIO_DEBUG
-	return Audio::FAudioDebugger::ToggleStatSounds(World, ViewportClient, Stream);
-#else // !ENABLE_AUDIO_DEBUG
-	return false;
-#endif // !ENABLE_AUDIO_DEBUG
-}
-
-int32 UEngine::RenderStatSounds(UWorld* World, FViewport* Viewport, FCanvas* Canvas, int32 X, int32 Y, const FVector* ViewLocation, const FRotator* ViewRotation)
-{
-#if ENABLE_AUDIO_DEBUG
-	const int32 FontHeight = GetStatsFont()->GetMaxCharHeight() + 2.0f;
-	Canvas->DrawShadowedString(X, Y, TEXT("'stat sounds' is deprecated. Use au.Debug.Sounds"), GetStatsFont(), FLinearColor::Yellow);
-	Y += FontHeight;
-
-	return Audio::FAudioDebugger::RenderStatSounds(World, Viewport, Canvas, X, Y);
-#else // !ENABLE_AUDIO_DEBUG
-	return Y;
-#endif // !ENABLE_AUDIO_DEBUG
-}
 #endif // !UE_BUILD_SHIPPING
 
 // AI
