Convert checkf to ensure as it happens more frequently than I thought.

diff --git a/Engine/Source/Runtime/Engine/Private/Materials/MaterialShader.cpp b/Engine/Source/Runtime/Engine/Private/Materials/MaterialShader.cpp
--- a/Engine/Source/Runtime/Engine/Private/Materials/MaterialShader.cpp
+++ b/Engine/Source/Runtime/Engine/Private/Materials/MaterialShader.cpp
@@ -2819,7 +2819,7 @@ FMaterialShaderMap* FMaterialShaderMap::AcquireFinalizedClone()
 	Clone->bCompiledSuccessfully = bCompiledSuccessfully;
 	Clone->bIsPersistent = bIsPersistent;
 	Clone->AssignCopy(*this);
-	checkf(reinterpret_cast<uint64>(Clone->GetMaterialPath()) > reinterpret_cast<uint64>(Clone->GetContent()), TEXT("String embedded in memory image is located outside of it (image start %p > string %p). Pointer/offset overflow while freezing?"),
+	ensureMsgf(reinterpret_cast<uint64>(Clone->GetMaterialPath()) > reinterpret_cast<uint64>(Clone->GetContent()), TEXT("String embedded in memory image is located outside of it (image start %p > string %p). Pointer/offset overflow while freezing?"),
 		Clone->GetContent(), Clone->GetMaterialPath()
 		);
 #if WITH_EDITOR
