Change the ensure for destroying a NetDriver during its tick to not trigger for PendingNetGames

diff --git a/Engine/Source/Runtime/Engine/Private/UnrealEngine.cpp b/Engine/Source/Runtime/Engine/Private/UnrealEngine.cpp
--- a/Engine/Source/Runtime/Engine/Private/UnrealEngine.cpp
+++ b/Engine/Source/Runtime/Engine/Private/UnrealEngine.cpp
@@ -12740,7 +12740,7 @@ void DestroyNamedNetDriver_Local(FWorldContext &Context, FName NetDriverName)
 		{
 			UE_LOG(LogNet, Log, TEXT("DestroyNamedNetDriver %s [%s]"), *NetDriver->GetName(), *NetDriverName.ToString());
 			
-			ensureMsgf(!NetDriver->IsInTick(), TEXT("Attempting to destroy NetDriver %s [%s] while it is ticking."), *NetDriver->GetName(), *NetDriverName.ToString());
+			ensureMsgf(!(NetDriver->IsInTick() && NetDriverName != NAME_PendingNetDriver), TEXT("Attempting to destroy NetDriver %s [%s] while it is ticking."), *NetDriver->GetName(), *NetDriverName.ToString());
 
 			Context.ActiveNetDrivers.RemoveAtSwap(Index);
 			NetDriver->SetWorld(NULL);
