Some support for visual logs in multithread.

#rnx

diff --git a/Engine/Source/Runtime/Engine/Private/VisualLogger/VisualLogger.cpp b/Engine/Source/Runtime/Engine/Private/VisualLogger/VisualLogger.cpp
--- a/Engine/Source/Runtime/Engine/Private/VisualLogger/VisualLogger.cpp
+++ b/Engine/Source/Runtime/Engine/Private/VisualLogger/VisualLogger.cpp
@@ -119,30 +119,53 @@ bool FVisualLogger::IsClassAllowed(const UClass& InClass) const
 
 void FVisualLogger::AddObjectToAllowList(const UObject& InObject)
 {
-	const int32 PrevNum = ObjectAllowList.Num();
-	ObjectAllowList.Add(&InObject);
+	FWriteScopeLock Lock(EntryRWLock);
+	bool bIsAlreadyInSet = false;
+	ObjectAllowList.Add(&InObject, &bIsAlreadyInSet);
 
-	const bool bChanged = (PrevNum != ObjectAllowList.Num());
-	if (bChanged)
+	if (!bIsAlreadyInSet)
 	{
 		FVisualLogEntry* CurrentEntry = CurrentEntryPerObject.Find(&InObject);
 		if (CurrentEntry)
 		{
-			CurrentEntry->bPassedObjectAllowList = true;
-			CurrentEntry->UpdateAllowedToLog();
+			CurrentEntry->SetPassedObjectAllowList(true);
+		}
+
+	    // We are not locking the thread entry map, we do not expect other threads adding new entries while editing allow list
+	    // @todo we should add a thread access detector
+		for (FVisualLoggerObjectEntryMap* ThreadCurrentEntryMap : ThreadCurrentEntryMaps)
+		{
+			checkf(ThreadCurrentEntryMap, TEXT("Unexpected null map entry pointer"));
+			FVisualLogEntry* ThreadCurrentEntry = ThreadCurrentEntryMap->Find(&InObject);
+			if (ThreadCurrentEntry)
+			{
+				ThreadCurrentEntry->SetPassedObjectAllowList(true);
+			}
 		}
 	}
 }
 
 void FVisualLogger::ClearObjectAllowList()
 {
+	FWriteScopeLock Lock(EntryRWLock);
 	for (FObjectKey It : ObjectAllowList)
 	{
 		FVisualLogEntry* CurrentEntry = CurrentEntryPerObject.Find(It);
 		if (CurrentEntry)
 		{
-			CurrentEntry->bPassedObjectAllowList = false;
-			CurrentEntry->UpdateAllowedToLog();
+			CurrentEntry->SetPassedObjectAllowList(false);
+		}
+
+	    // We are not locking the thread entry map, we do not expect other threads adding new entries while editing allow list
+	    // @todo we should add a thread access detector
+		for (FVisualLoggerObjectEntryMap* ThreadCurrentEntryMap : ThreadCurrentEntryMaps)
+		{
+			checkf(ThreadCurrentEntryMap, TEXT("Unexpected null map entry pointer"));
+			FVisualLogEntry* ThreadCurrentEntry = ThreadCurrentEntryMap->Find(It);
+			if (ThreadCurrentEntry)
+			{
+				ThreadCurrentEntry->SetPassedObjectAllowList(false);
+			}
 		}
 	}
 
@@ -156,33 +179,110 @@ bool FVisualLogger::IsObjectAllowed(const UObject* InObject) const
 
 FVisualLogEntry* FVisualLogger::GetLastEntryForObject(const UObject* Object)
 {
-	const UObject* LogOwner = FVisualLogger::FindRedirection(Object);
-	return CurrentEntryPerObject.Contains(LogOwner) ? &CurrentEntryPerObject[LogOwner] : nullptr;
+	const UObject* LogOwner = nullptr;
+	{
+		FReadScopeLock Lock(RedirectRWLock);
+		LogOwner = FindRedirectionInternal(Object);
+	}
+	if (LogOwner == nullptr)
+	{
+		return nullptr;
+	}
+
+	FVisualLoggerObjectEntryMap& ThreadCurrentEntryPerObject = GetThreadCurrentEntryMap();
+	return ThreadCurrentEntryPerObject.Find(LogOwner);
 }
 
-FVisualLogEntry* FVisualLogger::GetEntryToWrite(const UObject* Object, const float TimeStamp, const ECreateIfNeeded ShouldCreate)
+FVisualLogEntry* FVisualLogger::GetEntryToWrite(const UObject* Object, const float TimeStamp, ECreateIfNeeded ShouldCreate)
 {
-	const UObject* const LogOwner = FindRedirection(Object);
+	const UObject* LogOwner = nullptr;
+	{
+		FReadScopeLock Lock(RedirectRWLock);
+		LogOwner = FindRedirectionInternal(Object);
+	}
 	if (LogOwner == nullptr)
 	{
 		return nullptr;
 	}
 
-	// Entry can be created or reused (after being flushed) and will need to be initialized
+	FVisualLoggerObjectEntryMap& ThreadCurrentEntryPerObject = GetThreadCurrentEntryMap();
+	FVisualLogEntry* ThreadCurrentEntry = ThreadCurrentEntryPerObject.Find(LogOwner);
 	bool bInitializeEntry = false;
-	FVisualLogEntry* CurrentEntry = nullptr;
+	if (ThreadCurrentEntry)
+	{
+		if (ThreadCurrentEntry->ShouldLog(ShouldCreate))
+		{
+			if (ThreadCurrentEntry->ShouldFlush(TimeStamp))
+			{
+				FWriteScopeLock Lock(EntryRWLock);
+				if (FVisualLogEntry* GlobalCurrentEntry = GetEntryToWriteInternal(LogOwner, ThreadCurrentEntry->TimeStamp, ShouldCreate))
+				{
+					ThreadCurrentEntry->MoveTo(*GlobalCurrentEntry);
+				}
+			}
 
-	if (CurrentEntryPerObject.Contains(LogOwner))
+			bInitializeEntry = !ThreadCurrentEntry->bIsInitialized;
+		}
+	}
+	else
 	{
-		CurrentEntry = &CurrentEntryPerObject[LogOwner];
+		// Check to see if it exists in the global map, 
+		// if so then force creation for the current thread
+		if (ShouldCreate != ECreateIfNeeded::Create)
+		{
+			FReadScopeLock Lock(EntryRWLock);
+			if (CurrentEntryPerObject.Find(LogOwner))
+			{
+				ShouldCreate = ECreateIfNeeded::Create;
+			}
+		}
 
-		// We serialize and reinitialize entries only when allowed to log and parameter
-		// indicates that new entry can be created. Otherwise we simply return current entry.  
-		if (CurrentEntry->bIsAllowedToLog && ShouldCreate == ECreateIfNeeded::Create)
+		if (ShouldCreate == ECreateIfNeeded::Create)
 		{
-			// Same LogOwner can be used for logs at different time in the frame so need to flush entry right away
-			// Other entries will be flushed in batch using FlushEntries
-			if (CurrentEntry->bIsInitialized && TimeStamp > CurrentEntry->TimeStamp)
+			ThreadCurrentEntry = &ThreadCurrentEntryPerObject.Add(LogOwner);
+
+			CalculateEntryAllowLogging(ThreadCurrentEntry, LogOwner, Object);
+		}
+
+		bInitializeEntry = ThreadCurrentEntry->bIsAllowedToLog;
+	}
+
+	if (bInitializeEntry)
+	{
+		checkf(ThreadCurrentEntry != nullptr, TEXT("bInitializeEntry can only be true when CurrentEntry is valid."));
+		ThreadCurrentEntry->InitializeEntry(TimeStamp);
+	}
+
+	if (ThreadCurrentEntry != nullptr && ThreadCurrentEntry->bIsAllowedToLog)
+	{
+		bIsFlushRequired = true;
+	}
+	else
+	{
+		ThreadCurrentEntry = nullptr;
+	}
+
+	return ThreadCurrentEntry;
+}
+
+
+FVisualLogEntry* FVisualLogger::GetEntryToWriteInternal(const UObject* Object, const float TimeStamp, const ECreateIfNeeded ShouldCreate)
+{
+	const UObject* LogOwner = FindRedirectionInternal(Object);
+	if (LogOwner == nullptr)
+	{
+		return nullptr;
+	}
+
+	// Entry can be created or reused (after being flushed) and will need to be initialized
+	bool bInitializeEntry = false;
+	FVisualLogEntry* CurrentEntry = CurrentEntryPerObject.Find(LogOwner);
+
+	if (CurrentEntry)
+	{
+		if (CurrentEntry->ShouldLog(ShouldCreate))
+		{
+			if (CurrentEntry->ShouldFlush(TimeStamp))
 			{
 				FlushEntry(*CurrentEntry, LogOwner);
 			}
@@ -206,11 +306,7 @@ FVisualLogEntry* FVisualLogger::GetEntryToWrite(const UObject* Object, const flo
 		ObjectToClassNameMap.Add(LogOwner, *(LogOwner->GetClass()->GetName()));
 		ObjectToWorldMap.Add(LogOwner, World);
 
-		// IsClassAllowed isn't super fast, but this gets calculated only once for every 
-		// object trying to log something
-		CurrentEntry->bPassedClassAllowList = (ClassAllowList.Num() == 0) || IsClassAllowed(*LogOwner->GetClass()) || IsClassAllowed(*Object->GetClass());
-		CurrentEntry->bPassedObjectAllowList = IsObjectAllowed(LogOwner);
-		CurrentEntry->UpdateAllowedToLog();
+		CalculateEntryAllowLogging(CurrentEntry, LogOwner, Object);
 
 		bInitializeEntry = CurrentEntry->bIsAllowedToLog;
 	}
@@ -218,9 +314,8 @@ FVisualLogEntry* FVisualLogger::GetEntryToWrite(const UObject* Object, const flo
 	if (bInitializeEntry)
 	{
 		checkf(CurrentEntry != nullptr, TEXT("bInitializeEntry can only be true when CurrentEntry is valid."));
-		CurrentEntry->Reset();
-		CurrentEntry->TimeStamp = TimeStamp;
-		CurrentEntry->bIsInitialized = true;
+		CurrentEntry->InitializeEntry(TimeStamp);
+
 
 		if (const AActor* ObjectAsActor = Cast<AActor>(LogOwner))
 		{
@@ -228,6 +323,7 @@ FVisualLogEntry* FVisualLogger::GetEntryToWrite(const UObject* Object, const flo
 			CurrentEntry->bIsLocationValid = true;
 		}
 
+		FReadScopeLock RedirectScopeLock(RedirectRWLock);
 		FOwnerToChildrenRedirectionMap& RedirectionMap = GetRedirectionMap(LogOwner);
 		if (RedirectionMap.Contains(LogOwner))
 		{
@@ -264,6 +360,15 @@ FVisualLogEntry* FVisualLogger::GetEntryToWrite(const UObject* Object, const flo
 	return CurrentEntry;
 }
 
+void FVisualLogger::CalculateEntryAllowLogging(FVisualLogEntry* CurrentEntry, const UObject* LogOwner, const UObject* Object)
+{
+	// IsClassAllowed isn't super fast, but this gets calculated only once for every 
+	// object trying to log something
+	CurrentEntry->bPassedClassAllowList = (ClassAllowList.Num() == 0) || IsClassAllowed(*LogOwner->GetClass()) || IsClassAllowed(*Object->GetClass());
+	CurrentEntry->bPassedObjectAllowList = IsObjectAllowed(LogOwner);
+	CurrentEntry->UpdateAllowedToLog();
+}
+
 void FVisualLogger::Tick(float DeltaTime)
 {
 	if (bIsFlushRequired)
@@ -271,10 +376,35 @@ void FVisualLogger::Tick(float DeltaTime)
 		Flush();
 		bIsFlushRequired = false;
 	}
+
+	if (bContainsInvalidRedirects)
+	{
+		CleanupRedirects();
+		bContainsInvalidRedirects = false;
+	}
 }
 
 void FVisualLogger::Flush()
 {
+	FWriteScopeLock Lock(EntryRWLock);
+
+	// We are not locking the thread entry map, we do not expect other threads adding new entries while flushing
+	// @todo we should add a thread access detector
+	for (FVisualLoggerObjectEntryMap* ThreadCurrentEntryMap : ThreadCurrentEntryMaps)
+	{
+		checkf(ThreadCurrentEntryMap, TEXT("Unexpected null map entry pointer"));
+		for (auto& ThreadCurrentEntry : *ThreadCurrentEntryMap)
+		{
+			if (ThreadCurrentEntry.Value.bIsInitialized)
+			{
+				const UObject* OwnerObject = ThreadCurrentEntry.Key.ResolveObjectPtrEvenIfPendingKill();
+				if (FVisualLogEntry* GlobalCurrentEntry = GetEntryToWriteInternal(OwnerObject, ThreadCurrentEntry.Value.TimeStamp, ECreateIfNeeded::Create))
+				{
+					ThreadCurrentEntry.Value.MoveTo(*GlobalCurrentEntry);
+				}
+			}
+		}
+	}
 	for (auto &CurrentEntry : CurrentEntryPerObject)
 	{
 		if (CurrentEntry.Value.bIsInitialized)
@@ -296,6 +426,19 @@ void FVisualLogger::FlushEntry(FVisualLogEntry& Entry, const FObjectKey& ObjectK
 	Entry.Reset();
 }
 
+thread_local FVisualLogger::FVisualLoggerObjectEntryMap* GThreadCurrentEntryPerObject = nullptr;
+FVisualLogger::FVisualLoggerObjectEntryMap& FVisualLogger::GetThreadCurrentEntryMap()
+{
+	if (GThreadCurrentEntryPerObject == nullptr)
+	{
+		GThreadCurrentEntryPerObject = new FVisualLoggerObjectEntryMap();
+		FWriteScopeLock Lock(EntryRWLock);
+		ThreadCurrentEntryMaps.Add(GThreadCurrentEntryPerObject);
+	}
+	checkf(GThreadCurrentEntryPerObject, TEXT("Unexpected null map entry pointer"));
+	return *GThreadCurrentEntryPerObject;
+}
+
 void FVisualLogger::EventLog(const UObject* Object, const FName EventTag1, const FVisualLogEventBase& Event1, const FVisualLogEventBase& Event2, const FVisualLogEventBase& Event3, const FVisualLogEventBase& Event4, const FVisualLogEventBase& Event5, const FVisualLogEventBase& Event6)
 {
 	EventLog(Object, EventTag1, Event1, Event2, Event3, Event4, Event5);
@@ -453,6 +596,11 @@ void FVisualLogger::Cleanup(UWorld* OldWorld, const bool bReleaseMemory)
                 ObjectToWorldMap.Reset();
                 ChildToOwnerMap.Reset();
                 CurrentEntryPerObject.Reset();
+				for (FVisualLoggerObjectEntryMap* ThreadCurrentEntryMap : ThreadCurrentEntryMaps)
+				{
+					checkf(ThreadCurrentEntryMap, TEXT("Unexpected null map entry pointer"));
+					ThreadCurrentEntryMap->Reset();
+				}
                 ObjectToNameMap.Reset();
                 ObjectToClassNameMap.Reset();
             }
@@ -465,6 +613,11 @@ void FVisualLogger::Cleanup(UWorld* OldWorld, const bool bReleaseMemory)
 						FObjectKey Obj = It.Key();
                         ObjectToWorldMap.Remove(Obj);
                         CurrentEntryPerObject.Remove(Obj);
+						for (FVisualLoggerObjectEntryMap* ThreadCurrentEntryMap : ThreadCurrentEntryMaps)
+						{
+							checkf(ThreadCurrentEntryMap, TEXT("Unexpected null map entry pointer"));
+							ThreadCurrentEntryMap->Remove(Obj);
+						}
                         ObjectToNameMap.Remove(Obj);
                         ObjectToClassNameMap.Remove(Obj);
                     }
@@ -487,6 +640,11 @@ void FVisualLogger::Cleanup(UWorld* OldWorld, const bool bReleaseMemory)
 		ObjectToWorldMap.Reset();
 		ChildToOwnerMap.Reset();
 		CurrentEntryPerObject.Reset();
+		for (FVisualLoggerObjectEntryMap* ThreadCurrentEntryMap : ThreadCurrentEntryMaps)
+		{
+			checkf(ThreadCurrentEntryMap, TEXT("Unexpected null map entry pointer"));
+			ThreadCurrentEntryMap->Reset();
+		}
 		ObjectToNameMap.Reset();
 		ObjectToClassNameMap.Reset();
 	}
@@ -520,16 +678,16 @@ FVisualLogger::FOwnerToChildrenRedirectionMap& FVisualLogger::GetRedirectionMap(
 	return WorldToRedirectionMap.FindOrAdd(World);
 }
 
-void FVisualLogger::Redirect(const UObject* FromObject, const UObject* ToObject)
+UObject* FVisualLogger::RedirectInternal(const UObject* FromObject, const UObject* ToObject)
 {
 	if (FromObject == ToObject || FromObject == nullptr || ToObject == nullptr)
 	{
-		return;
+		return nullptr;
 	}
 
 	const TWeakObjectPtr<const UObject> FromWeakPtr(FromObject);
-	UObject* OldRedirection = FindRedirection(FromObject);
-	UObject* NewRedirection = FindRedirection(ToObject);
+	UObject* OldRedirection = FindRedirectionInternal(FromObject);
+	UObject* NewRedirection = FindRedirectionInternal(ToObject);
 
 	if (OldRedirection != NewRedirection)
 	{
@@ -544,22 +702,19 @@ void FVisualLogger::Redirect(const UObject* FromObject, const UObject* ToObject)
 		OwnerToChildrenMap.FindOrAdd(NewRedirection).AddUnique(FromWeakPtr);
 	}
 
-	FChildToOwnerRedirectionMap& ChildToOwnerMap = FVisualLogger::Get().GetChildToOwnerRedirectionMap();
 	ChildToOwnerMap.FindOrAdd(FromWeakPtr.Get(true/*bEvenIfPendingKill*/)) = ToObject;
 
-	UE_CVLOG(FromObject != nullptr, FromObject, LogVisual, Log, TEXT("Redirected '%s' to '%s'"), *FromObject->GetName(), *NewRedirection->GetName());
+	return NewRedirection;
 }
 
-UObject* FVisualLogger::FindRedirection(const UObject* Object)
+UObject* FVisualLogger::FindRedirectionInternal(const UObject* Object) const
 {
-	FChildToOwnerRedirectionMap& Map = FVisualLogger::Get().GetChildToOwnerRedirectionMap();
-
 	TWeakObjectPtr<const UObject> TargetWeakPtr(Object);
-	TWeakObjectPtr<const UObject>* Parent = &TargetWeakPtr;
+	const TWeakObjectPtr<const UObject>* Parent = &TargetWeakPtr;
 
 	while (Parent)
 	{
-		Parent = Map.Find(TargetWeakPtr.Get(/*bEvenIfPendingKill*/true));
+		Parent = ChildToOwnerMap.Find(TargetWeakPtr.Get(/*bEvenIfPendingKill*/true));
 		if (Parent)
 		{
 			if (Parent->IsValid())
@@ -569,7 +724,7 @@ UObject* FVisualLogger::FindRedirection(const UObject* Object)
 			else
 			{
 				Parent = nullptr;
-				Map.Remove(TargetWeakPtr.Get(/*bEvenIfPendingKill*/true));
+				bContainsInvalidRedirects = true;
 			}
 		}
 	}
@@ -578,6 +733,19 @@ UObject* FVisualLogger::FindRedirection(const UObject* Object)
 	return const_cast<UObject*>(Target);
 }
 
+void FVisualLogger::CleanupRedirects()
+{
+	FWriteScopeLock Lock(RedirectRWLock);
+	for (auto It = ChildToOwnerMap.CreateIterator(); It; ++It)
+	{
+		if(!It.Value().IsValid())
+		{
+			FObjectKey Obj = It.Key();
+			ChildToOwnerMap.Remove(Obj);
+		}
+	}
+}
+
 void FVisualLogger::SetIsRecording(const bool InIsRecording)
 {
 	if (InIsRecording == false && InIsRecording != !!bIsRecording && FParse::Param(FCommandLine::Get(), TEXT("LogNavOctree")))
diff --git a/Engine/Source/Runtime/Engine/Private/VisualLogger/VisualLoggerTypes.cpp b/Engine/Source/Runtime/Engine/Private/VisualLogger/VisualLoggerTypes.cpp
--- a/Engine/Source/Runtime/Engine/Private/VisualLogger/VisualLoggerTypes.cpp
+++ b/Engine/Source/Runtime/Engine/Private/VisualLogger/VisualLoggerTypes.cpp
@@ -130,6 +130,13 @@ FVisualLogEntry::FVisualLogEntry(float InTimeStamp, FVector InLocation, const UO
 	}
 }
 
+void FVisualLogEntry::InitializeEntry(const float InTimeStamp)
+{
+	Reset();
+	TimeStamp = InTimeStamp;
+	bIsInitialized = true;
+}
+
 void FVisualLogEntry::Reset()
 {
 	TimeStamp = -1;
@@ -144,6 +151,12 @@ void FVisualLogEntry::Reset()
 	bIsInitialized = false;
 }
 
+void FVisualLogEntry::SetPassedObjectAllowList(const bool bPassed)
+{
+	bPassedObjectAllowList = bPassed;
+	UpdateAllowedToLog();
+}
+
 void FVisualLogEntry::UpdateAllowedToLog()
 {
 	bIsAllowedToLog = bPassedClassAllowList || bPassedObjectAllowList;
@@ -154,6 +167,22 @@ int32 FVisualLogEntry::AddEvent(const FVisualLogEventBase& Event)
 	return Events.Add(Event);
 }
 
+void FVisualLogEntry::MoveTo(FVisualLogEntry& Other)
+{
+	ensureMsgf(bIsInitialized && Other.bIsInitialized, TEXT("Both entries need to be initialized to move to the other"));
+	ensureMsgf(TimeStamp == Other.TimeStamp, TEXT("Can only move similar entries"));
+	ensureMsgf(bPassedClassAllowList == Other.bPassedClassAllowList, TEXT("Can only move similar entries"));
+	ensureMsgf(bPassedObjectAllowList == Other.bPassedObjectAllowList, TEXT("Can only move similar entries"));
+	ensureMsgf(bIsAllowedToLog == Other.bIsAllowedToLog, TEXT("Can only move similar entries"));
+	Other.Events.Append(Events);
+	Other.LogLines.Append(LogLines);
+	Other.Status.Append(Status);
+	Other.ElementsToDraw.Append(ElementsToDraw);
+	Other.HistogramSamples.Append(HistogramSamples);
+	Other.DataBlocks.Append(DataBlocks);
+	Reset();
+}
+
 void FVisualLogEntry::AddText(const FString& TextLine, const FName& CategoryName, ELogVerbosity::Type Verbosity)
 {
 	LogLines.Add(FVisualLogLine(CategoryName, Verbosity, TextLine));
diff --git a/Engine/Source/Runtime/Engine/Public/VisualLogger/VisualLogger.h b/Engine/Source/Runtime/Engine/Public/VisualLogger/VisualLogger.h
--- a/Engine/Source/Runtime/Engine/Public/VisualLogger/VisualLogger.h
+++ b/Engine/Source/Runtime/Engine/Public/VisualLogger/VisualLogger.h
@@ -16,8 +16,8 @@
 
 #if ENABLE_VISUAL_LOG
 
-#define REDIRECT_TO_VLOG(Dest) FVisualLogger::Get().Redirect(this, Dest)
-#define REDIRECT_OBJECT_TO_VLOG(Src, Dest) FVisualLogger::Get().Redirect(Src, Dest)
+#define REDIRECT_TO_VLOG(Dest) FVisualLogger::Redirect(this, Dest)
+#define REDIRECT_OBJECT_TO_VLOG(Src, Dest) FVisualLogger::Redirect(Src, Dest)
 
 #define CONNECT_WITH_VLOG(Dest)
 #define CONNECT_OBJECT_WITH_VLOG(Src, Dest)
@@ -450,10 +450,24 @@ public:
 	void Cleanup(UWorld* OldWorld, bool bReleaseMemory = false);
 
 	/** Set log owner redirection from one object to another, to combine logs */
-	static void Redirect(const UObject* FromObject, const UObject* ToObject);
+	static void Redirect(const UObject* FromObject, const UObject* ToObject)
+	{ 
+		FVisualLogger& Logger = FVisualLogger::Get();
+		UObject* NewRedirection = nullptr;
+		{
+			FWriteScopeLock Lock(Logger.RedirectRWLock);
+			NewRedirection = Logger.RedirectInternal(FromObject, ToObject);
+		}
+		UE_CVLOG(FromObject != nullptr && NewRedirection != nullptr, FromObject, LogVisual, Log, TEXT("Redirected '%s' to '%s'"), *FromObject->GetName(), *NewRedirection->GetName());
+	}
 
 	/** find and return redirection object for given object*/
-	static UObject* FindRedirection(const UObject* Object);
+	static UObject* FindRedirection(const UObject* Object)
+	{ 
+		FVisualLogger& Logger = FVisualLogger::Get();
+		FReadScopeLock Lock(Logger.RedirectRWLock);
+		return Logger.FindRedirectionInternal(Object);
+	}
 
 	/** blocks all categories from logging. It can be bypassed with the category allow list */
 	void BlockAllCategories(const bool bInBlock) { bBlockedAllCategories = bInBlock; }
@@ -561,10 +575,26 @@ public:
 	UE_DEPRECATED(5.0, "Use AddObjectToAllowList instead")
 	void AddWhitelistedObject(const UObject& InObject) { AddObjectToAllowList(InObject); }
 
+	typedef TMap<FObjectKey, FVisualLogEntry> FVisualLoggerObjectEntryMap;
+
 private:
 	FVisualLogger();
 	virtual void Serialize(const TCHAR* V, ELogVerbosity::Type Verbosity, const FName& Category) override { ensureMsgf(0, TEXT("Regular serialize is forbiden for visual logs")); }
 
+	FVisualLoggerObjectEntryMap& GetThreadCurrentEntryMap();
+
+	/** Get global entry to write where all logs are combined together, not thread safe */
+	FVisualLogEntry* GetEntryToWriteInternal(const UObject* Object, float TimeStamp, ECreateIfNeeded ShouldCreate);
+
+	/** Redirect internal implementation, not thread safe */
+	UObject* RedirectInternal(const UObject* FromObject, const UObject* ToObject);
+	/** Find redirects internal implementation, not thread safe */
+	UObject* FindRedirectionInternal(const UObject* Object) const;
+	/** Cleanup invalid redirects */
+	void CleanupRedirects();
+	/** Figure out all conditions if this entry is allowed to log */
+	void CalculateEntryAllowLogging(FVisualLogEntry* CurrentEntry, const UObject* LogOwner, const UObject* Object);
+
 protected:
 	/** Flushes entries recorded in the frame */
 	void Tick(float DeltaTime);
@@ -600,7 +630,11 @@ protected:
 	// last generated unique id for given times tamp
 	TMap<float, int32> LastUniqueIds;
 	// Current entry with all data
-	TMap<FObjectKey, FVisualLogEntry>	CurrentEntryPerObject;
+	FVisualLoggerObjectEntryMap CurrentEntryPerObject;
+	// Threads current entry maps
+	TArray<FVisualLoggerObjectEntryMap*> ThreadCurrentEntryMaps;
+	// Read Write lock protecting object entries
+	mutable FRWLock EntryRWLock;
 	// Map to contain names for Objects (they can be destroyed after while)
 	TMap<FObjectKey, FName> ObjectToNameMap;
 	// Map to contain class names for Objects (they can be destroyed after while)
@@ -610,6 +644,8 @@ protected:
 	// for any object that has requested redirection this map holds where we should
 	// redirect the traffic to
 	FChildToOwnerRedirectionMap ChildToOwnerMap;
+	// Read Write lock protecting redirection maps (ChildToOwnerMap and ObjectToWorldMap)
+	mutable FRWLock RedirectRWLock;
 	// if set all categories are blocked from logging
 	bool bBlockedAllCategories : 1;
 	// if set we are recording to file
@@ -623,6 +659,8 @@ protected:
 	bool bForceUniqueLogNames : 1;
 	/** Indicates that entries were added/updated and that a flush is required */
 	bool bIsFlushRequired : 1;
+	/** Indicates there are entries in the redirection map that are invalid */
+	mutable bool bContainsInvalidRedirects : 1;
 	// start recording time
 	float StartRecordingToFileTime;
 	/** Delegate to set project specific file name for vlogs */
diff --git a/Engine/Source/Runtime/Engine/Public/VisualLogger/VisualLoggerTypes.h b/Engine/Source/Runtime/Engine/Public/VisualLogger/VisualLoggerTypes.h
--- a/Engine/Source/Runtime/Engine/Public/VisualLogger/VisualLoggerTypes.h
+++ b/Engine/Source/Runtime/Engine/Public/VisualLogger/VisualLoggerTypes.h
@@ -181,7 +181,22 @@ struct ENGINE_API FVisualLogEntry
 	FVisualLogEntry(const AActor* InActor, TArray<TWeakObjectPtr<UObject> >* Children);
 	FVisualLogEntry(float InTimeStamp, FVector InLocation, const UObject* Object, TArray<TWeakObjectPtr<UObject> >* Children);
 
+	bool ShouldLog(const ECreateIfNeeded ShouldCreate) const
+	{
+		// We serialize and reinitialize entries only when allowed to log and parameter
+		// indicates that new entry can be created.
+		return bIsAllowedToLog && ShouldCreate == ECreateIfNeeded::Create;
+	}
+
+	bool ShouldFlush(float InTimeStamp) const
+	{
+		//Same LogOwner can be used for logs at different time in the frame so need to flush entry right away
+		return bIsInitialized && InTimeStamp > TimeStamp;
+	}
+
+	void InitializeEntry( const float InTimeStamp );
 	void Reset();
+	void SetPassedObjectAllowList(const bool bPassed);
 	void UpdateAllowedToLog();
 
 	void AddText(const FString& TextLine, const FName& CategoryName, ELogVerbosity::Type Verbosity);
@@ -223,6 +238,8 @@ struct ENGINE_API FVisualLogEntry
 	// find index of status category
 	int32 FindStatusIndex(const FString& CategoryName);
 
+	// Moves all content to provided entry and reseting our content.
+	void MoveTo(FVisualLogEntry& Other);
 
 #endif // ENABLE_VISUAL_LOG
 };
