Fixed ortho view rendering after WORLD_MAX bump in CL 18595474. New WORLD_MAX was offsetting ortho view origin too far and breaking precision.

diff --git a/Engine/Source/Runtime/Engine/Private/SceneView.cpp b/Engine/Source/Runtime/Engine/Private/SceneView.cpp
--- a/Engine/Source/Runtime/Engine/Private/SceneView.cpp
+++ b/Engine/Source/Runtime/Engine/Private/SceneView.cpp
@@ -548,7 +548,7 @@ void FViewMatrices::Init(const FMinimalInitializer& Initializer)
 	}
 	else if (Initializer.bUseFauxOrthoViewPos)
 	{
-		auto DistanceToViewOrigin = WORLD_MAX;
+		auto DistanceToViewOrigin = UE_OLD_WORLD_MAX;
 		ViewOrigin = FVector(InvViewMatrix.TransformVector(FVector(0, 0, -1).GetSafeNormal())) * DistanceToViewOrigin + LocalViewOrigin;
 		bViewOriginIsFudged = true;
 	}
