Only compile shaders that are used to draw instead of compiling every shader in the each material's shader map on PostLoad.

# Results
1. Opening QAGame editor and opening TM-ShaderModels by default.
2. A cold DDC and every shader needing to be compiled.

* 63.78% less shaders are compiled.
* 1.19x speedup opening the editor and TM-ShaderModels level.
* 5.46x speedup opening TM-ShaderModels level.

# Changes

- This functionality is experimental and is disabled by default.  It can be enabled by setting r.ShaderCompiler.JobCacheDDC to true.
- This functionality is explicitly an editor only feature.  In other words the cooker or commandlets still compile and cache full shader maps.
- Completely skip compilation of shaders during PostLoad of Materials and MaterialInstances.
- In FShaderJobCache::Find does a DDC Get for the shader if we can't find it locally in memory.
- FShaderJobCache::Add does a DDC Put to store individual shaders in the DDC.
- Skip trying to load full shader maps from the DDC as a minor optimization.

# Unreal Insights Counters

- Shaders/JobCache/DDCRequests
- Shaders/JobCache/DDCHits
- Shaders/JobCache/DDCBytesRecieved
- Shaders/JobCache/DDCBytesSent

# todo
- Currently just storing the results with the local cache policy (UE::DerivedData::ECachePolicy::StoreLocal)
- How do we handle rendering that _needs_ shaders immediately before rendering?

diff --git a/Engine/Source/Runtime/Engine/Private/Materials/Material.cpp b/Engine/Source/Runtime/Engine/Private/Materials/Material.cpp
--- a/Engine/Source/Runtime/Engine/Private/Materials/Material.cpp
+++ b/Engine/Source/Runtime/Engine/Private/Materials/Material.cpp
@@ -3224,7 +3224,15 @@ void UMaterial::PostLoad()
 				}
 			}
 
-			CacheResourceShadersForRendering(false);
+			const bool bSkipCompilationOnPostLoad = IsShaderJobCacheDDCEnabled();
+			if (bSkipCompilationOnPostLoad)
+			{
+				CacheResourceShadersForRendering(false, EMaterialShaderPrecompileMode::None);
+			}
+			else
+			{
+				CacheResourceShadersForRendering(false);
+			}
 		}
 	}
 	INC_FLOAT_STAT_BY(STAT_ShaderCompiling_MaterialLoading,(float)MaterialLoadTime);
diff --git a/Engine/Source/Runtime/Engine/Private/Materials/MaterialInstance.cpp b/Engine/Source/Runtime/Engine/Private/Materials/MaterialInstance.cpp
--- a/Engine/Source/Runtime/Engine/Private/Materials/MaterialInstance.cpp
+++ b/Engine/Source/Runtime/Engine/Private/Materials/MaterialInstance.cpp
@@ -51,6 +51,7 @@
 #include "UObject/UE5MainStreamObjectVersion.h"
 #include "UObject/FortniteMainBranchObjectVersion.h"
 #include "ShaderCompilerCore.h"
+#include "ShaderCompiler.h"
 
 DECLARE_CYCLE_STAT(TEXT("MaterialInstance CopyMatInstParams"), STAT_MaterialInstance_CopyMatInstParams, STATGROUP_Shaders);
 DECLARE_CYCLE_STAT(TEXT("MaterialInstance Serialize"), STAT_MaterialInstance_Serialize, STATGROUP_Shaders);
@@ -2581,8 +2582,18 @@ void UMaterialInstance::PostLoad()
 	{
 		SCOPE_SECONDS_COUNTER(MaterialLoadTime);
 
+		const bool bSkipCompilationOnPostLoad = IsShaderJobCacheDDCEnabled();
+
 		// Make sure static parameters are up to date and shaders are cached for the current platform
-		InitStaticPermutation();
+		if (bSkipCompilationOnPostLoad)
+		{
+			InitStaticPermutation(EMaterialShaderPrecompileMode::None);
+		}
+		else
+		{
+			InitStaticPermutation();
+		}
+
 #if WITH_EDITOR && 0 // the cooker will kick BeginCacheForCookedPlatformData on its own. If a commandlet is going to do rendering (e.g. rebuilding HLOD), InitStaticPermutation should have already created the necessary resources
 		// enable caching in postload for derived data cache commandlet and cook by the book
 		ITargetPlatformManagerModule* TPM = GetTargetPlatformManager();
diff --git a/Engine/Source/Runtime/Engine/Private/Materials/MaterialShader.cpp b/Engine/Source/Runtime/Engine/Private/Materials/MaterialShader.cpp
--- a/Engine/Source/Runtime/Engine/Private/Materials/MaterialShader.cpp
+++ b/Engine/Source/Runtime/Engine/Private/Materials/MaterialShader.cpp
@@ -16,6 +16,7 @@
 #include "ProfilingDebugging/CookStats.h"
 #include "UObject/ReleaseObjectVersion.h"
 #include "UObject/EditorObjectVersion.h"
+#include "ProfilingDebugging/CountersTrace.h"
 #include "ProfilingDebugging/LoadTimeTracker.h"
 #include "Misc/ScopeRWLock.h"
 
@@ -1317,6 +1318,11 @@ void FMaterialShaderMap::GetAllOutdatedTypes(TArray<const FShaderType*>& Outdate
 #endif
 }
 
+TRACE_DECLARE_INT_COUNTER(Shaders_FMaterialShaderMapDDCRequests, TEXT("Shaders/FMaterialShaderMap/DDCRequests"));
+TRACE_DECLARE_INT_COUNTER(Shaders_FMaterialShaderMapDDCHits, TEXT("Shaders/FMaterialShaderMap/DDCHits"));
+TRACE_DECLARE_MEMORY_COUNTER(Shaders_FMaterialShaderMapDDCBytesReceived, TEXT("Shaders/FMaterialShaderMap/DDCBytesRecieved"));
+TRACE_DECLARE_MEMORY_COUNTER(Shaders_FMaterialShaderMapDDCBytesSent, TEXT("Shaders/FMaterialShaderMap/DDCBytesSent"));
+
 void FMaterialShaderMap::LoadFromDerivedDataCache(const FMaterial* Material, const FMaterialShaderMapId& ShaderMapId, EShaderPlatform InPlatform, const ITargetPlatform* TargetPlatform, TRefCountPtr<FMaterialShaderMap>& InOutShaderMap, FString& OutDDCKeyDesc)
 {
 	TRACE_CPUPROFILER_EVENT_SCOPE(FMaterialShaderMap::LoadFromDerivedDataCache);
@@ -1376,6 +1382,8 @@ void FMaterialShaderMap::LoadFromDerivedDataCache(const FMaterial* Material, con
 			//   - since the get call is synchronous, this can cause a hitch if there's network latency
 			if (CheckCache && Material->IsPersistent() && GetDerivedDataCacheRef().GetSynchronous(*DataKey, CachedData, Material->GetFriendlyName()))
 			{
+				TRACE_COUNTER_INCREMENT(Shaders_FMaterialShaderMapDDCHits);
+				TRACE_COUNTER_ADD(Shaders_FMaterialShaderMapDDCBytesReceived, CachedData.Num());
 				COOK_STAT(Timer.AddHit(CachedData.Num()));
 				InOutShaderMap = new FMaterialShaderMap();
 				FMemoryReader Ar(CachedData, true);
@@ -1393,6 +1401,7 @@ void FMaterialShaderMap::LoadFromDerivedDataCache(const FMaterial* Material, con
 			}
 			else
 			{
+				TRACE_COUNTER_INCREMENT(Shaders_FMaterialShaderMapDDCRequests);
 				// We should be build the data later, and we can track that the resource was built there when we push it to the DDC.
 				COOK_STAT(Timer.TrackCyclesOnly());
 				InOutShaderMap = nullptr;
@@ -1415,6 +1424,8 @@ void FMaterialShaderMap::SaveToDerivedDataCache()
 	FMemoryWriter Ar(SaveData, true);
 	Serialize(Ar);
 
+	TRACE_COUNTER_ADD(Shaders_FMaterialShaderMapDDCBytesSent, SaveData.Num());
+
 	FString DataKey = GetMaterialShaderMapKeyString(ShaderMapId, GetShaderPlatform());
 	UE_LOG(LogMaterial, Verbose, TEXT("Saved shaders for %s from DDC (key hash: %s)"), GetMaterialPath(), *FSHA1_HashString(DataKey));
 	GetDerivedDataCacheRef().Put(*DataKey, SaveData, FStringView(GetFriendlyName()));
diff --git a/Engine/Source/Runtime/Engine/Private/Materials/MaterialShared.cpp b/Engine/Source/Runtime/Engine/Private/Materials/MaterialShared.cpp
--- a/Engine/Source/Runtime/Engine/Private/Materials/MaterialShared.cpp
+++ b/Engine/Source/Runtime/Engine/Private/Materials/MaterialShared.cpp
@@ -2227,12 +2227,18 @@ bool FMaterial::CacheShaders(const FMaterialShaderMapId& ShaderMapId, EShaderPla
 				}
 			}
 
+			// If we are loading individual shaders from the shader job cache don't attempt to load full maps.
+			const bool bSkipCompilationOnPostLoad = IsShaderJobCacheDDCEnabled();
+
 			// Attempt to load from the derived data cache if we are uncooked and don't have any shadermap.
 			// If we have an incomplete shadermap, continue with it to prevent creation of duplicate shadermaps for the same ShaderMapId
 			if (!ShaderMap && !FPlatformProperties::RequiresCookedData())
 			{
 				TRefCountPtr<FMaterialShaderMap> LoadedShaderMap;
-				FMaterialShaderMap::LoadFromDerivedDataCache(this, ShaderMapId, Platform, TargetPlatform, LoadedShaderMap, DDCKeyHash);
+				if (bSkipCompilationOnPostLoad == false)
+				{
+					FMaterialShaderMap::LoadFromDerivedDataCache(this, ShaderMapId, Platform, TargetPlatform, LoadedShaderMap, DDCKeyHash);
+				}
 				ShaderMap = LoadedShaderMap;
 			}
 
diff --git a/Engine/Source/Runtime/Engine/Private/ShaderCompiler/ShaderCompiler.cpp b/Engine/Source/Runtime/Engine/Private/ShaderCompiler/ShaderCompiler.cpp
--- a/Engine/Source/Runtime/Engine/Private/ShaderCompiler/ShaderCompiler.cpp
+++ b/Engine/Source/Runtime/Engine/Private/ShaderCompiler/ShaderCompiler.cpp
@@ -66,6 +66,12 @@
 #include "Misc/LargeWorldRenderPosition.h"
 
 #if WITH_EDITOR
+#include "Compression/OodleDataCompression.h"
+#include "DerivedDataCache.h"
+#include "DerivedDataCacheRecord.h"
+#include "DerivedDataRequestOwner.h"
+#include "Virtualization/PayloadId.h"
+
 #include "TextureCompiler.h"
 #include "Rendering/StaticLightingSystemInterface.h"
 #endif
@@ -108,6 +114,23 @@ static FAutoConsoleVariableRef CVarShaderCompilerMaxJobCacheMemoryPercent(
 	ECVF_Default
 );
 
+static TAutoConsoleVariable<bool> CVarJobCacheDDC(
+	TEXT("r.ShaderCompiler.JobCacheDDC"),
+	false,
+	TEXT("Skips compilation of all shaders on Material and Material Instance PostLoad and relies on on-demand shader compilation to compile what is needed."),
+	ECVF_ReadOnly);
+
+bool IsShaderJobCacheDDCEnabled()
+{
+	// For now we only support the editor and not commandlets like the cooker.
+	if (GIsEditor && !IsRunningCommandlet())
+	{
+		return CVarJobCacheDDC.GetValueOnAnyThread();
+	}
+
+	return false;
+}
+
 int32 GShaderCompilerDumpCompileJobInputs = 0;
 static FAutoConsoleVariableRef CVarShaderCompilerDumpCompileJobInputs(
 	TEXT("r.ShaderCompiler.DumpCompileJobInputs"),
@@ -7608,12 +7631,28 @@ void FShaderPipelineCompileJob::SerializeOutput(FArchive& Ar)
 	}
 }
 
-TRACE_DECLARE_INT_COUNTER(Shaders_JobCacheSearchAttempts, TEXT("Shaders/JobCacheSearchAttempts"));
-TRACE_DECLARE_INT_COUNTER(Shaders_JobCacheHits, TEXT("Shaders/JobCacheHits"));
+TRACE_DECLARE_INT_COUNTER(Shaders_JobCacheSearchAttempts, TEXT("Shaders/JobCache/SearchAttempts"));
+TRACE_DECLARE_INT_COUNTER(Shaders_JobCacheHits, TEXT("Shaders/JobCache/Hits"));
+
+TRACE_DECLARE_INT_COUNTER(Shaders_JobCacheDDCRequests, TEXT("Shaders/JobCache/DDCRequests"));
+TRACE_DECLARE_INT_COUNTER(Shaders_JobCacheDDCHits, TEXT("Shaders/JobCache/DDCHits"));
+TRACE_DECLARE_MEMORY_COUNTER(Shaders_JobCacheDDCBytesReceived, TEXT("Shaders/JobCache/DDCBytesRecieved"));
+TRACE_DECLARE_MEMORY_COUNTER(Shaders_JobCacheDDCBytesSent, TEXT("Shaders/JobCache/DDCBytesSent"));
+
+#if WITH_EDITOR
+namespace
+{
+	/** The FCacheBucket used with the DDC, cached to avoid recreating it for each request */
+	UE::DerivedData::FCacheBucket ShaderJobCacheDDCBucket = UE::DerivedData::FCacheBucket(TEXT("FShaderJobCacheShaders"));
+	UE::DerivedData::FValueId ShaderJobCacheId = UE::DerivedData::FValueId::FromName("FShaderJobCacheShaderID");
+}
+#endif
+
 FShaderJobCache::FJobCachedOutput* FShaderJobCache::Find(const FJobInputHash& Hash)
 {
 	++TotalSearchAttempts;
 	TRACE_COUNTER_INCREMENT(Shaders_JobCacheSearchAttempts);
+	TRACE_CPUPROFILER_EVENT_SCOPE(FShaderJobCache::Find);
 
 	if (ShaderCompiler::IsJobCacheEnabled())
 	{
@@ -7630,6 +7669,83 @@ FShaderJobCache::FJobCachedOutput* FShaderJobCache::Find(const FJobInputHash& Ha
 			(*CannedOutput)->NumHits++;
 			return &(*CannedOutput)->JobOutput;
 		}
+#if WITH_EDITOR
+		else
+		{
+			// If we didn't find it in memory search the DDC if it's enabled.
+			const bool bCachePerShaderDDC = IsShaderJobCacheDDCEnabled();
+			if (bCachePerShaderDDC)
+			{
+				FSharedBuffer Results;
+
+				TRACE_COUNTER_INCREMENT(Shaders_JobCacheDDCRequests);
+
+				UE::DerivedData::FRequestOwner RequestOwner(UE::DerivedData::EPriority::Blocking);
+				RequestOwner.KeepAlive();
+
+				UE::DerivedData::FCacheGetRequest Request;
+				Request.Name = TEXT("FShaderJobCache");
+				// Create key.
+				Request.Key.Bucket = ShaderJobCacheDDCBucket;
+				Request.Key.Hash = Hash;
+				Request.Policy = UE::DerivedData::ECachePolicy::QueryLocal;
+
+				UE::DerivedData::GetCache().Get(
+					{ Request },
+					RequestOwner,
+					[&Results](UE::DerivedData::FCacheGetResponse&& Response)
+					{
+						switch (Response.Status)
+						{
+							case UE::DerivedData::EStatus::Ok:
+							{
+								Results = Response.Record.GetValue(ShaderJobCacheId).GetData().Decompress();
+
+								TRACE_COUNTER_INCREMENT(Shaders_JobCacheDDCHits);
+
+								break;
+							}
+							case UE::DerivedData::EStatus::Error:
+							{
+								break;
+							}
+							case UE::DerivedData::EStatus::Canceled:
+							{
+								break;
+							}
+						}
+					});
+
+				RequestOwner.Wait();
+				if (!Results.IsNull())
+				{
+					const uint64 OutputsOriginalSize = Outputs.GetAllocatedSize();
+
+					// Create a new entry to store in the FShaderJobCache
+					FStoredOutput* NewStoredOutput = new FStoredOutput();
+					NewStoredOutput->NumHits = 0;
+					NewStoredOutput->NumReferences = 1;
+					NewStoredOutput->JobOutput.Reserve(Results.GetSize());
+					NewStoredOutput->JobOutput.SetNum(Results.GetSize());
+
+					TRACE_COUNTER_ADD(Shaders_JobCacheDDCBytesReceived, Results.GetSize());
+
+					check(NewStoredOutput->JobOutput.GetAllocatedSize() == Results.GetSize());
+					FMemory::Memcpy(NewStoredOutput->JobOutput.GetData(),  Results.GetData(), Results.GetSize());
+
+					// Generate an output hash and cache the result in the FShaderJobCache
+					FJobOutputHash NewOutputHash = FBlake3::HashBuffer(NewStoredOutput->JobOutput.GetData(), NewStoredOutput->JobOutput.Num());
+					Outputs.Add(NewOutputHash, NewStoredOutput);
+					InputHashToOutput.Add(Hash, NewOutputHash);
+
+					CurrentlyAllocatedMemory += NewStoredOutput->JobOutput.GetAllocatedSize() + sizeof(FStoredOutput);
+
+					// return the results.
+					return &NewStoredOutput->JobOutput;
+				}
+			}
+		}
+#endif
 	}
 
 	return nullptr;
@@ -7649,6 +7765,8 @@ FShaderJobCache::FShaderJobCache()
 
 void FShaderJobCache::Add(const FJobInputHash& Hash, const FJobCachedOutput& Contents, int32 InitialHitCount)
 {
+	TRACE_CPUPROFILER_EVENT_SCOPE(FShaderJobCache::Add);
+
 	if (!ShaderCompiler::IsJobCacheEnabled())
 	{
 		return;
@@ -7657,7 +7775,6 @@ void FShaderJobCache::Add(const FJobInputHash& Hash, const FJobCachedOutput& Con
 	FJobOutputHash* ExistingOutputHash = InputHashToOutput.Find(Hash);
 	if (ExistingOutputHash)
 	{
-		// we can arrive here due to cloned jobs ignoring our normal caching rules
 		return;
 	}
 
@@ -7667,8 +7784,10 @@ void FShaderJobCache::Add(const FJobInputHash& Hash, const FJobCachedOutput& Con
 	const uint64 InputHashToOutputOriginalSize = InputHashToOutput.GetAllocatedSize();
 	InputHashToOutput.Add(Hash, OutputHash);
 
+	const bool bCachePerShaderDDC = IsShaderJobCacheDDCEnabled();
+
 	FStoredOutput** CannedOutput = Outputs.Find(OutputHash);
-	if (CannedOutput)
+	if (CannedOutput && (bCachePerShaderDDC == false))
 	{
 		// update the output hit count
 		(*CannedOutput)->NumReferences++;
@@ -7685,6 +7804,28 @@ void FShaderJobCache::Add(const FJobInputHash& Hash, const FJobCachedOutput& Con
 
 		CurrentlyAllocatedMemory += NewStoredOutput->JobOutput.GetAllocatedSize() + sizeof(FStoredOutput);
 
+#if WITH_EDITOR
+		if (bCachePerShaderDDC)
+		{
+			// Create key.
+			UE::DerivedData::FCacheKey Key;
+			Key.Bucket = ShaderJobCacheDDCBucket;
+			Key.Hash = Hash;
+			UE::DerivedData::FCacheRecordBuilder RecordBuilder(Key);
+
+			RecordBuilder.AddValue(ShaderJobCacheId, FSharedBuffer::MakeView(MakeMemoryView(Contents)));
+
+			TRACE_COUNTER_ADD(Shaders_JobCacheDDCBytesSent, Contents.GetAllocatedSize());
+
+			UE::DerivedData::FRequestOwner RequestOwner(UE::DerivedData::EPriority::Normal);
+			RequestOwner.KeepAlive();
+			UE::DerivedData::GetCache().Put(
+				{ {{TEXT("FShaderJobCache")}, RecordBuilder.Build(), UE::DerivedData::ECachePolicy::StoreLocal} },
+				RequestOwner
+			);
+		}
+#endif
+
 		// delete the previous cache entries if we have a budget
 		const uint64 MemoryBudgetBytes = GetCurrentMemoryBudget();
 		if (MemoryBudgetBytes)
diff --git a/Engine/Source/Runtime/Engine/Public/ShaderCompiler.h b/Engine/Source/Runtime/Engine/Public/ShaderCompiler.h
--- a/Engine/Source/Runtime/Engine/Public/ShaderCompiler.h
+++ b/Engine/Source/Runtime/Engine/Public/ShaderCompiler.h
@@ -80,6 +80,8 @@ inline const TCHAR* ShaderCompileJobPriorityToString(EShaderCompileJobPriority v
 	}
 }
 
+extern ENGINE_API bool IsShaderJobCacheDDCEnabled();
+
 /** Results for a single compiled shader map. */
 struct FShaderMapCompileResults
 {
