Fix BSP undo/redo crash

UModelSceneProxy shares data between Game/Render thread. UModel is typically updated by deregistering its components and flushing render commands before updating the vertex buffers. This happens in PreEditUndo. The resources are then initialized and component reregistered in PostEditUndo.

The previous change to immediately process the list of added primitives after the create action of an actor (AActor::PostEditUndo) breaks this contract since the order of the AActor::PostEditUndo in the transaction list is not guaranteed. This change addresses the issue by skipping the immediate UpdateAllPrimitiveSceneInfo for Brush actors. This is a least invasive solution that restores the prior behavior solely for BSP.

#rnx

diff --git a/Engine/Source/Runtime/Engine/Private/ActorEditor.cpp b/Engine/Source/Runtime/Engine/Private/ActorEditor.cpp
--- a/Engine/Source/Runtime/Engine/Private/ActorEditor.cpp
+++ b/Engine/Source/Runtime/Engine/Private/ActorEditor.cpp
@@ -547,8 +547,12 @@ void AActor::PostEditUndo()
 		Super::PostEditUndo();
 	}
 
+	// Do not immediately update all primitive scene infos for brush actor
+	// undo/redo transactions since they require the render thread to wait until
+	// after the transactions are processed to guarantee that the model data
+	// is safe to access.
 	UWorld* World = GetWorld();
-	if (World && World->Scene)
+	if (World && World->Scene && !FActorEditorUtils::IsABrush(this))
 	{
 		ENQUEUE_RENDER_COMMAND(UpdateAllPrimitiveSceneInfosCmd)([Scene = World->Scene](FRHICommandListImmediate& RHICmdList) {
 			Scene->UpdateAllPrimitiveSceneInfos(RHICmdList);
diff --git a/Engine/Source/Runtime/Engine/Private/ActorEditorUtils.cpp b/Engine/Source/Runtime/Engine/Private/ActorEditorUtils.cpp
--- a/Engine/Source/Runtime/Engine/Private/ActorEditorUtils.cpp
+++ b/Engine/Source/Runtime/Engine/Private/ActorEditorUtils.cpp
@@ -37,6 +37,11 @@ namespace FActorEditorUtils
 		return bIsBuilder;
 	}
 
+	bool IsABrush( const AActor* InActor )
+	{
+		return Cast<ABrush>(InActor) != nullptr;
+	}
+
 	bool IsAPreviewOrInactiveActor( const AActor* InActor )
 	{
 		UWorld* World = InActor ? InActor->GetWorld() : NULL;
diff --git a/Engine/Source/Runtime/Engine/Public/ActorEditorUtils.h b/Engine/Source/Runtime/Engine/Public/ActorEditorUtils.h
--- a/Engine/Source/Runtime/Engine/Public/ActorEditorUtils.h
+++ b/Engine/Source/Runtime/Engine/Public/ActorEditorUtils.h
@@ -16,6 +16,13 @@ namespace FActorEditorUtils
 	*/
 	ENGINE_API bool IsABuilderBrush( const AActor* InActor );
 
+	/**
+	 * Return whether this actor is an ABrush or not.
+	 *
+	 * @return true if this actor is an ABrush, false otherwise.
+	 */
+	ENGINE_API bool IsABrush( const AActor* InActor );
+
 	/** 
 	* Return whether this actor is a preview actor or not.
 	*
