Re-introduce clearing out raw animation curve data for Animation Sequences

diff --git a/Engine/Source/Runtime/Engine/Private/Animation/AnimSequence.cpp b/Engine/Source/Runtime/Engine/Private/Animation/AnimSequence.cpp
--- a/Engine/Source/Runtime/Engine/Private/Animation/AnimSequence.cpp
+++ b/Engine/Source/Runtime/Engine/Private/Animation/AnimSequence.cpp
@@ -585,8 +585,33 @@ void UAnimSequence::Serialize(FArchive& Ar)
 	Ar.UsingCustomVersion(FFrameworkObjectVersion::GUID);
 	Ar.UsingCustomVersion(FUE5MainStreamObjectVersion::GUID);
 
+	FRawCurveTracks RawCurveCache;
+	
+	if (Ar.IsCooking())
+	{
+		PRAGMA_DISABLE_DEPRECATION_WARNINGS
+		RawCurveCache.FloatCurves = MoveTemp(RawCurveData.FloatCurves);
+		RawCurveData.FloatCurves.Reset();
+
+#if WITH_EDITORONLY_DATA
+		RawCurveCache.TransformCurves = MoveTemp(RawCurveData.TransformCurves);
+		RawCurveData.TransformCurves.Reset();
+#endif
+		PRAGMA_ENABLE_DEPRECATION_WARNINGS
+	}
+
 	Super::Serialize(Ar);
 
+	if (Ar.IsCooking())
+	{
+		PRAGMA_DISABLE_DEPRECATION_WARNINGS
+		RawCurveData.FloatCurves = MoveTemp(RawCurveCache.FloatCurves);
+#if WITH_EDITORONLY_DATA
+		RawCurveData.TransformCurves = MoveTemp(RawCurveCache.TransformCurves);
+#endif		
+		PRAGMA_ENABLE_DEPRECATION_WARNINGS
+	}
+
 	FStripDataFlags StripFlags( Ar );
 	if( !StripFlags.IsEditorDataStripped() )
 	{
