Static analysis fix

#rnx

diff --git a/Engine/Source/Runtime/Engine/Private/Particles/FXBudget.cpp b/Engine/Source/Runtime/Engine/Private/Particles/FXBudget.cpp
--- a/Engine/Source/Runtime/Engine/Private/Particles/FXBudget.cpp
+++ b/Engine/Source/Runtime/Engine/Private/Particles/FXBudget.cpp
@@ -332,12 +332,9 @@ void FFXBudget::Reset()
 {
 	if (StatsListener.IsValid())
 	{
-		if (StatsListener.IsValid())
-		{
-			FParticlePerfStatsManager::RemoveListener(StatsListener);
-			StatsListener = MakeShared<FParticlePerfStatsListener_FXBudget, ESPMode::ThreadSafe>();
-			FParticlePerfStatsManager::AddListener(StatsListener);
-		}
+		FParticlePerfStatsManager::RemoveListener(StatsListener);
+		StatsListener = MakeShared<FParticlePerfStatsListener_FXBudget, ESPMode::ThreadSafe>();
+		FParticlePerfStatsManager::AddListener(StatsListener);
 	}
 }
 
