Adding missing push model mark dirty calls for some replicated properties

diff --git a/Engine/Source/Runtime/Engine/Classes/Components/SceneComponent.h b/Engine/Source/Runtime/Engine/Classes/Components/SceneComponent.h
--- a/Engine/Source/Runtime/Engine/Classes/Components/SceneComponent.h
+++ b/Engine/Source/Runtime/Engine/Classes/Components/SceneComponent.h
@@ -208,6 +208,15 @@ private:
 	UPROPERTY(Transient, Replicated)
 	uint8 bShouldSnapRotationWhenAttached : 1;
 
+	/** Sets bShouldBeAttached, push model aware. */
+	void SetShouldBeAttached(bool bNewShouldBeAttached);
+
+	/** Sets bShouldSnapLocationWhenAttached, push model aware. */
+	void SetShouldSnapLocationWhenAttached(bool bShouldSnapLocation);
+
+	/** Sets bShouldSnapRotationWhenAttached, push model aware. */
+	void SetShouldSnapRotationWhenAttached(bool bShouldSnapRotation);
+
 	/**
 	 * Whether or not the cached PhysicsVolume this component overlaps should be updated when the component is moved.
 	 * @see GetPhysicsVolume()
diff --git a/Engine/Source/Runtime/Engine/Private/Components/SceneComponent.cpp b/Engine/Source/Runtime/Engine/Private/Components/SceneComponent.cpp
--- a/Engine/Source/Runtime/Engine/Private/Components/SceneComponent.cpp
+++ b/Engine/Source/Runtime/Engine/Private/Components/SceneComponent.cpp
@@ -77,7 +77,7 @@ USceneComponent::USceneComponent(const FObjectInitializer& ObjectInitializer /*=
 	// default behavior is visible
 	SetVisibleFlag(true);
 	bAutoActivate = false;
-	bShouldBeAttached = AttachParent != nullptr;
+	SetShouldBeAttached(AttachParent != nullptr);
 }
 
 #if WITH_EDITORONLY_DATA
@@ -692,9 +692,9 @@ void USceneComponent::OnRegister()
 			// Failed to attach, we need to clear AttachParent so we don't think we're actually attached when we're not.
 			SetAttachParent(nullptr);
 			SetAttachSocketName(NAME_None);
-			bShouldBeAttached = false;
-			bShouldSnapLocationWhenAttached = false;
-			bShouldSnapRotationWhenAttached = false;
+			SetShouldBeAttached(false);
+			SetShouldSnapLocationWhenAttached(false);
+			SetShouldSnapRotationWhenAttached(false);
 		}
 	}
 	
@@ -1787,7 +1787,7 @@ void USceneComponent::SetupAttachment(class USceneComponent* InParent, FName InS
 				{
 					SetAttachParent(InParent);
 					SetAttachSocketName(InSocketName);
-					bShouldBeAttached = AttachParent != nullptr;
+					SetShouldBeAttached(AttachParent != nullptr);
 				}
 			}
 		}
@@ -1850,8 +1850,8 @@ bool USceneComponent::AttachToComponent(USceneComponent* Parent, const FAttachme
 		ensureMsgf(!AttachmentRules.bWeldSimulatedBodies, TEXT("AttachToComponent when called from a constructor cannot weld simulated bodies. Consider calling SetupAttachment directly instead."));
 		ensureMsgf(AttachmentRules.LocationRule == EAttachmentRule::KeepRelative && AttachmentRules.RotationRule == EAttachmentRule::KeepRelative && AttachmentRules.ScaleRule == EAttachmentRule::KeepRelative, TEXT("AttachToComponent when called from a constructor is only setting up attachment and will always be treated as KeepRelative. Consider calling SetupAttachment directly instead."));
 		SetupAttachment(Parent, SocketName);
-		bShouldSnapLocationWhenAttached = false;
-		bShouldSnapRotationWhenAttached = false;
+		SetShouldSnapLocationWhenAttached(false);
+		SetShouldSnapRotationWhenAttached(false);
 
 		return true;
 	}
@@ -2014,10 +2014,10 @@ bool USceneComponent::AttachToComponent(USceneComponent* Parent, const FAttachme
 		// Save pointer from child to parent
 		SetAttachParent(Parent);
 		SetAttachSocketName(SocketName);
-		bShouldBeAttached = AttachParent != nullptr;
+		SetShouldBeAttached(AttachParent != nullptr);
 
-		bShouldSnapLocationWhenAttached = AttachmentRules.LocationRule == EAttachmentRule::SnapToTarget;
-		bShouldSnapRotationWhenAttached = AttachmentRules.RotationRule == EAttachmentRule::SnapToTarget;
+		SetShouldSnapLocationWhenAttached(AttachmentRules.LocationRule == EAttachmentRule::SnapToTarget);
+		SetShouldSnapRotationWhenAttached(AttachmentRules.RotationRule == EAttachmentRule::SnapToTarget);
 
 		OnAttachmentChanged();
 
@@ -2217,10 +2217,10 @@ void USceneComponent::DetachFromComponent(const FDetachmentTransformRules& Detac
 #endif
 		SetAttachParent(nullptr);
 		SetAttachSocketName(NAME_None);
-		bShouldBeAttached = 0;
+		SetShouldBeAttached(false);
 
-		bShouldSnapLocationWhenAttached = false;
-		bShouldSnapRotationWhenAttached = false;
+		SetShouldSnapLocationWhenAttached(false);
+		SetShouldSnapRotationWhenAttached(false);
 
 		OnAttachmentChanged();
 
@@ -3320,16 +3320,16 @@ void USceneComponent::PostRepNotifies()
 		}
 		else
 		{
-			const uint8 bOldShouldBeAttached = bShouldBeAttached;
-			const uint8 bOldShouldSnapLocationWhenAttached = bShouldSnapLocationWhenAttached;
-			const uint8 bOldShouldSnapRotationWhenAttached = bShouldSnapRotationWhenAttached;
+			const bool bOldShouldBeAttached = bShouldBeAttached;
+			const bool bOldShouldSnapLocationWhenAttached = bShouldSnapLocationWhenAttached;
+			const bool bOldShouldSnapRotationWhenAttached = bShouldSnapRotationWhenAttached;
 
 			AttachToComponent(NetOldAttachParent, FAttachmentTransformRules::KeepRelativeTransform, NetOldAttachSocketName);
 
 			// restore to what we have received from the server
-			bShouldBeAttached = bOldShouldBeAttached;
-			bShouldSnapLocationWhenAttached = bOldShouldSnapLocationWhenAttached;
-			bShouldSnapRotationWhenAttached = bOldShouldSnapRotationWhenAttached;
+			SetShouldBeAttached(bOldShouldBeAttached);
+			SetShouldSnapLocationWhenAttached(bOldShouldSnapLocationWhenAttached);
+			SetShouldSnapRotationWhenAttached(bOldShouldSnapRotationWhenAttached);
 		}
 
 		bNetUpdateAttachment = false;
@@ -4009,4 +4009,22 @@ void USceneComponent::ModifiedAttachChildren()
 	MARK_PROPERTY_DIRTY_FROM_NAME(USceneComponent, AttachChildren, this);
 }
 
+void USceneComponent::SetShouldBeAttached(bool bNewShouldBeAttached)
+{
+	bShouldBeAttached = bNewShouldBeAttached;
+	MARK_PROPERTY_DIRTY_FROM_NAME(USceneComponent, bShouldBeAttached, this);
+}
+
+void USceneComponent::SetShouldSnapLocationWhenAttached(bool bShouldSnapLocation)
+{
+	bShouldSnapLocationWhenAttached = bShouldSnapLocation;
+	MARK_PROPERTY_DIRTY_FROM_NAME(USceneComponent, bShouldSnapLocationWhenAttached, this);
+}
+
+void USceneComponent::SetShouldSnapRotationWhenAttached(bool bShouldSnapRotation)
+{
+	bShouldSnapRotationWhenAttached = bShouldSnapRotation;
+	MARK_PROPERTY_DIRTY_FROM_NAME(USceneComponent, bShouldSnapRotationWhenAttached, this);
+}
+
 #undef LOCTEXT_NAMESPACE
