Fix CIS warning

#rnx

diff --git a/Engine/Source/Runtime/Engine/Private/VisualLogger/VisualLogger.cpp b/Engine/Source/Runtime/Engine/Private/VisualLogger/VisualLogger.cpp
--- a/Engine/Source/Runtime/Engine/Private/VisualLogger/VisualLogger.cpp
+++ b/Engine/Source/Runtime/Engine/Private/VisualLogger/VisualLogger.cpp
@@ -242,9 +242,9 @@ FVisualLogEntry* FVisualLogger::GetEntryToWrite(const UObject* Object, const flo
 			ThreadCurrentEntry = &ThreadCurrentEntryPerObject.Add(LogOwner);
 
 			CalculateEntryAllowLogging(ThreadCurrentEntry, LogOwner, Object);
-		}
 
-		bInitializeEntry = ThreadCurrentEntry->bIsAllowedToLog;
+			bInitializeEntry = ThreadCurrentEntry->bIsAllowedToLog;
+		}
 	}
 
 	if (bInitializeEntry)
