Fixing Niagara's previews.

- Fixed bug in DPManager preventing proper expansion of CVars from DPs.
- Adding ECVF_Preview flag for CVars so cvars can easily opt into the previews themselves without having to be added to a white list.
- Handling the recent change in the DPMan to use a Preview DP instead of actually changing ActiveDP. Tracking PreviewDP and having Niagara read that if it's set.
- Fixed Preview Platform DP names for Switch.

diff --git a/Engine/Plugins/FX/Niagara/Source/Niagara/Private/NiagaraPlatformSet.cpp b/Engine/Plugins/FX/Niagara/Source/Niagara/Private/NiagaraPlatformSet.cpp
--- a/Engine/Plugins/FX/Niagara/Source/Niagara/Private/NiagaraPlatformSet.cpp
+++ b/Engine/Plugins/FX/Niagara/Source/Niagara/Private/NiagaraPlatformSet.cpp
@@ -48,7 +48,7 @@ static FAutoConsoleVariableRef CVarNiagaraQualityLevel(
 	GNiagaraQualityLevel,
 	TEXT("The quality level for Niagara Effects. \n"),
 	FConsoleVariableDelegate::CreateStatic(&FNiagaraPlatformSet::OnQualityLevelChanged),
-	ECVF_Scalability
+	ECVF_Scalability | ECVF_Preview
 );
 
 // Override platform device profile
@@ -116,6 +116,12 @@ static FAutoConsoleCommand GCmdSetNiagaraPlatformOverride(
 
 static UDeviceProfile* NiagaraGetActiveDeviceProfile()
 {
+	UDeviceProfileManager& DPMan = UDeviceProfileManager::Get();
+	if (UDeviceProfile* DPPreview = DPMan.GetPreviewDeviceProfile())
+	{
+		return DPPreview;
+	}
+
 	UDeviceProfile* ActiveProfile = GNiagaraPlatformOverride.Get(); 
 	if (ActiveProfile == nullptr)
 	{
diff --git a/Engine/Source/Runtime/Core/Public/HAL/IConsoleManager.h b/Engine/Source/Runtime/Core/Public/HAL/IConsoleManager.h
--- a/Engine/Source/Runtime/Core/Public/HAL/IConsoleManager.h
+++ b/Engine/Source/Runtime/Core/Public/HAL/IConsoleManager.h
@@ -90,6 +90,9 @@ enum EConsoleVariableFlags
 	/* those cvars control other cvars with the flag ECVF_Scalability, names should start with "sg." */
 	ECVF_ScalabilityGroup = 0x80,
 
+	/* CVars with this flag will be included in the device profile previews. */
+	ECVF_Preview = 0x100,
+
 	// ------------------------------------------------
 
 	/* Set flags */
diff --git a/Engine/Source/Runtime/Engine/Classes/DeviceProfiles/DeviceProfileManager.h b/Engine/Source/Runtime/Engine/Classes/DeviceProfiles/DeviceProfileManager.h
--- a/Engine/Source/Runtime/Engine/Classes/DeviceProfiles/DeviceProfileManager.h
+++ b/Engine/Source/Runtime/Engine/Classes/DeviceProfiles/DeviceProfileManager.h
@@ -136,6 +136,13 @@ public:
 	 */
 	UDeviceProfile* GetActiveProfile() const;
 
+	/**
+	 * Get the currently previewing device profile. Can be null.
+	 *
+	 * @return The preview profile.
+	 */
+	UDeviceProfile* GetPreviewDeviceProfile() const;
+
 	/**
 	* Get a list of all possible parent profiles for a given device profile
 	*
@@ -274,6 +281,9 @@ private:
 	// Holds the device profile that has been overridden, null no override active.
 	UDeviceProfile* BaseDeviceProfile = nullptr;
 
+	// Holds the device profile that we are previewing.
+	UDeviceProfile* PreviewDeviceProfile = nullptr;
+
 	// Stores any scalability group settings set by the active device profile.
 	static TMap<FString, FString> DeviceProfileScalabilityCVars;
 	
diff --git a/Engine/Source/Runtime/Engine/Private/DeviceProfiles/DeviceProfileManager.cpp b/Engine/Source/Runtime/Engine/Private/DeviceProfiles/DeviceProfileManager.cpp
--- a/Engine/Source/Runtime/Engine/Private/DeviceProfiles/DeviceProfileManager.cpp
+++ b/Engine/Source/Runtime/Engine/Private/DeviceProfiles/DeviceProfileManager.cpp
@@ -166,7 +166,7 @@ void UDeviceProfileManager::ProcessDeviceProfileIniSettings(const FString& Devic
 	check(ConfigSystem);
 
 	TArray< FString > AvailableProfiles;
-	GConfig->GetSectionNames( GDeviceProfilesIni, AvailableProfiles );
+	ConfigSystem->GetSectionNames( GDeviceProfilesIni, AvailableProfiles );
 
 	// Look up the ini for this tree as we are far too early to use the UObject system
 	AvailableProfiles.Remove( TEXT( "DeviceProfiles" ) );
@@ -400,8 +400,13 @@ void UDeviceProfileManager::ProcessDeviceProfileIniSettings(const FString& Devic
 									UE_LOG(LogInit, Log, TEXT("Skipping Device Profile CVar due to PreviewDenylistCVars: [[%s]]"), *CVarKey);
 									continue;
 								}
-
-								if (PreviewAllowlistCVars.Num() > 0 && !PreviewAllowlistCVars.ContainsByPredicate([&CVarKey](const FString& Entry) { return CVarKey.StartsWith(Entry); }))
+								
+								bool bIsPreviewCVar = false;
+								if (IConsoleVariable* CVar = IConsoleManager::Get().FindConsoleVariable(*CVarKey))
+								{
+									bIsPreviewCVar = CVar->TestFlags(ECVF_Preview);
+								}
+								if (bIsPreviewCVar == false && (PreviewAllowlistCVars.Num() > 0 && !PreviewAllowlistCVars.ContainsByPredicate([&CVarKey](const FString& Entry) { return CVarKey.StartsWith(Entry); })))
 								{
 									UE_LOG(LogInit, Log, TEXT("Skipping Device Profile CVar due to PreviewAllowlistCVars: [[%s]]"), *CVarKey);
 									continue;
@@ -949,6 +954,8 @@ void UDeviceProfileManager::SetPreviewDeviceProfile(UDeviceProfile* DeviceProfil
 
 	RestorePreviewDeviceProfile();
 
+	PreviewDeviceProfile = DeviceProfile;
+
 	UE_LOG(LogDeviceProfileManager, Log, TEXT("SetPreviewDeviceProfile preview to %s"), *DeviceProfile->GetName());
 	// apply the preview DP cvars.
 	for (const auto& Pair : DeviceProfile->GetAllExpandedCVars())
@@ -960,6 +967,7 @@ void UDeviceProfileManager::SetPreviewDeviceProfile(UDeviceProfile* DeviceProfil
 			// remember the previous value so we can restore
 			FString OldValue = CVar->GetString();
 			PreviewPushedSettings.Add(Pair.Key, OldValue);
+			//UE_LOG(LogDeviceProfileManager, Log, TEXT("Pushing Device Profile CVar: [[%s:%s]]"), *Pair.Key, *Pair.Value);
 			// cheat CVar can only be set in ConsoleVariables.ini
 			if (!CVar->TestFlags(EConsoleVariableFlags::ECVF_Cheat))
 			{
@@ -981,6 +989,7 @@ void UDeviceProfileManager::RestorePreviewDeviceProfile()
 		// this sets us back to non-preview state.
 		RestorePushedState(PreviewPushedSettings);
 	}
+	PreviewDeviceProfile = nullptr;
 }
 #endif
 
@@ -1262,6 +1271,10 @@ UDeviceProfile* UDeviceProfileManager::GetActiveProfile() const
 	return ActiveDeviceProfile;
 }
 
+UDeviceProfile* UDeviceProfileManager::GetPreviewDeviceProfile() const
+{
+	return PreviewDeviceProfile;
+}
 
 void UDeviceProfileManager::GetAllPossibleParentProfiles(const UDeviceProfile* ChildProfile, OUT TArray<UDeviceProfile*>& PossibleParentProfiles) const
 {
