Revert to UE4 WORLD_MAX. Added macros to help fix new WORLD_MAX issues.

diff --git a/Engine/Source/Runtime/Engine/Public/EngineDefines.h b/Engine/Source/Runtime/Engine/Public/EngineDefines.h
--- a/Engine/Source/Runtime/Engine/Public/EngineDefines.h
+++ b/Engine/Source/Runtime/Engine/Public/EngineDefines.h
@@ -28,10 +28,12 @@
 /*-----------------------------------------------------------------------------
 	Size of the world.
 -----------------------------------------------------------------------------*/
-#define UE_OLD_WORLD_MAX			2097152.0				/* UE4 world size */
+#define UE_OLD_WORLD_MAX			2097152.0					/* UE4 maximum world size */
+#define UE_OLD_HALF_WORLD_MAX		(UE_OLD_WORLD_MAX * 0.5)	/* UE4 half maximum world size */
+#define UE_OLD_HALF_WORLD_MAX1		(UE_OLD_HALF_WORLD_MAX - 1)	/* UE4 half maximum world size minus one */
 
 #ifndef UE_USE_UE4_WORLD_MAX
-#define UE_USE_UE4_WORLD_MAX		(UE_LARGE_WORLD_COORDINATES_DISABLED)			// Force UE4 WORLD_MAX for converted UE4 titles that explicitly rely on it.
+#define UE_USE_UE4_WORLD_MAX		1 //(UE_LARGE_WORLD_COORDINATES_DISABLED)			// Force UE4 WORLD_MAX for converted UE4 titles that explicitly rely on it.
 #endif
 
 #if UE_USE_UE4_WORLD_MAX
@@ -44,7 +46,7 @@
 #define HALF_WORLD_MAX				(WORLD_MAX * 0.5)		/* Half the maximum size of the world */
 #define HALF_WORLD_MAX1				(HALF_WORLD_MAX - 1.0)	/* Half the maximum size of the world minus one */
 
-
-
+#define UE_FLOAT_HUGE_DISTANCE		1048576.0				/* Maximum distance representable by a float whilst maintaining precision of at least 0.0625 units (1/16th of a cm) - Precision issues may occur for positions/distances represented by float types that exceed this value */
+#define UE_DOUBLE_HUGE_DISTANCE		562949953421312.0		/* Maximum distance representable by a double whilst maintaining precision of at least 0.0625 units (1/16th of a cm) - Precision issues may occur for positions/distances represented by double types that exceed this value */
 
 #define DEFAULT_ORTHOZOOM			10000.0					/* Default 2D viewport zoom */
