Provide BP bindings for adding min LOD per platform.

diff --git a/Engine/Source/Runtime/Engine/Classes/Engine/StaticMesh.h b/Engine/Source/Runtime/Engine/Classes/Engine/StaticMesh.h
--- a/Engine/Source/Runtime/Engine/Classes/Engine/StaticMesh.h
+++ b/Engine/Source/Runtime/Engine/Classes/Engine/StaticMesh.h
@@ -815,6 +815,28 @@ public:
 		return INDEX_NONE;
 	}
 
+	UFUNCTION(BlueprintCallable, Category=StaticMesh)
+	void SetMinimumLODForPlatforms(const TMap<FName, int32>& PlatformMinimumLODs)
+	{
+#if WITH_EDITORONLY_DATA
+		WaitUntilAsyncPropertyReleased(EStaticMeshAsyncProperties::MinLOD);
+		PRAGMA_DISABLE_DEPRECATION_WARNINGS
+		MinLOD.PerPlatform = PlatformMinimumLODs;
+		PRAGMA_ENABLE_DEPRECATION_WARNINGS
+#endif
+	}
+
+	UFUNCTION(BlueprintCallable, Category=StaticMesh)
+	void SetMinimumLODForPlatform(const FName& PlatformName, int32 InMinLOD)
+	{
+#if WITH_EDITORONLY_DATA
+		WaitUntilAsyncPropertyReleased(EStaticMeshAsyncProperties::MinLOD);
+		PRAGMA_DISABLE_DEPRECATION_WARNINGS
+		MinLOD.PerPlatform.Add(PlatformName, InMinLOD);
+		PRAGMA_ENABLE_DEPRECATION_WARNINGS
+#endif
+	}
+
 private:
 	UE_DEPRECATED(5.00, "This must be protected for async build, always use the accessors even internally.")
 	UPROPERTY(BlueprintGetter = GetStaticMaterials, BlueprintSetter = SetStaticMaterials, Category = StaticMesh)
