Fix morph targets vertex buffer not updating correctly with the CPU update path for mobile

diff --git a/Engine/Source/Runtime/Engine/Private/GPUSkinCache.cpp b/Engine/Source/Runtime/Engine/Private/GPUSkinCache.cpp
--- a/Engine/Source/Runtime/Engine/Private/GPUSkinCache.cpp
+++ b/Engine/Source/Runtime/Engine/Private/GPUSkinCache.cpp
@@ -177,6 +177,7 @@ ENGINE_API bool DoSkeletalMeshIndexBuffersNeedSRV()
 	return IsGPUSkinCacheAvailable(GMaxRHIShaderPlatform);
 }
 
+UE_DEPRECATED(5.0, "This function is no longer in use and will be removed.")
 ENGINE_API bool DoRecomputeSkinTangentsOnGPU_RT()
 {
 	// currently only implemented and tested on Window SM5 (needs Compute, Atomics, SRV for index buffers, UAV for VertexBuffers)
diff --git a/Engine/Source/Runtime/Engine/Private/SkeletalRenderGPUSkin.cpp b/Engine/Source/Runtime/Engine/Private/SkeletalRenderGPUSkin.cpp
--- a/Engine/Source/Runtime/Engine/Private/SkeletalRenderGPUSkin.cpp
+++ b/Engine/Source/Runtime/Engine/Private/SkeletalRenderGPUSkin.cpp
@@ -673,7 +673,7 @@ void FSkeletalMeshObjectGPUSkin::ProcessUpdatedDynamicData(EGPUSkinCacheEntryMod
 		// only update if the morph data changed and there are weighted morph targets
 		if(bMorphNeedsUpdate)
 		{
-			UpdateMorphVertexBuffer(RHICmdList, Mode, LOD, LODData, MorphVertexBuffer);
+			UpdateMorphVertexBuffer(RHICmdList, Mode, LOD, LODData, bGPUSkinCacheEnabled, MorphVertexBuffer);
 		}
 	}
 	else
@@ -927,7 +927,8 @@ void FSkeletalMeshObjectGPUSkin::UpdateRayTracingGeometry(FSkeletalMeshLODRender
 
 #endif // RHI_RAYTRACING
 
-void FSkeletalMeshObjectGPUSkin::UpdateMorphVertexBuffer(FRHICommandListImmediate& RHICmdList, EGPUSkinCacheEntryMode Mode, FSkeletalMeshObjectLOD& LOD, const FSkeletalMeshLODRenderData& LODData, FMorphVertexBuffer& MorphVertexBuffer)
+void FSkeletalMeshObjectGPUSkin::UpdateMorphVertexBuffer(FRHICommandListImmediate& RHICmdList, EGPUSkinCacheEntryMode Mode, FSkeletalMeshObjectLOD& LOD, const FSkeletalMeshLODRenderData& LODData, 
+															bool bGPUSkinCacheEnabled, FMorphVertexBuffer& MorphVertexBuffer)
 {
 	QUICK_SCOPE_CYCLE_COUNTER(STAT_FSkeletalMeshObjectGPUSkin_ProcessUpdatedDynamicData_UpdateMorphBuffer);
 	if (UseGPUMorphTargets(FeatureLevel))
@@ -939,7 +940,7 @@ void FSkeletalMeshObjectGPUSkin::UpdateMorphVertexBuffer(FRHICommandListImmediat
 	else
 	{
 		// update the morph data for the lod (before SkinCache)
-		LOD.UpdateMorphVertexBufferCPU(DynamicData->ActiveMorphTargets, DynamicData->MorphTargetWeights, MorphVertexBuffer);
+		LOD.UpdateMorphVertexBufferCPU(DynamicData->ActiveMorphTargets, DynamicData->MorphTargetWeights, DynamicData->SectionIdsUseByActiveMorphTargets, bGPUSkinCacheEnabled, MorphVertexBuffer);
 	}
 
 	if (LOD.MorphVertexBufferPool.IsDoubleBuffered())
@@ -1222,17 +1223,33 @@ void FSkeletalMeshObjectGPUSkin::FSkeletalMeshObjectLOD::UpdateSkinWeights(FSkel
 	
 }
 
-void FSkeletalMeshObjectGPUSkin::FSkeletalMeshObjectLOD::UpdateMorphVertexBufferCPU(const TArray<FActiveMorphTarget>& ActiveMorphTargets, const TArray<float>& MorphTargetWeights, FMorphVertexBuffer& MorphVertexBuffer)
+void FSkeletalMeshObjectGPUSkin::FSkeletalMeshObjectLOD::UpdateMorphVertexBufferCPU(const TArray<FActiveMorphTarget>& ActiveMorphTargets, const TArray<float>& MorphTargetWeights, 
+																					const TArray<int32>& SectionIdsUseByActiveMorphTargets, bool bGPUSkinCacheEnabled, FMorphVertexBuffer& MorphVertexBuffer)
 {
 	SCOPE_CYCLE_COUNTER(STAT_MorphVertexBuffer_Update);
 
 	if (IsValidRef(MorphVertexBuffer.VertexBufferRHI))
 	{
-		extern ENGINE_API bool DoRecomputeSkinTangentsOnGPU_RT();
-		bool bBlendTangentsOnCPU = !DoRecomputeSkinTangentsOnGPU_RT();
-
 		// LOD of the skel mesh is used to find number of vertices in buffer
 		FSkeletalMeshLODRenderData& LodData = SkelMeshRenderData->LODRenderData[LODIndex];
+
+		// Whether all sections of the LOD perform GPU recompute tangent
+		bool bAllSectionsDoGPURecomputeTangent = bGPUSkinCacheEnabled && GSkinCacheRecomputeTangents > 0;
+		if (bAllSectionsDoGPURecomputeTangent && GSkinCacheRecomputeTangents == 2)
+		{
+			for (int32 i = 0; i < LodData.RenderSections.Num(); ++i)
+			{
+				const FSkelMeshRenderSection& RenderSection = LodData.RenderSections[i];
+				if (!RenderSection.bRecomputeTangent)
+				{
+					bAllSectionsDoGPURecomputeTangent = false;
+					break;
+				}
+			}
+		}
+
+		// If the LOD performs GPU skin cache recompute tangent, then there is no need to update tangents here
+		bool bBlendTangentsOnCPU = !bAllSectionsDoGPURecomputeTangent;
 		
 		const bool bUseGPUMorphTargets = UseGPUMorphTargets(FeatureLevel);
 		MorphVertexBuffer.RecreateResourcesIfRequired(bUseGPUMorphTargets);
@@ -1337,6 +1354,8 @@ void FSkeletalMeshObjectGPUSkin::FSkeletalMeshObjectLOD::UpdateMorphVertexBuffer
 			SCOPE_CYCLE_COUNTER(STAT_MorphVertexBuffer_RhiUnlock);
 			// Unlock the buffer.
 			RHIUnlockBuffer(MorphVertexBuffer.VertexBufferRHI);
+			// Copy the section Ids use by all active morph targets
+			MorphVertexBuffer.SectionIds = SectionIdsUseByActiveMorphTargets;
 			// set update flag
 			MorphVertexBuffer.bHasBeenUpdated = true;
 		}
diff --git a/Engine/Source/Runtime/Engine/Private/SkeletalRenderGPUSkin.h b/Engine/Source/Runtime/Engine/Private/SkeletalRenderGPUSkin.h
--- a/Engine/Source/Runtime/Engine/Private/SkeletalRenderGPUSkin.h
+++ b/Engine/Source/Runtime/Engine/Private/SkeletalRenderGPUSkin.h
@@ -584,7 +584,8 @@ protected:
 		 * @param ActiveMorphTargets - Morph to accumulate. assumed to be weighted and have valid targets
 		 * @param MorphTargetWeights - All Morph weights
 		 */
-		void UpdateMorphVertexBufferCPU(const TArray<FActiveMorphTarget>& ActiveMorphTargets, const TArray<float>& MorphTargetWeights, FMorphVertexBuffer& MorphVertexBuffer);
+		void UpdateMorphVertexBufferCPU(const TArray<FActiveMorphTarget>& ActiveMorphTargets, const TArray<float>& MorphTargetWeights, const TArray<int32>& SectionIdsUseByActiveMorphTargets, 
+										bool bGPUSkinCacheEnabled, FMorphVertexBuffer& MorphVertexBuffer);
 		void UpdateMorphVertexBufferGPU(FRHICommandListImmediate& RHICmdList, const TArray<float>& MorphTargetWeights, const FMorphTargetVertexInfoBuffers& MorphTargetVertexInfoBuffers, 
 										const TArray<int32>& SectionIdsUseByActiveMorphTargets, const FName& OwnerName, EGPUSkinCacheEntryMode Mode, FMorphVertexBuffer& MorphVertexBuffer);
 
@@ -613,7 +614,7 @@ protected:
 
 	void ProcessUpdatedDynamicData(EGPUSkinCacheEntryMode Mode, FGPUSkinCache* GPUSkinCache, FRHICommandListImmediate& RHICmdList, uint32 FrameNumberToPrepare, uint32 RevisionNumber, bool bMorphNeedsUpdate, int32 LODIndex);
 
-	virtual void UpdateMorphVertexBuffer(FRHICommandListImmediate& RHICmdList, EGPUSkinCacheEntryMode Mode, FSkeletalMeshObjectLOD& LOD, const FSkeletalMeshLODRenderData& LODData, FMorphVertexBuffer& MorphVertexBuffer);
+	virtual void UpdateMorphVertexBuffer(FRHICommandListImmediate& RHICmdList, EGPUSkinCacheEntryMode Mode, FSkeletalMeshObjectLOD& LOD, const FSkeletalMeshLODRenderData& LODData, bool bGPUSkinCacheEnabled, FMorphVertexBuffer& MorphVertexBuffer);
 
 	void WaitForRHIThreadFenceForDynamicData();
 
diff --git a/Engine/Source/Runtime/Engine/Public/GPUSkinCache.h b/Engine/Source/Runtime/Engine/Public/GPUSkinCache.h
--- a/Engine/Source/Runtime/Engine/Public/GPUSkinCache.h
+++ b/Engine/Source/Runtime/Engine/Public/GPUSkinCache.h
@@ -61,6 +61,7 @@ extern ENGINE_API bool GPUSkinCacheNeedsDuplicatedVertices();
 
 // Is it actually enabled?
 extern ENGINE_API int32 GEnableGPUSkinCache;
+extern int32 GSkinCacheRecomputeTangents;
 
 class FGPUSkinCacheEntry;
 
