Fixed optimization using s.LevelStreamingAddPrimitiveGranularity : Support cases where AActor::PostRegisterAllComponents triggers code that either unregisters or re-registers components pushed in FRegisterComponentContext.

diff --git a/Engine/Source/Runtime/Engine/Private/Components/ActorComponent.cpp b/Engine/Source/Runtime/Engine/Private/Components/ActorComponent.cpp
--- a/Engine/Source/Runtime/Engine/Private/Components/ActorComponent.cpp
+++ b/Engine/Source/Runtime/Engine/Private/Components/ActorComponent.cpp
@@ -109,11 +109,19 @@ void FRegisterComponentContext::Process()
 		{
 			FOptionalTaskTagScope Scope(ETaskTag::EParallelGameThread);
 			UPrimitiveComponent* Component = AddPrimitiveBatches[Index];
-			if (IsValid(Component))
+
+			// AActor::PostRegisterAllComponents (called by AActor::IncrementalRegisterComponents) can trigger code 
+			// that either unregisters or re-registers components. If unregistered, skip this component.
+			// If re-registered, FRegisterComponentContext is not passed, so SceneProxy can be created.
+			if (IsValid(Component) && Component->IsRegistered())
 			{
 				if (Component->IsRenderStateCreated() || !bAppCanEverRender)
 				{
-					Scene->AddPrimitive(Component);
+					// Skip if SceneProxy is already created
+					if (Component->SceneProxy == nullptr)
+					{
+						Scene->AddPrimitive(Component);
+					}
 				}
 				else // Fallback for some edge case where the component renderstate are missing
 				{
