Fix inconsistent serialization of floating point vectors in FTextureBuildSettings triggered by LWC.  Always use float (not double) precision for these vectors.
Fix legacy DDC key generation for textures to always use double (not float) precision when composing the key.

diff --git a/Engine/Source/Developer/TextureBuild/Private/TextureBuildFunction.cpp b/Engine/Source/Developer/TextureBuild/Private/TextureBuildFunction.cpp
--- a/Engine/Source/Developer/TextureBuild/Private/TextureBuildFunction.cpp
+++ b/Engine/Source/Developer/TextureBuild/Private/TextureBuildFunction.cpp
@@ -66,14 +66,14 @@ static void ReadCbField(FCbFieldView Field, FColor& OutValue)
 	OutValue.B = It++->AsUInt8(OutValue.B);
 }
 
-static void ReadCbField(FCbFieldView Field, FVector2D& OutValue)
+static void ReadCbField(FCbFieldView Field, FVector2f& OutValue)
 {
 	FCbFieldViewIterator It = Field.AsArrayView().CreateViewIterator();
 	OutValue.X = It++->AsFloat(OutValue.X);
 	OutValue.Y = It++->AsFloat(OutValue.Y);
 }
 
-static void ReadCbField(FCbFieldView Field, FVector4& OutValue)
+static void ReadCbField(FCbFieldView Field, FVector4f& OutValue)
 {
 	FCbFieldViewIterator It = Field.AsArrayView().CreateViewIterator();
 	OutValue.X = It++->AsFloat(OutValue.X);
diff --git a/Engine/Source/Developer/TextureCompressor/Public/TextureCompressorModule.h b/Engine/Source/Developer/TextureCompressor/Public/TextureCompressorModule.h
--- a/Engine/Source/Developer/TextureCompressor/Public/TextureCompressorModule.h
+++ b/Engine/Source/Developer/TextureCompressor/Public/TextureCompressorModule.h
@@ -73,7 +73,7 @@ struct FTextureBuildSettings
 	/** Color adjustment parameters. */
 	FColorAdjustmentParameters ColorAdjustment;
 	/** Channel values to compare to when preserving alpha coverage. */
-	FVector4 AlphaCoverageThresholds;
+	FVector4f AlphaCoverageThresholds;
 	/** The desired amount of mip sharpening. */
 	float MipSharpening;
 	/** For angular filtered cubemaps, the mip level which contains convolution with the diffuse cosine lobe. */
@@ -103,13 +103,13 @@ struct FTextureBuildSettings
 	/** Whether the texture has a defined source color space. */
 	bool bHasColorSpaceDefinition;
 	/** Red chromaticity coordinate of the source color space. */
-	FVector2D RedChromaticityCoordinate;
+	FVector2f RedChromaticityCoordinate;
 	/** Green chromaticity coordinate of the source color space. */
-	FVector2D GreenChromaticityCoordinate;
+	FVector2f GreenChromaticityCoordinate;
 	/** Blue chromaticity coordinate of the source color space. */
-	FVector2D BlueChromaticityCoordinate;
+	FVector2f BlueChromaticityCoordinate;
 	/** White chromaticity coordinate of the source color space. */
-	FVector2D WhiteChromaticityCoordinate;
+	FVector2f WhiteChromaticityCoordinate;
 	/** Chromatic adaption method applied if the source white point differs from the working color space white point. */
 	uint8 ChromaticAdaptationMethod;
 	/** Whether the texture should use the legacy gamma space for converting to sRGB */
diff --git a/Engine/Source/Runtime/Engine/Private/TextureDerivedData.cpp b/Engine/Source/Runtime/Engine/Private/TextureDerivedData.cpp
--- a/Engine/Source/Runtime/Engine/Private/TextureDerivedData.cpp
+++ b/Engine/Source/Runtime/Engine/Private/TextureDerivedData.cpp
@@ -95,8 +95,8 @@ static void SerializeForKey(FArchive& Ar, const FTextureBuildSettings& Settings)
 	float TempFloat;
 	uint8 TempByte;
 	FColor TempColor;
-	FVector2D TempVector2D;
-	FVector4 TempVector4;
+	FVector2d TempVector2d;
+	FVector4d TempVector4d;
 	UE::Color::FColorSpace TempColorSpace;
 	FGuid TempGuid;
 
@@ -129,10 +129,10 @@ static void SerializeForKey(FArchive& Ar, const FTextureBuildSettings& Settings)
 		TempUint32 = UE::Color::COLORSPACE_VER; Ar << TempUint32;
 		TempColorSpace = UE::Color::FColorSpace::GetWorking(); Ar << TempColorSpace;
 
-		TempVector2D = Settings.RedChromaticityCoordinate; Ar << TempVector2D;
-		TempVector2D = Settings.GreenChromaticityCoordinate; Ar << TempVector2D;
-		TempVector2D = Settings.BlueChromaticityCoordinate; Ar << TempVector2D;
-		TempVector2D = Settings.WhiteChromaticityCoordinate; Ar << TempVector2D;
+		TempVector2d = Settings.RedChromaticityCoordinate; Ar << TempVector2d;
+		TempVector2d = Settings.GreenChromaticityCoordinate; Ar << TempVector2d;
+		TempVector2d = Settings.BlueChromaticityCoordinate; Ar << TempVector2d;
+		TempVector2d = Settings.WhiteChromaticityCoordinate; Ar << TempVector2d;
 		TempByte = Settings.ChromaticAdaptationMethod; Ar << TempByte;
 	}
 
@@ -141,7 +141,7 @@ static void SerializeForKey(FArchive& Ar, const FTextureBuildSettings& Settings)
 
 	if (Settings.AlphaCoverageThresholds != FVector4(0, 0, 0, 0))
 	{
-		TempVector4 = Settings.AlphaCoverageThresholds; Ar << TempVector4;
+		TempVector4d = Settings.AlphaCoverageThresholds; Ar << TempVector4d;
 	}
 	
 	TempByte = Settings.bComputeBokehAlpha; Ar << TempByte;
diff --git a/Engine/Source/Runtime/Engine/Private/TextureDerivedDataBuildUtils.cpp b/Engine/Source/Runtime/Engine/Private/TextureDerivedDataBuildUtils.cpp
--- a/Engine/Source/Runtime/Engine/Private/TextureDerivedDataBuildUtils.cpp
+++ b/Engine/Source/Runtime/Engine/Private/TextureDerivedDataBuildUtils.cpp
@@ -40,7 +40,7 @@ static void WriteCbField(FCbWriter& Writer, FAnsiStringView Name, const FColor&
 	Writer.EndArray();
 }
 
-static void WriteCbField(FCbWriter& Writer, FAnsiStringView Name, const FVector2D& Value)
+static void WriteCbField(FCbWriter& Writer, FAnsiStringView Name, const FVector2f& Value)
 {
 	Writer.BeginArray(Name);
 	Writer.AddFloat(Value.X);
@@ -48,7 +48,7 @@ static void WriteCbField(FCbWriter& Writer, FAnsiStringView Name, const FVector2
 	Writer.EndArray();
 }
 
-static void WriteCbField(FCbWriter& Writer, FAnsiStringView Name, const FVector4& Value)
+static void WriteCbField(FCbWriter& Writer, FAnsiStringView Name, const FVector4f& Value)
 {
 	Writer.BeginArray(Name);
 	Writer.AddFloat(Value.X);
