Changes default screen percentage behavior in editor to take more advantage of TSR's upscaling behavior already enabled by default

1) Editor viewports now display at high DPI by default for better Gen5-console feeling in editor on high DPI monitors.

2) Realtime editor viewports now display at screen percentage based on the display pixel count. Fix screen percentage independent of the pixel range could either render not enough pixel at low display resolution or be too slow at high display resolution. Given the level editor viewport can scale in size based on editor user layout preference, monitor, the screen percentage automatically scalling based on display pixel count gives a better experience of quality consistency/GPU cost for rendering frame, especially when going level editor fullscreen with F11 or not. The curve to compute rendering resolution from display resolution is in BaseEngine.ini

3) Non-realtime editor viewports still display at a screen percentage based on OS DPI scale, independent of the display pixel count because use spatial upscaler and therefore need consistent sharpness not worst than editor fonts, but keep same rendering resolution as before to not increase out of GPU memory errors. The performance of potentially rendering at high resolution than a realtime viewport is ignored given the non real-time viewports are not rendering every frame.

4) PIE viewports now override r.ScreenPercentage by default to have consistent resolution quality and performance between realtime editor viewport and gameplay viewport by default without risk of introducing out of GPU memory errors loosing possibly unsaved editor work if the game's graphic settings default to a screen percentage higher. To advertise discoverability of this change, any update to the r.ScreenPercentage cvar will lead to a message display in console/log if overriden by editor settings or dynamic resolution settings.

5) As much as a project can change the default anti-aliasing method, can also configure on a per project basis the default screen percentage behavior in project preference that the editor user can then override locally for his/her own use with his editor settings. For instance mobile only or forward renderer MSAA only video games.

6) When working very high end work-in progress rendering tech project in a econemic context where lattest GPUs are hard to get, the range of GPU performance can vary greatly between contributor of the same projects. To provide more robustness for lowerend GPU of collegues with tech not yet scaling, the project settings now offer a MaxRenderingResolution for editor viewports that can also be overriden on a per user basis.

7) Automated screenshot already takes care of forcing r.ScreenPercentage.

diff --git a/Engine/Config/BaseEngine.ini b/Engine/Config/BaseEngine.ini
--- a/Engine/Config/BaseEngine.ini
+++ b/Engine/Config/BaseEngine.ini
@@ -1659,6 +1659,14 @@ r.Shaders.SymbolPathOverride=""
 ; Base shader symbols off of source and not generated shaders
 r.Shaders.AllowUniqueSymbols=0
 
+[Rendering.AutoScreenPercentage]
+MinDisplayResolution=720
+MinRenderingResolution=720
+MidDisplayResolution=2160
+MidRenderingResolution=1080
+MaxDisplayResolution=4320
+MaxRenderingResolution=1440
+
 [LogFiles]
 PurgeLogsDays=5
 MaxLogFilesOnDisk=10
diff --git a/Engine/Source/Developer/FunctionalTesting/Private/AutomationBlueprintFunctionLibrary.cpp b/Engine/Source/Developer/FunctionalTesting/Private/AutomationBlueprintFunctionLibrary.cpp
--- a/Engine/Source/Developer/FunctionalTesting/Private/AutomationBlueprintFunctionLibrary.cpp
+++ b/Engine/Source/Developer/FunctionalTesting/Private/AutomationBlueprintFunctionLibrary.cpp
@@ -243,6 +243,9 @@ FAutomationTestScreenshotEnvSetup::FAutomationTestScreenshotEnvSetup()
 	, ContactShadows(TEXT("r.ContactShadows"))
 	, TonemapperGamma(TEXT("r.TonemapperGamma"))
 	, TonemapperSharpen(TEXT("r.Tonemapper.Sharpen"))
+	, ScreenPercentage(TEXT("r.ScreenPercentage"))
+	, ScreenPercentageMode(TEXT("r.ScreenPercentage.Mode"))
+	, EditorViewportOverrideGameScreenPercentage(TEXT("r.Editor.Viewport.OverridePIEScreenPercentage"))
 	, SecondaryScreenPercentage(TEXT("r.SecondaryScreenPercentage.GameViewport"))
 {
 }
@@ -276,8 +279,18 @@ void FAutomationTestScreenshotEnvSetup::Setup(UWorld* InWorld, FAutomationScreen
 		//TonemapperSharpen.Set(0);
 	}
 
+	// Forces ScreenPercentage=100
+	{
+		if (GIsEditor)
+		{
+			EditorViewportOverrideGameScreenPercentage.Set(0);
+		}
+		ScreenPercentageMode.Set(0);
+		ScreenPercentage.Set(100.f);
+	}
+
 	// Ignore High-DPI settings
-	SecondaryScreenPercentage.Set(100.f); 
+	SecondaryScreenPercentage.Set(100.f);
 
 	InOutOptions.SetToleranceAmounts(InOutOptions.Tolerance);
 
@@ -313,6 +326,9 @@ void FAutomationTestScreenshotEnvSetup::Restore()
 	ContactShadows.Restore();
 	TonemapperGamma.Restore();
 	//TonemapperSharpen.Restore();
+	ScreenPercentage.Restore();
+	ScreenPercentageMode.Restore();
+	EditorViewportOverrideGameScreenPercentage.Restore();
 	SecondaryScreenPercentage.Restore();
 
 	AutomationViewExtension.Reset();
diff --git a/Engine/Source/Developer/FunctionalTesting/Public/AutomationBlueprintFunctionLibrary.h b/Engine/Source/Developer/FunctionalTesting/Public/AutomationBlueprintFunctionLibrary.h
--- a/Engine/Source/Developer/FunctionalTesting/Public/AutomationBlueprintFunctionLibrary.h
+++ b/Engine/Source/Developer/FunctionalTesting/Public/AutomationBlueprintFunctionLibrary.h
@@ -234,6 +234,9 @@ private:
 	FConsoleVariableSwapperTempl<int32> ContactShadows;
 	FConsoleVariableSwapperTempl<float> TonemapperGamma;
 	FConsoleVariableSwapperTempl<float> TonemapperSharpen;
+	FConsoleVariableSwapperTempl<float> ScreenPercentage;
+	FConsoleVariableSwapperTempl<int32> ScreenPercentageMode;
+	FConsoleVariableSwapperTempl<int32> EditorViewportOverrideGameScreenPercentage;
 	FConsoleVariableSwapperTempl<float> SecondaryScreenPercentage;
 
 	TWeakObjectPtr<UWorld> WorldPtr;
diff --git a/Engine/Source/Editor/LevelEditor/Private/SLevelViewportToolBar.cpp b/Engine/Source/Editor/LevelEditor/Private/SLevelViewportToolBar.cpp
--- a/Engine/Source/Editor/LevelEditor/Private/SLevelViewportToolBar.cpp
+++ b/Engine/Source/Editor/LevelEditor/Private/SLevelViewportToolBar.cpp
@@ -1324,8 +1324,8 @@ void SLevelViewportToolBar::OnFOVValueChanged( float NewValue )
 
 TSharedRef<SWidget> SLevelViewportToolBar::GenerateScreenPercentageMenu() const
 {
-	const int32 PreviewScreenPercentageMin = ISceneViewFamilyScreenPercentage::kMinTAAUpsampleResolutionFraction * 100.0f;
-	const int32 PreviewScreenPercentageMax = ISceneViewFamilyScreenPercentage::kMaxTAAUpsampleResolutionFraction * 100.0f;
+	const int32 PreviewScreenPercentageMin = ISceneViewFamilyScreenPercentage::kMinTSRResolutionFraction * 100.0f;
+	const int32 PreviewScreenPercentageMax = ISceneViewFamilyScreenPercentage::kMaxTSRResolutionFraction * 100.0f;
 
 	return
 		SNew(SBox)
@@ -1343,8 +1343,8 @@ TSharedRef<SWidget> SLevelViewportToolBar::GenerateScreenPercentageMenu() const
 					SNew(SSpinBox<int32>)
 					.Style(&FAppStyle::Get(), "Menu.SpinBox")
 					.Font(FEditorStyle::GetFontStyle(TEXT("MenuItem.Font")))
-					.MinValue(PreviewScreenPercentageMin)
-					.MaxValue(PreviewScreenPercentageMax)
+					.MinSliderValue(PreviewScreenPercentageMin)
+					.MaxSliderValue(PreviewScreenPercentageMax)
 					.Value(this, &SLevelViewportToolBar::OnGetScreenPercentageValue)
 					.OnValueChanged(const_cast<SLevelViewportToolBar*>(this), &SLevelViewportToolBar::OnScreenPercentageValueChanged)
 				]
diff --git a/Engine/Source/Editor/UnrealEd/Classes/Editor/EditorPerformanceSettings.h b/Engine/Source/Editor/UnrealEd/Classes/Editor/EditorPerformanceSettings.h
--- a/Engine/Source/Editor/UnrealEd/Classes/Editor/EditorPerformanceSettings.h
+++ b/Engine/Source/Editor/UnrealEd/Classes/Editor/EditorPerformanceSettings.h
@@ -6,9 +6,28 @@
 #include "UObject/ObjectMacros.h"
 #include "UObject/Object.h"
 #include "Engine/DeveloperSettings.h"
+
 #include "EditorPerformanceSettings.generated.h"
 
 
+// Mode for the computation of the screen percentage (r.ScreenPercentage.Mode).
+UENUM()
+enum class EEditorUserScreenPercentageModeOverride
+{
+	// Uses the project's default configured in the project settings.
+	ProjectDefault,
+
+	// Directly controls the screen percentage with the r.ScreenPercentage cvar
+	Manual UMETA(DisplayName = "Manual"),
+
+	// Automatic control the screen percentage based on the display resolution, r.ScreenPercentage.Auto.*
+	BasedOnDisplayResolution UMETA(DisplayName = "Based on display resolution"),
+
+	// Based on DPI scale.
+	BasedOnDPIScale UMETA(DisplayName = "Based on operating system's DPI scale"),
+};
+
+
 UCLASS(minimalapi, config=EditorSettings, meta=(DisplayName = "Performance", ToolTip="Settings to tweak the performance of the editor"))
 class UEditorPerformanceSettings : public UDeveloperSettings
 {
@@ -30,13 +49,6 @@ class UEditorPerformanceSettings : public UDeveloperSettings
 	UPROPERTY(EditAnywhere, config, Category=EditorPerformance)
 	uint32 bMonitorEditorPerformance:1;
 
-	/** 
-	 * By default the editor will adjust scene scaling (quality) for high DPI in order to ensure consistent performance with very large render targets.
-	 * Enabling this will disable automatic adjusting and render at the full resolution of the viewport
-	 */
-	UPROPERTY(EditAnywhere, config, Category=EditorPerformance, meta=(DisplayName="Disable DPI Based Editor Viewport Scaling", ConsoleVariable="Editor.OverrideDPIBasedEditorViewportScaling"))
-	bool bOverrideDPIBasedEditorViewportScaling;
-
 	/** When enabled, Shared Data Cache performance notifications may be displayed when not connected to a shared cache */
 	UPROPERTY(EditAnywhere, config, Category = EditorPerformance, meta = (DisplayName = "Enable Shared Data Cache Performance Notifications"))
 	uint32 bEnableSharedDDCPerformanceNotifications : 1;
@@ -44,6 +56,56 @@ class UEditorPerformanceSettings : public UDeveloperSettings
 	/** When enabled, a warning will appear in the viewport when your editors scalability settings are non-default and you may be viewing a low quality scene */
 	UPROPERTY(EditAnywhere, config, Category = EditorPerformance, meta = (DisplayName = "Enable Scalability Warning Indicator"))
 	uint32 bEnableScalabilityWarningIndicator : 1;
+	
+	/** 
+	 * By default the editor will adjust scene scaling (quality) for high DPI in order to ensure consistent performance with very large render targets.
+	 * Enabling this will disable automatic adjusting and render at the full resolution of the viewport
+	 */
+	UPROPERTY(EditAnywhere, config, Category=ViewportResolution, meta=(
+		DisplayName="Displays viewports' at high DPI",
+		ConsoleVariable="r.Editor.Viewport.HighDPI"))
+	uint32 bDisplayHighDPIViewports : 1;
+	
+	UPROPERTY(EditAnywhere, config, Category=ViewportResolution, meta=(
+		DisplayName="Override game screen percentage settings with editor settings in PIE",
+		ConsoleVariable="r.Editor.Viewport.OverridePIEScreenPercentage"))
+	uint32 bOverridePIEScreenPercentage : 1;
+
+	UPROPERTY(EditAnywhere, config, Category=ViewportResolution, meta=(
+		DisplayName="Override project's default screen percentage mode for realtime editor viewports."))
+	EEditorUserScreenPercentageModeOverride RealtimeScreenPercentageMode;
+
+	UPROPERTY(EditAnywhere, config, Category=ViewportResolution, meta=(
+		DisplayName="Override project's default screen percentage mode for non-realtime editor viewports."))
+	EEditorUserScreenPercentageModeOverride NonRealtimeScreenPercentageMode;
+
+	UPROPERTY(EditAnywhere, config, Category=ViewportResolution, meta=(
+		DisplayName="Override project's manual screen percentage"))
+	uint32 bOverrideManualScreenPercentage : 1;
+
+	/** Editor viewport screen percentage */
+	UPROPERTY(EditAnywhere, config, Category=ViewportResolution, meta=(
+		EditCondition="bOverrideManualScreenPercentage",
+		DisplayName="Editor viewport screen percentage."))
+	float ManualScreenPercentage;
+
+	UPROPERTY(EditAnywhere, config, Category=ViewportResolution, meta=(
+		DisplayName="Override project's minimum viewport rendering resolution"))
+	uint32 bOverrideMinViewportRenderingResolution : 1;
+
+	UPROPERTY(EditAnywhere, config, Category=ViewportResolution, meta=(
+		EditCondition="bOverrideMinViewportRenderingResolution",
+		DisplayName="Minimum default rendering resolution to use for editor viewports."))
+	int32 MinViewportRenderingResolution;
+
+	UPROPERTY(EditAnywhere, config, Category=ViewportResolution, meta=(
+		DisplayName="Override project's maximum viewport rendering resolution"))
+	uint32 bOverrideMaxViewportRenderingResolution : 1;
+
+	UPROPERTY(EditAnywhere, config, Category=ViewportResolution, meta=(
+		EditCondition="bOverrideMaxViewportRenderingResolution",
+		DisplayName="Maximum default rendering resolution to use for editor viewports."))
+	int32 MaxViewportRenderingResolution;
 
 public:
 	/** UObject interface */
diff --git a/Engine/Source/Editor/UnrealEd/Private/EditorPerformanceSettings.cpp b/Engine/Source/Editor/UnrealEd/Private/EditorPerformanceSettings.cpp
--- a/Engine/Source/Editor/UnrealEd/Private/EditorPerformanceSettings.cpp
+++ b/Engine/Source/Editor/UnrealEd/Private/EditorPerformanceSettings.cpp
@@ -1,13 +1,16 @@
 // Copyright Epic Games, Inc. All Rights Reserved.
 
 #include "Editor/EditorPerformanceSettings.h"
-#include "EditorViewportClient.h"
-#include "Editor.h"
+#include "Settings/EditorProjectSettings.h"
 
-static TAutoConsoleVariable<int32> CVarOverrideDPIBasedEditorViewportScaling(
-	TEXT("Editor.OverrideDPIBasedEditorViewportScaling"),
-	0,
-	TEXT("Sets whether or not we should globally override screen percentage in editor and PIE viewports"),
+TAutoConsoleVariable<int32> CVarEditorViewportHighDPI(
+	TEXT("r.Editor.Viewport.HighDPI"), 1,
+	TEXT("Controls whether editor & PIE viewports can be displayed at high DPI."),
+	ECVF_Default);
+
+TAutoConsoleVariable<int32> CVarEditorViewportOverrideGameScreenPercentage(
+	TEXT("r.Editor.Viewport.OverridePIEScreenPercentage"), 1,
+	TEXT("Apply editor viewports' default screen percentage settings to game viewport clients in PIE."),
 	ECVF_Default);
 
 UEditorPerformanceSettings::UEditorPerformanceSettings(const FObjectInitializer& ObjectInitializer)
@@ -16,9 +19,18 @@ UEditorPerformanceSettings::UEditorPerformanceSettings(const FObjectInitializer&
 	, bThrottleCPUWhenNotForeground(true)
 	, bDisableRealtimeViewportsInRemoteSessions(true)
 	, bMonitorEditorPerformance(false)
-	, bOverrideDPIBasedEditorViewportScaling(false)
 	, bEnableSharedDDCPerformanceNotifications(true)
 	, bEnableScalabilityWarningIndicator(true)
+	, bDisplayHighDPIViewports(true)
+	, bOverridePIEScreenPercentage(true)
+	, RealtimeScreenPercentageMode(EEditorUserScreenPercentageModeOverride::ProjectDefault)
+	, NonRealtimeScreenPercentageMode(EEditorUserScreenPercentageModeOverride::ProjectDefault)
+	, bOverrideManualScreenPercentage(false)
+	, ManualScreenPercentage(100.0f)
+	, bOverrideMinViewportRenderingResolution(false)
+	, MinViewportRenderingResolution(720)
+	, bOverrideMaxViewportRenderingResolution(false)
+	, MaxViewportRenderingResolution(2160)
 {
 
 }
@@ -27,30 +39,36 @@ void UEditorPerformanceSettings::PostInitProperties()
 {
 	Super::PostInitProperties();
 
-	CVarOverrideDPIBasedEditorViewportScaling->Set(bOverrideDPIBasedEditorViewportScaling, ECVF_SetByProjectSetting);
+	CVarEditorViewportHighDPI->Set(bDisplayHighDPIViewports != 0, ECVF_SetByProjectSetting);
+	CVarEditorViewportOverrideGameScreenPercentage->Set(bOverridePIEScreenPercentage != 0, ECVF_SetByProjectSetting);
 
+	UEditorPerformanceProjectSettings::ExportResolutionValuesToConsoleVariables();
 }
 
 void UEditorPerformanceSettings::PostEditChangeProperty(FPropertyChangedEvent& PropertyChangedEvent)
 {
 	Super::PostEditChangeProperty(PropertyChangedEvent);
 
-	ExportValuesToConsoleVariables(PropertyChangedEvent.Property);
-	if (PropertyChangedEvent.Property && (PropertyChangedEvent.Property->GetFName() == GET_MEMBER_NAME_CHECKED(UEditorPerformanceSettings, bOverrideDPIBasedEditorViewportScaling)))
+	if (!PropertyChangedEvent.Property)
+	{
+		return;
+	}
+		
+	if (PropertyChangedEvent.Property->GetFName() == GET_MEMBER_NAME_CHECKED(UEditorPerformanceSettings, bDisplayHighDPIViewports) ||
+		PropertyChangedEvent.Property->GetFName() == GET_MEMBER_NAME_CHECKED(UEditorPerformanceSettings, bOverridePIEScreenPercentage) ||
+		PropertyChangedEvent.Property->GetFName() == GET_MEMBER_NAME_CHECKED(UEditorPerformanceSettings, RealtimeScreenPercentageMode) ||
+		PropertyChangedEvent.Property->GetFName() == GET_MEMBER_NAME_CHECKED(UEditorPerformanceSettings, NonRealtimeScreenPercentageMode) ||
+		PropertyChangedEvent.Property->GetFName() == GET_MEMBER_NAME_CHECKED(UEditorPerformanceSettings, bOverrideManualScreenPercentage) ||
+		PropertyChangedEvent.Property->GetFName() == GET_MEMBER_NAME_CHECKED(UEditorPerformanceSettings, ManualScreenPercentage) ||
+		PropertyChangedEvent.Property->GetFName() == GET_MEMBER_NAME_CHECKED(UEditorPerformanceSettings, bOverrideMinViewportRenderingResolution) ||
+		PropertyChangedEvent.Property->GetFName() == GET_MEMBER_NAME_CHECKED(UEditorPerformanceSettings, MinViewportRenderingResolution) ||
+		PropertyChangedEvent.Property->GetFName() == GET_MEMBER_NAME_CHECKED(UEditorPerformanceSettings, bOverrideMaxViewportRenderingResolution) ||
+		PropertyChangedEvent.Property->GetFName() == GET_MEMBER_NAME_CHECKED(UEditorPerformanceSettings, MaxViewportRenderingResolution))
+	{
+		UEditorPerformanceProjectSettings::ExportResolutionValuesToConsoleVariables();
+	}
+	else
 	{
-		// Tell all viewports to refresh their screen percentage when the dpi scaling override changes
-		for (FEditorViewportClient* Client : GEditor->GetAllViewportClients())
-		{
-			if (Client)
-			{
-				Client->RequestUpdateDPIScale();
-				Client->Invalidate();
-			}
-		}
-
-		if (GEngine->GameViewport)
-		{
-			GEngine->GameViewport->RequestUpdateDPIScale();
-		}
+		ExportValuesToConsoleVariables(PropertyChangedEvent.Property);
 	}
 }
\ No newline at end of file
diff --git a/Engine/Source/Editor/UnrealEd/Private/EditorViewportClient.cpp b/Engine/Source/Editor/UnrealEd/Private/EditorViewportClient.cpp
--- a/Engine/Source/Editor/UnrealEd/Private/EditorViewportClient.cpp
+++ b/Engine/Source/Editor/UnrealEd/Private/EditorViewportClient.cpp
@@ -82,16 +82,10 @@ static TAutoConsoleVariable<int32> CVarEditorViewportTest(
 	TEXT("1..7: Various Configuations"),
 	ECVF_RenderThreadSafe);
 
-static TAutoConsoleVariable<float> CVarEditorViewportScreenPercentage(
-	TEXT("r.Editor.ViewportScreenPercentage"),
-	100.0f,
-	TEXT("Default screen percentage of the editor's viewports"),
-	ECVF_RenderThreadSafe);
-
 static bool GetDefaultLowDPIPreviewValue()
 {
-	static auto CVarEnableEditorScreenPercentageOverride = IConsoleManager::Get().FindConsoleVariable(TEXT("Editor.OverrideDPIBasedEditorViewportScaling"));
-	return CVarEnableEditorScreenPercentageOverride->GetInt() == 0;
+	static auto CVarEditorViewportHighDPI = IConsoleManager::Get().FindConsoleVariable(TEXT("r.Editor.Viewport.HighDPI"));
+	return CVarEditorViewportHighDPI->GetInt() == 0;
 }
 
 float ComputeOrthoZoomFactor(const float ViewportWidth)
@@ -445,7 +439,6 @@ FEditorViewportClient::FEditorViewportClient(FEditorModeTools* InModeTools, FPre
 	, MovingPreviewLightSavedScreenPos(ForceInitToZero)
 	, MovingPreviewLightTimer(0.0f)
 	, bLockFlightCamera(false)
-	, PreviewResolutionFraction(CVarEditorViewportScreenPercentage.GetValueOnGameThread() / 100.0f)
 	, SceneDPIMode(ESceneDPIMode::EditorDefault)
 	, PerspViewModeIndex(DefaultPerspectiveViewMode)
 	, OrthoViewModeIndex(DefaultOrthoViewMode)
@@ -2683,18 +2676,60 @@ bool FEditorViewportClient::SupportsPreviewResolutionFraction() const
 	return true;
 }
 
+
+float FEditorViewportClient::GetDefaultPrimaryResolutionFractionTarget() const
+{
+	FStaticResolutionFractionHeuristic StaticHeuristic;
+	StaticHeuristic.Settings.PullEditorRenderingSettings(bIsRealtime);
+
+	if (SupportsLowDPIPreview() && IsLowDPIPreview()) // TODO: && ViewFamily.SupportsScreenPercentage())
+	{
+		StaticHeuristic.SecondaryViewFraction = GetDPIDerivedResolutionFraction();
+	}
+
+	StaticHeuristic.TotalDisplayedPixelCount = FMath::Max(Viewport->GetSizeXY().X * Viewport->GetSizeXY().Y, 1);
+	StaticHeuristic.DPIScale = GetDPIScale();
+	return StaticHeuristic.ResolveResolutionFraction();
+}
+
 int32 FEditorViewportClient::GetPreviewScreenPercentage() const
 {
+	float ResolutionFraction = 1.0f;
+	if (PreviewResolutionFraction.IsSet())
+	{
+		ResolutionFraction = PreviewResolutionFraction.GetValue();
+	}
+	else
+	{
+		ResolutionFraction = GetDefaultPrimaryResolutionFractionTarget();
+	}
+
 	// We expose the resolution fraction derived from DPI, to not lie to the artist when screen percentage = 100%.
 	return FMath::RoundToInt(FMath::Clamp(
-		PreviewResolutionFraction,
-		ISceneViewFamilyScreenPercentage::kMinTAAUpsampleResolutionFraction,
-		ISceneViewFamilyScreenPercentage::kMaxTAAUpsampleResolutionFraction) * 100.0f);
+		ResolutionFraction,
+		ISceneViewFamilyScreenPercentage::kMinTSRResolutionFraction,
+		ISceneViewFamilyScreenPercentage::kMaxTSRResolutionFraction) * 100.0f);
 }
 
 void FEditorViewportClient::SetPreviewScreenPercentage(int32 PreviewScreenPercentage)
 {
-	PreviewResolutionFraction = PreviewScreenPercentage / 100.0f;
+	float AutoResolutionFraction = GetDefaultPrimaryResolutionFractionTarget();
+	int32 AutoScreenPercentage = FMath::RoundToInt(FMath::Clamp(
+		AutoResolutionFraction,
+		ISceneViewFamilyScreenPercentage::kMinTSRResolutionFraction,
+		ISceneViewFamilyScreenPercentage::kMaxTSRResolutionFraction) * 100.0f);
+
+	float NewResolutionFraction = PreviewScreenPercentage / 100.0f;
+	if (NewResolutionFraction >= ISceneViewFamilyScreenPercentage::kMinTSRResolutionFraction &&
+		NewResolutionFraction <= ISceneViewFamilyScreenPercentage::kMaxTSRResolutionFraction &&
+		PreviewScreenPercentage != AutoScreenPercentage)
+	{
+		PreviewResolutionFraction = NewResolutionFraction;
+	}
+	else
+	{
+		PreviewResolutionFraction.Reset();
+	}
 }
 
 bool FEditorViewportClient::SupportsLowDPIPreview() const
@@ -2702,7 +2737,7 @@ bool FEditorViewportClient::SupportsLowDPIPreview() const
 	return GetDPIDerivedResolutionFraction() < 1.0f;
 }
 
-bool FEditorViewportClient::IsLowDPIPreview()
+bool FEditorViewportClient::IsLowDPIPreview() const
 {
 	if (SceneDPIMode == ESceneDPIMode::EditorDefault)
 	{
@@ -3834,7 +3869,6 @@ void FEditorViewportClient::Draw(FViewport* InViewport, FCanvas* Canvas)
 	ViewFamily.LandscapeLODOverride = LandscapeLODOverride;
 
 	// Setup the screen percentage and upscaling method for the view family.
-	bool bFinalScreenPercentageShowFlag;
 	{
 		checkf(ViewFamily.GetScreenPercentageInterface() == nullptr,
 			TEXT("Some code has tried to set up an alien screen percentage driver, that could be wrong if not supported very well by the RHI."));
@@ -3850,7 +3884,28 @@ void FEditorViewportClient::Draw(FViewport* InViewport, FCanvas* Canvas)
 		{
 			GCustomEditorStaticScreenPercentage->SetupEditorViewFamily(ViewFamily, this);
 		}
+	}
 
+	FSceneView* View = nullptr;
+
+	// Stereo rendering
+	const bool bStereoDeviceActive = bStereoRendering && GEngine->StereoRenderingDevice.IsValid();
+	int32 NumViews = bStereoRendering ? GEngine->StereoRenderingDevice->GetDesiredNumberOfViews(bStereoRendering) : 1;
+	for( int StereoViewIndex = 0; StereoViewIndex < NumViews; ++StereoViewIndex )
+	{
+		View = CalcSceneView( &ViewFamily, bStereoRendering ? StereoViewIndex : INDEX_NONE);
+
+		SetupViewForRendering(ViewFamily,*View);
+
+	    FSlateRect SafeFrame;
+	    View->CameraConstrainedViewRect = View->UnscaledViewRect;
+	    if (CalculateEditorConstrainedViewRect(SafeFrame, Viewport, Canvas->GetDPIScale()))
+	    {
+		    View->CameraConstrainedViewRect = FIntRect(SafeFrame.Left, SafeFrame.Top, SafeFrame.Right, SafeFrame.Bottom);
+	    }
+ 	}
+
+	{
 		// If a screen percentage interface was not set by one of the view extension, then set the legacy one.
 		if (ViewFamily.GetScreenPercentageInterface() == nullptr)
 		{
@@ -3859,7 +3914,14 @@ void FEditorViewportClient::Draw(FViewport* InViewport, FCanvas* Canvas)
 			// If not doing VR rendering, apply preview resolution fraction.
 			if (!bStereoRendering && SupportsPreviewResolutionFraction() && ViewFamily.SupportsScreenPercentage())
 			{
-				GlobalResolutionFraction = PreviewResolutionFraction;
+				if (PreviewResolutionFraction.IsSet())
+				{
+					GlobalResolutionFraction = PreviewResolutionFraction.GetValue();
+				}
+				else
+				{
+					GlobalResolutionFraction = GetDefaultPrimaryResolutionFractionTarget();
+				}
 
 				// Force screen percentage's engine show flag to be turned on for preview screen percentage.
 				ViewFamily.EngineShowFlags.ScreenPercentage = (GlobalResolutionFraction != 1.0);
@@ -3871,38 +3933,14 @@ void FEditorViewportClient::Draw(FViewport* InViewport, FCanvas* Canvas)
 		}
 
 		check(ViewFamily.GetScreenPercentageInterface() != nullptr);
-
-		bFinalScreenPercentageShowFlag = ViewFamily.EngineShowFlags.ScreenPercentage;
 	}
 
-	FSceneView* View = nullptr;
-
-	// Stereo rendering
-	const bool bStereoDeviceActive = bStereoRendering && GEngine->StereoRenderingDevice.IsValid();
-	int32 NumViews = bStereoRendering ? GEngine->StereoRenderingDevice->GetDesiredNumberOfViews(bStereoRendering) : 1;
-	for( int StereoViewIndex = 0; StereoViewIndex < NumViews; ++StereoViewIndex )
-	{
-		View = CalcSceneView( &ViewFamily, bStereoRendering ? StereoViewIndex : INDEX_NONE);
-
-		SetupViewForRendering(ViewFamily,*View);
-
-	    FSlateRect SafeFrame;
-	    View->CameraConstrainedViewRect = View->UnscaledViewRect;
-	    if (CalculateEditorConstrainedViewRect(SafeFrame, Viewport, Canvas->GetDPIScale()))
-	    {
-		    View->CameraConstrainedViewRect = FIntRect(SafeFrame.Left, SafeFrame.Top, SafeFrame.Right, SafeFrame.Bottom);
-	    }
- 	}
-
 	if (IsAspectRatioConstrained())
 	{
 		// Clear the background to black if the aspect ratio is constrained, as the scene view won't write to all pixels.
 		Canvas->Clear(FLinearColor::Black);
 	}
 
-	// Make sure the engine show flag for screen percentage is still what it was when setting up the screen percentage interface
-	ViewFamily.EngineShowFlags.ScreenPercentage = bFinalScreenPercentageShowFlag;
-
 	// Draw the 3D scene
 	GetRendererModule().BeginRenderingViewFamily(Canvas,&ViewFamily);
 
diff --git a/Engine/Source/Editor/UnrealEd/Private/SCommonEditorViewportToolbarBase.cpp b/Engine/Source/Editor/UnrealEd/Private/SCommonEditorViewportToolbarBase.cpp
--- a/Engine/Source/Editor/UnrealEd/Private/SCommonEditorViewportToolbarBase.cpp
+++ b/Engine/Source/Editor/UnrealEd/Private/SCommonEditorViewportToolbarBase.cpp
@@ -491,8 +491,8 @@ void SCommonEditorViewportToolbarBase::OnFOVValueChanged(float NewValue) const
 
 TSharedRef<SWidget> SCommonEditorViewportToolbarBase::GenerateScreenPercentageMenu() const
 {
-	const int32 PreviewScreenPercentageMin = ISceneViewFamilyScreenPercentage::kMinTAAUpsampleResolutionFraction * 100.0f;
-	const int32 PreviewScreenPercentageMax = ISceneViewFamilyScreenPercentage::kMaxTAAUpsampleResolutionFraction * 100.0f;
+	const int32 PreviewScreenPercentageMin = ISceneViewFamilyScreenPercentage::kMinTSRResolutionFraction * 100.0f;
+	const int32 PreviewScreenPercentageMax = ISceneViewFamilyScreenPercentage::kMaxTSRResolutionFraction * 100.0f;
 
 	return
 		SNew(SBox)
@@ -510,8 +510,8 @@ TSharedRef<SWidget> SCommonEditorViewportToolbarBase::GenerateScreenPercentageMe
 					SNew(SSpinBox<int32>)
 					.Style(&FAppStyle::Get(), "Menu.SpinBox")
 					.Font(FEditorStyle::GetFontStyle(TEXT("MenuItem.Font")))
-					.MinValue(PreviewScreenPercentageMin)
-					.MaxValue(PreviewScreenPercentageMax)
+					.MinSliderValue(PreviewScreenPercentageMin)
+					.MaxSliderValue(PreviewScreenPercentageMax)
 					.Value(this, &SCommonEditorViewportToolbarBase::OnGetScreenPercentageValue)
 					.OnValueChanged(const_cast<SCommonEditorViewportToolbarBase*>(this), &SCommonEditorViewportToolbarBase::OnScreenPercentageValueChanged)
 				]
diff --git a/Engine/Source/Editor/UnrealEd/Private/Settings/EditorProjectSettings.cpp b/Engine/Source/Editor/UnrealEd/Private/Settings/EditorProjectSettings.cpp
--- a/Engine/Source/Editor/UnrealEd/Private/Settings/EditorProjectSettings.cpp
+++ b/Engine/Source/Editor/UnrealEd/Private/Settings/EditorProjectSettings.cpp
@@ -4,6 +4,36 @@
 #include "UObject/UnrealType.h"
 #include "Toolkits/ToolkitManager.h"
 #include "BlueprintEditor.h"
+#include "Editor/EditorPerformanceSettings.h"
+#include "EditorViewportClient.h"
+#include "Editor.h"
+
+
+TAutoConsoleVariable<int32> CVarEditorViewportDefaultScreenPercentageMode(
+	TEXT("r.Editor.Viewport.ScreenPercentageMode.RealTime"), 1,
+	TEXT("Controls the default screen percentage mode for realtime editor viewports."),
+	ECVF_Default);
+
+TAutoConsoleVariable<int32> CVarEditorViewportNonRealtimeDefaultScreenPercentageMode(
+	TEXT("r.Editor.Viewport.ScreenPercentageMode.NonRealTime"), 2,
+	TEXT("Controls the default screen percentage mode for non-realtime editor viewports."),
+	ECVF_Default);
+
+TAutoConsoleVariable<float> CVarEditorViewportDefaultScreenPercentage(
+	TEXT("r.Editor.Viewport.ScreenPercentage"), 100,
+	TEXT("Controls the editor viewports' default screen percentage when using r.Editor.Viewport.ScreenPercentageMode=0."),
+	ECVF_Default);
+
+TAutoConsoleVariable<int32> CVarEditorViewportDefaultMinRenderingResolution(
+	TEXT("r.Editor.Viewport.MinRenderingResolution"), 720,
+	TEXT("Controls the minimum number of rendered pixel by default in editor viewports."),
+	ECVF_Default);
+
+TAutoConsoleVariable<int32> CVarEditorViewportDefaultMaxRenderingResolution(
+	TEXT("r.Editor.Viewport.MaxRenderingResolution"), 2160,
+	TEXT("Controls the absolute maximum number of rendered pixel in editor viewports."),
+	ECVF_Default);
+
 
 EUnit ConvertDefaultInputUnits(EDefaultLocationUnit In)
 {
@@ -140,3 +170,150 @@ void ULevelEditor2DSettings::PostEditChangeProperty(struct FPropertyChangedEvent
 
 	Super::PostEditChangeProperty(PropertyChangedEvent);
 }
+
+
+/* UEditorPerformanceProjectSettings
+*****************************************************************************/
+
+UEditorPerformanceProjectSettings::UEditorPerformanceProjectSettings(const FObjectInitializer& ObjectInitializer)
+	: Super(ObjectInitializer)
+	, RealtimeScreenPercentageMode(EScreenPercentageMode::BasedOnDisplayResolution)
+	, NonRealtimeScreenPercentageMode(EScreenPercentageMode::BasedOnDPIScale)
+	, ManualScreenPercentage(100.0f)
+	, MinViewportRenderingResolution(720)
+	, MaxViewportRenderingResolution(2160)
+{
+}
+
+void UEditorPerformanceProjectSettings::PostInitProperties()
+{
+	Super::PostInitProperties();
+	ExportResolutionValuesToConsoleVariables();
+}
+
+void UEditorPerformanceProjectSettings::PostEditChangeProperty(struct FPropertyChangedEvent& PropertyChangedEvent)
+{
+	Super::PostEditChangeProperty(PropertyChangedEvent);
+	if (!PropertyChangedEvent.Property)
+	{
+		return;
+	}
+
+	if (PropertyChangedEvent.Property->GetFName() == GET_MEMBER_NAME_CHECKED(UEditorPerformanceProjectSettings, RealtimeScreenPercentageMode) ||
+		PropertyChangedEvent.Property->GetFName() == GET_MEMBER_NAME_CHECKED(UEditorPerformanceProjectSettings, NonRealtimeScreenPercentageMode) ||
+		PropertyChangedEvent.Property->GetFName() == GET_MEMBER_NAME_CHECKED(UEditorPerformanceProjectSettings, ManualScreenPercentage) ||
+		PropertyChangedEvent.Property->GetFName() == GET_MEMBER_NAME_CHECKED(UEditorPerformanceProjectSettings, MinViewportRenderingResolution) ||
+		PropertyChangedEvent.Property->GetFName() == GET_MEMBER_NAME_CHECKED(UEditorPerformanceProjectSettings, MaxViewportRenderingResolution))
+	{
+		ExportResolutionValuesToConsoleVariables();
+	}
+	else
+	{
+		ExportValuesToConsoleVariables(PropertyChangedEvent.Property);
+	}
+}
+
+// static
+void UEditorPerformanceProjectSettings::ExportResolutionValuesToConsoleVariables()
+{
+	// Apply project settings settings
+	{
+		const UEditorPerformanceProjectSettings* EditorProjectSettings = GetDefault<UEditorPerformanceProjectSettings>();
+
+		CVarEditorViewportDefaultScreenPercentageMode->Set(int32(EditorProjectSettings->RealtimeScreenPercentageMode), ECVF_SetByProjectSetting);
+		CVarEditorViewportNonRealtimeDefaultScreenPercentageMode->Set(int32(EditorProjectSettings->NonRealtimeScreenPercentageMode), ECVF_SetByProjectSetting);
+		CVarEditorViewportDefaultScreenPercentage->Set(EditorProjectSettings->ManualScreenPercentage, ECVF_SetByProjectSetting);
+		CVarEditorViewportDefaultMinRenderingResolution->Set(EditorProjectSettings->MinViewportRenderingResolution, ECVF_SetByProjectSetting);
+		CVarEditorViewportDefaultMaxRenderingResolution->Set(EditorProjectSettings->MaxViewportRenderingResolution, ECVF_SetByProjectSetting);
+	}
+
+
+	// Override with the per editor user settings
+	{
+		const UEditorPerformanceSettings* EditorUserSettings = GetDefault<UEditorPerformanceSettings>();
+
+		{
+			if (EditorUserSettings->RealtimeScreenPercentageMode == EEditorUserScreenPercentageModeOverride::Manual)
+			{
+				CVarEditorViewportDefaultScreenPercentageMode->Set(int32(EScreenPercentageMode::Manual), ECVF_SetByProjectSetting);
+			}
+			else if (EditorUserSettings->RealtimeScreenPercentageMode == EEditorUserScreenPercentageModeOverride::BasedOnDisplayResolution)
+			{
+				CVarEditorViewportDefaultScreenPercentageMode->Set(int32(EScreenPercentageMode::BasedOnDisplayResolution), ECVF_SetByProjectSetting);
+			}
+			else if (EditorUserSettings->RealtimeScreenPercentageMode == EEditorUserScreenPercentageModeOverride::BasedOnDPIScale)
+			{
+				CVarEditorViewportDefaultScreenPercentageMode->Set(int32(EScreenPercentageMode::BasedOnDPIScale), ECVF_SetByProjectSetting);
+			}
+		}
+
+		{
+			if (EditorUserSettings->NonRealtimeScreenPercentageMode == EEditorUserScreenPercentageModeOverride::Manual)
+			{
+				CVarEditorViewportNonRealtimeDefaultScreenPercentageMode->Set(int32(EScreenPercentageMode::Manual), ECVF_SetByProjectSetting);
+			}
+			else if (EditorUserSettings->NonRealtimeScreenPercentageMode == EEditorUserScreenPercentageModeOverride::BasedOnDisplayResolution)
+			{
+				CVarEditorViewportNonRealtimeDefaultScreenPercentageMode->Set(int32(EScreenPercentageMode::BasedOnDisplayResolution), ECVF_SetByProjectSetting);
+			}
+			else if (EditorUserSettings->NonRealtimeScreenPercentageMode == EEditorUserScreenPercentageModeOverride::BasedOnDPIScale)
+			{
+				CVarEditorViewportNonRealtimeDefaultScreenPercentageMode->Set(int32(EScreenPercentageMode::BasedOnDPIScale), ECVF_SetByProjectSetting);
+			}
+		}
+
+		if (EditorUserSettings->bOverrideManualScreenPercentage)
+		{
+			CVarEditorViewportDefaultScreenPercentage->Set(EditorUserSettings->ManualScreenPercentage, ECVF_SetByProjectSetting);
+		}
+
+		if (EditorUserSettings->bOverrideMinViewportRenderingResolution)
+		{
+			CVarEditorViewportDefaultMinRenderingResolution->Set(EditorUserSettings->MinViewportRenderingResolution, ECVF_SetByProjectSetting);
+		}
+
+		if (EditorUserSettings->bOverrideMaxViewportRenderingResolution)
+		{
+			CVarEditorViewportDefaultMaxRenderingResolution->Set(EditorUserSettings->MaxViewportRenderingResolution, ECVF_SetByProjectSetting);
+		}
+	}
+
+	// Tell all viewports to refresh their screen percentage when the dpi scaling override changes
+	if (GEngine && GEditor)
+	{
+		for (FEditorViewportClient* Client : GEditor->GetAllViewportClients())
+		{
+			if (Client)
+			{
+				Client->RequestUpdateDPIScale();
+				Client->Invalidate();
+			}
+		}
+
+		if (GEngine->GameViewport)
+		{
+			GEngine->GameViewport->RequestUpdateDPIScale();
+		}
+	}
+}
+
+void SetupEditorResolutionFractionHeuristicSettings(FStaticResolutionFractionHeuristic::FUserSettings* OutputSettings, bool bIsRealTime)
+{
+	FStaticResolutionFractionHeuristic::FUserSettings Settings;
+
+	// Setup up with the per project settings
+	{
+		if (bIsRealTime)
+		{
+			OutputSettings->Mode = EScreenPercentageMode(FMath::Clamp(CVarEditorViewportDefaultScreenPercentageMode->GetInt(), 0, 2));
+		}
+		else
+		{
+			OutputSettings->Mode = EScreenPercentageMode(FMath::Clamp(CVarEditorViewportNonRealtimeDefaultScreenPercentageMode->GetInt(), 0, 2));
+		}
+		OutputSettings->GlobalResolutionFraction = CVarEditorViewportDefaultScreenPercentage->GetFloat() / 100.0f;
+		OutputSettings->MinRenderingResolution = CVarEditorViewportDefaultMinRenderingResolution->GetInt();
+		OutputSettings->MaxRenderingResolution = CVarEditorViewportDefaultMaxRenderingResolution->GetInt();
+	}
+
+}
diff --git a/Engine/Source/Editor/UnrealEd/Public/EditorViewportClient.h b/Engine/Source/Editor/UnrealEd/Public/EditorViewportClient.h
--- a/Engine/Source/Editor/UnrealEd/Public/EditorViewportClient.h
+++ b/Engine/Source/Editor/UnrealEd/Public/EditorViewportClient.h
@@ -1282,6 +1282,9 @@ public:
 	/** @return True if PreviewResolutionFraction is supported. */
 	bool SupportsPreviewResolutionFraction() const;
 
+	/** @return default resolution fraction for UI based on display resolution and user settings. */
+	float GetDefaultPrimaryResolutionFractionTarget() const;
+
 	/** @return preview screen percentage for UI. */
 	int32 GetPreviewScreenPercentage() const;
 
@@ -1292,7 +1295,7 @@ public:
 	bool SupportsLowDPIPreview() const;
 
 	/** @return whether previewing for low DPI. */
-	bool IsLowDPIPreview();
+	bool IsLowDPIPreview() const;
 
 	/** Set whether previewing for low DPI. */
 	void SetLowDPIPreview(bool LowDPIPreview);
@@ -1831,12 +1834,12 @@ private:
 
 private:
 	/** Controles resolution fraction for previewing in editor viewport at different screen percentage. */
-	float PreviewResolutionFraction;
+	TOptional<float> PreviewResolutionFraction;
 
 	// DPI mode for scene rendering.
 	enum class ESceneDPIMode
 	{
-		// Uses Editor.OverrideDPIBasedEditorViewportScaling.
+		// Uses r.Editor.Viewport.HighDPI.
 		EditorDefault,
 
 		// Force emulating low DPI.
diff --git a/Engine/Source/Editor/UnrealEd/Public/Settings/EditorProjectSettings.h b/Engine/Source/Editor/UnrealEd/Public/Settings/EditorProjectSettings.h
--- a/Engine/Source/Editor/UnrealEd/Public/Settings/EditorProjectSettings.h
+++ b/Engine/Source/Editor/UnrealEd/Public/Settings/EditorProjectSettings.h
@@ -6,6 +6,7 @@
 #include "UObject/ObjectMacros.h"
 #include "Engine/DeveloperSettings.h"
 #include "Math/UnitConversion.h"
+#include "LegacyScreenPercentageDriver.h"
 
 #include "EditorProjectSettings.generated.h"
 
@@ -157,6 +158,48 @@ public:
 
 };
 
+/**
+ * Configure per-project performance settings for the Editor
+ */
+UCLASS(config=Editor, meta=(DisplayName="Performance"), defaultconfig)
+class UNREALED_API UEditorPerformanceProjectSettings : public UDeveloperSettings
+{
+	GENERATED_UCLASS_BODY()
+
+public:
+	
+	UPROPERTY(EditAnywhere, config, Category=ViewportResolution, meta=(
+		DisplayName="Default screen percentage mode for realtime editor viewports."))
+	EScreenPercentageMode RealtimeScreenPercentageMode;
+
+	UPROPERTY(EditAnywhere, config, Category=ViewportResolution, meta=(
+		DisplayName="Default screen percentage mode for non-realtime editor viewports."))
+	EScreenPercentageMode NonRealtimeScreenPercentageMode;
+
+	UPROPERTY(EditAnywhere, config, Category=ViewportResolution, meta=(
+		EditCondition="RealtimeScreenPercentageMode == EScreenPercentageMode::Manual || NonRealtimeScreenPercentageMode == EScreenPercentageMode::Manual",
+		DisplayName="Manual screen percentage to be set by default for editor viewports."))
+	float ManualScreenPercentage;
+
+	UPROPERTY(EditAnywhere, config, Category=ViewportResolution, meta=(
+		DisplayName="Minimum default rendering resolution to use for editor viewports."))
+	int32 MinViewportRenderingResolution;
+
+	UPROPERTY(EditAnywhere, config, Category=ViewportResolution, meta=(
+		DisplayName="Maximum default rendering resolution to use for editor viewports."))
+	int32 MaxViewportRenderingResolution;
+
+
+	static void ExportResolutionValuesToConsoleVariables();
+
+public:
+	// UObject interface
+	virtual void PostInitProperties() override;
+	virtual void PostEditChangeProperty(struct FPropertyChangedEvent& PropertyChangedEvent) override;
+	// End of UObject interface
+
+};
+
 
 UCLASS(config=Editor, meta=(DisplayName="Derived Data"), defaultconfig)
 class UNREALED_API UDDCProjectSettings : public UDeveloperSettings
diff --git a/Engine/Source/Runtime/Core/Private/HAL/ConsoleManager.cpp b/Engine/Source/Runtime/Core/Private/HAL/ConsoleManager.cpp
--- a/Engine/Source/Runtime/Core/Private/HAL/ConsoleManager.cpp
+++ b/Engine/Source/Runtime/Core/Private/HAL/ConsoleManager.cpp
@@ -2631,15 +2631,6 @@ static TAutoConsoleVariable<int32> CVarDepthOfFieldQuality(
 	TEXT(" 4: extremely high quality, intended for non realtime cutscenes, CircleDOF only (very slow)"),
 	ECVF_Scalability | ECVF_RenderThreadSafe);
 
-static TAutoConsoleVariable<float> CVarScreenPercentage(
-	TEXT("r.ScreenPercentage"),
-	100.0f,
-	TEXT("To render in lower resolution and upscale for better performance (combined up with the blenable post process setting).\n")
-	TEXT("70 is a good value for low aliasing and performance, can be verified with 'show TestImage'\n")
-	TEXT("in percent, >0 and <=100, larger numbers are possible (supersampling) but the downsampling quality is improvable.")
-	TEXT("<0 is treated like 100."),
-	ECVF_Scalability | ECVF_Default);
-
 static TAutoConsoleVariable<int32> CVarHighResScreenshotDelay(
 	TEXT("r.HighResScreenshotDelay"),
 	4,
diff --git a/Engine/Source/Runtime/Engine/Private/GameViewportClient.cpp b/Engine/Source/Runtime/Engine/Private/GameViewportClient.cpp
--- a/Engine/Source/Runtime/Engine/Private/GameViewportClient.cpp
+++ b/Engine/Source/Runtime/Engine/Private/GameViewportClient.cpp
@@ -117,15 +117,6 @@ static TAutoConsoleVariable<float> CVarSecondaryScreenPercentage( // TODO: make
 	TEXT(" 1: override secondary screen percentage."),
 	ECVF_Default);
 
-static TAutoConsoleVariable<float> CVarScreenPercentageMaxResolution(
-	TEXT("r.ScreenPercentage.MaxResolution"),
-	0.0f,
-	TEXT("Controls maximum number of rendered pixel before any upscaling such that doesn't go higher than the specified 16:9 resolution ")
-	TEXT("of this variable. For instance set this value to 1440 so that you are not rendering more than 2560x1440 = 3.6M pixels. ")
-	TEXT("This is useful to set this on PC in your project's DefaultEditor.ini so you are not rendering more pixel on PC in PIE that you would ")
-	TEXT("in average on console with your project specific dynamic resolution settings."),
-	ECVF_Default);
-
 #if CSV_PROFILER
 struct FCsvLocalPlayer
 {
@@ -1642,42 +1633,28 @@ void UGameViewportClient::Draw(FViewport* InViewport, FCanvas* SceneCanvas)
 		// If a screen percentage interface was not set by dynamic resolution, then create one matching legacy behavior.
 		if (ViewFamily.GetScreenPercentageInterface() == nullptr)
 		{
-			bool AllowPostProcessSettingsScreenPercentage = false;
 			float GlobalResolutionFraction = 1.0f;
 
 			if (ViewFamily.EngineShowFlags.ScreenPercentage)
 			{
-				// Allow FPostProcessSettings::ScreenPercentage.
-				AllowPostProcessSettingsScreenPercentage = true;
+				// Get global view fraction.
+				FStaticResolutionFractionHeuristic StaticHeuristic;
 
-				// Get global view fraction set by r.ScreenPercentage.
-				GlobalResolutionFraction = FLegacyScreenPercentageDriver::GetCVarResolutionFraction();
-
-				// Max the rendering resolution to average target resolution used on platform that have dynamic resolution.
-				float MaxRenderingHeight = CVarScreenPercentageMaxResolution.GetValueOnGameThread();
-				if (MaxRenderingHeight > 0.0f)
+#if WITH_EDITOR
+				if (FStaticResolutionFractionHeuristic::FUserSettings::EditorOverridePIESettings())
 				{
-					int32 MaxRenderingWidth = MaxRenderingHeight * (1920.0f / 1080.0f);
-					int32 MaxRenderingPixelCount = MaxRenderingWidth * MaxRenderingHeight;
-
-					int32 TotalRenderingPixelCount = 0;
-
-					for (const FSceneView* View : ViewFamily.Views)
-					{
-						// Number of pixel drawn in ViewFamily.RenderTarget
-						int32 DisplayPixelCount = View->CameraConstrainedViewRect.Area();
-
-						// Number of pixel drawn before secondary spatial upscale.
-						int32 PrimaryScreenPercentageUpscaledPixelCount = DisplayPixelCount * (ViewFamily.SecondaryViewFraction * ViewFamily.SecondaryViewFraction);
+					StaticHeuristic.Settings.PullEditorRenderingSettings(/* bIsRealTime = */ true);
+				}
+				else
+#endif
+				{
+					StaticHeuristic.Settings.PullRunTimeRenderingSettings();
+				}
 
-						TotalRenderingPixelCount += PrimaryScreenPercentageUpscaledPixelCount;
-					}
+				StaticHeuristic.PullViewFamilyRenderingSettings(ViewFamily);
+				StaticHeuristic.DPIScale = GetDPIScale();
 
-					// Compute max resolution fraction such that the total number of pixel doesn't go over the CVarScreenPercentageMaxResolution.
-					float MaxGlobalResolutionFraction = FMath::Sqrt(float(MaxRenderingPixelCount) / float(TotalRenderingPixelCount));
-				
-					GlobalResolutionFraction = FMath::Min(GlobalResolutionFraction, MaxGlobalResolutionFraction);
-				}
+				GlobalResolutionFraction = StaticHeuristic.ResolveResolutionFraction();
 			}
 
 			ViewFamily.SetScreenPercentageInterface(new FLegacyScreenPercentageDriver(
diff --git a/Engine/Source/Runtime/Engine/Private/LegacyScreenPercentageDriver.cpp b/Engine/Source/Runtime/Engine/Private/LegacyScreenPercentageDriver.cpp
--- a/Engine/Source/Runtime/Engine/Private/LegacyScreenPercentageDriver.cpp
+++ b/Engine/Source/Runtime/Engine/Private/LegacyScreenPercentageDriver.cpp
@@ -1,12 +1,108 @@
 // Copyright Epic Games, Inc. All Rights Reserved.
 
-/*=============================================================================
-	Renderer.cpp: Renderer module implementation.
-=============================================================================*/
-
 #include "LegacyScreenPercentageDriver.h"
 #include "UnrealEngine.h"
+#include "Misc/ConfigCacheIni.h"
+
+
+static TAutoConsoleVariable<int32> CVarScreenPercentageMode(
+	TEXT("r.ScreenPercentage.Mode"), 0,
+	TEXT("Selects mode to control the screen percentage.\n")
+	TEXT(" 0: Controls the view's screen percentage based on r.ScreenPercentage\n")
+	TEXT(" 1: Controls the view's screen percentage based on displayed resolution with r.ScreenPercentage.Auto.* (default)\n"),
+	ECVF_Scalability | ECVF_Default);
+
+static TAutoConsoleVariable<float> CVarScreenPercentage(
+	TEXT("r.ScreenPercentage"), 100.0f,
+	TEXT("To render in lower resolution and upscale for better performance (combined up with the blenable post process setting).\n")
+	TEXT("70 is a good value for low aliasing and performance, can be verified with 'show TestImage'\n")
+	TEXT("in percent, >0 and <=100, larger numbers are possible (supersampling) but the downsampling quality is improvable.")
+	TEXT("<0 is treated like 100."),
+	ECVF_Scalability | ECVF_Default);
+
+static TAutoConsoleVariable<float> CVarAutoPixelCountMultiplier(
+	TEXT("r.ScreenPercentage.Auto.PixelCountMultiplier"), 1.0f,
+	TEXT(""),
+	ECVF_Default);
+
+static TAutoConsoleVariable<float> CVarScreenPercentageMinResolution(
+	TEXT("r.ScreenPercentage.MinResolution"),
+	0.0f,
+	TEXT("Controls the absolute minimum number of rendered pixel."),
+	ECVF_Default);
+
+static TAutoConsoleVariable<float> CVarScreenPercentageMaxResolution(
+	TEXT("r.ScreenPercentage.MaxResolution"),
+	0.0f,
+	TEXT("Controls the absolute maximum number of rendered pixel before any upscaling such that doesn't go higher than the specified 16:9 resolution ")
+	TEXT("of this variable. For instance set this value to 1440 so that you are not rendering more than 2560x1440 = 3.6M pixels. ")
+	TEXT("This is useful to set this on PC in your project's DefaultEditor.ini so you are not rendering more pixel on PC in PIE that you would ")
+	TEXT("in average on console with your project specific dynamic resolution settings."),
+	ECVF_Default);
+
+#if !UE_BUILD_SHIPPING
+
+void OnScreenPercentageChange(IConsoleVariable* Var)
+{
+	float SP = Var->GetFloat();
+
+	if (!GEngine || !GEngine->GameViewport)
+	{
+		return;
+	}
+
+#if WITH_EDITOR
+	if (FStaticResolutionFractionHeuristic::FUserSettings::EditorOverridePIESettings())
+	{
+		UE_LOG(LogEngine, Display, TEXT("r.ScreenPercentage=%f is ignored because overriden by editor settings (r.Editor.Viewport.OverridePIEScreenPercentage). Change this behavior in Edit -> Editor Preferences -> Performance"), SP);
+		return;
+	}
+#endif
 
+	// Check whether dynamic resolution is overriding r.ScreenPercentage.
+	{
+		FDynamicResolutionStateInfos DynResInfo;
+		GEngine->GetDynamicResolutionCurrentStateInfos(/* out */ DynResInfo);
+
+		if (DynResInfo.Status == EDynamicResolutionStatus::Enabled)
+		{
+			UE_LOG(LogEngine, Display, TEXT("r.ScreenPercentage=%f is ignored because overriden by dynamic resolution settings (r.DynamicRes.OperationMode)."), SP);
+			return;
+		}
+		else if (DynResInfo.Status == EDynamicResolutionStatus::DebugForceEnabled)
+		{
+			UE_LOG(LogEngine, Display, TEXT("r.ScreenPercentage=%f is ignored because overriden by dynamic resolution settings (r.DynamicRes.TestScreenPercentage)."), SP);
+			return;
+		}
+	}
+
+	// Check whether the screen percentage mode honor r.ScreenPercentage
+	int32 SPMode = CVarScreenPercentageMode.GetValueOnGameThread();
+	if (SPMode != 0)
+	{
+		UE_LOG(LogEngine, Display, TEXT("r.ScreenPercentage=%f is ignored because overriden by the screen percentage mode (r.ScreenPercentage.Mode=%d)."), SP, SPMode);
+		return;
+	}
+}
+
+#endif // !UE_BUILD_SHIPPING
+
+void InitScreenPercentage()
+{
+#if !UE_BUILD_SHIPPING
+	CVarScreenPercentage.AsVariable()->SetOnChangedCallback(FConsoleVariableDelegate::CreateStatic(&OnScreenPercentageChange));
+#endif
+}
+
+static float GetRenderingPixelCount(float Rendering169Height)
+{
+	return Rendering169Height * FMath::Abs(Rendering169Height) * (1920.0f / 1080.0f);
+}
+
+static float GetResolutionFraction(float ScreenPercentage)
+{
+	return ScreenPercentage / 100.0f;
+}
 
 FLegacyScreenPercentageDriver::FLegacyScreenPercentageDriver(
 	const FSceneViewFamily& InViewFamily,
@@ -66,7 +162,7 @@ float FLegacyScreenPercentageDriver::GetCVarResolutionFraction()
 	check(IsInGameThread());
 	static const auto ScreenPercentageCVar = IConsoleManager::Get().FindTConsoleVariableDataFloat(TEXT("r.ScreenPercentage"));
 
-	float GlobalFraction = ScreenPercentageCVar->GetValueOnAnyThread() / 100.0f;
+	float GlobalFraction = GetResolutionFraction(CVarScreenPercentage.GetValueOnGameThread());
 	if (GlobalFraction <= 0.0)
 	{
 		GlobalFraction = 1.0f;
@@ -74,3 +170,167 @@ float FLegacyScreenPercentageDriver::GetCVarResolutionFraction()
 
 	return GlobalFraction;
 }
+
+#if WITH_EDITOR
+// static
+bool FStaticResolutionFractionHeuristic::FUserSettings::EditorOverridePIESettings()
+{
+	if (!GIsEditor)
+	{
+		return false;
+	}
+
+	static auto CVarEditorViewportOverrideGameScreenPercentage = IConsoleManager::Get().FindConsoleVariable(TEXT("r.Editor.Viewport.OverridePIEScreenPercentage"));
+	if (!CVarEditorViewportOverrideGameScreenPercentage)
+	{
+		return false;
+	}
+
+	return CVarEditorViewportOverrideGameScreenPercentage->GetInt() != 0;
+}
+#endif
+
+void FStaticResolutionFractionHeuristic::FUserSettings::PullRunTimeRenderingSettings()
+{
+	Mode = EScreenPercentageMode(FMath::Clamp(CVarScreenPercentageMode.GetValueOnGameThread(), 0, 2));
+	GlobalResolutionFraction = FLegacyScreenPercentageDriver::GetCVarResolutionFraction();
+	MinRenderingResolution = CVarScreenPercentageMinResolution.GetValueOnGameThread();
+	MaxRenderingResolution = CVarScreenPercentageMaxResolution.GetValueOnGameThread();
+	AutoPixelCountMultiplier = CVarAutoPixelCountMultiplier.GetValueOnGameThread();
+}
+
+void FStaticResolutionFractionHeuristic::FUserSettings::PullEditorRenderingSettings(bool bIsRealTime)
+#if WITH_EDITOR
+{
+	static auto CVarEditorViewportDefaultScreenPercentageRealTimeMode = IConsoleManager::Get().FindConsoleVariable(TEXT("r.Editor.Viewport.ScreenPercentageMode.RealTime"));
+	static auto CVarEditorViewportDefaultScreenPercentageMode = IConsoleManager::Get().FindConsoleVariable(TEXT("r.Editor.Viewport.ScreenPercentageMode.NonRealTime"));
+	static auto CVarEditorViewportDefaultScreenPercentage = IConsoleManager::Get().FindConsoleVariable(TEXT("r.Editor.Viewport.ScreenPercentage"));
+	static auto CVarEditorViewportDefaultMinRenderingResolution = IConsoleManager::Get().FindConsoleVariable(TEXT("r.Editor.Viewport.MinRenderingResolution"));
+	static auto CVarEditorViewportDefaultMaxRenderingResolution = IConsoleManager::Get().FindConsoleVariable(TEXT("r.Editor.Viewport.MaxRenderingResolution"));
+
+	if (bIsRealTime)
+	{
+		Mode = EScreenPercentageMode(FMath::Clamp(CVarEditorViewportDefaultScreenPercentageRealTimeMode->GetInt(), 0, 2));
+	}
+	else
+	{
+		Mode = EScreenPercentageMode(FMath::Clamp(CVarEditorViewportDefaultScreenPercentageMode->GetInt(), 0, 2));
+	}
+	GlobalResolutionFraction = GetResolutionFraction(CVarEditorViewportDefaultScreenPercentage->GetInt());
+	MinRenderingResolution = CVarEditorViewportDefaultMinRenderingResolution->GetInt();
+	MaxRenderingResolution = CVarEditorViewportDefaultMaxRenderingResolution->GetInt();
+}
+#else
+{
+	unimplemented();
+}
+#endif
+
+void FStaticResolutionFractionHeuristic::PullViewFamilyRenderingSettings(const FSceneViewFamily& ViewFamily)
+{
+	check(ViewFamily.Views.Num() > 0);
+
+	SecondaryViewFraction = ViewFamily.SecondaryViewFraction;
+
+	TotalDisplayedPixelCount = 0;
+	for (const FSceneView* View : ViewFamily.Views)
+	{
+		// Number of pixel drawn in ViewFamily.RenderTarget
+		int32 DisplayPixelCount = View->CameraConstrainedViewRect.Area();
+		
+		TotalDisplayedPixelCount += DisplayPixelCount;
+	}
+}
+
+float FStaticResolutionFractionHeuristic::ResolveResolutionFraction() const
+{
+	// Compute number of pixel being rendered.
+	float LocalTotalDisplayedPixelCount = FMath::Max(TotalDisplayedPixelCount, 1) * SecondaryViewFraction * SecondaryViewFraction;
+
+	float GlobalResolutionFraction = 1.0f;
+	if (Settings.Mode == EScreenPercentageMode::BasedOnDisplayResolution)
+	{
+		static bool bInitPixelCount = false;
+		static float AutoMinDisplayResolution;
+		static float AutoMinRenderingResolution;
+		static float AutoMidDisplayResolution;
+		static float AutoMidRenderingResolution;
+		static float AutoMaxDisplayResolution;
+		static float AutoMaxRenderingResolution;
+
+		if (!bInitPixelCount)
+		{
+			bool bSuccess = true;
+			bSuccess = bSuccess && GConfig->GetFloat(TEXT("Rendering.AutoScreenPercentage"), TEXT("MinDisplayResolution"), AutoMinDisplayResolution, GEngineIni);
+			bSuccess = bSuccess && GConfig->GetFloat(TEXT("Rendering.AutoScreenPercentage"), TEXT("MidDisplayResolution"), AutoMidDisplayResolution, GEngineIni);
+			bSuccess = bSuccess && GConfig->GetFloat(TEXT("Rendering.AutoScreenPercentage"), TEXT("MaxDisplayResolution"), AutoMaxDisplayResolution, GEngineIni);
+			bSuccess = bSuccess && GConfig->GetFloat(TEXT("Rendering.AutoScreenPercentage"), TEXT("MinRenderingResolution"), AutoMinRenderingResolution, GEngineIni);
+			bSuccess = bSuccess && GConfig->GetFloat(TEXT("Rendering.AutoScreenPercentage"), TEXT("MidRenderingResolution"), AutoMidRenderingResolution, GEngineIni);
+			bSuccess = bSuccess && GConfig->GetFloat(TEXT("Rendering.AutoScreenPercentage"), TEXT("MaxRenderingResolution"), AutoMaxRenderingResolution, GEngineIni);
+			if (!bSuccess)
+			{
+				UE_LOG(LogEngine, Fatal, TEXT("Failed to load Rendering.AutoScreenPercentage ini section."));
+			}
+			bInitPixelCount = true;
+		}
+
+		float LerpedRenderingPixelCount = 0.0f;
+		if (LocalTotalDisplayedPixelCount < GetRenderingPixelCount(AutoMinDisplayResolution))
+		{
+			LerpedRenderingPixelCount = LocalTotalDisplayedPixelCount * GetRenderingPixelCount(AutoMinRenderingResolution) / GetRenderingPixelCount(AutoMinDisplayResolution);
+		}
+		else if (LocalTotalDisplayedPixelCount > GetRenderingPixelCount(AutoMaxDisplayResolution))
+		{
+			LerpedRenderingPixelCount = LocalTotalDisplayedPixelCount * GetRenderingPixelCount(AutoMaxRenderingResolution) / GetRenderingPixelCount(AutoMaxDisplayResolution);
+		}
+		else
+		{
+			float MinDisplayPixelCount   = GetRenderingPixelCount(AutoMinDisplayResolution);
+			float MaxDisplayPixelCount   = GetRenderingPixelCount(AutoMidDisplayResolution);
+			float MinRenderingPixelCount = GetRenderingPixelCount(AutoMinRenderingResolution);
+			float MaxRenderingPixelCount = GetRenderingPixelCount(AutoMidRenderingResolution);
+			if (LocalTotalDisplayedPixelCount > MaxDisplayPixelCount)
+			{
+				MinDisplayPixelCount   = GetRenderingPixelCount(AutoMidDisplayResolution);
+				MaxDisplayPixelCount   = GetRenderingPixelCount(AutoMaxDisplayResolution);
+				MinRenderingPixelCount = GetRenderingPixelCount(AutoMidRenderingResolution);
+				MaxRenderingPixelCount = GetRenderingPixelCount(AutoMaxRenderingResolution);
+			}
+
+			float ReslutionFractionLerp = FMath::Clamp((LocalTotalDisplayedPixelCount - MinDisplayPixelCount) / (MaxDisplayPixelCount - MinDisplayPixelCount), 0.0f, 1.0f);
+			LerpedRenderingPixelCount = FMath::Lerp(MinRenderingPixelCount, MaxRenderingPixelCount, ReslutionFractionLerp);
+		}
+
+		GlobalResolutionFraction = FMath::Sqrt(Settings.AutoPixelCountMultiplier * LerpedRenderingPixelCount / LocalTotalDisplayedPixelCount);
+	}
+	else if (Settings.Mode == EScreenPercentageMode::BasedOnDPIScale)
+	{
+		GlobalResolutionFraction = 1.0f / DPIScale;
+	}
+	else
+	{
+		ensure(Settings.Mode == EScreenPercentageMode::Manual);
+		GlobalResolutionFraction = Settings.GlobalResolutionFraction;
+	}
+
+	// Max the rendering resolution to average target resolution used on platform that have dynamic resolution.
+	const float MaxRenderingPixelCount = GetRenderingPixelCount(Settings.MaxRenderingResolution);
+	if (MaxRenderingPixelCount > 0.0f)
+	{
+		// Compute max resolution fraction such that the total number of pixel doesn't go over the CVarScreenPercentageMaxResolution.
+		float MaxGlobalResolutionFraction = FMath::Sqrt(float(MaxRenderingPixelCount) / float(LocalTotalDisplayedPixelCount));
+
+		GlobalResolutionFraction = FMath::Min(GlobalResolutionFraction, MaxGlobalResolutionFraction);
+	}
+
+	// Min the rendering resolution to avoid any upscaling at very low resolutions.
+	const float MinRenderingPixelCount = GetRenderingPixelCount(Settings.MinRenderingResolution);
+	if (MinRenderingPixelCount > 0.0f)
+	{
+		float MinGlobalResolutionFraction = FMath::Min(FMath::Sqrt(float(MinRenderingPixelCount) / float(LocalTotalDisplayedPixelCount)), 1.0f);
+
+		GlobalResolutionFraction = FMath::Max(GlobalResolutionFraction, MinGlobalResolutionFraction);
+	}
+
+	return GlobalResolutionFraction;
+}
diff --git a/Engine/Source/Runtime/Engine/Private/UnrealEngine.cpp b/Engine/Source/Runtime/Engine/Private/UnrealEngine.cpp
--- a/Engine/Source/Runtime/Engine/Private/UnrealEngine.cpp
+++ b/Engine/Source/Runtime/Engine/Private/UnrealEngine.cpp
@@ -274,6 +274,7 @@ IMPLEMENT_MODULE( FEngineModule, Engine );
 
 #define LOCTEXT_NAMESPACE "UnrealEngine"
 
+void InitScreenPercentage();
 
 void OnChangeEngineCVarRequiringRecreateRenderState(IConsoleVariable* Var)
 {
@@ -299,6 +300,8 @@ void FEngineModule::StartupModule()
 	SuspendTextureStreamingRenderTasks = &SuspendRenderAssetStreaming;
 	ResumeTextureStreamingRenderTasks = &ResumeRenderAssetStreaming;
 
+	InitScreenPercentage();
+
 	FParticleSystemWorldManager::OnStartup();
 
 #if WITH_EDITOR
diff --git a/Engine/Source/Runtime/Engine/Public/LegacyScreenPercentageDriver.h b/Engine/Source/Runtime/Engine/Public/LegacyScreenPercentageDriver.h
--- a/Engine/Source/Runtime/Engine/Public/LegacyScreenPercentageDriver.h
+++ b/Engine/Source/Runtime/Engine/Public/LegacyScreenPercentageDriver.h
@@ -29,7 +29,6 @@ public:
 	/** Gets the view rect fraction from the r.ScreenPercentage cvar. */
 	static float GetCVarResolutionFraction();
 
-
 private:
 	// View family to take care of.
 	const FSceneViewFamily& ViewFamily;
@@ -46,3 +45,71 @@ private:
 	virtual float GetPrimaryResolutionFraction_RenderThread() const override;
 	virtual ISceneViewFamilyScreenPercentage* Fork_GameThread(const class FSceneViewFamily& ForkedViewFamily) const override;
 };
+
+// Mode for the computation of the screen percentage (r.ScreenPercentage.Mode).
+UENUM()
+enum class EScreenPercentageMode
+{
+	// Directly controls the screen percentage with the r.ScreenPercentage cvar
+	Manual UMETA(DisplayName="Manual"),
+
+	// Automatic control the screen percentage based on the display resolution, r.ScreenPercentage.Auto.*
+	BasedOnDisplayResolution UMETA(DisplayName="Based on display resolution"),
+
+	// Based on DPI scale.
+	BasedOnDPIScale UMETA(DisplayName="Based on operating system's DPI scale"),
+};
+
+/**
+ * Heuristic to automatically compute a default resolution fraction based user settings and display information.
+ */
+struct ENGINE_API FStaticResolutionFractionHeuristic
+{
+	// User configurable settings
+	struct ENGINE_API FUserSettings
+	{
+		EScreenPercentageMode Mode = EScreenPercentageMode::Manual;
+
+		// r.ScreenPercentage when Mode = EMode::Manual.
+		float GlobalResolutionFraction = 1.0f;
+
+		// r.ScreenPercentage.{Min,Max}Resolution
+		float MinRenderingResolution = 0.0f;
+		float MaxRenderingResolution = 0.0f;
+
+		// r.ScreenPercentage.Auto.* Mode = EMode::BasedOnDisplayResolution.
+		float AutoPixelCountMultiplier = 1.0f;
+
+		/** Return whether should use the editor settings for PIE. */
+#if WITH_EDITOR
+		static bool EditorOverridePIESettings();
+#else
+		static inline bool EditorOverridePIESettings()
+		{
+			return false;
+		}
+#endif
+
+		/** Pulls the user settings from the gameplay runtime cvars. */
+		void PullRunTimeRenderingSettings();
+
+		/** Pulls the user settings from the editor cvars. */
+		void PullEditorRenderingSettings(bool bIsRealTime);
+	};
+
+	FUserSettings Settings;
+
+	// [Required] total number of pixel being display into viewport.
+	int32 TotalDisplayedPixelCount = 0;
+
+	// [Required] secondary resolution fraction.
+	float SecondaryViewFraction = 1.0f;
+
+	// [Required] DPI scale
+	float DPIScale = 1.0f;
+
+	// Fetches rendering information from the ViewFamily.
+	void PullViewFamilyRenderingSettings(const FSceneViewFamily& ViewFamily);
+
+	float ResolveResolutionFraction() const;
+};
