Fix bug in DDC2 texture build path when the source texture has only the top mip, and the other mips are invalid/uninitialized.  This is addressed by ensuring that we only use those mips if the TMGS_LeaveExistingMips MipGenSettings value is in use.  This mirrors the behavior in FTextureSourceData::Init in the older texture build path.  Also ensured that the behavior for quantity of slices is mirrored from there also.

diff --git a/Engine/Source/Runtime/Engine/Private/TextureDerivedDataBuildUtils.cpp b/Engine/Source/Runtime/Engine/Private/TextureDerivedDataBuildUtils.cpp
--- a/Engine/Source/Runtime/Engine/Private/TextureDerivedDataBuildUtils.cpp
+++ b/Engine/Source/Runtime/Engine/Private/TextureDerivedDataBuildUtils.cpp
@@ -190,7 +190,7 @@ static void WriteOutputSettings(FCbWriter& Writer, int32 NumInlineMips)
 	Writer.EndObject();
 }
 
-static void WriteSource(FCbWriter& Writer, const UTexture& Texture, int32 LayerIndex)
+static void WriteSource(FCbWriter& Writer, const UTexture& Texture, int32 LayerIndex, const FTextureBuildSettings& BuildSettings)
 {
 	const FTextureSource& Source = Texture.Source;
 
@@ -210,12 +210,13 @@ static void WriteSource(FCbWriter& Writer, const UTexture& Texture, int32 LayerI
 	Writer.AddInteger("CompressionFormat", CompressionFormat);
 	Writer.AddInteger("SourceFormat", Source.GetFormat(LayerIndex));
 	Writer.AddInteger("GammaSpace", static_cast<uint8>(GammaSpace));
-	Writer.AddInteger("NumSlices", Source.GetNumSlices());
+	Writer.AddInteger("NumSlices", (BuildSettings.bCubemap || BuildSettings.bTextureArray || BuildSettings.bVolume) ? Source.GetNumSlices() : 1);
 	Writer.AddInteger("SizeX", Source.GetSizeX());
 	Writer.AddInteger("SizeY", Source.GetSizeY());
 	Writer.BeginArray("Mips");
+	int32 NumMips = BuildSettings.MipGenSettings == TMGS_LeaveExistingMips ? Source.GetNumMips() : 1;
 	int64 Offset = 0;
-	for (int32 MipIndex = 0, MipCount = Source.GetNumMips(); MipIndex < MipCount; ++MipIndex)
+	for (int32 MipIndex = 0, MipCount = NumMips; MipIndex < MipCount; ++MipIndex)
 	{
 		Writer.BeginObject();
 		Writer.AddInteger("Offset", Offset);
@@ -287,12 +288,12 @@ FCbObject SaveTextureBuildSettings(const UTexture& Texture, const FTextureBuildS
 	WriteOutputSettings(Writer, NumInlineMips);
 
 	Writer.SetName("Source");
-	WriteSource(Writer, Texture, LayerIndex);
+	WriteSource(Writer, Texture, LayerIndex, BuildSettings);
 
 	if (bUseCompositeTexture && Texture.CompositeTexture)
 	{
 		Writer.SetName("CompositeSource");
-		WriteSource(Writer, *Texture.CompositeTexture, LayerIndex);
+		WriteSource(Writer, *Texture.CompositeTexture, LayerIndex, BuildSettings);
 	}
 
 	Writer.EndObject();
