Fixes variable shadowing on non-unity build

diff --git a/Engine/Source/Editor/UnrealEd/Private/EditorViewportClient.cpp b/Engine/Source/Editor/UnrealEd/Private/EditorViewportClient.cpp
--- a/Engine/Source/Editor/UnrealEd/Private/EditorViewportClient.cpp
+++ b/Engine/Source/Editor/UnrealEd/Private/EditorViewportClient.cpp
@@ -84,8 +84,8 @@ static TAutoConsoleVariable<int32> CVarEditorViewportTest(
 
 static bool GetDefaultLowDPIPreviewValue()
 {
-	static auto CVarEditorViewportHighDPI = IConsoleManager::Get().FindConsoleVariable(TEXT("r.Editor.Viewport.HighDPI"));
-	return CVarEditorViewportHighDPI->GetInt() == 0;
+	static auto CVarEditorViewportHighDPIPtr = IConsoleManager::Get().FindConsoleVariable(TEXT("r.Editor.Viewport.HighDPI"));
+	return CVarEditorViewportHighDPIPtr->GetInt() == 0;
 }
 
 float ComputeOrthoZoomFactor(const float ViewportWidth)
diff --git a/Engine/Source/Runtime/Engine/Private/UnrealClient.cpp b/Engine/Source/Runtime/Engine/Private/UnrealClient.cpp
--- a/Engine/Source/Runtime/Engine/Private/UnrealClient.cpp
+++ b/Engine/Source/Runtime/Engine/Private/UnrealClient.cpp
@@ -2336,9 +2336,9 @@ float FCommonViewportClient::GetDPIDerivedResolutionFraction() const
 			return 1.0f;
 		}
 
-		static auto CVarEnableEditorScreenPercentageOverride = IConsoleManager::Get().FindConsoleVariable(TEXT("Editor.OverrideDPIBasedEditorViewportScaling"));
+		static auto CVarEditorViewportHighDPIPtr = IConsoleManager::Get().FindConsoleVariable(TEXT("r.Editor.Viewport.HighDPI"));
 
-		if (CVarEnableEditorScreenPercentageOverride && CVarEnableEditorScreenPercentageOverride->GetInt() == 0)
+		if (CVarEditorViewportHighDPIPtr && CVarEditorViewportHighDPIPtr->GetInt() == 0)
 		{
 			return FMath::Min(1.0f / GetDPIScale(), 1.0f);
 		}
