Avoid doing per-shader DDC for default materials (used for fallback rendering with ODSC) and for global shaders.

- These shaders are always required so we will always have complete shader maps for them.
- It is more optimal to do a single DDC Get/Put on a full, complete map then thousands of individual requests in this case.

diff --git a/Engine/Source/Runtime/Engine/Private/Materials/MeshMaterialShader.cpp b/Engine/Source/Runtime/Engine/Private/Materials/MeshMaterialShader.cpp
--- a/Engine/Source/Runtime/Engine/Private/Materials/MeshMaterialShader.cpp
+++ b/Engine/Source/Runtime/Engine/Private/Materials/MeshMaterialShader.cpp
@@ -38,6 +38,8 @@ static void PrepareMeshMaterialShaderCompileJob(EShaderPlatform Platform,
 	NewJob->Input.SharedEnvironment = MaterialEnvironment;
 	FShaderCompilerEnvironment& ShaderEnvironment = NewJob->Input.Environment;
 	ShaderEnvironment.TargetPlatform = MaterialEnvironment->TargetPlatform;
+	NewJob->bIsDefaultMaterial = Material->IsDefaultMaterial();
+	NewJob->bIsGlobalShader = false;
 
 	const FMaterialShaderParameters MaterialParameters(Material);
 
diff --git a/Engine/Source/Runtime/Engine/Private/ShaderCompiler/ShaderCompiler.cpp b/Engine/Source/Runtime/Engine/Private/ShaderCompiler/ShaderCompiler.cpp
--- a/Engine/Source/Runtime/Engine/Private/ShaderCompiler/ShaderCompiler.cpp
+++ b/Engine/Source/Runtime/Engine/Private/ShaderCompiler/ShaderCompiler.cpp
@@ -530,8 +530,10 @@ void FShaderCompileJobCollection::SubmitJobs(const TArray<FShaderCommonCompileJo
 				{
 					const FShaderCommonCompileJob::FInputHash& InputHash = Job->GetInputHash();
 
+					const bool bCheckDDC = !(Job->bIsDefaultMaterial || Job->bIsGlobalShader);
+
 					// see if we can find the job in the cache first
-					if (TArray<uint8>* ExistingOutput = CompletedJobsCache.Find(InputHash))
+					if (TArray<uint8>* ExistingOutput = CompletedJobsCache.Find(InputHash, bCheckDDC))
 					{
 						UE_LOG(LogShaderCompilers, UE_SHADERCACHE_LOG_LEVEL, TEXT("There is already a cached job with the ihash %s, processing the new one immediately."), *LexToString(InputHash));
 						FMemoryReader MemReader(*ExistingOutput);
@@ -689,8 +691,9 @@ void FShaderCompileJobCollection::AddToCacheAndProcessPending(FShaderCommonCompi
 
 	if (FinishedJob->bSucceeded)
 	{
+		const bool bAddToDDC = !(FinishedJob->bIsDefaultMaterial || FinishedJob->bIsGlobalShader);
 		// we only cache jobs that succeded
-		CompletedJobsCache.Add(InputHash, Output, NumOutstandingJobsWithSameHash);
+		CompletedJobsCache.Add(InputHash, Output, NumOutstandingJobsWithSameHash, bAddToDDC);
 	}
 
 	// remove ourselves from the jobs in flight, if we were there (if this job is a cloned job it might not have been)
@@ -6289,6 +6292,8 @@ static void PrepareGlobalShaderCompileJob(EShaderPlatform Platform,
 	static FString GlobalName(TEXT("Global"));
 
 	NewJob->bErrorsAreLikelyToBeCode = true;
+	NewJob->bIsGlobalShader = true;
+	NewJob->bIsDefaultMaterial = false;
 
 	// Compile the shader environment passed in with the shader type's source code.
 	::GlobalBeginCompileShader(
@@ -7648,7 +7653,7 @@ namespace
 }
 #endif
 
-FShaderJobCache::FJobCachedOutput* FShaderJobCache::Find(const FJobInputHash& Hash)
+FShaderJobCache::FJobCachedOutput* FShaderJobCache::Find(const FJobInputHash& Hash, const bool bCheckDDC)
 {
 	++TotalSearchAttempts;
 	TRACE_COUNTER_INCREMENT(Shaders_JobCacheSearchAttempts);
@@ -7673,7 +7678,7 @@ FShaderJobCache::FJobCachedOutput* FShaderJobCache::Find(const FJobInputHash& Ha
 		else
 		{
 			// If we didn't find it in memory search the DDC if it's enabled.
-			const bool bCachePerShaderDDC = IsShaderJobCacheDDCEnabled();
+			const bool bCachePerShaderDDC = IsShaderJobCacheDDCEnabled() && bCheckDDC;
 			if (bCachePerShaderDDC)
 			{
 				FSharedBuffer Results;
@@ -7763,7 +7768,7 @@ FShaderJobCache::FShaderJobCache()
 	CurrentlyAllocatedMemory = sizeof(*this) + InputHashToOutput.GetAllocatedSize() + Outputs.GetAllocatedSize();
 }
 
-void FShaderJobCache::Add(const FJobInputHash& Hash, const FJobCachedOutput& Contents, int32 InitialHitCount)
+void FShaderJobCache::Add(const FJobInputHash& Hash, const FJobCachedOutput& Contents, int32 InitialHitCount, const bool bAddToDDC)
 {
 	TRACE_CPUPROFILER_EVENT_SCOPE(FShaderJobCache::Add);
 
@@ -7784,7 +7789,7 @@ void FShaderJobCache::Add(const FJobInputHash& Hash, const FJobCachedOutput& Con
 	const uint64 InputHashToOutputOriginalSize = InputHashToOutput.GetAllocatedSize();
 	InputHashToOutput.Add(Hash, OutputHash);
 
-	const bool bCachePerShaderDDC = IsShaderJobCacheDDCEnabled();
+	const bool bCachePerShaderDDC = IsShaderJobCacheDDCEnabled() && bAddToDDC;
 
 	FStoredOutput** CannedOutput = Outputs.Find(OutputHash);
 	if (CannedOutput && (bCachePerShaderDDC == false))
diff --git a/Engine/Source/Runtime/Engine/Public/ShaderCompiler.h b/Engine/Source/Runtime/Engine/Public/ShaderCompiler.h
--- a/Engine/Source/Runtime/Engine/Public/ShaderCompiler.h
+++ b/Engine/Source/Runtime/Engine/Public/ShaderCompiler.h
@@ -137,9 +137,12 @@ public:
 		After a job is bFinalized it will be bReleased when ReleaseJob() is invoked, which means that the shader compile thread
 		is no longer processing the job; which is useful for non standard job handling (Niagara as an example). */
 	uint8 bReleased : 1;
-
 	/** Whether we hashed the inputs */
 	uint8 bInputHashSet : 1;
+	/** Whether or not we are a default material. */
+	uint8 bIsDefaultMaterial : 1;
+	/** Whether or not we are a global shader. */
+	uint8 bIsGlobalShader : 1;
 	/** Hash of all the job inputs */
 	FInputHash InputHash;
 
@@ -202,7 +205,10 @@ protected:
 		bSucceeded(false),
 		bErrorsAreLikelyToBeCode(false),
 		bReleased(false),
-		bInputHashSet(false)
+		bInputHashSet(false),
+		bIsDefaultMaterial(false),
+		bIsGlobalShader(false)
+
 	{
 		check(InPriroity != EShaderCompileJobPriority::None);
 	}
@@ -342,10 +348,10 @@ public:
 	using FJobCachedOutput = TArray<uint8>;
 
 	/** Looks for the job in the cache, returns null if not found */
-	FJobCachedOutput* Find(const FJobInputHash& Hash);
+	FJobCachedOutput* Find(const FJobInputHash& Hash, const bool bCheckDDC);
 
 	/** Adds a job output to the cache */
-	void Add(const FJobInputHash& Hash, const FJobCachedOutput& Contents, int InitialHitCount);
+	void Add(const FJobInputHash& Hash, const FJobCachedOutput& Contents, int InitialHitCount, const bool bAddToDDC);
 
 	/** Calculates memory used by the cache*/
 	uint64 GetAllocatedMemory();
