Fixes crashes when  PIE in Lyra with net mode = listen server/client

diff --git a/Engine/Source/Runtime/Engine/Private/GameViewportClient.cpp b/Engine/Source/Runtime/Engine/Private/GameViewportClient.cpp
--- a/Engine/Source/Runtime/Engine/Private/GameViewportClient.cpp
+++ b/Engine/Source/Runtime/Engine/Private/GameViewportClient.cpp
@@ -1637,7 +1637,7 @@ void UGameViewportClient::Draw(FViewport* InViewport, FCanvas* SceneCanvas)
 		{
 			float GlobalResolutionFraction = 1.0f;
 
-			if (ViewFamily.EngineShowFlags.ScreenPercentage)
+			if (ViewFamily.EngineShowFlags.ScreenPercentage && !bDisableWorldRendering && ViewFamily.Views.Num() > 0)
 			{
 				// Get global view fraction.
 				FStaticResolutionFractionHeuristic StaticHeuristic;
