New FActorToken for message logs that contains information to select loaded or unloaded actors

original code contributed by: The Coalition

diff --git a/Engine/Source/Developer/MessageLog/Private/UserInterface/SMessageLogMessageListRow.cpp b/Engine/Source/Developer/MessageLog/Private/UserInterface/SMessageLogMessageListRow.cpp
--- a/Engine/Source/Developer/MessageLog/Private/UserInterface/SMessageLogMessageListRow.cpp
+++ b/Engine/Source/Developer/MessageLog/Private/UserInterface/SMessageLogMessageListRow.cpp
@@ -331,6 +331,15 @@ void SMessageLogMessageListRow::CreateMessage(const TSharedRef<SHorizontalBox>&
 	}
 		break;
 #endif
+
+	case EMessageToken::Actor:
+	{
+		const TSharedRef<FActorToken> ActorToken = StaticCastSharedRef<FActorToken>(InMessageToken);
+
+		IconBrushName = FName("Icons.Search");
+		RowContent = CreateHyperlink(InMessageToken, ActorToken->ToText());
+	}
+		break;
 	}
 
 	if (RowContent.IsValid())
diff --git a/Engine/Source/Editor/SceneOutliner/Private/ActorBrowsingMode.cpp b/Engine/Source/Editor/SceneOutliner/Private/ActorBrowsingMode.cpp
--- a/Engine/Source/Editor/SceneOutliner/Private/ActorBrowsingMode.cpp
+++ b/Engine/Source/Editor/SceneOutliner/Private/ActorBrowsingMode.cpp
@@ -57,6 +57,8 @@ FActorBrowsingMode::FActorBrowsingMode(SSceneOutliner* InSceneOutliner, TWeakObj
 
 	GEngine->OnLevelActorDeleted().AddRaw(this, &FActorBrowsingMode::OnLevelActorDeleted);
 
+	GEditor->OnSelectUnloadedActorsEvent().AddRaw(this, &FActorBrowsingMode::OnSelectUnloadedActors);
+
 	FEditorDelegates::OnEditCutActorsBegin.AddRaw(this, &FActorBrowsingMode::OnEditCutActorsBegin);
 	FEditorDelegates::OnEditCutActorsEnd.AddRaw(this, &FActorBrowsingMode::OnEditCutActorsEnd);
 	FEditorDelegates::OnEditCopyActorsBegin.AddRaw(this, &FActorBrowsingMode::OnEditCopyActorsBegin);
@@ -177,6 +179,8 @@ FActorBrowsingMode::~FActorBrowsingMode()
 	FSceneOutlinerDelegates::Get().OnComponentsUpdated.RemoveAll(this);
 
 	GEngine->OnLevelActorDeleted().RemoveAll(this);
+
+	GEditor->OnSelectUnloadedActorsEvent().RemoveAll(this);
 		
 	FEditorDelegates::OnEditCutActorsBegin.RemoveAll(this);
 	FEditorDelegates::OnEditCutActorsEnd.RemoveAll(this);
@@ -590,6 +594,24 @@ void FActorBrowsingMode::OnLevelActorDeleted(AActor* Actor)
 	ApplicableActors.Remove(Actor);
 }
 
+void FActorBrowsingMode::OnSelectUnloadedActors(const TArray<FGuid>& ActorGuids)
+{
+	TArray<FSceneOutlinerTreeItemPtr> ItemsToSelect;
+	ItemsToSelect.Reserve(ActorGuids.Num());
+	for (const FGuid& ActorGuid : ActorGuids)
+	{
+		if (FSceneOutlinerTreeItemPtr ItemPtr = SceneOutliner->GetTreeItem(ActorGuid))
+		{
+			ItemsToSelect.Add(ItemPtr);
+		}
+	}
+	
+	if (ItemsToSelect.Num())
+	{
+		SceneOutliner->SetItemSelection(ItemsToSelect, true);
+	}
+}
+
 void FActorBrowsingMode::OnActorDescRemoved(FWorldPartitionActorDesc* InActorDesc)
 {
 	ApplicableUnloadedActors.Remove(InActorDesc);
diff --git a/Engine/Source/Editor/SceneOutliner/Public/ActorBrowsingMode.h b/Engine/Source/Editor/SceneOutliner/Public/ActorBrowsingMode.h
--- a/Engine/Source/Editor/SceneOutliner/Public/ActorBrowsingMode.h
+++ b/Engine/Source/Editor/SceneOutliner/Public/ActorBrowsingMode.h
@@ -61,6 +61,9 @@ public:
 	void OnComponentsUpdated();
 	/** Called by the engine when an actor is deleted */
 	void OnLevelActorDeleted(AActor* Actor);
+
+	/** Called by the editor to allow selection of unloaded actors */
+	void OnSelectUnloadedActors(const TArray<FGuid>& ActorGuids);
 	
 	/** Called when an actor desc is removed */
 	void OnActorDescRemoved(FWorldPartitionActorDesc* InActorDesc);
diff --git a/Engine/Source/Editor/UnrealEd/Classes/Editor/EditorEngine.h b/Engine/Source/Editor/UnrealEd/Classes/Editor/EditorEngine.h
--- a/Engine/Source/Editor/UnrealEd/Classes/Editor/EditorEngine.h
+++ b/Engine/Source/Editor/UnrealEd/Classes/Editor/EditorEngine.h
@@ -693,6 +693,13 @@ public:
 	/** Called when the editor has been asked to perform an exec command on particle systems. */
 	DECLARE_EVENT_OneParam(UEditorEngine, FExecParticleInvoked, const TCHAR*);
 	FExecParticleInvoked& OnExecParticleInvoked() { return ExecParticleInvokedEvent; }
+
+	/** Called to allow selection of unloaded actors */
+	DECLARE_EVENT_OneParam(UEditorEngine, FSelectUnloadedActorsEvent, const TArray<FGuid>&);
+	FSelectUnloadedActorsEvent& OnSelectUnloadedActorsEvent() { return SelectUnloadedActorsEvent; }
+
+	void BroadcastSelectUnloadedActors(const TArray<FGuid>& ActorGuids) const { SelectUnloadedActorsEvent.Broadcast(ActorGuids); }
+
 	/**
 	 * Called before an actor or component is about to be translated, rotated, or scaled by the editor
 	 *
@@ -2870,6 +2877,9 @@ private:
 	/** Broadcasts after an Exec event on particles has been invoked.*/
 	FExecParticleInvoked ExecParticleInvokedEvent; 
 
+	/** Broadcasts to allow selection of unloaded actors */
+	FSelectUnloadedActorsEvent SelectUnloadedActorsEvent;
+
 	/** Delegate to be called when a matinee is requested to be opened */
 	FShouldOpenMatineeCallback ShouldOpenMatineeCallback;	
 
diff --git a/Engine/Source/Editor/UnrealEd/Private/UnrealEdMisc.cpp b/Engine/Source/Editor/UnrealEd/Private/UnrealEdMisc.cpp
--- a/Engine/Source/Editor/UnrealEd/Private/UnrealEdMisc.cpp
+++ b/Engine/Source/Editor/UnrealEd/Private/UnrealEdMisc.cpp
@@ -504,6 +504,7 @@ void FUnrealEdMisc::OnInit()
 	FUObjectToken::DefaultOnMessageTokenActivated().BindRaw(this, &FUnrealEdMisc::OnMessageTokenActivated);
 	FUObjectToken::DefaultOnGetObjectDisplayName().BindRaw(this, &FUnrealEdMisc::OnGetDisplayName);
 	FAssetNameToken::OnGotoAsset().BindRaw(this, &FUnrealEdMisc::OnGotoAsset);
+	FActorToken::DefaultOnMessageTokenActivated().BindRaw(this, &FUnrealEdMisc::OnActorTokenActivated);
 
 	// Register to receive notification of new key bindings
 	OnUserDefinedChordChangedDelegateHandle = FInputBindingManager::Get().RegisterUserDefinedChordChanged(FOnUserDefinedChordChanged::FDelegate::CreateRaw( this, &FUnrealEdMisc::OnUserDefinedChordChanged ));
@@ -1006,6 +1007,7 @@ void FUnrealEdMisc::OnExit()
 	FUObjectToken::DefaultOnMessageTokenActivated().Unbind();
 	FUObjectToken::DefaultOnGetObjectDisplayName().Unbind();
 	FAssetNameToken::OnGotoAsset().Unbind();
+	FActorToken::DefaultOnMessageTokenActivated().Unbind();
 
 	// Unregister message log UIs
 	FMessageLogModule& MessageLogModule = FModuleManager::LoadModuleChecked<FMessageLogModule>("MessageLog");
@@ -1369,15 +1371,7 @@ void FUnrealEdMisc::OnMessageTokenActivated(const TSharedRef<IMessageToken>& Tok
 
 			if (Actor && Actor->GetLevel() != nullptr)
 			{
-				// Select the actor
-				GEditor->SelectNone(false, true);
-				GEditor->SelectActor(Actor, /*InSelected=*/true, /*bNotify=*/false, /*bSelectEvenIfHidden=*/true); 
-				GEditor->NoteSelectionChange();
-				GEditor->MoveViewportCamerasToActor(*Actor, false);
-
-				// Update the property windows and create one if necessary
-				GUnrealEd->ShowActorProperties();
-				GUnrealEd->UpdateFloatingPropertyWindows();
+				SelectActorFromMessageToken(Actor);
 			}
 			else
 			{
@@ -1419,6 +1413,35 @@ void FUnrealEdMisc::OnMessageTokenActivated(const TSharedRef<IMessageToken>& Tok
 	}
 }
 
+void FUnrealEdMisc::OnActorTokenActivated(const TSharedRef<class IMessageToken>& Token)
+{
+	if (Token->GetType() == EMessageToken::Actor)
+	{
+		const TSharedRef<FActorToken> ActorToken = StaticCastSharedRef<FActorToken>(Token);
+		if (AActor* Actor = FindObject<AActor>(nullptr, *ActorToken->GetActorPath()))
+		{
+			SelectActorFromMessageToken(Actor);
+		}
+		else
+		{
+			GEditor->BroadcastSelectUnloadedActors({ ActorToken->GetActorGuid() });
+		}
+	}
+}
+
+void FUnrealEdMisc::SelectActorFromMessageToken(AActor* InActor)
+{
+	// Select the actor
+	GEditor->SelectNone(false, true);
+	GEditor->SelectActor(InActor, /*InSelected=*/true, /*bNotify=*/false, /*bSelectEvenIfHidden=*/true);
+	GEditor->NoteSelectionChange();
+	GEditor->MoveViewportCamerasToActor(*InActor, false);
+
+	// Update the property windows and create one if necessary
+	GUnrealEd->ShowActorProperties();
+	GUnrealEd->UpdateFloatingPropertyWindows();
+}
+
 FText FUnrealEdMisc::OnGetDisplayName(const UObject* InObject, const bool bFullPath)
 {
 	FText Name = LOCTEXT("DisplayNone", "<None>");
diff --git a/Engine/Source/Editor/UnrealEd/Public/UnrealEdMisc.h b/Engine/Source/Editor/UnrealEd/Public/UnrealEdMisc.h
--- a/Engine/Source/Editor/UnrealEd/Public/UnrealEdMisc.h
+++ b/Engine/Source/Editor/UnrealEd/Public/UnrealEdMisc.h
@@ -286,6 +286,9 @@ public:
 	/** Delegate used to go to assets in the content browser */
 	void OnGotoAsset(const FString& InAssetPath) const;
 
+	/** Delegate used on message log AActor token activation */
+	void OnActorTokenActivated(const TSharedRef<class IMessageToken>& Token);
+
 	/** Delegate used to update the map of asset update counts */
 	void OnObjectSaved(UObject* SavedObject, FObjectPreSaveContext SaveContext);
 
@@ -308,6 +311,8 @@ public:
 	FString FindMapFileFromPartialName(const FString& PartialMapName);
 
 private:
+	void SelectActorFromMessageToken(AActor* InActor);
+
 	void PreSaveWorld(class UWorld* World, FObjectPreSaveContext ObjectSaveContext);
 
 	/** The current state of the autosave */
diff --git a/Engine/Source/Runtime/Core/Private/Logging/TokenizedMessage.cpp b/Engine/Source/Runtime/Core/Private/Logging/TokenizedMessage.cpp
--- a/Engine/Source/Runtime/Core/Private/Logging/TokenizedMessage.cpp
+++ b/Engine/Source/Runtime/Core/Private/Logging/TokenizedMessage.cpp
@@ -211,4 +211,36 @@ TSharedRef<FDocumentationToken> FDocumentationToken::Create(const FString& InDoc
 	return MakeShareable(new FDocumentationToken(InDocumentationLink, InPreviewExcerptLink, InPreviewExcerptName));
 }
 
+FOnMessageTokenActivated FActorToken::DefaultMessageTokenActivated;
+
+TSharedRef<FActorToken> FActorToken::Create(const FString& InActorPath, const FGuid& InActorGuid, const FText& InMessage)
+{
+	return MakeShareable(new FActorToken(InActorPath, InActorGuid, InMessage));
+}
+
+FActorToken::FActorToken(const FString& InActorPath, const FGuid& InActorGuid, const FText& InMessage)
+	: ActorPath(InActorPath), ActorGuid(InActorGuid)
+{
+	if (!InMessage.IsEmpty())
+	{
+		CachedText = InMessage;
+	}
+	else
+	{
+		CachedText = FText::FromString(ActorPath);
+	}
+}
+
+const FOnMessageTokenActivated& FActorToken::GetOnMessageTokenActivated() const
+{
+	if (MessageTokenActivated.IsBound())
+	{
+		return MessageTokenActivated;
+	}
+	else
+	{
+		return DefaultMessageTokenActivated;
+	}
+}
+
 #undef LOCTEXT_NAMESPACE
diff --git a/Engine/Source/Runtime/Core/Public/Logging/TokenizedMessage.h b/Engine/Source/Runtime/Core/Public/Logging/TokenizedMessage.h
--- a/Engine/Source/Runtime/Core/Public/Logging/TokenizedMessage.h
+++ b/Engine/Source/Runtime/Core/Public/Logging/TokenizedMessage.h
@@ -33,6 +33,7 @@ namespace EMessageToken
 	enum Type
 	{
 		Action,
+		Actor,
 		AssetName,
 		Documentation,
 		Image,
@@ -593,3 +594,52 @@ private:
 	FString TutorialAssetName;
 };
 
+/** 
+ * Basic message token that defaults its activated method to select an actor in the opened level
+ */
+class FActorToken : public IMessageToken
+{
+public:
+	/** Factory method, tokens can only be constructed as shared refs */
+	CORE_API static TSharedRef<FActorToken> Create(const FString& InActorPath, const FGuid& InActorGuid, const FText& InMessage = FText());
+
+	/** Begin IMessageToken interface */
+	virtual EMessageToken::Type GetType() const override
+	{
+		return EMessageToken::Actor;
+	}
+
+	virtual const FOnMessageTokenActivated& GetOnMessageTokenActivated() const override;
+	/** End IMessageToken interface */
+
+	/** Get the actor name used by this token */
+	const FString& GetActorPath() const
+	{
+		return ActorPath;
+	}
+
+	/** Get the actor guid used by this token */
+	const FGuid& GetActorGuid() const
+	{
+		return ActorGuid;
+	}
+
+	/** Get the delegate for default token activation */
+	CORE_API static FOnMessageTokenActivated& DefaultOnMessageTokenActivated()
+	{
+		return DefaultMessageTokenActivated;
+	}
+
+private:
+	/** Private constructor */
+	FActorToken(const FString& InActorPath, const FGuid& InActorGuid, const FText& InMessage);
+
+	/** The actor path we will select */
+	FString ActorPath;
+
+	/** The actor guid we will select */
+	FGuid ActorGuid;
+
+	/** The default activation method, if any */
+	CORE_API static FOnMessageTokenActivated DefaultMessageTokenActivated;
+};
\ No newline at end of file
diff --git a/Engine/Source/Runtime/Engine/Private/WorldPartition/ErrorHandling/WorldPartitionStreamingGenerationMapCheckErrorHandler.cpp b/Engine/Source/Runtime/Engine/Private/WorldPartition/ErrorHandling/WorldPartitionStreamingGenerationMapCheckErrorHandler.cpp
--- a/Engine/Source/Runtime/Engine/Private/WorldPartition/ErrorHandling/WorldPartitionStreamingGenerationMapCheckErrorHandler.cpp
+++ b/Engine/Source/Runtime/Engine/Private/WorldPartition/ErrorHandling/WorldPartitionStreamingGenerationMapCheckErrorHandler.cpp
@@ -13,7 +13,7 @@ void FStreamingGenerationMapCheckErrorHandler::OnInvalidReference(const FWorldPa
 {
 	TSharedRef<FTokenizedMessage> Error = FMessageLog("MapCheck").Warning()
 		->AddToken(FTextToken::Create(LOCTEXT("MapCheck_WorldPartition_Actor", "Actor")))
-		->AddToken(FAssetNameToken::Create(ActorDescView.GetActorLabelOrName().ToString()))
+		->AddToken(FActorToken::Create(ActorDescView.GetActorPath().ToString(), ActorDescView.GetGuid(), FText::FromName(ActorDescView.GetActorLabelOrName())))
 		->AddToken(FTextToken::Create(LOCTEXT("MapCheck_WorldPartition_HaveMissingRefsTo", "have missing references to")))
 		->AddToken(FTextToken::Create(FText::FromString(ReferenceGuid.ToString())))
 		->AddToken(FMapErrorToken::Create(FName(TEXT("WorldPartition_MissingActorReference_CheckForErrors"))));
@@ -26,10 +26,10 @@ void FStreamingGenerationMapCheckErrorHandler::OnInvalidReferenceGridPlacement(c
 
 	TSharedRef<FTokenizedMessage> Error = FMessageLog("MapCheck").Error()
 		->AddToken(FTextToken::Create(ActorDescView.GetIsSpatiallyLoaded() ? SpatiallyLoadedActor : NonSpatiallyLoadedActor))
-		->AddToken(FAssetNameToken::Create(ActorDescView.GetActorLabelOrName().ToString()))
+		->AddToken(FActorToken::Create(ActorDescView.GetActorPath().ToString(), ActorDescView.GetGuid(), FText::FromName(ActorDescView.GetActorLabelOrName())))
 		->AddToken(FTextToken::Create(LOCTEXT("MapCheck_WorldPartition_References", "references")))
 		->AddToken(FTextToken::Create(ReferenceActorDescView.GetIsSpatiallyLoaded() ? SpatiallyLoadedActor : NonSpatiallyLoadedActor))
-		->AddToken(FAssetNameToken::Create(ReferenceActorDescView.GetActorLabelOrName().ToString()))
+		->AddToken(FActorToken::Create(ReferenceActorDescView.GetActorPath().ToString(), ReferenceActorDescView.GetGuid(), FText::FromName(ReferenceActorDescView.GetActorLabelOrName())))
 		->AddToken(FMapErrorToken::Create(FName(TEXT("WorldPartition_StreamedActorReferenceAlwaysLoadedActor_CheckForErrors"))));
 }
 
@@ -37,9 +37,9 @@ void FStreamingGenerationMapCheckErrorHandler::OnInvalidReferenceDataLayers(cons
 {
 	TSharedRef<FTokenizedMessage> Error = FMessageLog("MapCheck").Error()
 		->AddToken(FTextToken::Create(LOCTEXT("MapCheck_WorldPartition_Actor", "Actor")))
-		->AddToken(FAssetNameToken::Create(ActorDescView.GetActorLabelOrName().ToString()))
+		->AddToken(FActorToken::Create(ActorDescView.GetActorPath().ToString(), ActorDescView.GetGuid(), FText::FromName(ActorDescView.GetActorLabelOrName())))
 		->AddToken(FTextToken::Create(LOCTEXT("MapCheck_WorldPartition_ReferenceActorInOtherDataLayers", "references an actor in a different set of data layers")))
-		->AddToken(FAssetNameToken::Create(ReferenceActorDescView.GetActorLabelOrName().ToString()))
+		->AddToken(FActorToken::Create(ReferenceActorDescView.GetActorPath().ToString(), ReferenceActorDescView.GetGuid(), FText::FromName(ReferenceActorDescView.GetActorLabelOrName())))
 		->AddToken(FMapErrorToken::Create(FName(TEXT("WorldPartition_ActorReferenceActorInAnotherDataLayer_CheckForErrors"))));
 }
 
@@ -47,9 +47,9 @@ void FStreamingGenerationMapCheckErrorHandler::OnInvalidReferenceRuntimeGrid(con
 {
 	TSharedRef<FTokenizedMessage> Error = FMessageLog("MapCheck").Error()
 		->AddToken(FTextToken::Create(LOCTEXT("MapCheck_WorldPartition_Actor", "Actor")))
-		->AddToken(FAssetNameToken::Create(ActorDescView.GetActorLabelOrName().ToString()))
+		->AddToken(FActorToken::Create(ActorDescView.GetActorPath().ToString(), ActorDescView.GetGuid(), FText::FromName(ActorDescView.GetActorLabelOrName())))
 		->AddToken(FTextToken::Create(LOCTEXT("MapCheck_WorldPartition_ReferenceActorInOtherRuntimeGrid", "references an actor in a different runtime grid")))
-		->AddToken(FAssetNameToken::Create(ReferenceActorDescView.GetActorLabelOrName().ToString()))
+		->AddToken(FActorToken::Create(ReferenceActorDescView.GetActorPath().ToString(), ReferenceActorDescView.GetGuid(), FText::FromName(ReferenceActorDescView.GetActorLabelOrName())))
 		->AddToken(FMapErrorToken::Create(FName(TEXT("WorldPartition_ActorReferenceActorInAnotherRuntimeGrid_CheckForErrors"))));
 }
 
@@ -57,7 +57,7 @@ void FStreamingGenerationMapCheckErrorHandler::OnInvalidReferenceLevelScriptStre
 {
 	TSharedRef<FTokenizedMessage> Error = FMessageLog("MapCheck").Error()
 		->AddToken(FTextToken::Create(LOCTEXT("MapCheck_WorldPartition_LevelScriptBlueprintStreamedActorReference", "Level Script Blueprint references streamed actor")))
-		->AddToken(FAssetNameToken::Create(ActorDescView.GetActorLabelOrName().ToString()))
+		->AddToken(FActorToken::Create(ActorDescView.GetActorPath().ToString(), ActorDescView.GetGuid(), FText::FromName(ActorDescView.GetActorLabelOrName())))
 		->AddToken(FMapErrorToken::Create(FName(TEXT("WorldPartition_LevelScriptBlueprintRefefenceStreamed_CheckForErrors"))));
 }
 
@@ -65,7 +65,7 @@ void FStreamingGenerationMapCheckErrorHandler::OnInvalidReferenceLevelScriptData
 {
 	TSharedRef<FTokenizedMessage> Error = FMessageLog("MapCheck").Error()
 		->AddToken(FTextToken::Create(LOCTEXT("MapCheck_WorldPartition_LevelScriptBlueprintActorReference", "Level Script Blueprint references actor")))
-		->AddToken(FAssetNameToken::Create(ActorDescView.GetActorLabelOrName().ToString()))
+		->AddToken(FActorToken::Create(ActorDescView.GetActorPath().ToString(), ActorDescView.GetGuid(), FText::FromName(ActorDescView.GetActorLabelOrName())))
 		->AddToken(FTextToken::Create(LOCTEXT("MapCheck_WorldPartition_LevelScriptBlueprintDataLayerReference", "with a non empty set of data layers")))
 		->AddToken(FMapErrorToken::Create(FName(TEXT("WorldPartition_LevelScriptBlueprintRefefenceDataLayer_CheckForErrors"))));
 }
