Fix ensure with new support of Multithreaded VLogs

#rnx

diff --git a/Engine/Source/Runtime/Engine/Private/VisualLogger/VisualLogger.cpp b/Engine/Source/Runtime/Engine/Private/VisualLogger/VisualLogger.cpp
--- a/Engine/Source/Runtime/Engine/Private/VisualLogger/VisualLogger.cpp
+++ b/Engine/Source/Runtime/Engine/Private/VisualLogger/VisualLogger.cpp
@@ -268,7 +268,8 @@ FVisualLogEntry* FVisualLogger::GetEntryToWrite(const UObject* Object, const flo
 
 FVisualLogEntry* FVisualLogger::GetEntryToWriteInternal(const UObject* Object, const float TimeStamp, const ECreateIfNeeded ShouldCreate)
 {
-	const UObject* LogOwner = FindRedirectionInternal(Object);
+	// No redirection needed, it should have been done at the time of the thread entry was computed
+	const UObject* LogOwner = Object;
 	if (LogOwner == nullptr)
 	{
 		return nullptr;
