Fix Fast Array bug where changing values after they are replicated in a single frame can cause them never to send.

diff --git a/Engine/Source/Runtime/Engine/Private/DataReplication.cpp b/Engine/Source/Runtime/Engine/Private/DataReplication.cpp
--- a/Engine/Source/Runtime/Engine/Private/DataReplication.cpp
+++ b/Engine/Source/Runtime/Engine/Private/DataReplication.cpp
@@ -260,7 +260,7 @@ public:
 		FReplicationChangelistMgr& ChangelistMgr,
 		TArray<TSharedPtr<INetDeltaBaseState>>& CustomDeltaStates)
 	{
-		RepLayout.PreSendCustomDeltaProperties(Object, Connection, ChangelistMgr, CustomDeltaStates);
+		RepLayout.PreSendCustomDeltaProperties(Object, Connection, ChangelistMgr, Connection->Driver->ReplicationFrame, CustomDeltaStates);
 	}
 
 	static void PostSendCustomDeltaProperties(
diff --git a/Engine/Source/Runtime/Engine/Private/RepLayout.cpp b/Engine/Source/Runtime/Engine/Private/RepLayout.cpp
--- a/Engine/Source/Runtime/Engine/Private/RepLayout.cpp
+++ b/Engine/Source/Runtime/Engine/Private/RepLayout.cpp
@@ -512,8 +512,8 @@ struct FCustomDeltaChangelistState
 		ArrayStates.SetNum(NumArrays);
 	}
 
-	/** Index used to determine whether or not we've compared already on a given frame. */
-	uint32 CompareIndex = 0;
+	/** Last Replication Frame where we modified histories and caused compares to happen. */
+	uint32 LastReplicationFrame = 0;
 
 	/**
 	 * An array tracking the last compared history of Arrays.
@@ -7237,6 +7237,7 @@ void FRepLayout::PreSendCustomDeltaProperties(
 	UObject* Object,
 	UNetConnection* Connection,
 	FReplicationChangelistMgr& ChangelistMgr,
+	uint32 ReplicationFrame,
 	TArray<TSharedPtr<INetDeltaBaseState>>& CustomDeltaStates) const
 {
 	using namespace UE_RepLayout_Private;
@@ -7252,10 +7253,9 @@ void FRepLayout::PreSendCustomDeltaProperties(
 
 			// Check to see whether or not we need to do comparisons this frame.
 			// If we do, then run through our fast array states and generate new history items if needed.
-			if (CustomDeltaChangelistState.CompareIndex != static_cast<uint32>(GFrameCounter))
+			if (CustomDeltaChangelistState.LastReplicationFrame != ReplicationFrame)
 			{
-				const bool bIsInitial = (CustomDeltaChangelistState.CompareIndex == 0);
-				CustomDeltaChangelistState.CompareIndex = GFrameCounter;
+				CustomDeltaChangelistState.LastReplicationFrame = ReplicationFrame;
 
 				const FConstRepObjectDataBuffer ObjectData(Object);
 				const uint16 NumLifetimeCustomDeltaProperties = LocalLifetimeCustomPropertyState.GetNumCustomDeltaProperties();
@@ -7291,6 +7291,7 @@ void FRepLayout::PreSendCustomDeltaProperties(
 						{
 							FDeltaArrayHistoryState& FastArrayHistoryState = CustomDeltaChangelistState.ArrayStates[FastArrayNumber];
 
+							// If the fast array's ReplicationKey hasn't changed, then we can safely assume there's been no changes.
 							const int32 FastArrayReplicationKey = CustomDeltaProperty.GetFastArrayArrayReplicationKey(FastArraySerializer);
 							if (FastArrayHistoryState.ArrayReplicationKey != FastArrayReplicationKey)
 							{
diff --git a/Engine/Source/Runtime/Engine/Public/Net/RepLayout.h b/Engine/Source/Runtime/Engine/Public/Net/RepLayout.h
--- a/Engine/Source/Runtime/Engine/Public/Net/RepLayout.h
+++ b/Engine/Source/Runtime/Engine/Public/Net/RepLayout.h
@@ -1892,6 +1892,7 @@ private:
 		UObject* Object,
 		UNetConnection* Connection,
 		FReplicationChangelistMgr& ChangelistMgr,
+		uint32 ReplicationFrame,
 		TArray<TSharedPtr<INetDeltaBaseState>>& CustomDeltaStates) const;
 
 	void PostSendCustomDeltaProperties(
diff --git a/Engine/Source/Runtime/Net/Core/Classes/Net/Serialization/FastArraySerializer.h b/Engine/Source/Runtime/Net/Core/Classes/Net/Serialization/FastArraySerializer.h
--- a/Engine/Source/Runtime/Net/Core/Classes/Net/Serialization/FastArraySerializer.h
+++ b/Engine/Source/Runtime/Net/Core/Classes/Net/Serialization/FastArraySerializer.h
@@ -1593,8 +1593,11 @@ bool FFastArraySerializer::FastArrayDeltaSerialize_DeltaSerializeStructs(TArray<
 		// Get the old map if its there
 		TMap<int32, int32>* OldItemMap = nullptr;
 		int32 BaseReplicationKey = INDEX_NONE;
-		int32 OldChangelistHistory = INDEX_NONE;
-		int32 OldLastAckedHistory = INDEX_NONE;
+
+		// We set ChangelistHistory and LastAckedHistory to 0, as that's the default RepLayout expects.
+		// Setting them to anything else may cause issues with initial sends.
+		uint32 OldChangelistHistory = 0;
+		uint32 OldLastAckedHistory = 0;
 
 		// See if the array changed at all. If the ArrayReplicationKey matches we can skip checking individual items
 		if (Parms.OldState)
