Expose morph target and cloth information through the mesh deformer helper interface.

diff --git a/Engine/Source/Runtime/Engine/Private/GPUSkinVertexFactory.cpp b/Engine/Source/Runtime/Engine/Private/GPUSkinVertexFactory.cpp
--- a/Engine/Source/Runtime/Engine/Private/GPUSkinVertexFactory.cpp
+++ b/Engine/Source/Runtime/Engine/Private/GPUSkinVertexFactory.cpp
@@ -1258,14 +1258,13 @@ public:
 	{
 		// Call regular GPU skinning shader parameters
 		FGPUSkinVertexFactoryShaderParameters::GetElementShaderBindings(Scene, View, Shader, InputStreamType, FeatureLevel, VertexFactory, BatchElement, ShaderBindings, VertexStreams);
-		const auto* GPUSkinVertexFactory = (const FGPUBaseSkinVertexFactory*)VertexFactory;
-		// A little hacky; problem is we can't upcast from FGPUBaseSkinVertexFactory to FGPUBaseSkinAPEXClothVertexFactory as they are unrelated; a nice solution would be
-		// to use virtual inheritance, but that requires RTTI and complicates things further...
-		const FGPUBaseSkinAPEXClothVertexFactory::ClothShaderType& ClothShaderData = 
-			 GPUSkinVertexFactory->GetBoneInfluenceType() == DefaultBoneInfluence ? ((const TGPUSkinAPEXClothVertexFactory<DefaultBoneInfluence>*)GPUSkinVertexFactory)->GetClothShaderData() :
-																					((const TGPUSkinAPEXClothVertexFactory<UnlimitedBoneInfluence>*)GPUSkinVertexFactory)->GetClothShaderData();
+		FGPUBaseSkinVertexFactory const* GPUSkinVertexFactory = (const FGPUBaseSkinVertexFactory*)VertexFactory;
+		FGPUBaseSkinAPEXClothVertexFactory const* ClothVertexFactory = GPUSkinVertexFactory->GetClothVertexFactory();
+		check(ClothVertexFactory != nullptr);
 
-		ShaderBindings.Add(Shader->GetUniformBufferParameter<FAPEXClothUniformShaderParameters>(),ClothShaderData.GetClothUniformBuffer());
+		const FGPUBaseSkinAPEXClothVertexFactory::ClothShaderType& ClothShaderData = ClothVertexFactory->GetClothShaderData();
+
+		ShaderBindings.Add(Shader->GetUniformBufferParameter<FAPEXClothUniformShaderParameters>(), ClothShaderData.GetClothUniformBuffer());
 
 		uint32 FrameNumber = View->Family->FrameNumber;
 
@@ -1282,14 +1281,8 @@ public:
 			ShaderBindings.Add(PreviousClothLocalToWorldParameter, ClothShaderData.GetClothLocalToWorldForReading(true, FrameNumber));
 		}
 
-		ShaderBindings.Add(GPUSkinApexClothParameter,
-			 GPUSkinVertexFactory->GetBoneInfluenceType() == DefaultBoneInfluence ? ((const TGPUSkinAPEXClothVertexFactory<DefaultBoneInfluence>*)GPUSkinVertexFactory)->GetClothBuffer() :
-																					((const TGPUSkinAPEXClothVertexFactory<UnlimitedBoneInfluence>*)GPUSkinVertexFactory)->GetClothBuffer() );
-
-		int32 ClothIndexOffset =
-			 GPUSkinVertexFactory->GetBoneInfluenceType() == DefaultBoneInfluence ? ((const TGPUSkinAPEXClothVertexFactory<DefaultBoneInfluence>*)GPUSkinVertexFactory)->GetClothIndexOffset(BatchElement.MinVertexIndex) :
-																					((const TGPUSkinAPEXClothVertexFactory<UnlimitedBoneInfluence>*)GPUSkinVertexFactory)->GetClothIndexOffset(BatchElement.MinVertexIndex);
-
+		ShaderBindings.Add(GPUSkinApexClothParameter, ClothVertexFactory->GetClothBuffer());
+		int32 ClothIndexOffset = ClothVertexFactory->GetClothIndexOffset(BatchElement.MinVertexIndex);
 		FIntPoint GPUSkinApexClothStartIndexOffset(BatchElement.MinVertexIndex, ClothIndexOffset);
 		ShaderBindings.Add(GPUSkinApexClothStartIndexOffsetParameter, GPUSkinApexClothStartIndexOffset);
 	}
diff --git a/Engine/Source/Runtime/Engine/Private/SkeletalMeshDeformerHelpers.cpp b/Engine/Source/Runtime/Engine/Private/SkeletalMeshDeformerHelpers.cpp
--- a/Engine/Source/Runtime/Engine/Private/SkeletalMeshDeformerHelpers.cpp
+++ b/Engine/Source/Runtime/Engine/Private/SkeletalMeshDeformerHelpers.cpp
@@ -16,10 +16,67 @@ FRHIShaderResourceView* FSkeletalMeshDeformerHelpers::GetBoneBufferForReading(
 	}
 
 	FSkeletalMeshObjectGPUSkin* MeshObjectGPU = static_cast<FSkeletalMeshObjectGPUSkin*>(MeshObject);
-	const FSkeletalMeshObjectGPUSkin::FSkeletalMeshObjectLOD& LOD = MeshObjectGPU->LODs[LODIndex];
-	FGPUBaseSkinVertexFactory* BaseVertexFactory = LOD.GPUSkinVertexFactories.VertexFactories[SectionIndex].Get();
-	FShaderResourceViewRHIRef SRV = BaseVertexFactory->GetShaderData().GetBoneBufferForReading(bPreviousFrame).VertexBufferSRV;
-	return SRV.IsValid() ? SRV : nullptr;
+	FGPUBaseSkinVertexFactory const* BaseVertexFactory = MeshObjectGPU->GetBaseSkinVertexFactory(LODIndex, SectionIndex);
+	bool bHasBoneBuffer = BaseVertexFactory != nullptr && BaseVertexFactory->GetShaderData().HasBoneBufferForReading(bPreviousFrame);
+	if (!bHasBoneBuffer)
+	{
+		return nullptr;
+	}
+
+	return BaseVertexFactory->GetShaderData().GetBoneBufferForReading(bPreviousFrame).VertexBufferSRV;
+}
+
+FRHIShaderResourceView* FSkeletalMeshDeformerHelpers::GetMorphTargetBufferForReading(
+	FSkeletalMeshObject* MeshObject,
+	int32 LODIndex,
+	int32 SectionIndex,
+	uint32 FrameNumber,
+	bool bPreviousFrame)
+{
+	if (MeshObject->IsCPUSkinned())
+	{
+		return nullptr;
+	}
+
+	FSkeletalMeshObjectGPUSkin* MeshObjectGPU = static_cast<FSkeletalMeshObjectGPUSkin*>(MeshObject);
+	FGPUBaseSkinVertexFactory const* BaseVertexFactory = MeshObjectGPU->GetBaseSkinVertexFactory(LODIndex, SectionIndex);
+	FMorphVertexBuffer const* MorphVertexBuffer = BaseVertexFactory != nullptr ? BaseVertexFactory->GetMorphVertexBuffer(bPreviousFrame, FrameNumber) : nullptr;
+
+	return MorphVertexBuffer != nullptr ? MorphVertexBuffer->GetSRV() : nullptr;
+}
+
+FSkeletalMeshDeformerHelpers::FClothBuffers FSkeletalMeshDeformerHelpers::GetClothBuffersForReading(
+	FSkeletalMeshObject* MeshObject,
+	int32 LODIndex,
+	int32 SectionIndex,
+	uint32 FrameNumber,
+	bool bPreviousFrame)
+{
+	if (MeshObject->IsCPUSkinned())
+	{
+		return FClothBuffers();
+	}
+
+	FSkeletalMeshObjectGPUSkin* MeshObjectGPU = static_cast<FSkeletalMeshObjectGPUSkin*>(MeshObject);
+	FGPUBaseSkinVertexFactory const* BaseVertexFactory = MeshObjectGPU->GetBaseSkinVertexFactory(LODIndex, SectionIndex);
+	FGPUBaseSkinAPEXClothVertexFactory const* ClothVertexFactory = BaseVertexFactory != nullptr ? BaseVertexFactory->GetClothVertexFactory() : nullptr;
+
+	if (ClothVertexFactory == nullptr || !ClothVertexFactory->GetClothShaderData().HasClothBufferForReading(bPreviousFrame, FrameNumber))
+	{
+		return FClothBuffers();
+	}
+
+	FSkeletalMeshRenderData const& SkeletalMeshRenderData = MeshObject->GetSkeletalMeshRenderData();
+	FSkeletalMeshLODRenderData const* LodRenderData = SkeletalMeshRenderData.GetPendingFirstLOD(LODIndex);
+	FSkelMeshRenderSection const& RenderSection = LodRenderData->RenderSections[SectionIndex];
+
+	FClothBuffers Ret;
+	Ret.ClothInfluenceBuffer = ClothVertexFactory->GetClothBuffer();
+	Ret.ClothInfluenceBufferOffset = ClothVertexFactory->GetClothIndexOffset(RenderSection.BaseVertexIndex);
+	Ret.ClothSimulatedPositionAndNormalBuffer = ClothVertexFactory->GetClothShaderData().GetClothBufferForReading(bPreviousFrame, FrameNumber).VertexBufferSRV;
+	Ret.ClothLocalToWorld = MeshObjectGPU->DynamicData->ClothObjectLocalToWorld * ClothVertexFactory->GetClothShaderData().GetClothLocalToWorldForReading(bPreviousFrame, FrameNumber).Inverse();
+	Ret.ClothWorldToLocal = Ret.ClothLocalToWorld.Inverse();
+	return Ret;
 }
 
 void FSkeletalMeshDeformerHelpers::SetVertexFactoryBufferOverrides(
diff --git a/Engine/Source/Runtime/Engine/Private/SkeletalRenderGPUSkin.cpp b/Engine/Source/Runtime/Engine/Private/SkeletalRenderGPUSkin.cpp
--- a/Engine/Source/Runtime/Engine/Private/SkeletalRenderGPUSkin.cpp
+++ b/Engine/Source/Runtime/Engine/Private/SkeletalRenderGPUSkin.cpp
@@ -1367,7 +1367,6 @@ const FVertexFactory* FSkeletalMeshObjectGPUSkin::GetSkinVertexFactory(const FSc
 	checkSlow( LODs.IsValidIndex(LODIndex) );
 	checkSlow( DynamicData );
 
-	const FSkelMeshObjectLODInfo& MeshLODInfo = LODInfo[LODIndex];
 	const FSkeletalMeshObjectLOD& LOD = LODs[LODIndex];
 
 	// If a mesh deformer cache was used, return the passthrough vertex factory
@@ -1393,6 +1392,14 @@ const FVertexFactory* FSkeletalMeshObjectGPUSkin::GetSkinVertexFactory(const FSc
 		return LOD.GPUSkinVertexFactories.PassthroughVertexFactories[ChunkIdx].Get();
 	}
 
+	// No passthrough usage so return the base skin vertex factory.
+	return GetBaseSkinVertexFactory(LODIndex, ChunkIdx);
+}
+
+FGPUBaseSkinVertexFactory const* FSkeletalMeshObjectGPUSkin::GetBaseSkinVertexFactory(int32 LODIndex, int32 ChunkIdx) const
+{
+	const FSkeletalMeshObjectLOD& LOD = LODs[LODIndex];
+
 	// cloth simulation is updated & if this ChunkIdx is for ClothVertexFactory
 	if ( DynamicData->ClothingSimData.Num() > 0 
 		&& LOD.GPUSkinVertexFactories.ClothVertexFactories.IsValidIndex(ChunkIdx)  
diff --git a/Engine/Source/Runtime/Engine/Private/SkeletalRenderGPUSkin.h b/Engine/Source/Runtime/Engine/Private/SkeletalRenderGPUSkin.h
--- a/Engine/Source/Runtime/Engine/Private/SkeletalRenderGPUSkin.h
+++ b/Engine/Source/Runtime/Engine/Private/SkeletalRenderGPUSkin.h
@@ -405,6 +405,12 @@ public:
 
 	FSkinWeightVertexBuffer* GetSkinWeightVertexBuffer(int32 LODIndex) const;
 
+	/** 
+	 * Get the skin vertex factory for direct skinning. 
+	 * This is different from GetSkinVertexFactory because it ignores any passthrough vertex factories that may be in use.
+	 */
+	FGPUBaseSkinVertexFactory const* GetBaseSkinVertexFactory(int32 LODIndex, int32 ChunkIdx) const;
+
 	/** 
 	 * Vertex buffers that can be used for GPU skinning factories 
 	 */
diff --git a/Engine/Source/Runtime/Engine/Public/GPUSkinVertexFactory.h b/Engine/Source/Runtime/Engine/Public/GPUSkinVertexFactory.h
--- a/Engine/Source/Runtime/Engine/Public/GPUSkinVertexFactory.h
+++ b/Engine/Source/Runtime/Engine/Public/GPUSkinVertexFactory.h
@@ -226,6 +226,16 @@ public:
 			return UniformBuffer;
 		}
 		
+		bool HasBoneBufferForReading(bool bPrevious) const
+		{
+			const FVertexBufferAndSRV* RetPtr = &GetBoneBufferInternal(bPrevious);
+			if (!RetPtr->VertexBufferRHI.IsValid() && bPrevious)
+			{
+				RetPtr = &GetBoneBufferInternal(false);
+			}
+			return RetPtr->VertexBufferRHI.IsValid();
+		}
+
 		// @param bPrevious true:previous, false:current
 		const FVertexBufferAndSRV& GetBoneBufferForReading(bool bPrevious) const
 		{
@@ -357,6 +367,9 @@ public:
 	/** Morph vertex factory functions */
 	virtual void UpdateMorphVertexStream(const class FMorphVertexBuffer* MorphVertexBuffer) {}
 	virtual const class FMorphVertexBuffer* GetMorphVertexBuffer(bool bPrevious, uint32 FrameNumber) const { return nullptr; }
+	/** Cloth vertex factory access. */
+	virtual class FGPUBaseSkinAPEXClothVertexFactory* GetClothVertexFactory() { return nullptr; }
+	virtual class FGPUBaseSkinAPEXClothVertexFactory const* GetClothVertexFactory() const { return nullptr; }
 
 	virtual GPUSkinBoneInfluenceType GetBoneInfluenceType() const				{ return DefaultBoneInfluence; }
 	virtual uint32 GetNumBoneInfluences() const									{ return Data.IsValid() ? Data->NumBoneInfluences : 0; }
@@ -683,6 +696,16 @@ public:
 			return ClothSimulPositionNormalBuffer[Index];
 		}
 
+		bool HasClothBufferForReading(bool bPrevious, uint32 FrameNumber) const
+		{
+			int32 Index = GetMostRecentIndex(FrameNumber);
+			if (bPrevious && DoWeHavePreviousData())
+			{
+				Index = 1 - Index;
+			}
+			return ClothSimulPositionNormalBuffer[Index].VertexBufferRHI.IsValid();
+		}
+
 		// @param bPrevious true:previous, false:current
 		// @param FrameNumber usually from View.Family->FrameNumber
 		const FVertexBufferAndSRV& GetClothBufferForReading(bool bPrevious, uint32 FrameNumber) const
@@ -834,10 +857,16 @@ public:
 		return ClothShaderData;
 	}
 
+	static bool IsClothEnabled(EShaderPlatform Platform);
+
 	virtual FGPUBaseSkinVertexFactory* GetVertexFactory() = 0;
 	virtual const FGPUBaseSkinVertexFactory* GetVertexFactory() const = 0;
 
-	static bool IsClothEnabled(EShaderPlatform Platform);
+	/** Get buffer containing cloth influences. */
+	virtual FShaderResourceViewRHIRef GetClothBuffer() { return nullptr; }
+	virtual const FShaderResourceViewRHIRef GetClothBuffer() const { return nullptr; }
+	/** Get offset from vertex index to cloth influence index at a given vertex index. The offset will be constant for all vertices in the same section. */
+	virtual uint32 GetClothIndexOffset(uint32 VertexIndex, uint32 LODBias = 0) const { return 0; }
 
 protected:
 	ClothShaderType ClothShaderData;
@@ -863,17 +892,17 @@ class TGPUSkinAPEXClothVertexFactory : public FGPUBaseSkinAPEXClothVertexFactory
 	typedef TGPUSkinVertexFactory<BoneInfluenceType> Super;
 
 public:
-	inline FShaderResourceViewRHIRef GetClothBuffer()
+	inline FShaderResourceViewRHIRef GetClothBuffer() override
 	{
 		return ClothDataPtr ? ClothDataPtr->ClothBuffer : nullptr;
 	}
 
-	inline const FShaderResourceViewRHIRef GetClothBuffer() const
+	const FShaderResourceViewRHIRef GetClothBuffer() const override
 	{
 		return ClothDataPtr ? ClothDataPtr->ClothBuffer : nullptr;
 	}
 
-	inline uint32 GetClothIndexOffset(uint32 VertexIndex, uint32 LODBias = 0) const
+	uint32 GetClothIndexOffset(uint32 VertexIndex, uint32 LODBias = 0) const override
 	{
 		if (ClothDataPtr)
 		{
@@ -919,6 +948,16 @@ public:
 		return this;
 	}
 
+	virtual FGPUBaseSkinAPEXClothVertexFactory* GetClothVertexFactory() override 
+	{
+		return this; 
+	}
+	
+	virtual FGPUBaseSkinAPEXClothVertexFactory const* GetClothVertexFactory() const override 
+	{
+		return this; 
+	}
+
 	// FRenderResource interface.
 
 	/**
diff --git a/Engine/Source/Runtime/Engine/Public/SkeletalMeshDeformerHelpers.h b/Engine/Source/Runtime/Engine/Public/SkeletalMeshDeformerHelpers.h
--- a/Engine/Source/Runtime/Engine/Public/SkeletalMeshDeformerHelpers.h
+++ b/Engine/Source/Runtime/Engine/Public/SkeletalMeshDeformerHelpers.h
@@ -9,10 +9,13 @@ class FSkeletalMeshObject;
 class FRDGPooledBuffer;
 class FRHIShaderResourceView;
 
-/** Functions that expose some functionality of FSkeletalMeshObject required by MeshDeformer systems. */
+/** Functions that expose some internal functionality of FSkeletalMeshObject required by MeshDeformer systems. */
 class FSkeletalMeshDeformerHelpers
 {
 public:
+
+#pragma region GetInternals
+
 	/** Get direct access to bone matrix buffer SRV. */
 	ENGINE_API static FRHIShaderResourceView* GetBoneBufferForReading(
 		FSkeletalMeshObject* MeshObject,
@@ -20,6 +23,36 @@ public:
 		int32 SectionIndex,
 		bool bPreviousFrame);
 
+	/** Get direct access to morph target buffer SRV. */
+	ENGINE_API static FRHIShaderResourceView* GetMorphTargetBufferForReading(
+		FSkeletalMeshObject* MeshObject,
+		int32 LODIndex,
+		int32 SectionIndex,
+		uint32 FrameNumber,
+		bool bPreviousFrame);
+
+	/** Buffer SRVs from the cloth system. */
+	struct ENGINE_API FClothBuffers
+	{
+		int32 ClothInfluenceBufferOffset = 0;
+		FRHIShaderResourceView* ClothInfluenceBuffer = nullptr;
+		FRHIShaderResourceView* ClothSimulatedPositionAndNormalBuffer = nullptr;
+		FMatrix44f ClothLocalToWorld = FMatrix44f::Identity;
+		FMatrix44f ClothWorldToLocal = FMatrix44f::Identity;
+	};
+
+	/** Get direct access to cloth buffer SRVs. */
+	ENGINE_API static FClothBuffers GetClothBuffersForReading(
+		FSkeletalMeshObject* MeshObject,
+		int32 LODIndex,
+		int32 SectionIndex,
+		uint32 FrameNumber,
+		bool bPreviousFrame);
+
+#pragma endregion GetInternals
+
+#pragma region SetInternals
+
 	/** Buffer override behavior for SetVertexFactoryBufferOverrides. */
 	enum class EOverrideType
 	{
@@ -48,4 +81,6 @@ public:
 		TRefCountPtr<FRDGPooledBuffer> const& PositionBuffer);
 
 #endif // RHI_RAYTRACING
+
+#pragma endregion SetInternals
 };
