Fixing up warnings about one-shots with MetaSounds

- changing GetDuration(), IsLooping(), and IsOneShot() to be const for better consistency on API (const-cast the one case that is not const for caching data)
- Renaming SoundCue param in GameplayCueNotify for sounds to Sound since it's not a SoundCue.

#p4v-preflight-copy 18647645

diff --git a/Engine/Plugins/Runtime/GameplayAbilities/Source/GameplayAbilities/Public/GameplayCueNotifyTypes.h b/Engine/Plugins/Runtime/GameplayAbilities/Source/GameplayAbilities/Public/GameplayCueNotifyTypes.h
--- a/Engine/Plugins/Runtime/GameplayAbilities/Source/GameplayAbilities/Public/GameplayCueNotifyTypes.h
+++ b/Engine/Plugins/Runtime/GameplayAbilities/Source/GameplayAbilities/Public/GameplayCueNotifyTypes.h
@@ -401,6 +401,9 @@ public:
 
 	// Sound to play.
 	UPROPERTY(EditDefaultsOnly, BlueprintReadOnly, Category = GameplayCueNotify)
+	USoundBase* Sound;
+
+	UPROPERTY(meta = (DeprecatedProperty, DeprecationMessage = "5.0 - SoundCue is deprecated. Instead use the Sound property. The type is USoundBase not USoundCue."))
 	USoundBase* SoundCue;
 
 	// How long it should take to fade out.  Only used on looping gameplay cues.
diff --git a/Engine/Plugins/Runtime/Metasound/Source/MetasoundEngine/Private/MetasoundSource.cpp b/Engine/Plugins/Runtime/Metasound/Source/MetasoundEngine/Private/MetasoundSource.cpp
--- a/Engine/Plugins/Runtime/Metasound/Source/MetasoundEngine/Private/MetasoundSource.cpp
+++ b/Engine/Plugins/Runtime/Metasound/Source/MetasoundEngine/Private/MetasoundSource.cpp
@@ -467,7 +467,7 @@ bool UMetaSoundSource::SupportsSubtitles() const
 	return Super::SupportsSubtitles();
 }
 
-float UMetaSoundSource::GetDuration()
+float UMetaSoundSource::GetDuration() const
 {
 	// eh? this is kind of a weird field anyways.
 	return Super::GetDuration();
@@ -724,12 +724,12 @@ bool UMetaSoundSource::IsParameterValid(const FAudioParameter& InParameter, cons
 	return bIsValid;
 }
 
-bool UMetaSoundSource::IsLooping()
+bool UMetaSoundSource::IsLooping() const
 {
 	return !IsOneShot();
 }
 
-bool UMetaSoundSource::IsOneShot()
+bool UMetaSoundSource::IsOneShot() const
 {
 	using namespace Metasound::Frontend;
 
diff --git a/Engine/Plugins/Runtime/Metasound/Source/MetasoundEngine/Public/MetasoundSource.h b/Engine/Plugins/Runtime/Metasound/Source/MetasoundEngine/Public/MetasoundSource.h
--- a/Engine/Plugins/Runtime/Metasound/Source/MetasoundEngine/Public/MetasoundSource.h
+++ b/Engine/Plugins/Runtime/Metasound/Source/MetasoundEngine/Public/MetasoundSource.h
@@ -163,13 +163,13 @@ public:
 
 	virtual bool IsPlayable() const override;
 	virtual bool SupportsSubtitles() const override;
-	virtual float GetDuration() override;
+	virtual float GetDuration() const override;
 	virtual bool ImplementsParameterInterface(Audio::FParameterInterfacePtr InInterface) const override;
 	virtual ISoundGeneratorPtr CreateSoundGenerator(const FSoundGeneratorInitParams& InParams) override;
 	virtual TUniquePtr<Audio::IParameterTransmitter> CreateParameterTransmitter(Audio::FParameterTransmitterInitParams&& InParams) const override;
 	virtual bool IsParameterValid(const FAudioParameter& InParameter) const override;
-	virtual bool IsLooping() override;
-	virtual bool IsOneShot() override;
+	virtual bool IsLooping() const override;
+	virtual bool IsOneShot() const override;
 
 protected:
 	/** Gets all the default parameters for this Asset.  */
diff --git a/Engine/Plugins/Runtime/SoundUtilities/Source/SoundUtilities/Private/SoundSimple.cpp b/Engine/Plugins/Runtime/SoundUtilities/Source/SoundUtilities/Private/SoundSimple.cpp
--- a/Engine/Plugins/Runtime/SoundUtilities/Source/SoundUtilities/Private/SoundSimple.cpp
+++ b/Engine/Plugins/Runtime/SoundUtilities/Source/SoundUtilities/Private/SoundSimple.cpp
@@ -106,7 +106,7 @@ float USoundSimple::GetMaxDistance() const
 	return MaxDistance;
 }
 
-float USoundSimple::GetDuration()
+float USoundSimple::GetDuration() const
 {
 	return Duration;
 }
\ No newline at end of file
diff --git a/Engine/Plugins/Runtime/SoundUtilities/Source/SoundUtilities/Public/SoundSimple.h b/Engine/Plugins/Runtime/SoundUtilities/Source/SoundUtilities/Public/SoundSimple.h
--- a/Engine/Plugins/Runtime/SoundUtilities/Source/SoundUtilities/Public/SoundSimple.h
+++ b/Engine/Plugins/Runtime/SoundUtilities/Source/SoundUtilities/Public/SoundSimple.h
@@ -59,7 +59,7 @@ public:
 	virtual bool IsPlayable() const override;
 	virtual void Parse(class FAudioDevice* AudioDevice, const UPTRINT NodeWaveInstanceHash, FActiveSound& ActiveSound, const FSoundParseParameters& ParseParams, TArray<FWaveInstance*>& WaveInstances) override;
 	virtual float GetMaxDistance() const override;
-	virtual float GetDuration() override;
+	virtual float GetDuration() const override;
 	//~ End USoundBase Interface.
 
 protected:
diff --git a/Engine/Source/Runtime/Engine/Classes/Sound/DialogueSoundWaveProxy.h b/Engine/Source/Runtime/Engine/Classes/Sound/DialogueSoundWaveProxy.h
--- a/Engine/Source/Runtime/Engine/Classes/Sound/DialogueSoundWaveProxy.h
+++ b/Engine/Source/Runtime/Engine/Classes/Sound/DialogueSoundWaveProxy.h
@@ -36,7 +36,7 @@ public:
 	/** 
 	 * Returns the length of the sound
 	 */
-	virtual float GetDuration() override;
+	virtual float GetDuration() const override;
 
 	virtual float GetVolumeMultiplier() override;
 	virtual float GetPitchMultiplier() override;
diff --git a/Engine/Source/Runtime/Engine/Classes/Sound/SoundBase.h b/Engine/Source/Runtime/Engine/Classes/Sound/SoundBase.h
--- a/Engine/Source/Runtime/Engine/Classes/Sound/SoundBase.h
+++ b/Engine/Source/Runtime/Engine/Classes/Sound/SoundBase.h
@@ -227,7 +227,7 @@ public:
 	/**
 	 * Returns the length of the sound
 	 */
-	virtual float GetDuration();
+	virtual float GetDuration() const;
 
 	/** Returns whether or not this sound has a delay node, which means it's possible for the sound to not generate audio for a while. */
 	bool HasDelayNode() const;
@@ -253,10 +253,10 @@ public:
 	virtual class UCurveTable* GetCurveData() const { return nullptr; }
 
 	/** Returns whether or not this sound is looping. TODO: Deprecate this to only use IsOneshot() in a MetaSound world. */
-	virtual bool IsLooping();
+	virtual bool IsLooping() const;
 
 	/** Query if it's one shot. One shot is defined as a sound which is intended to have a fixed duration. */
-	virtual bool IsOneShot();
+	virtual bool IsOneShot() const;
 
 	/** Parses the Sound to generate the WaveInstances to play. */
 	virtual void Parse( class FAudioDevice* AudioDevice, const UPTRINT NodeWaveInstanceHash, FActiveSound& ActiveSound, const FSoundParseParameters& ParseParams, TArray<FWaveInstance*>& WaveInstances ) { }
diff --git a/Engine/Source/Runtime/Engine/Classes/Sound/SoundCue.h b/Engine/Source/Runtime/Engine/Classes/Sound/SoundCue.h
--- a/Engine/Source/Runtime/Engine/Classes/Sound/SoundCue.h
+++ b/Engine/Source/Runtime/Engine/Classes/Sound/SoundCue.h
@@ -171,7 +171,7 @@ public:
 	virtual float GetVolumeMultiplier() override;
 	virtual float GetPitchMultiplier() override;
 	virtual float GetMaxDistance() const override;
-	virtual float GetDuration() override;
+	virtual float GetDuration() const override;
 	virtual const FSoundAttenuationSettings* GetAttenuationSettingsToApply() const override;
 	virtual float GetSubtitlePriority() const override;
 	virtual bool GetSoundWavesWithCookedAnalysisData(TArray<USoundWave*>& OutSoundWaves) override;
diff --git a/Engine/Source/Runtime/Engine/Classes/Sound/SoundSourceBus.h b/Engine/Source/Runtime/Engine/Classes/Sound/SoundSourceBus.h
--- a/Engine/Source/Runtime/Engine/Classes/Sound/SoundSourceBus.h
+++ b/Engine/Source/Runtime/Engine/Classes/Sound/SoundSourceBus.h
@@ -59,7 +59,7 @@ public:
 
 	//~ Begin USoundBase Interface.
 	virtual bool IsPlayable() const override;
-	virtual float GetDuration() override;
+	virtual float GetDuration() const override;
 	//~ End USoundBase Interface.
 
 protected:
diff --git a/Engine/Source/Runtime/Engine/Classes/Sound/SoundWave.h b/Engine/Source/Runtime/Engine/Classes/Sound/SoundWave.h
--- a/Engine/Source/Runtime/Engine/Classes/Sound/SoundWave.h
+++ b/Engine/Source/Runtime/Engine/Classes/Sound/SoundWave.h
@@ -781,7 +781,7 @@ public:
 	//~ Begin USoundBase Interface.
 	virtual bool IsPlayable() const override;
 	virtual void Parse(class FAudioDevice* AudioDevice, const UPTRINT NodeWaveInstanceHash, FActiveSound& ActiveSound, const FSoundParseParameters& ParseParams, TArray<FWaveInstance*>& WaveInstances) override;
-	virtual float GetDuration() override;
+	virtual float GetDuration() const override;
 	virtual float GetSubtitlePriority() const override;
 	virtual bool SupportsSubtitles() const override;
 	virtual bool GetSoundWavesWithCookedAnalysisData(TArray<USoundWave*>& OutSoundWaves) override;
diff --git a/Engine/Source/Runtime/Engine/Private/Animation/AnimNotify_PlaySound.cpp b/Engine/Source/Runtime/Engine/Private/Animation/AnimNotify_PlaySound.cpp
--- a/Engine/Source/Runtime/Engine/Private/Animation/AnimNotify_PlaySound.cpp
+++ b/Engine/Source/Runtime/Engine/Private/Animation/AnimNotify_PlaySound.cpp
@@ -39,9 +39,9 @@ void UAnimNotify_PlaySound::Notify(class USkeletalMeshComponent* MeshComp, class
 	// Don't call super to avoid call back in to blueprints
 	if (Sound && MeshComp)
 	{
-		if (Sound->IsLooping() && !Sound->IsA<USoundWaveProcedural>())
+		if (!Sound->IsOneShot())
 		{
-			UE_LOG(LogAudio, Warning, TEXT("PlaySound notify: Anim %s tried to spawn infinitely looping sound asset %s. Spawning suppressed."), *GetNameSafe(Animation), *GetNameSafe(Sound));
+			UE_LOG(LogAudio, Warning, TEXT("PlaySound notify: Anim %s tried to play a sound asset which is not a one-shot: '%s'. Spawning suppressed."), *GetNameSafe(Animation), *GetNameSafe(Sound));
 			return;
 		}
 
@@ -83,7 +83,7 @@ void UAnimNotify_PlaySound::ValidateAssociatedAssets()
 {
 	static const FName NAME_AssetCheck("AssetCheck");
 
-	if ((Sound != nullptr) && (Sound->IsLooping()) && (!Sound->IsA<USoundWaveProcedural>()))
+	if (Sound != nullptr && !Sound->IsOneShot())
 	{
 		UObject* ContainingAsset = GetContainingAsset();
 
diff --git a/Engine/Source/Runtime/Engine/Private/DialogueWave.cpp b/Engine/Source/Runtime/Engine/Private/DialogueWave.cpp
--- a/Engine/Source/Runtime/Engine/Private/DialogueWave.cpp
+++ b/Engine/Source/Runtime/Engine/Private/DialogueWave.cpp
@@ -433,7 +433,7 @@ float UDialogueSoundWaveProxy::GetMaxDistance() const
 	return SoundWave->GetMaxDistance();
 }
 
-float UDialogueSoundWaveProxy::GetDuration()
+float UDialogueSoundWaveProxy::GetDuration() const
 {
 	return SoundWave->GetDuration();
 }
diff --git a/Engine/Source/Runtime/Engine/Private/SoundBase.cpp b/Engine/Source/Runtime/Engine/Private/SoundBase.cpp
--- a/Engine/Source/Runtime/Engine/Private/SoundBase.cpp
+++ b/Engine/Source/Runtime/Engine/Private/SoundBase.cpp
@@ -66,7 +66,7 @@ float USoundBase::GetMaxDistance() const
 	return WORLD_MAX;
 }
 
-float USoundBase::GetDuration()
+float USoundBase::GetDuration() const
 {
 	return Duration;
 }
@@ -96,12 +96,12 @@ float USoundBase::GetPitchMultiplier()
 	return 1.f;
 }
 
-bool USoundBase::IsOneShot()
+bool USoundBase::IsOneShot() const
 {
 	return !IsLooping();
 }
 
-bool USoundBase::IsLooping()
+bool USoundBase::IsLooping() const
 {
 	return (GetDuration() >= INDEFINITELY_LOOPING_DURATION);
 }
diff --git a/Engine/Source/Runtime/Engine/Private/SoundCue.cpp b/Engine/Source/Runtime/Engine/Private/SoundCue.cpp
--- a/Engine/Source/Runtime/Engine/Private/SoundCue.cpp
+++ b/Engine/Source/Runtime/Engine/Private/SoundCue.cpp
@@ -561,12 +561,15 @@ float USoundCue::GetMaxDistance() const
 	return GIsEditor ? FindMaxDistanceInternal() : MaxDistance;
 }
 
-float USoundCue::GetDuration()
+float USoundCue::GetDuration() const
 {
 	// Always recalc the duration when in the editor as it could change
 	if (GIsEditor || (Duration < SMALL_NUMBER) || HasDelayNode())
 	{
-		CacheAggregateValues();
+		// This needs to be cached here vs an earlier point due to the need to parse sound cues and load order issues.
+		// Alternative is to make getters not const, this is preferable. 
+		USoundCue* ThisSoundCue = const_cast<USoundCue*>(this);
+		ThisSoundCue->CacheAggregateValues();
 	}
 
 	return Duration;
diff --git a/Engine/Source/Runtime/Engine/Private/SoundSourceBus.cpp b/Engine/Source/Runtime/Engine/Private/SoundSourceBus.cpp
--- a/Engine/Source/Runtime/Engine/Private/SoundSourceBus.cpp
+++ b/Engine/Source/Runtime/Engine/Private/SoundSourceBus.cpp
@@ -57,7 +57,7 @@ bool USoundSourceBus::IsPlayable() const
 	return true;
 }
 
-float USoundSourceBus::GetDuration()
+float USoundSourceBus::GetDuration() const
 {
 	return (SourceBusDuration > 0.0f) ? SourceBusDuration : INDEFINITELY_LOOPING_DURATION;
 }
diff --git a/Engine/Source/Runtime/Engine/Private/SoundWave.cpp b/Engine/Source/Runtime/Engine/Private/SoundWave.cpp
--- a/Engine/Source/Runtime/Engine/Private/SoundWave.cpp
+++ b/Engine/Source/Runtime/Engine/Private/SoundWave.cpp
@@ -2196,13 +2196,12 @@ void USoundWave::Parse(FAudioDevice* AudioDevice, const UPTRINT NodeWaveInstance
 
 		const USoundBase* Sound = ActiveSound.GetSound();
 
-		// Don't need to worry about logging orphaned sounds for procedural sounds
-		if (!Sound->IsA<USoundWaveProcedural>())
+		if (!Sound->IsOneShot())
 		{
 #if !(NO_LOGGING || UE_BUILD_SHIPPING || UE_BUILD_TEST)
 			if (!ActiveSound.bWarnedAboutOrphanedLooping && ActiveSound.GetAudioComponentID() == 0 && ActiveSound.FadeOut == FActiveSound::EFadeOut::None)
 			{
-				UE_LOG(LogAudio, Warning, TEXT("Detected orphaned looping sound '%s'."), *ActiveSound.GetSound()->GetName());
+				UE_LOG(LogAudio, Warning, TEXT("Detected orphaned sound which is not a one-shot '%s'. Sounds which are not one-shots need to have audio components or they risk being orphaned."), *ActiveSound.GetSound()->GetName());
 				ActiveSound.bWarnedAboutOrphanedLooping = true;
 			}
 #endif
@@ -2534,7 +2533,7 @@ bool USoundWave::IsPlayable() const
 	return true;
 }
 
-float USoundWave::GetDuration()
+float USoundWave::GetDuration() const
 {
 	return (bLooping ? INDEFINITELY_LOOPING_DURATION : Duration);
 }
