Engine: Rename IsSplitScreen to HasMultipleLocalPlayers since the original name was misleading.
Cameras: Only apply splitscreen scaling on shakes when the camera manager is actually running in a splitscreen context.

diff --git a/Engine/Source/Runtime/Engine/Classes/Engine/Engine.h b/Engine/Source/Runtime/Engine/Classes/Engine/Engine.h
--- a/Engine/Source/Runtime/Engine/Classes/Engine/Engine.h
+++ b/Engine/Source/Runtime/Engine/Classes/Engine/Engine.h
@@ -2537,9 +2537,13 @@ public:
 	/** @return the currently active audio device */
 	FAudioDeviceHandle GetActiveAudioDevice();
 
-	/** @return whether we're currently running in split screen (more than one local player) */
+	/** @return whether we currently have more than one local player */
+	UE_DEPRECATED(5.1, "IsSplitScreen was only ever checking if there are more than one local player. Use HasMultipleLocalPlayers instead.")
 	virtual bool IsSplitScreen(UWorld *InWorld);
 
+	/** @returns whether there are currently multiple local players in the given world */
+	virtual bool HasMultipleLocalPlayers(UWorld* InWorld);
+
 	/** @return whether we're currently running with stereoscopic 3D enabled for the specified viewport (or globally, if viewport is nullptr) */
 	bool IsStereoscopic3D(FViewport* InViewport = nullptr);
 
diff --git a/Engine/Source/Runtime/Engine/Classes/Kismet/KismetSystemLibrary.h b/Engine/Source/Runtime/Engine/Classes/Kismet/KismetSystemLibrary.h
--- a/Engine/Source/Runtime/Engine/Classes/Kismet/KismetSystemLibrary.h
+++ b/Engine/Source/Runtime/Engine/Classes/Kismet/KismetSystemLibrary.h
@@ -208,10 +208,15 @@ class ENGINE_API UKismetSystemLibrary : public UBlueprintFunctionLibrary
 	UFUNCTION(BlueprintPure, Category="Networking", meta=(WorldContext="WorldContextObject"))
 	static bool IsStandalone(const UObject* WorldContextObject);
 	
-	/** Returns whether we're currently running in split screen (more than one local player). */
-	UFUNCTION(BlueprintPure, Category = "Viewport", meta = (WorldContext = "WorldContextObject"))
+	/** Returns whether we currently have more than one local player. */
+	UE_DEPRECATED(5.1, "IsSplitScreen was only ever checking if there are more than one local player. Use HasMultipleLocalPlayers instead.")
+	UFUNCTION(BlueprintPure, Category = "Viewport", meta = (WorldContext = "WorldContextObject", DeprecatedFunction, DeprecationMessage = "Use HasMultipleLocalPlayers instead"))
 	static bool IsSplitScreen(const UObject* WorldContextObject);
 
+	/** Returns whether there are currently multiple local players in the given world */
+	UFUNCTION(BlueprintPure, Category = "Viewport", meta = (WorldContext = "WorldContextObject"))
+	static bool HasMultipleLocalPlayers(const UObject* WorldContextObject);
+
 	/** Returns whether this is a build that is packaged for distribution */
 	UFUNCTION(BlueprintPure, Category="Development", meta=(BlueprintThreadSafe))
 	static bool IsPackagedForDistribution();
diff --git a/Engine/Source/Runtime/Engine/Private/Camera/CameraModifier_CameraShake.cpp b/Engine/Source/Runtime/Engine/Private/Camera/CameraModifier_CameraShake.cpp
--- a/Engine/Source/Runtime/Engine/Private/Camera/CameraModifier_CameraShake.cpp
+++ b/Engine/Source/Runtime/Engine/Private/Camera/CameraModifier_CameraShake.cpp
@@ -6,8 +6,11 @@
 #include "Camera/PlayerCameraManager.h"
 #include "Engine/Canvas.h"
 #include "Engine/Engine.h"
+#include "Engine/GameViewportClient.h"
+#include "Engine/LocalPlayer.h"
 #include "DisplayDebugHelpers.h"
 #include "EngineGlobals.h"
+#include "GameFramework/PlayerController.h"
 
 DEFINE_LOG_CATEGORY_STATIC(LogCameraShake, Warning, All);
 
@@ -95,9 +98,16 @@ UCameraShakeBase* UCameraModifier_CameraShake::AddCameraShake(TSubclassOf<UCamer
 		const bool bIsCustomInitialized = Params.Initializer.IsBound();
 
 		// Adjust for splitscreen
-		if (CameraOwner != nullptr && GEngine->IsSplitScreen(CameraOwner->GetWorld()))
+		if (CameraOwner != nullptr && CameraOwner->PCOwner != nullptr)
 		{
-			Scale *= SplitScreenShakeScale;
+			const ULocalPlayer* LocalPlayer = CameraOwner->PCOwner->GetLocalPlayer();
+			if (LocalPlayer != nullptr && LocalPlayer->ViewportClient != nullptr)
+			{
+				if (LocalPlayer->ViewportClient->GetCurrentSplitscreenConfiguration() != ESplitScreenType::None)
+				{
+					Scale *= SplitScreenShakeScale;
+				}
+			}
 		}
 
 		UCameraShakeBase const* const ShakeCDO = GetDefault<UCameraShakeBase>(ShakeClass);
diff --git a/Engine/Source/Runtime/Engine/Private/ContentStreaming.cpp b/Engine/Source/Runtime/Engine/Private/ContentStreaming.cpp
--- a/Engine/Source/Runtime/Engine/Private/ContentStreaming.cpp
+++ b/Engine/Source/Runtime/Engine/Private/ContentStreaming.cpp
@@ -508,7 +508,7 @@ void IStreamingManager::SetupViewInfos( float DeltaTime )
 	// Apply a split-screen factor if we have multiple players on the same machine, and they currently have individual views.
 	float SplitScreenFactor = 1.0f;
 	
-	if ( bHaveMultiplePlayerViews && GEngine->IsSplitScreen(NULL) )
+	if ( bHaveMultiplePlayerViews && GEngine->HasMultipleLocalPlayers(NULL) )
 	{
 		SplitScreenFactor = 0.75f;
 	}
diff --git a/Engine/Source/Runtime/Engine/Private/KismetSystemLibrary.cpp b/Engine/Source/Runtime/Engine/Private/KismetSystemLibrary.cpp
--- a/Engine/Source/Runtime/Engine/Private/KismetSystemLibrary.cpp
+++ b/Engine/Source/Runtime/Engine/Private/KismetSystemLibrary.cpp
@@ -233,9 +233,14 @@ bool UKismetSystemLibrary::IsStandalone(const UObject* WorldContextObject)
 }
 
 bool UKismetSystemLibrary::IsSplitScreen(const UObject* WorldContextObject)
+{
+	return HasMultipleLocalPlayers(WorldContextObject);
+}
+
+bool UKismetSystemLibrary::HasMultipleLocalPlayers(const UObject* WorldContextObject)
 {
 	UWorld* World = GEngine->GetWorldFromContextObject(WorldContextObject, EGetWorldErrorMode::LogAndReturnNull);
-	return World ? GEngine->IsSplitScreen(World) : false;
+	return World ? GEngine->HasMultipleLocalPlayers(World) : false;
 }
 
 bool UKismetSystemLibrary::IsPackagedForDistribution()
diff --git a/Engine/Source/Runtime/Engine/Private/UnrealEngine.cpp b/Engine/Source/Runtime/Engine/Private/UnrealEngine.cpp
--- a/Engine/Source/Runtime/Engine/Private/UnrealEngine.cpp
+++ b/Engine/Source/Runtime/Engine/Private/UnrealEngine.cpp
@@ -3762,8 +3762,13 @@ void UEngine::RecordHMDAnalytics()
 	}
 }
 
-/** @return whether we're currently running in split screen (more than one local player) */
+/** @return whether we currently have more than one local player */
 bool UEngine::IsSplitScreen(UWorld *InWorld)
+{
+	return HasMultipleLocalPlayers(InWorld);
+}
+
+bool UEngine::HasMultipleLocalPlayers(UWorld* InWorld)
 {
 	if (InWorld == NULL)
 	{
