Fix static meshes not rendering when the default minimum LOD index is lower than the effective minimum LOD index.

#rnx

diff --git a/Engine/Source/Runtime/Engine/Private/StaticMeshRender.cpp b/Engine/Source/Runtime/Engine/Private/StaticMeshRender.cpp
--- a/Engine/Source/Runtime/Engine/Private/StaticMeshRender.cpp
+++ b/Engine/Source/Runtime/Engine/Private/StaticMeshRender.cpp
@@ -2368,7 +2368,9 @@ FPrimitiveSceneProxy* UStaticMeshComponent::CreateSceneProxy()
 	}
 
 	const FStaticMeshLODResourcesArray& LODResources = GetStaticMesh()->GetRenderData()->LODResources;
-	if (LODResources.Num() == 0	|| LODResources[FMath::Clamp<int32>(GetStaticMesh()->GetDefaultMinLOD(), 0, LODResources.Num()-1)].VertexBuffers.StaticMeshVertexBuffer.GetNumVertices() == 0)
+	const int32 SMCurrentMinLOD = GetStaticMesh()->GetMinLODIdx();
+	const int32 EffectiveMinLOD = bOverrideMinLOD ? MinLOD : SMCurrentMinLOD;
+	if (LODResources.Num() == 0	|| LODResources[FMath::Clamp<int32>(EffectiveMinLOD, 0, LODResources.Num()-1)].VertexBuffers.StaticMeshVertexBuffer.GetNumVertices() == 0)
 	{
 		UE_LOG(LogStaticMesh, Verbose, TEXT("Skipping CreateSceneProxy for StaticMeshComponent %s (LOD problems)"), *GetFullName());
 		return nullptr;
