Type truncation warning cleanup

#rnx

diff --git a/Engine/Source/Runtime/Engine/Classes/Components/CapsuleComponent.h b/Engine/Source/Runtime/Engine/Classes/Components/CapsuleComponent.h
--- a/Engine/Source/Runtime/Engine/Classes/Components/CapsuleComponent.h
+++ b/Engine/Source/Runtime/Engine/Classes/Components/CapsuleComponent.h
@@ -202,12 +202,12 @@ FORCEINLINE void UCapsuleComponent::SetCapsuleHalfHeight(float HalfHeight, bool
 FORCEINLINE float UCapsuleComponent::GetScaledCapsuleRadius() const
 {
 	const FVector& ComponentScale = GetComponentTransform().GetScale3D();
-	return CapsuleRadius * (ComponentScale.X < ComponentScale.Y ? ComponentScale.X : ComponentScale.Y);
+	return CapsuleRadius * UE_REAL_TO_FLOAT(ComponentScale.X < ComponentScale.Y ? ComponentScale.X : ComponentScale.Y);
 }
 
 FORCEINLINE float UCapsuleComponent::GetScaledCapsuleHalfHeight() const
 {
-	return CapsuleHalfHeight * GetComponentTransform().GetScale3D().Z;
+	return CapsuleHalfHeight * UE_REAL_TO_FLOAT(GetComponentTransform().GetScale3D().Z);
 }
 
 FORCEINLINE float UCapsuleComponent::GetScaledCapsuleHalfHeight_WithoutHemisphere() const
diff --git a/Engine/Source/Runtime/Engine/Classes/Engine/TextureRenderTarget2DArray.h b/Engine/Source/Runtime/Engine/Classes/Engine/TextureRenderTarget2DArray.h
--- a/Engine/Source/Runtime/Engine/Classes/Engine/TextureRenderTarget2DArray.h
+++ b/Engine/Source/Runtime/Engine/Classes/Engine/TextureRenderTarget2DArray.h
@@ -76,11 +76,11 @@ class UTextureRenderTarget2DArray : public UTextureRenderTarget
 	ENGINE_API class UTexture2DArray* ConstructTexture2DArray(UObject* InOuter, const FString& NewTexName, EObjectFlags InFlags);
 
 	//~ Begin UTexture Interface.
-	virtual float GetSurfaceWidth() const  override { return SizeX; }
+	virtual float GetSurfaceWidth() const  override { return static_cast<float>(SizeX); }
 	// PVS-Studio notices that the implementation of GetSurfaceWidth is identical to this one
 	// and warns us. In this case, it is intentional, so we disable the warning:
-	virtual float GetSurfaceHeight()const  override { return SizeX; }	 //-V524
-	virtual float GetSurfaceDepth() const override { return 0; }
+	virtual float GetSurfaceHeight()const  override { return static_cast<float>(SizeX); }	 //-V524
+	virtual float GetSurfaceDepth() const override { return 0.0f; }
 	virtual uint32 GetSurfaceArraySize() const override { return Slices; }
 	virtual FTextureResource* CreateResource() override;
 	virtual EMaterialValueType GetMaterialType() const override;
diff --git a/Engine/Source/Runtime/Engine/Classes/Engine/TextureRenderTargetCube.h b/Engine/Source/Runtime/Engine/Classes/Engine/TextureRenderTargetCube.h
--- a/Engine/Source/Runtime/Engine/Classes/Engine/TextureRenderTargetCube.h
+++ b/Engine/Source/Runtime/Engine/Classes/Engine/TextureRenderTargetCube.h
@@ -67,11 +67,11 @@ class UTextureRenderTargetCube : public UTextureRenderTarget
 	ENGINE_API class UTextureCube* ConstructTextureCube(UObject* InOuter, const FString& NewTexName, EObjectFlags InFlags);
 
 	//~ Begin UTexture Interface.
-	virtual float GetSurfaceWidth() const  override { return SizeX; }
+	virtual float GetSurfaceWidth() const  override { return static_cast<float>(SizeX); }
 	// PVS-Studio notices that the implementation of GetSurfaceWidth is identical to this one
 	// and warns us. In this case, it is intentional, so we disable the warning:
-	virtual float GetSurfaceHeight()const  override { return SizeX; }	 //-V524
-	virtual float GetSurfaceDepth() const override { return 0; }
+	virtual float GetSurfaceHeight()const  override { return static_cast<float>(SizeX); }	 //-V524
+	virtual float GetSurfaceDepth() const override { return 0.0f; }
 	virtual uint32 GetSurfaceArraySize() const override { return 6; }
 	virtual FTextureResource* CreateResource() override;
 	virtual EMaterialValueType GetMaterialType() const override;
diff --git a/Engine/Source/Runtime/Engine/Classes/Engine/TextureRenderTargetVolume.h b/Engine/Source/Runtime/Engine/Classes/Engine/TextureRenderTargetVolume.h
--- a/Engine/Source/Runtime/Engine/Classes/Engine/TextureRenderTargetVolume.h
+++ b/Engine/Source/Runtime/Engine/Classes/Engine/TextureRenderTargetVolume.h
@@ -75,9 +75,9 @@ class UTextureRenderTargetVolume : public UTextureRenderTarget
 	ENGINE_API class UVolumeTexture* ConstructTextureVolume(UObject* InOuter, const FString& NewTexName, EObjectFlags InFlags);
 
 	//~ Begin UTexture Interface.
-	virtual float GetSurfaceWidth() const  override { return SizeX; }
-	virtual float GetSurfaceHeight()const  override { return SizeY; }
-	virtual float GetSurfaceDepth()const  override { return SizeZ; }
+	virtual float GetSurfaceWidth() const  override { return static_cast<float>(SizeX); }
+	virtual float GetSurfaceHeight()const  override { return static_cast<float>(SizeY); }
+	virtual float GetSurfaceDepth()const  override { return static_cast<float>(SizeZ); }
 	virtual uint32 GetSurfaceArraySize() const override { return 0; }
 	virtual FTextureResource* CreateResource() override;
 	virtual EMaterialValueType GetMaterialType() const override;
diff --git a/Engine/Source/Runtime/Engine/Classes/Engine/VolumeTexture.h b/Engine/Source/Runtime/Engine/Classes/Engine/VolumeTexture.h
--- a/Engine/Source/Runtime/Engine/Classes/Engine/VolumeTexture.h
+++ b/Engine/Source/Runtime/Engine/Classes/Engine/VolumeTexture.h
@@ -82,9 +82,9 @@ public:
 	}
 
 	//~ Begin UTexture Interface
-	virtual float GetSurfaceWidth() const override { return GetSizeX(); }
-	virtual float GetSurfaceHeight() const override { return GetSizeY(); }
-	virtual float GetSurfaceDepth() const override { return GetSizeZ(); }
+	virtual float GetSurfaceWidth() const override { return static_cast<float>(GetSizeX()); }
+	virtual float GetSurfaceHeight() const override { return static_cast<float>(GetSizeY()); }
+	virtual float GetSurfaceDepth() const override { return static_cast<float>(GetSizeZ()); }
 	virtual uint32 GetSurfaceArraySize() const override { return 0; }
 	virtual FTextureResource* CreateResource() override;
 #if WITH_EDITOR
diff --git a/Engine/Source/Runtime/Engine/Public/AudioStreamingCache.h b/Engine/Source/Runtime/Engine/Public/AudioStreamingCache.h
--- a/Engine/Source/Runtime/Engine/Public/AudioStreamingCache.h
+++ b/Engine/Source/Runtime/Engine/Public/AudioStreamingCache.h
@@ -239,7 +239,7 @@ private:
 		FCacheElementDebugInfo()
 			: NumTotalChunks(0)
 			, NumTimesTouched(0)
-			, TimeLoadStarted(0.0)
+			, TimeLoadStarted(0)
 			, TimeToLoad(0.0)
 			, AverageLocationInCacheWhenNeeded(0.0f)
 			, LoadingBehavior(ESoundWaveLoadingBehavior::Uninitialized)
diff --git a/Engine/Source/Runtime/Landscape/Classes/LandscapeGrassType.h b/Engine/Source/Runtime/Landscape/Classes/LandscapeGrassType.h
--- a/Engine/Source/Runtime/Landscape/Classes/LandscapeGrassType.h
+++ b/Engine/Source/Runtime/Landscape/Classes/LandscapeGrassType.h
@@ -115,8 +115,8 @@ struct FGrassVariety
 		, GrassDensity(400)
 		, bUseGrid(true)
 		, PlacementJitter(1.0f)
-		, StartCullDistance(10000.0f)
-		, EndCullDistance(10000.0f)
+		, StartCullDistance(10000)
+		, EndCullDistance(10000)
 		, MinLOD(-1)
 		, Scaling(EGrassScaling::Uniform)
 		, ScaleX(1.0f, 1.0f)
diff --git a/Engine/Source/Runtime/Slate/Public/Framework/Docking/TabManager.h b/Engine/Source/Runtime/Slate/Public/Framework/Docking/TabManager.h
--- a/Engine/Source/Runtime/Slate/Public/Framework/Docking/TabManager.h
+++ b/Engine/Source/Runtime/Slate/Public/Framework/Docking/TabManager.h
@@ -651,7 +651,7 @@ class SLATE_API FTabManager : public TSharedFromThis<FTabManager>
 			protected:
 				FArea( const float InWidth, const float InHeight )
 				: WindowPlacement(Placement_Automatic)
-				, UnscaledWindowPosition(FVector2D(0,0))
+				, UnscaledWindowPosition(FVector2D(0.0, 0.0))
 				, UnscaledWindowSize(InWidth, InHeight)
 				, bIsMaximized( false )
 				{
@@ -797,7 +797,7 @@ class SLATE_API FTabManager : public TSharedFromThis<FTabManager>
 
 		static TSharedRef<FArea> NewArea( const FVector2D& WindowSize )
 		{
-			return MakeShareable( new FArea( WindowSize.X, WindowSize.Y ) );
+			return MakeShareable( new FArea( UE_REAL_TO_FLOAT(WindowSize.X), UE_REAL_TO_FLOAT(WindowSize.Y) ) );
 		}
 		
 		static TSharedRef<FStack> NewStack() 
