Validate Attribute bone index typing / mapping

Animaton attribute index remapping:
- Introduced FMeshAttributeContainer which stores the attributes according to FMeshPoseIndex
- Added FinalizeAttributeEvaluationResults to SkeletalMeshComponent - which remaps the evaluated compact-pose indexed attributes to mesh-indices
- Added CopyFrom template for different bone index types
- Updated CopyPoseFromMesh
   * Added SourceBoneToTarget as attributes are mapped from source to target vs target to source for Bones
- Ensure that attributes are only inserted during evaluation if its bone index is present in the required bones
- Added testing for remapping attributes
    * Remapping between LODs with removed bones
    * Remapping between Skeletal meshes with different skeleton hierarchies

diff --git a/Engine/Source/Editor/Persona/Private/AnimationEditorViewportClient.cpp b/Engine/Source/Editor/Persona/Private/AnimationEditorViewportClient.cpp
--- a/Engine/Source/Editor/Persona/Private/AnimationEditorViewportClient.cpp
+++ b/Engine/Source/Editor/Persona/Private/AnimationEditorViewportClient.cpp
@@ -742,7 +742,7 @@ void FAnimationViewportClient::ShowAttributeNames(FCanvas* Canvas, FSceneView* V
 		const int32 HalfX = Viewport->GetSizeXY().X / 2 / GetDPIScale();
 		const int32 HalfY = Viewport->GetSizeXY().Y / 2 / GetDPIScale();
 
-		const UE::Anim::FHeapAttributeContainer& Attributes = MeshComponent->GetCustomAttributes();
+		const UE::Anim::FMeshAttributeContainer& Attributes = MeshComponent->GetCustomAttributes();
 
 		const int32 TransformAnimationAttributeTypeIndex = Attributes.FindTypeIndex(FTransformAnimationAttribute::StaticStruct());
 		if (TransformAnimationAttributeTypeIndex != INDEX_NONE)
@@ -1617,7 +1617,7 @@ void FAnimationViewportClient::DrawAttributes(UDebugSkelMeshComponent* MeshCompo
 {
 	if (MeshComponent && MeshComponent->SkeletalMesh)
 	{
-		const UE::Anim::FHeapAttributeContainer& Attributes = MeshComponent->GetCustomAttributes();
+		const UE::Anim::FMeshAttributeContainer& Attributes = MeshComponent->GetCustomAttributes();
 
 		const int32 TransformAnimationAttributeTypeIndex = Attributes.FindTypeIndex(FTransformAnimationAttribute::StaticStruct());
 		if (TransformAnimationAttributeTypeIndex != INDEX_NONE)
diff --git a/Engine/Source/Editor/UnrealEd/Private/Animation/DebugSkelMeshComponent.cpp b/Engine/Source/Editor/UnrealEd/Private/Animation/DebugSkelMeshComponent.cpp
--- a/Engine/Source/Editor/UnrealEd/Private/Animation/DebugSkelMeshComponent.cpp
+++ b/Engine/Source/Editor/UnrealEd/Private/Animation/DebugSkelMeshComponent.cpp
@@ -687,7 +687,7 @@ void UDebugSkelMeshComponent::GenSpaceBases(TArray<FTransform>& OutSpaceBases)
 	TempBoneSpaceTransforms.AddUninitialized(OutSpaceBases.Num());
 	FVector TempRootBoneTranslation;
 	FBlendedHeapCurve TempCurve;
-	UE::Anim::FHeapAttributeContainer TempAtttributes;
+	UE::Anim::FMeshAttributeContainer TempAtttributes;
 	DoInstancePreEvaluation();
 	PerformAnimationEvaluation(SkeletalMesh, AnimScriptInstance, OutSpaceBases, TempBoneSpaceTransforms, TempRootBoneTranslation, TempCurve, TempAtttributes);
 	DoInstancePostEvaluation();
diff --git a/Engine/Source/Runtime/AnimGraphRuntime/Private/AnimNodes/AnimNode_CopyPoseFromMesh.cpp b/Engine/Source/Runtime/AnimGraphRuntime/Private/AnimNodes/AnimNode_CopyPoseFromMesh.cpp
--- a/Engine/Source/Runtime/AnimGraphRuntime/Private/AnimNodes/AnimNode_CopyPoseFromMesh.cpp
+++ b/Engine/Source/Runtime/AnimGraphRuntime/Private/AnimNodes/AnimNode_CopyPoseFromMesh.cpp
@@ -234,7 +234,7 @@ void FAnimNode_CopyPoseFromMesh::Evaluate_AnyThread(FPoseContext& Output)
 	if (bCopyCustomAttributes)
 	{	
 		const FBoneContainer& RequiredBones = OutPose.GetBoneContainer();
-		UE::Anim::Attributes::CopyAndRemapAttributes(SourceCustomAttributes, Output.CustomAttributes, BoneMapToSource, RequiredBones);		
+		UE::Anim::Attributes::CopyAndRemapAttributes(SourceCustomAttributes, Output.CustomAttributes, SourceBoneToTarget, RequiredBones);		
 	}
 }
 
@@ -315,6 +315,15 @@ void FAnimNode_CopyPoseFromMesh::ReinitializeMeshComponent(USkeletalMeshComponen
 					}
 				}
 			}
+
+			if (bCopyCustomAttributes)
+			{
+				SourceBoneToTarget.Reserve(BoneMapToSource.Num());
+				Algo::Transform(BoneMapToSource, SourceBoneToTarget, [](const TPair<int32, int32>& Pair)
+				{
+					return TPair<int32, int32>(Pair.Value, Pair.Key);
+				});
+			}
 		}
 	}
 }
diff --git a/Engine/Source/Runtime/AnimGraphRuntime/Public/AnimNodes/AnimNode_CopyPoseFromMesh.h b/Engine/Source/Runtime/AnimGraphRuntime/Public/AnimNodes/AnimNode_CopyPoseFromMesh.h
--- a/Engine/Source/Runtime/AnimGraphRuntime/Public/AnimNodes/AnimNode_CopyPoseFromMesh.h
+++ b/Engine/Source/Runtime/AnimGraphRuntime/Public/AnimNodes/AnimNode_CopyPoseFromMesh.h
@@ -66,6 +66,7 @@ private:
 	// cache of target space bases to source space bases
 	TMap<int32, int32> BoneMapToSource;
 	TMap<FName, SmartName::UID_Type> CurveNameToUIDMap;
+	TMap<int32, int32> SourceBoneToTarget;
 
 	// Cached transforms, copied on the game thread
 	TArray<FTransform> SourceMeshTransformArray;
@@ -74,7 +75,7 @@ private:
 	TMap<FName, float> SourceCurveList;
 
 	// Cached attributes, copied on the game thread
-	UE::Anim::FHeapAttributeContainer SourceCustomAttributes;
+	UE::Anim::FMeshAttributeContainer SourceCustomAttributes;
 
 	// reinitialize mesh component 
 	void ReinitializeMeshComponent(USkeletalMeshComponent* NewSkeletalMeshComponent, USkeletalMeshComponent* TargetMeshComponent);
diff --git a/Engine/Source/Runtime/AnimationCore/Public/BoneIndices.h b/Engine/Source/Runtime/AnimationCore/Public/BoneIndices.h
--- a/Engine/Source/Runtime/AnimationCore/Public/BoneIndices.h
+++ b/Engine/Source/Runtime/AnimationCore/Public/BoneIndices.h
@@ -21,6 +21,8 @@ struct FBoneIndexBase
 
 	FORCEINLINE explicit operator bool() const { return IsValid(); }
 
+	friend FORCEINLINE uint32 GetTypeHash(const FBoneIndexBase& Index) { return GetTypeHash(Index.BoneIndex); }
+
 protected:
 	int32 BoneIndex;
 };
diff --git a/Engine/Source/Runtime/Engine/Classes/Animation/AttributeBlendOperator.h b/Engine/Source/Runtime/Engine/Classes/Animation/AttributeBlendOperator.h
--- a/Engine/Source/Runtime/Engine/Classes/Animation/AttributeBlendOperator.h
+++ b/Engine/Source/Runtime/Engine/Classes/Animation/AttributeBlendOperator.h
@@ -200,7 +200,7 @@ namespace UE
 				{
 					const FAttributeId& Identifier = It.GetIdentifier();
 					AttributeType* OutAttribute = OutAttributes->FindOrAdd<AttributeType>(Identifier);
-					(*OutAttribute) = It.GetHighestBoneWeightedValue();
+					AttributeType::StaticStruct()->CopyScriptStruct(OutAttribute, &It.GetHighestBoneWeightedValue());
 				});
 
 				BlendData.ForEachUniqueAttribute<AttributeType>([OutAttributes](typename FAttributeBlendData::template TSingleIterator<AttributeType>& It) -> void
@@ -221,7 +221,7 @@ namespace UE
 							const bool bHighestWeight = It.IsHighestBoneWeighted();
 							if (bHighestWeight || ExistingIndex == INDEX_NONE)
 							{
-								OutAttribute = Attribute;
+								AttributeType::StaticStruct()->CopyScriptStruct(OutAttributePtr, &Attribute);
 							}
 						}
 					}
@@ -261,7 +261,7 @@ namespace UE
 
 						// Find or add as the attribute might already exist, so Add would fail
 						AttributeType& OutAttribute = *OutAttributes->FindOrAdd<AttributeType>(Identifier);
-						OutAttribute = Attribute;
+						AttributeType::StaticStruct()->CopyScriptStruct(&OutAttribute, &Attribute);
 					}
 				});
 			}			
@@ -312,7 +312,7 @@ namespace UE
 						const float AttributeWeight = It.GetWeight();
 						if (ExistingIndex == INDEX_NONE || AttributeWeight > 0.5f)
 						{
-							OutAttribute = Attribute;
+							AttributeType::StaticStruct()->CopyScriptStruct(OutAttributePtr, &Attribute);
 						}
 					}
 				});
@@ -354,8 +354,8 @@ namespace UE
 						// 'Accumulate' value if the attribute did not yet exist, or based upon being the highest weighted attribute
 						const float AttributeWeight = It.GetWeight();
 						if (ExistingIndex == INDEX_NONE)
-						{
-							OutAttribute = Attribute;
+						{							
+							AttributeType::StaticStruct()->CopyScriptStruct(OutAttributePtr, &Attribute);
 						}
 					}
 				});
diff --git a/Engine/Source/Runtime/Engine/Classes/Animation/AttributesContainer.h b/Engine/Source/Runtime/Engine/Classes/Animation/AttributesContainer.h
--- a/Engine/Source/Runtime/Engine/Classes/Animation/AttributesContainer.h
+++ b/Engine/Source/Runtime/Engine/Classes/Animation/AttributesContainer.h
@@ -11,6 +11,7 @@
 
 #include "Animation/AnimTypes.h"
 #include "Animation/WrappedAttribute.h"
+#include "BoneContainer.h"
 
 #include "BoneIndices.h"
 
@@ -41,6 +42,28 @@ namespace UE
 			int32 Index;
 		};
 
+		namespace Conversion
+		{
+			template<typename BoneIndexTypeA, typename BoneIndexTypeB>
+			BoneIndexTypeB MakeIndex(BoneIndexTypeA BoneIndex, const FBoneContainer& BoneContainer)
+			{
+				return BoneIndexTypeB();
+			}
+
+			template<>
+			inline FMeshPoseBoneIndex MakeIndex(FCompactPoseBoneIndex BoneIndex, const FBoneContainer& BoneContainer)
+			{
+				
+				return BoneContainer.MakeMeshPoseIndex(BoneIndex);
+			}
+			
+			template<>
+			inline FCompactPoseBoneIndex MakeIndex(FMeshPoseBoneIndex BoneIndex, const FBoneContainer& BoneContainer)
+			{
+				return BoneContainer.MakeCompactPoseIndex(BoneIndex);
+			}
+		}
+
 		/** Runtime container for Animation Attributes, providing a TMap like interface. Used in FStack*/
 		template<class BoneIndexType, typename InAllocator>
 		struct TAttributeContainer
@@ -67,6 +90,61 @@ namespace UE
 					}
 				}
 			}
+					
+		public:
+
+			template<class OtherBoneIndexType, typename OtherAllocator>
+			void CopyFrom(const TAttributeContainer<OtherBoneIndexType, OtherAllocator>& Other, const FBoneContainer& BoneContainer)
+			{
+				UniqueTypes = Other.UniqueTypes;
+				AttributeIdentifiers.Empty(Other.AttributeIdentifiers.Num());
+				UniqueTypedBoneIndices.Empty(Other.UniqueTypedBoneIndices.Num());
+				Values.Empty(Other.Values.Num());
+				
+				const int32 NumTypes = Other.UniqueTypedBoneIndices.Num();
+				for (int32 TypeIndex = 0; TypeIndex < NumTypes; ++TypeIndex)
+				{
+					// Generate mapping for all contained unique bones to other container its index type
+					TMap<int32, BoneIndexType> IndexMapping;
+					const TArray<int32>& OtherTypeBoneIndices = Other.UniqueTypedBoneIndices[TypeIndex];
+					for (int32 Index : OtherTypeBoneIndices)
+					{
+						OtherBoneIndexType OtherIndex(Index);						
+						BoneIndexType MyIndex = Conversion::MakeIndex<OtherBoneIndexType, BoneIndexType>(OtherIndex, BoneContainer);
+						if (MyIndex.IsValid())
+						{
+							IndexMapping.Add(Index, MyIndex);
+						}			
+					}
+					
+					const TWeakObjectPtr<const UScriptStruct>& WeakScriptStruct = UniqueTypes[TypeIndex];
+					const TArray<FAttributeId>& OtherTypeAttributeIdentifiers = Other.AttributeIdentifiers[TypeIndex];
+					const TArray<TWrappedAttribute<OtherAllocator>>& OtherValuesArray = Other.Values[TypeIndex];					
+
+					// Populate an array entry for the type
+					TArray<FAttributeId>& TypeAttributeIdentifiers = AttributeIdentifiers.AddDefaulted_GetRef();
+					TArray<TWrappedAttribute<InAllocator>>& ValuesArray = Values.AddDefaulted_GetRef();
+					TArray<int32>& TypeBoneIndices = UniqueTypedBoneIndices.AddDefaulted_GetRef();
+
+					const int32 NumOtherAttributes = OtherTypeAttributeIdentifiers.Num();
+					for (int32 AttributeIndex = 0; AttributeIndex < NumOtherAttributes; ++AttributeIndex)
+					{
+						const FAttributeId& Id = OtherTypeAttributeIdentifiers[AttributeIndex];
+						const TWrappedAttribute<OtherAllocator>& Value = OtherValuesArray[AttributeIndex];
+
+						// If the bone index was mapped successfully - add the remapped attribute
+						if (const BoneIndexType* MyIndexPtr = IndexMapping.Find(Id.GetIndex()))
+						{
+							FAttributeId MappedId(Id.GetName(), MyIndexPtr->GetInt(), Id.GetNamespace());
+							TypeAttributeIdentifiers.Add(MappedId);
+							TWrappedAttribute<InAllocator>& StructData = ValuesArray.Add_GetRef(WeakScriptStruct.Get());							
+							WeakScriptStruct->CopyScriptStruct(StructData.template GetPtr<void>(), Value.template GetPtr<void>());
+							
+							TypeBoneIndices.AddUnique(MyIndexPtr->GetInt());
+						}
+					}	
+				}
+			}
 
 			/** Copies all contained data from another Container instance using the same memory allocator */
 			void CopyFrom(const TAttributeContainer<BoneIndexType, InAllocator>& Other)
diff --git a/Engine/Source/Runtime/Engine/Classes/Components/SkeletalMeshComponent.h b/Engine/Source/Runtime/Engine/Classes/Components/SkeletalMeshComponent.h
--- a/Engine/Source/Runtime/Engine/Classes/Components/SkeletalMeshComponent.h
+++ b/Engine/Source/Runtime/Engine/Classes/Components/SkeletalMeshComponent.h
@@ -124,8 +124,8 @@ struct FAnimationEvaluationContext
 	FBlendedHeapCurve	CachedCurve;
 
 	// attribute data, swapped in from the component when we are running parallel eval
-	UE::Anim::FHeapAttributeContainer CustomAttributes;
-	UE::Anim::FHeapAttributeContainer CachedCustomAttributes;
+	UE::Anim::FMeshAttributeContainer CustomAttributes;
+	UE::Anim::FMeshAttributeContainer CachedCustomAttributes;
 
 	FAnimationEvaluationContext()
 	{
@@ -417,8 +417,8 @@ private:
 	FBlendedHeapCurve CachedCurve;
 
 	/** Current and cached attribute evaluation data, used for Update Rate optimization */
-	UE::Anim::FHeapAttributeContainer CachedAttributes;
-	UE::Anim::FHeapAttributeContainer CustomAttributes;
+	UE::Anim::FMeshAttributeContainer CachedAttributes;
+	UE::Anim::FMeshAttributeContainer CustomAttributes;
 public:
 	/** 
 	 * Get float type attribute value.
@@ -1892,12 +1892,12 @@ public:
 	* @param	OutCurves				Blended Curve
 	*/
 #if WITH_EDITOR
-	void PerformAnimationEvaluation(const USkeletalMesh* InSkeletalMesh, UAnimInstance* InAnimInstance, TArray<FTransform>& OutSpaceBases, TArray<FTransform>& OutBoneSpaceTransforms, FVector& OutRootBoneTranslation, FBlendedHeapCurve& OutCurve, UE::Anim::FHeapAttributeContainer& OutAttributes);
+	void PerformAnimationEvaluation(const USkeletalMesh* InSkeletalMesh, UAnimInstance* InAnimInstance, TArray<FTransform>& OutSpaceBases, TArray<FTransform>& OutBoneSpaceTransforms, FVector& OutRootBoneTranslation, FBlendedHeapCurve& OutCurve, UE::Anim::FMeshAttributeContainer& OutAttributes);
 
 	UE_DEPRECATED(4.26, "Please use PerformAnimationEvaluation with different signature")
 	void PerformAnimationEvaluation(const USkeletalMesh* InSkeletalMesh, UAnimInstance* InAnimInstance, TArray<FTransform>& OutSpaceBases, TArray<FTransform>& OutBoneSpaceTransforms, FVector& OutRootBoneTranslation, FBlendedHeapCurve& OutCurve);
 #endif
-	void PerformAnimationProcessing(const USkeletalMesh* InSkeletalMesh, UAnimInstance* InAnimInstance, bool bInDoEvaluation, TArray<FTransform>& OutSpaceBases, TArray<FTransform>& OutBoneSpaceTransforms, FVector& OutRootBoneTranslation, FBlendedHeapCurve& OutCurve, UE::Anim::FHeapAttributeContainer& OutAttributes);
+	void PerformAnimationProcessing(const USkeletalMesh* InSkeletalMesh, UAnimInstance* InAnimInstance, bool bInDoEvaluation, TArray<FTransform>& OutSpaceBases, TArray<FTransform>& OutBoneSpaceTransforms, FVector& OutRootBoneTranslation, FBlendedHeapCurve& OutCurve, UE::Anim::FMeshAttributeContainer& OutAttributes);
 
 	UE_DEPRECATED(4.26, "Please use PerformAnimationEvaluation with different signature")
 	void PerformAnimationProcessing(const USkeletalMesh* InSkeletalMesh, UAnimInstance* InAnimInstance, bool bInDoEvaluation, TArray<FTransform>& OutSpaceBases, TArray<FTransform>& OutBoneSpaceTransforms, FVector& OutRootBoneTranslation, FBlendedHeapCurve& OutCurve);
@@ -2374,6 +2374,9 @@ private:
 	// Finalizes pose to OutBoneSpaceTransforms
 	void FinalizePoseEvaluationResult(const USkeletalMesh* InMesh, TArray<FTransform>& OutBoneSpaceTransforms, FVector& OutRootBoneTranslation, FCompactPose& InFinalPose) const;
 
+	// Finalizes attributes (remapping from compact to mesh bone-indices)
+	void FinalizeAttributeEvaluationResults(const FBoneContainer& BoneContainer, const UE::Anim::FHeapAttributeContainer& FinalContainer, UE::Anim::FMeshAttributeContainer& OutContainer) const;
+
 	friend class FParallelBlendPhysicsTask;
 	
 	void BlendInPhysicsInternal(FTickFunction& ThisTickFunction);
@@ -2471,12 +2474,12 @@ public:
 
 private:
 	/** Temporary array of attributes that are active on this component - keeps same buffer index as SpaceBases - Please check SkinnedMeshComponent*/
-	UE::Anim::FHeapAttributeContainer AttributesArray[2];
+	UE::Anim::FMeshAttributeContainer AttributesArray[2];
 
-	UE::Anim::FHeapAttributeContainer& GetEditableCustomAttributes() { return AttributesArray[CurrentEditableComponentTransforms]; }
+	UE::Anim::FMeshAttributeContainer& GetEditableCustomAttributes() { return AttributesArray[CurrentEditableComponentTransforms]; }
 
 public:
-	const UE::Anim::FHeapAttributeContainer& GetCustomAttributes() const { return AttributesArray[CurrentReadComponentTransforms]; }
+	const UE::Anim::FMeshAttributeContainer& GetCustomAttributes() const { return AttributesArray[CurrentReadComponentTransforms]; }
 public:
 	/** Skeletal mesh component should not be able to have its mobility set to static */
 	virtual const bool CanHaveStaticMobility() const override { return false; }
diff --git a/Engine/Source/Runtime/Engine/Private/Animation/AnimSequence.cpp b/Engine/Source/Runtime/Engine/Private/Animation/AnimSequence.cpp
--- a/Engine/Source/Runtime/Engine/Private/Animation/AnimSequence.cpp
+++ b/Engine/Source/Runtime/Engine/Private/Animation/AnimSequence.cpp
@@ -4755,7 +4755,11 @@ void UAnimSequence::EvaluateAttributes(FAnimationPoseData& OutAnimationPoseData,
 		for (const FAnimatedBoneAttribute& Attribute : DataModel->GetAttributes())
 		{
 			const FCompactPoseBoneIndex PoseBoneIndex = RequiredBones.GetCompactPoseIndexFromSkeletonIndex(Attribute.Identifier.GetBoneIndex());
-			UE::Anim::Attributes::GetAttributeValue(OutAttributes, PoseBoneIndex, Attribute, ExtractionContext);
+			// Only add attribute if the bone its tied to exists in the currently evaluated set of bones
+			if(PoseBoneIndex.IsValid())
+			{
+				UE::Anim::Attributes::GetAttributeValue(OutAttributes, PoseBoneIndex, Attribute, ExtractionContext);
+			}
 		}
 	}
 	else
@@ -4764,10 +4768,13 @@ void UAnimSequence::EvaluateAttributes(FAnimationPoseData& OutAnimationPoseData,
 		for (const TPair<FAnimationAttributeIdentifier, FAttributeCurve>& BakedAttribute : AttributeCurves)
 		{
 			const FCompactPoseBoneIndex PoseBoneIndex = RequiredBones.GetCompactPoseIndexFromSkeletonIndex(BakedAttribute.Key.GetBoneIndex());
-			UE::Anim::FAttributeId Info(BakedAttribute.Key.GetName(), PoseBoneIndex);
-			uint8* AttributePtr = OutAttributes.FindOrAdd(BakedAttribute.Key.GetType(), Info);
-
-			BakedAttribute.Value.EvaluateToPtr(BakedAttribute.Key.GetType(), ExtractionContext.CurrentTime, AttributePtr);
+			// Only add attribute if the bone its tied to exists in the currently evaluated set of bones
+			if(PoseBoneIndex.IsValid())
+			{
+				UE::Anim::FAttributeId Info(BakedAttribute.Key.GetName(), PoseBoneIndex);
+				uint8* AttributePtr = OutAttributes.FindOrAdd(BakedAttribute.Key.GetType(), Info);
+				BakedAttribute.Value.EvaluateToPtr(BakedAttribute.Key.GetType(), ExtractionContext.CurrentTime, AttributePtr);
+			}
 		}
 	}
 }
diff --git a/Engine/Source/Runtime/Engine/Private/Animation/AttributesRuntime.cpp b/Engine/Source/Runtime/Engine/Private/Animation/AttributesRuntime.cpp
--- a/Engine/Source/Runtime/Engine/Private/Animation/AttributesRuntime.cpp
+++ b/Engine/Source/Runtime/Engine/Private/Animation/AttributesRuntime.cpp
@@ -155,7 +155,7 @@ void Attributes::ConvertToAdditive(const FStackAttributeContainer& BaseAttribute
 	}
 }
 
-void Attributes::CopyAndRemapAttributes(const FHeapAttributeContainer& SourceAttributes, FStackAttributeContainer& TargetAttributes, const TMap<int32, int32>& BoneMapToSource, const FBoneContainer& RequiredBones)
+void Attributes::CopyAndRemapAttributes(const FMeshAttributeContainer& SourceAttributes, FStackAttributeContainer& TargetAttributes, const TMap<int32, int32>& BoneMapToSource, const FBoneContainer& RequiredBones)
 {
 	QUICK_SCOPE_CYCLE_COUNTER(STAT_CopyAndRemapAttributes);
 
@@ -173,24 +173,26 @@ void Attributes::CopyAndRemapAttributes(const FHeapAttributeContainer& SourceAtt
 			{
 				const FAttributeId& AttributeId = AttributeIds[EntryIndex];
 
-				const FCompactPoseBoneIndex PoseBoneIndex = FCompactPoseBoneIndex(AttributeId.GetIndex());
-				const FMeshPoseBoneIndex MeshBoneIndex = RequiredBones.MakeMeshPoseIndex(PoseBoneIndex);
-				const int32* Value = BoneMapToSource.Find(MeshBoneIndex.GetInt());
-
+				const int32* Value = BoneMapToSource.Find(AttributeId.GetIndex());
 				// If there is no remapping the attribute will be dropped
 				if (Value)
 				{
-					const int32 RemappedBoneIndex = *Value;
-					const FAttributeId NewInfo(AttributeId.GetName(), RemappedBoneIndex, AttributeId.GetNamespace());
-					uint8* NewAttribute = TargetAttributes.FindOrAdd(ScriptStruct, NewInfo);
-					ScriptStruct->CopyScriptStruct(NewAttribute, SourceValues[EntryIndex].GetPtr<void>());
+					const FMeshPoseBoneIndex MeshPoseIndex(*Value);
+					const FSkeletonPoseBoneIndex RemappedPoseBoneIndex = RequiredBones.GetSkeletonPoseIndexFromMeshPoseIndex(MeshPoseIndex);
+					const FCompactPoseBoneIndex CompactPoseIndex = RequiredBones.GetCompactPoseIndexFromSkeletonPoseIndex(RemappedPoseBoneIndex);
+					if (CompactPoseIndex.IsValid())
+					{
+						const FAttributeId NewInfo(AttributeId.GetName(), CompactPoseIndex.GetInt(), AttributeId.GetNamespace());
+						uint8* NewAttribute = TargetAttributes.FindOrAdd(ScriptStruct, NewInfo);
+						ScriptStruct->CopyScriptStruct(NewAttribute, SourceValues[EntryIndex].GetPtr<void>());
+					}
 				}
 			}
 		}
 	}
 }
 
-void Attributes::InterpolateAttributes(FHeapAttributeContainer& FromAttributes, const FHeapAttributeContainer& ToAttributes, float Alpha)
+void Attributes::InterpolateAttributes(FMeshAttributeContainer& FromAttributes, const FMeshAttributeContainer& ToAttributes, float Alpha)
 {
 	QUICK_SCOPE_CYCLE_COUNTER(STAT_InterpolateAttributes);
 
diff --git a/Engine/Source/Runtime/Engine/Private/Components/SkeletalMeshComponent.cpp b/Engine/Source/Runtime/Engine/Private/Components/SkeletalMeshComponent.cpp
--- a/Engine/Source/Runtime/Engine/Private/Components/SkeletalMeshComponent.cpp
+++ b/Engine/Source/Runtime/Engine/Private/Components/SkeletalMeshComponent.cpp
@@ -478,6 +478,12 @@ void USkeletalMeshComponent::FinalizePoseEvaluationResult(const USkeletalMesh* I
 	OutRootBoneTranslation = OutBoneSpaceTransforms[0].GetTranslation() - InMesh->GetRefSkeleton().GetRefBonePose()[0].GetTranslation();
 }
 
+void USkeletalMeshComponent::FinalizeAttributeEvaluationResults(const FBoneContainer& BoneContainer, const UE::Anim::FHeapAttributeContainer& FinalContainer,
+	UE::Anim::FMeshAttributeContainer& OutContainer) const
+{
+	OutContainer.CopyFrom(FinalContainer, BoneContainer);
+}
+
 bool USkeletalMeshComponent::NeedToSpawnAnimScriptInstance() const
 {
 	IAnimClassInterface* AnimClassInterface = IAnimClassInterface::GetFromClass(AnimClass);
@@ -1983,20 +1989,20 @@ void USkeletalMeshComponent::UpdateSlaveComponent()
 
 #if WITH_EDITOR
 
-void USkeletalMeshComponent::PerformAnimationEvaluation(const USkeletalMesh* InSkeletalMesh, UAnimInstance* InAnimInstance, TArray<FTransform>& OutSpaceBases, TArray<FTransform>& OutBoneSpaceTransforms, FVector& OutRootBoneTranslation, FBlendedHeapCurve& OutCurve, UE::Anim::FHeapAttributeContainer& OutAttributes)
+void USkeletalMeshComponent::PerformAnimationEvaluation(const USkeletalMesh* InSkeletalMesh, UAnimInstance* InAnimInstance, TArray<FTransform>& OutSpaceBases, TArray<FTransform>& OutBoneSpaceTransforms, FVector& OutRootBoneTranslation, FBlendedHeapCurve& OutCurve, UE::Anim::FMeshAttributeContainer& OutAttributes)
 {
 	PerformAnimationProcessing(InSkeletalMesh, InAnimInstance, true, OutSpaceBases, OutBoneSpaceTransforms, OutRootBoneTranslation, OutCurve, OutAttributes);
 }
 
 void USkeletalMeshComponent::PerformAnimationEvaluation(const USkeletalMesh* InSkeletalMesh, UAnimInstance* InAnimInstance, TArray<FTransform>& OutSpaceBases, TArray<FTransform>& OutBoneSpaceTransforms, FVector& OutRootBoneTranslation, FBlendedHeapCurve& OutCurve)
 {
-	UE::Anim::FHeapAttributeContainer Attributes;
+	UE::Anim::FMeshAttributeContainer Attributes;
 	PerformAnimationEvaluation(InSkeletalMesh, InAnimInstance, OutSpaceBases, OutBoneSpaceTransforms, OutRootBoneTranslation, OutCurve, Attributes);
 }
 
 #endif
 
-void USkeletalMeshComponent::PerformAnimationProcessing(const USkeletalMesh* InSkeletalMesh, UAnimInstance* InAnimInstance, bool bInDoEvaluation, TArray<FTransform>& OutSpaceBases, TArray<FTransform>& OutBoneSpaceTransforms, FVector& OutRootBoneTranslation, FBlendedHeapCurve& OutCurve, UE::Anim::FHeapAttributeContainer& OutAttributes)
+void USkeletalMeshComponent::PerformAnimationProcessing(const USkeletalMesh* InSkeletalMesh, UAnimInstance* InAnimInstance, bool bInDoEvaluation, TArray<FTransform>& OutSpaceBases, TArray<FTransform>& OutBoneSpaceTransforms, FVector& OutRootBoneTranslation, FBlendedHeapCurve& OutCurve, UE::Anim::FMeshAttributeContainer& OutAttributes)
 {
 	CSV_SCOPED_TIMING_STAT(Animation, WorkerThreadTickTime);
 	ANIM_MT_SCOPE_CYCLE_COUNTER(PerformAnimEvaluation, !IsInGameThread());
@@ -2025,13 +2031,20 @@ void USkeletalMeshComponent::PerformAnimationProcessing(const USkeletalMesh* InS
 		FMemMark Mark(FMemStack::Get());
 		FCompactPose EvaluatedPose;
 
+		UE::Anim::FHeapAttributeContainer Attributes;		
+
 		// evaluate pure animations, and fill up BoneSpaceTransforms
-		EvaluateAnimation(InSkeletalMesh, InAnimInstance, OutRootBoneTranslation, OutCurve, EvaluatedPose, OutAttributes);
-		EvaluatePostProcessMeshInstance(OutBoneSpaceTransforms, EvaluatedPose, OutCurve, InSkeletalMesh, OutRootBoneTranslation, OutAttributes);
+		EvaluateAnimation(InSkeletalMesh, InAnimInstance, OutRootBoneTranslation, OutCurve, EvaluatedPose, Attributes);
+		EvaluatePostProcessMeshInstance(OutBoneSpaceTransforms, EvaluatedPose, OutCurve, InSkeletalMesh, OutRootBoneTranslation, Attributes);
 
 		// Finalize the transforms from the evaluation
 		FinalizePoseEvaluationResult(InSkeletalMesh, OutBoneSpaceTransforms, OutRootBoneTranslation, EvaluatedPose);
 
+		if (EvaluatedPose.IsValid())
+		{
+			FinalizeAttributeEvaluationResults(EvaluatedPose.GetBoneContainer(), Attributes, OutAttributes);
+		}
+
 		// Fill SpaceBases from LocalAtoms
 		FillComponentSpaceTransforms(InSkeletalMesh, OutBoneSpaceTransforms, OutSpaceBases);
 	}
@@ -2040,7 +2053,7 @@ void USkeletalMeshComponent::PerformAnimationProcessing(const USkeletalMesh* InS
 
 void USkeletalMeshComponent::PerformAnimationProcessing(const USkeletalMesh* InSkeletalMesh, UAnimInstance* InAnimInstance, bool bInDoEvaluation, TArray<FTransform>& OutSpaceBases, TArray<FTransform>& OutBoneSpaceTransforms, FVector& OutRootBoneTranslation, FBlendedHeapCurve& OutCurve)
 {
-	UE::Anim::FHeapAttributeContainer Attributes;	
+	UE::Anim::FMeshAttributeContainer Attributes;	
 	PerformAnimationProcessing(InSkeletalMesh, InAnimInstance, bInDoEvaluation, OutSpaceBases, OutBoneSpaceTransforms, OutRootBoneTranslation, OutCurve, Attributes);
 }
 
@@ -4476,7 +4489,7 @@ bool USkeletalMeshComponent::FindAttributeChecked(const FName& BoneName, const F
 
 	if (SkeletalMesh)
 	{
-		const UE::Anim::FHeapAttributeContainer& Attributes = GetCustomAttributes();
+		const UE::Anim::FMeshAttributeContainer& Attributes = GetCustomAttributes();
 		const int32 BoneIndex = SkeletalMesh->GetRefSkeleton().FindBoneIndex(BoneName);
 
 		const CustomAttributeType* AttributePtr = Attributes.Find<CustomAttributeType>(UE::Anim::FAttributeId(AttributeName, FCompactPoseBoneIndex(BoneIndex)));
diff --git a/Engine/Source/Runtime/Engine/Public/Animation/AttributesRuntime.h b/Engine/Source/Runtime/Engine/Public/Animation/AttributesRuntime.h
--- a/Engine/Source/Runtime/Engine/Public/Animation/AttributesRuntime.h
+++ b/Engine/Source/Runtime/Engine/Public/Animation/AttributesRuntime.h
@@ -39,6 +39,7 @@ namespace UE
 	{
 		struct ENGINE_API FStackAttributeContainer : public TAttributeContainer<FCompactPoseBoneIndex, FAnimStackAllocator> {};
 		struct ENGINE_API FHeapAttributeContainer : public TAttributeContainer<FCompactPoseBoneIndex, FDefaultAllocator> {};
+		struct ENGINE_API FMeshAttributeContainer : public TAttributeContainer<FMeshPoseBoneIndex, FDefaultAllocator> {};
 
 		/** Accessor for internal data of the TAttributeContainer.  The TAttributeContainer methods are preferred - this accessor should only be used in limited situations */
 		template<class BoneIndexType, typename InAllocator>
@@ -89,10 +90,10 @@ namespace UE
 			static void ConvertToAdditive(const FStackAttributeContainer& BaseAttributes, FStackAttributeContainer& OutAdditiveAttributes);
 
 			/** Copy attributes from source, and remap the bone indices according to BoneMapToSource */
-			static void CopyAndRemapAttributes(const FHeapAttributeContainer& SourceAttributes, FStackAttributeContainer& OutAttributes, const TMap<int32, int32>& BoneMapToSource, const FBoneContainer& RequiredBones);
+			static void CopyAndRemapAttributes(const FMeshAttributeContainer& SourceAttributes, FStackAttributeContainer& OutAttributes, const TMap<int32, int32>& BoneMapToSource, const FBoneContainer& RequiredBones);
 
 			/** Interpolates between two sets of attributes */
-			static void InterpolateAttributes(FHeapAttributeContainer& FromAttributes, const FHeapAttributeContainer& ToAttributes, float Alpha);
+			static void InterpolateAttributes(FMeshAttributeContainer& FromAttributes, const FMeshAttributeContainer& ToAttributes, float Alpha);
 
 			/** Mirror (swap) attributes with the specified MirrorDataTable. Attributes are swapped using the bone mapping such that bones which are mirrored swap attributes */
 			static void MirrorAttributes(FStackAttributeContainer& Attributes, const UMirrorDataTable& MirrorDataTable);
