Animation attributes - normalization trait
- Added normalization trait/concept for animation attributes - which allows for normalization functionality after accumulation/blending operations
- Implemented normalization for Transform animation attribute
- Replaced override and accumulate with blend for StateMachine node
- Ensure that additive control rig attributes are actually made additive

diff --git a/Engine/Plugins/Animation/ControlRig/Source/ControlRig/Private/AnimNode_ControlRigBase.cpp b/Engine/Plugins/Animation/ControlRig/Source/ControlRig/Private/AnimNode_ControlRigBase.cpp
--- a/Engine/Plugins/Animation/ControlRig/Source/ControlRig/Private/AnimNode_ControlRigBase.cpp
+++ b/Engine/Plugins/Animation/ControlRig/Source/ControlRig/Private/AnimNode_ControlRigBase.cpp
@@ -442,6 +442,8 @@ void FAnimNode_ControlRigBase::Evaluate_AnyThread(FPoseContext& Output)
 			AdditivePose.Curve.ConvertToAdditive(SourcePose.Curve);
 			Output = SourcePose;
 
+			UE::Anim::Attributes::ConvertToAdditive(SourcePose.CustomAttributes, AdditivePose.CustomAttributes);
+
 			FAnimationPoseData BaseAnimationPoseData(Output);
 			const FAnimationPoseData AdditiveAnimationPoseData(AdditivePose);
 			FAnimationRuntime::AccumulateAdditivePose(BaseAnimationPoseData, AdditiveAnimationPoseData, InternalBlendAlpha, AAT_LocalSpaceBase);
diff --git a/Engine/Source/Runtime/Engine/Classes/Animation/AttributeBlendOperator.h b/Engine/Source/Runtime/Engine/Classes/Animation/AttributeBlendOperator.h
--- a/Engine/Source/Runtime/Engine/Classes/Animation/AttributeBlendOperator.h
+++ b/Engine/Source/Runtime/Engine/Classes/Animation/AttributeBlendOperator.h
@@ -65,6 +65,11 @@ namespace UE
 								OutAttribute.Accumulate(Attribute, AttributeWeight, It.GetAdditiveType());
 							}
 						}
+
+						if constexpr (UE::Anim::TAttributeTypeTraits<AttributeType>::RequiresNormalization)
+						{
+							OutAttribute.Normalize();
+						}
 					}
 				});
 
@@ -89,6 +94,11 @@ namespace UE
 						{
 							OutAttribute = Attribute.Multiply(AttributeWeight);
 						}
+
+						if constexpr (UE::Anim::TAttributeTypeTraits<AttributeType>::RequiresNormalization)
+						{
+							OutAttribute.Normalize();
+						}
 					}
 				});				
 			}
@@ -157,6 +167,11 @@ namespace UE
 								}
 							}
 						}
+						
+						if constexpr (UE::Anim::TAttributeTypeTraits<AttributeType>::RequiresNormalization)
+						{
+							OutAttribute.Normalize();
+						}
 					}
 				});
 
@@ -188,6 +203,11 @@ namespace UE
 								OutAttribute.Accumulate(Attribute, AttributeWeight, It.GetAdditiveType());
 							}
 						}
+
+						if constexpr (UE::Anim::TAttributeTypeTraits<AttributeType>::RequiresNormalization)
+						{
+							OutAttribute.Normalize();
+						}
 					}
 				});
 			}
@@ -244,6 +264,11 @@ namespace UE
 
 						const float AttributeWeight = It.GetWeight();
 						OutAttribute = Attribute.Multiply(AttributeWeight);
+
+						if constexpr (UE::Anim::TAttributeTypeTraits<AttributeType>::RequiresNormalization)
+						{
+							OutAttribute.Normalize();
+						}
 					}
 				});
 			}
@@ -288,7 +313,12 @@ namespace UE
 						}
 						else
 						{
-							OutAttribute = Attribute.Multiply(AttributeWeight);
+							OutAttribute = Attribute.Multiply(AttributeWeight);							
+						}
+
+						if constexpr (UE::Anim::TAttributeTypeTraits<AttributeType>::RequiresNormalization)
+						{
+							OutAttribute.Normalize();
 						}
 					}
 				});
@@ -333,6 +363,11 @@ namespace UE
 						AttributeType& OutAttribute = *OutAttributePtr;
 
 						OutAttribute.MakeAdditive(Attribute);
+
+						if constexpr (UE::Anim::TAttributeTypeTraits<AttributeType>::RequiresNormalization)
+						{
+							OutAttribute.Normalize();
+						}
 					}
 				});
 			}
@@ -371,6 +406,11 @@ namespace UE
 
 				Output = TypedFrom;
 				Output.Interpolate(TypedTo, Alpha);
+
+				if constexpr (UE::Anim::TAttributeTypeTraits<AttributeType>::RequiresNormalization)
+				{
+					Output.Normalize();
+				}
 			}
 
 			/* Interpolate operation for non-blendable; step interpolated attribute types */
diff --git a/Engine/Source/Runtime/Engine/Classes/Animation/BuiltInAttributeTypes.h b/Engine/Source/Runtime/Engine/Classes/Animation/BuiltInAttributeTypes.h
--- a/Engine/Source/Runtime/Engine/Classes/Animation/BuiltInAttributeTypes.h
+++ b/Engine/Source/Runtime/Engine/Classes/Animation/BuiltInAttributeTypes.h
@@ -124,6 +124,11 @@ struct FTransformAnimationAttribute
 	{
 		FAnimationRuntime::ConvertTransformToAdditive(Value, BaseAttribute.Value);
 	}
+
+	void Normalize()
+	{
+		Value.NormalizeRotation();
+	}
 	
 	FTransformAnimationAttribute Multiply(const float Weight) const
 	{
@@ -183,6 +188,16 @@ namespace UE
 			};
 		};
 
+		/** Transform attribute requires normalization */
+		template<>
+		struct TAttributeTypeTraits<FTransformAnimationAttribute> : public TAttributeTypeTraitsBase<FTransformAnimationAttribute>
+		{
+			enum
+			{
+				RequiresNormalization = true,
+			};
+		};
+
 		/** Non blendable types*/
 		template<>
 		struct TAttributeTypeTraits<FNonBlendableTransformAnimationAttribute> : public TAttributeTypeTraitsBase<FNonBlendableTransformAnimationAttribute>
diff --git a/Engine/Source/Runtime/Engine/Private/Animation/AnimNode_StateMachine.cpp b/Engine/Source/Runtime/Engine/Private/Animation/AnimNode_StateMachine.cpp
--- a/Engine/Source/Runtime/Engine/Private/Animation/AnimNode_StateMachine.cpp
+++ b/Engine/Source/Runtime/Engine/Private/Animation/AnimNode_StateMachine.cpp
@@ -923,9 +923,8 @@ void FAnimNode_StateMachine::EvaluateTransitionStandardBlendInternal(FPoseContex
 	// blend curve in
 	Output.Curve.Override(PreviousStateResult.Curve, 1.0 - Transition.Alpha);
 	Output.Curve.Accumulate(NextStateResult.Curve, Transition.Alpha);
-
-	UE::Anim::Attributes::OverrideAttributes(PreviousStateResult.CustomAttributes, Output.CustomAttributes, 1.0 - Transition.Alpha);
-	UE::Anim::Attributes::AccumulateAttributes(NextStateResult.CustomAttributes, Output.CustomAttributes, Transition.Alpha, AAT_None);
+		
+	UE::Anim::Attributes::BlendAttributes({ PreviousStateResult.CustomAttributes, NextStateResult.CustomAttributes }, { 1.0f - Transition.Alpha, Transition.Alpha }, { 0, 1 }, Output.CustomAttributes);
 }
 
 void FAnimNode_StateMachine::EvaluateTransitionCustomBlend(FPoseContext& Output, FAnimationActiveTransitionEntry& Transition, bool bIntermediatePoseIsValid)
diff --git a/Engine/Source/Runtime/Engine/Public/Animation/AttributeTraits.h b/Engine/Source/Runtime/Engine/Public/Animation/AttributeTraits.h
--- a/Engine/Source/Runtime/Engine/Public/Animation/AttributeTraits.h
+++ b/Engine/Source/Runtime/Engine/Public/Animation/AttributeTraits.h
@@ -19,7 +19,9 @@ namespace UE
 				/** Determines whether or not the type has an associated user-defined implementation of IAttributeBlendOperator. False by default */
 				WithCustomBlendOperator = false,
 				/** Determines whether or not the type should be step-interpolated rather than linearly. False by default */
-				StepInterpolate = false
+				StepInterpolate = false,
+				/** Determines whether or not the type should be normalized after any accumulation / multiplication has been applied. False by default */
+				RequiresNormalization = false
 			};
 		};
 
diff --git a/Engine/Source/Runtime/Engine/Public/Animation/AttributeTypes.h b/Engine/Source/Runtime/Engine/Public/Animation/AttributeTypes.h
--- a/Engine/Source/Runtime/Engine/Public/Animation/AttributeTypes.h
+++ b/Engine/Source/Runtime/Engine/Public/Animation/AttributeTypes.h
@@ -31,6 +31,15 @@ namespace UE
 			);
 		};
 
+		/** Concept used to verify a user-defined attribute type with its TAttributeTypeTraits::RequiresNormalization value set to true */
+		struct CNormalizedAttribute
+		{
+			template <typename T>
+			auto Requires(T& Val) -> decltype(
+				Val.Normalize()
+			);
+		};
+
 		struct ENGINE_API AttributeTypes
 		{
 		protected:
@@ -69,6 +78,11 @@ namespace UE
 				if constexpr (UE::Anim::TAttributeTypeTraits<AttributeType>::IsBlendable)	
 				{
 					static_assert(TModels<CBlendableAttribute, AttributeType>::Value, "Missing function implementations required for Attribute blending");
+
+					if  constexpr (UE::Anim::TAttributeTypeTraits<AttributeType>::RequiresNormalization)
+					{
+						static_assert(TModels<CNormalizedAttribute, AttributeType>::Value, "Missing function implementations required for Attribute normalization");
+					}
 					
 					if constexpr (!UE::Anim::TAttributeTypeTraits<AttributeType>::StepInterpolate)
 					{
