The "HighResShot" console command will no longer remember the last filename= parameter passed to it and overwrite it on subsequent calls without the parameter.

diff --git a/Engine/Source/Runtime/Engine/Private/UnrealClient.cpp b/Engine/Source/Runtime/Engine/Private/UnrealClient.cpp
--- a/Engine/Source/Runtime/Engine/Private/UnrealClient.cpp
+++ b/Engine/Source/Runtime/Engine/Private/UnrealClient.cpp
@@ -2218,7 +2218,10 @@ ENGINE_API bool GetHighResScreenShotInput(const TCHAR* Cmd, FOutputDevice& Ar, u
 	const FString FilenameSearchString = TEXT("filename=");
 
 	// FParse::Value has better handling of escape characters than FParse::Token
-	FParse::Value(Cmd, *FilenameSearchString, OutFilenameOverride);
+	if (!FParse::Value(Cmd, *FilenameSearchString, OutFilenameOverride))
+	{
+		OutFilenameOverride.Reset();
+	}
 
 	FString Arg;
 	while (FParse::Token(Cmd, Arg, true))
