PR #8796: Implemented Key parameter for PrintString and PrintText nodes (Contributed by [REDACTED_PERSON])

diff --git a/Engine/Source/Runtime/Engine/Classes/Kismet/KismetSystemLibrary.h b/Engine/Source/Runtime/Engine/Classes/Kismet/KismetSystemLibrary.h
--- a/Engine/Source/Runtime/Engine/Classes/Kismet/KismetSystemLibrary.h
+++ b/Engine/Source/Runtime/Engine/Classes/Kismet/KismetSystemLibrary.h
@@ -417,9 +417,10 @@ class ENGINE_API UKismetSystemLibrary : public UBlueprintFunctionLibrary
 	 * @param	bPrintToConsole	Whether or not to print the output to the console
 	 * @param	TextColor		The color of the text to display
 	 * @param	Duration		The display duration (if Print to Screen is True). Using negative number will result in loading the duration time from the config.
+	 * @param	Key				If a non-empty key is provided, the message will replace any existing messages with the same key.
 	 */
 	UFUNCTION(BlueprintCallable, meta=(WorldContext="WorldContextObject", CallableWithoutWorldContext, Keywords = "log print", AdvancedDisplay = "2", DevelopmentOnly), Category="Development")
-	static void PrintString(const UObject* WorldContextObject, const FString& InString = FString(TEXT("Hello")), bool bPrintToScreen = true, bool bPrintToLog = true, FLinearColor TextColor = FLinearColor(0.0, 0.66, 1.0), float Duration = 2.f);
+	static void PrintString(const UObject* WorldContextObject, const FString& InString = FString(TEXT("Hello")), bool bPrintToScreen = true, bool bPrintToLog = true, FLinearColor TextColor = FLinearColor(0.0, 0.66, 1.0), float Duration = 2.f, const FName Key = NAME_None);
 
 	/**
 	 * Prints text to the log, and optionally, to the screen
@@ -431,9 +432,10 @@ class ENGINE_API UKismetSystemLibrary : public UBlueprintFunctionLibrary
 	 * @param	bPrintToConsole	Whether or not to print the output to the console
 	 * @param	TextColor		The color of the text to display
 	 * @param	Duration		The display duration (if Print to Screen is True). Using negative number will result in loading the duration time from the config.
+	 * @param	Key				If a non-empty key is provided, the message will replace any existing messages with the same key.
 	 */
 	UFUNCTION(BlueprintCallable, meta=(WorldContext="WorldContextObject", CallableWithoutWorldContext, Keywords = "log", AdvancedDisplay = "2", DevelopmentOnly), Category="Development")
-	static void PrintText(const UObject* WorldContextObject, const FText InText = INVTEXT("Hello"), bool bPrintToScreen = true, bool bPrintToLog = true, FLinearColor TextColor = FLinearColor(0.0, 0.66, 1.0), float Duration = 2.f);
+	static void PrintText(const UObject* WorldContextObject, const FText InText = INVTEXT("Hello"), bool bPrintToScreen = true, bool bPrintToLog = true, FLinearColor TextColor = FLinearColor(0.0, 0.66, 1.0), float Duration = 2.f, const FName Key = NAME_None);
 
 	/**
 	 * Prints a warning string to the log and the screen. Meant to be used as a way to inform the user that they misused the node.
diff --git a/Engine/Source/Runtime/Engine/Private/KismetSystemLibrary.cpp b/Engine/Source/Runtime/Engine/Private/KismetSystemLibrary.cpp
--- a/Engine/Source/Runtime/Engine/Private/KismetSystemLibrary.cpp
+++ b/Engine/Source/Runtime/Engine/Private/KismetSystemLibrary.cpp
@@ -282,7 +282,7 @@ void UKismetSystemLibrary::LogString(const FString& InString, bool bPrintToLog)
 #endif
 }
 
-void UKismetSystemLibrary::PrintString(const UObject* WorldContextObject, const FString& InString, bool bPrintToScreen, bool bPrintToLog, FLinearColor TextColor, float Duration)
+void UKismetSystemLibrary::PrintString(const UObject* WorldContextObject, const FString& InString, bool bPrintToScreen, bool bPrintToLog, FLinearColor TextColor, float Duration, const FName Key)
 {
 #if !(UE_BUILD_SHIPPING || UE_BUILD_TEST) // Do not Print in Shipping or Test
 
@@ -344,7 +344,12 @@ void UKismetSystemLibrary::PrintString(const UObject* WorldContextObject, const
 			{
 				GConfig->GetFloat( TEXT("Kismet"), TEXT("PrintStringDuration"), Duration, GEngineIni );
 			}
-			GEngine->AddOnScreenDebugMessage((uint64)-1, Duration, TextColor.ToFColor(true), FinalDisplayString);
+			uint64 InnerKey = -1;
+			if (Key != NAME_None)
+			{
+				InnerKey = GetTypeHash(Key);
+			}
+			GEngine->AddOnScreenDebugMessage(InnerKey, Duration, TextColor.ToFColor(true), FinalDisplayString);
 		}
 		else
 		{
@@ -354,10 +359,10 @@ void UKismetSystemLibrary::PrintString(const UObject* WorldContextObject, const
 #endif
 }
 
-void UKismetSystemLibrary::PrintText(const UObject* WorldContextObject, const FText InText, bool bPrintToScreen, bool bPrintToLog, FLinearColor TextColor, float Duration)
+void UKismetSystemLibrary::PrintText(const UObject* WorldContextObject, const FText InText, bool bPrintToScreen, bool bPrintToLog, FLinearColor TextColor, float Duration, const FName Key)
 {
 #if !(UE_BUILD_SHIPPING || UE_BUILD_TEST) // Do not Print in Shipping or Test
-	PrintString(WorldContextObject, InText.ToString(), bPrintToScreen, bPrintToLog, TextColor, Duration);
+	PrintString(WorldContextObject, InText.ToString(), bPrintToScreen, bPrintToLog, TextColor, Duration, Key);
 #endif
 }
 
