Added tooltips to Lumen visualization modes

diff --git a/Engine/Source/Runtime/Engine/Private/LumenVisualizationData.cpp b/Engine/Source/Runtime/Engine/Private/LumenVisualizationData.cpp
--- a/Engine/Source/Runtime/Engine/Private/LumenVisualizationData.cpp
+++ b/Engine/Source/Runtime/Engine/Private/LumenVisualizationData.cpp
@@ -10,20 +10,44 @@
 static FLumenVisualizationData GLumenVisualizationData;
 
 // Must match appropriate values in r.Lumen.Visualize.Mode
-#define VISUALIZE_LUMEN_SCENE		1
-#define VISUALIZE_REFLECTION_VIEW	2
-#define VISUALIZE_SURFACE_CACHE		3
-#define VISUALIZE_OVERVIEW			4
+#define LUMEN_VISUALIZE_LUMEN_SCENE		1
+#define LUMEN_VISUALIZE_REFLECTION_VIEW	2
+#define LUMEN_VISUALIZE_SURFACE_CACHE	3
+#define LUMEN_VISUALIZE_OVERVIEW		4
 
 void FLumenVisualizationData::Initialize()
 {
 	if (!bIsInitialized)
 	{
-		AddVisualizationMode(TEXT("Overview"), LOCTEXT("Overview", "Overview"), FModeType::Overview, VISUALIZE_OVERVIEW, true);
-
-		AddVisualizationMode(TEXT("LumenScene"), LOCTEXT("LumenScene", "Lumen Scene"), FModeType::Standard, VISUALIZE_LUMEN_SCENE, true);
-		AddVisualizationMode(TEXT("ReflectionView"), LOCTEXT("ReflectionView", "Reflection View"), FModeType::Standard, VISUALIZE_REFLECTION_VIEW, true);
-		AddVisualizationMode(TEXT("SurfaceCache"), LOCTEXT("SurfaceCache", "Surface Cache"), FModeType::Standard, VISUALIZE_SURFACE_CACHE, true);
+		AddVisualizationMode(
+			TEXT("Overview"),
+			LOCTEXT("Overview", "Overview"),
+			LOCTEXT("OverviewDesc", "All Lumen view mode tiles overlayed on top"),
+			FModeType::Overview,
+			LUMEN_VISUALIZE_OVERVIEW,
+			true);
+
+		AddVisualizationMode(
+			TEXT("LumenScene"),
+			LOCTEXT("LumenScene", "Lumen Scene"),
+			LOCTEXT("LumenSceneDesc", "Visualizes Lumen scene representation in a highest possible quality and view distance"),
+			FModeType::Standard,
+			LUMEN_VISUALIZE_LUMEN_SCENE,
+			true);
+
+		AddVisualizationMode(
+			TEXT("ReflectionView"),
+			LOCTEXT("ReflectionView", "Reflection View"),
+			LOCTEXT("ReflectionViewDesc", "Visualizes Lumen scene representation with current reflection settings"),
+			FModeType::Standard, LUMEN_VISUALIZE_REFLECTION_VIEW, true);
+
+		AddVisualizationMode(
+			TEXT("SurfaceCache"),
+			LOCTEXT("SurfaceCache", "Surface Cache"),
+			LOCTEXT("SurfaceCacheDesc", "Visualizes Lumen surface cache and marks with pink missing surface cache coverage"),
+			FModeType::Standard,
+			LUMEN_VISUALIZE_SURFACE_CACHE,
+			true);
 
 		ConfigureConsoleCommand();
 
@@ -54,6 +78,7 @@ void FLumenVisualizationData::ConfigureConsoleCommand()
 void FLumenVisualizationData::AddVisualizationMode(
 	const TCHAR* ModeString,
 	const FText& ModeText,
+	const FText& ModeDesc,
 	const FModeType ModeType,
 	int32 ModeID,
 	bool DefaultComposited
@@ -65,7 +90,7 @@ void FLumenVisualizationData::AddVisualizationMode(
 	Record.ModeString			= FString(ModeString);
 	Record.ModeName				= ModeName;
 	Record.ModeText				= ModeText;
-	Record.ModeDesc				= FText::GetEmpty();
+	Record.ModeDesc				= ModeDesc;
 	Record.ModeType				= ModeType;
 	Record.ModeID				= ModeID;
 	Record.DefaultComposited	= DefaultComposited;
diff --git a/Engine/Source/Runtime/Engine/Public/LumenVisualizationData.h b/Engine/Source/Runtime/Engine/Public/LumenVisualizationData.h
--- a/Engine/Source/Runtime/Engine/Public/LumenVisualizationData.h
+++ b/Engine/Source/Runtime/Engine/Public/LumenVisualizationData.h
@@ -70,6 +70,7 @@ private:
 	void AddVisualizationMode(
 		const TCHAR* ModeString,
 		const FText& ModeText,
+		const FText& ModeDesc,
 		const FModeType ModeType,
 		int32 ModeID,
 		bool DefaultComposited
