Physx: cooking with debug editor results in physics interface failure

Calling PxRigidDynamic::putToSleep returns this error message: "Body must be in a scene."
Our calling code doesn't always gate PutToSleep with a check for IsInScene. Adding this allow cooking to succeed.

Original repro: VRTemplate on 4.27 with PhyX enabled

Original review: 18536624

diff --git a/Engine/Source/Runtime/Engine/Private/PhysicsEngine/BodyInstance.cpp b/Engine/Source/Runtime/Engine/Private/PhysicsEngine/BodyInstance.cpp
--- a/Engine/Source/Runtime/Engine/Private/PhysicsEngine/BodyInstance.cpp
+++ b/Engine/Source/Runtime/Engine/Private/PhysicsEngine/BodyInstance.cpp
@@ -4262,7 +4262,7 @@ void FBodyInstance::InitDynamicProperties_AssumesLocked()
 		FPhysicsInterface::SetSolverVelocityIterationCount_AssumesLocked(ActorHandle, VelocityIterCount);
 
 		CreateDOFLock();
-		if(!IsRigidBodyKinematic_AssumesLocked(ActorHandle))
+		if(FPhysicsInterface::IsInScene(ActorHandle) && !IsRigidBodyKinematic_AssumesLocked(ActorHandle))
 		{
 			if(!bStartAwake && !bComponentAwake)
 			{
