Added NativeThreadSafeUpdateAnimation to UAnimInstance. This provides an override point for a native C++ implementation of the worker thread animation update (BlueprintThreadSafeUpdateAnimation).

diff --git a/Engine/Source/Runtime/Engine/Classes/Animation/AnimInstance.h b/Engine/Source/Runtime/Engine/Classes/Animation/AnimInstance.h
--- a/Engine/Source/Runtime/Engine/Classes/Animation/AnimInstance.h
+++ b/Engine/Source/Runtime/Engine/Classes/Animation/AnimInstance.h
@@ -1279,8 +1279,11 @@ public:
 	// Native initialization override point
 	virtual void NativeInitializeAnimation();
 	// Native update override point. It is usually a good idea to simply gather data in this step and 
-	// for the bulk of the work to be done in NativeUpdateAnimation.
+	// for the bulk of the work to be done in NativeThreadSafeUpdateAnimation.
 	virtual void NativeUpdateAnimation(float DeltaSeconds);
+	// Native thread safe update override point. Executed on a worker thread just prior to graph update 
+	// for linked anim instances, only called when the hosting node(s) are relevant
+	virtual void NativeThreadSafeUpdateAnimation(float DeltaSeconds);
 	// Native update override point. Can be called from a worker thread. This is a good place to do any
 	// heavy lifting (as opposed to NativeUpdateAnimation_GameThread()).
 	// This function should not be used. Worker thread updates should be performed in the FAnimInstanceProxy attached to this instance.
diff --git a/Engine/Source/Runtime/Engine/Private/Animation/AnimInstance.cpp b/Engine/Source/Runtime/Engine/Private/Animation/AnimInstance.cpp
--- a/Engine/Source/Runtime/Engine/Private/Animation/AnimInstance.cpp
+++ b/Engine/Source/Runtime/Engine/Private/Animation/AnimInstance.cpp
@@ -58,6 +58,7 @@ DEFINE_STAT(STAT_PostUpdateAnimation);
 DEFINE_STAT(STAT_BlueprintUpdateAnimation);
 DEFINE_STAT(STAT_BlueprintPostEvaluateAnimation);
 DEFINE_STAT(STAT_NativeUpdateAnimation);
+DEFINE_STAT(STAT_NativeThreadSafeUpdateAnimation);
 DEFINE_STAT(STAT_Montage_Advance);
 DEFINE_STAT(STAT_Montage_UpdateWeight);
 DEFINE_STAT(STAT_UpdateCurves);
@@ -790,6 +791,10 @@ void UAnimInstance::NativeUpdateAnimation(float DeltaSeconds)
 {
 }
 
+void UAnimInstance::NativeThreadSafeUpdateAnimation(float DeltaSeconds)
+{
+}
+
 void UAnimInstance::NativeUpdateAnimation_WorkerThread(float DeltaSeconds)
 {
 }
diff --git a/Engine/Source/Runtime/Engine/Private/Animation/AnimInstanceProxy.cpp b/Engine/Source/Runtime/Engine/Private/Animation/AnimInstanceProxy.cpp
--- a/Engine/Source/Runtime/Engine/Private/Animation/AnimInstanceProxy.cpp
+++ b/Engine/Source/Runtime/Engine/Private/Animation/AnimInstanceProxy.cpp
@@ -1112,7 +1112,12 @@ void FAnimInstanceProxy::UpdateAnimation_WithRoot(const FAnimationUpdateContext&
 					return EAnimSubsystemEnumeration::Continue;
 				});
 			}
-			
+
+			{
+				SCOPE_CYCLE_COUNTER(STAT_NativeThreadSafeUpdateAnimation);
+				CastChecked<UAnimInstance>(GetAnimInstanceObject())->NativeThreadSafeUpdateAnimation(CurrentDeltaSeconds);
+			}
+
 			{
 				SCOPE_CYCLE_COUNTER(STAT_BlueprintUpdateAnimation);
 				CastChecked<UAnimInstance>(GetAnimInstanceObject())->BlueprintThreadSafeUpdateAnimation(CurrentDeltaSeconds);
diff --git a/Engine/Source/Runtime/Engine/Public/Animation/AnimStats.h b/Engine/Source/Runtime/Engine/Public/Animation/AnimStats.h
--- a/Engine/Source/Runtime/Engine/Public/Animation/AnimStats.h
+++ b/Engine/Source/Runtime/Engine/Public/Animation/AnimStats.h
@@ -35,6 +35,7 @@ DECLARE_CYCLE_STAT_EXTERN(TEXT("UpdateAnimation"), STAT_UpdateAnimation, STATGRO
 DECLARE_CYCLE_STAT_EXTERN(TEXT("PostUpdateAnimation"), STAT_PostUpdateAnimation, STATGROUP_Anim, );
 DECLARE_CYCLE_STAT_EXTERN(TEXT("BlueprintUpdateAnimation"), STAT_BlueprintUpdateAnimation, STATGROUP_Anim, );
 DECLARE_CYCLE_STAT_EXTERN(TEXT("NativeUpdateAnimation"), STAT_NativeUpdateAnimation, STATGROUP_Anim, );
+DECLARE_CYCLE_STAT_EXTERN(TEXT("NativeThreadSafeUpdateAnimation"), STAT_NativeThreadSafeUpdateAnimation, STATGROUP_Anim, );
 DECLARE_CYCLE_STAT_EXTERN(TEXT("BlueprintPostEvaluateAnimation"), STAT_BlueprintPostEvaluateAnimation, STATGROUP_Anim, );
 DECLARE_CYCLE_STAT_EXTERN(TEXT("MontageAdvance"), STAT_Montage_Advance, STATGROUP_Anim, );
 DECLARE_CYCLE_STAT_EXTERN(TEXT("MontageUpdateWeight"), STAT_Montage_UpdateWeight, STATGROUP_Anim, );
