Set defaults for Lumen ShowFlags to false.
This fixes an issue where PIE sessions would default to having Lumen Visualization enabled, which is a regression in behavior, causing confusion for developers.
Related to CL-18658215

diff --git a/Engine/Source/Runtime/Engine/Public/ShowFlags.h b/Engine/Source/Runtime/Engine/Public/ShowFlags.h
--- a/Engine/Source/Runtime/Engine/Public/ShowFlags.h
+++ b/Engine/Source/Runtime/Engine/Public/ShowFlags.h
@@ -373,6 +373,7 @@ private:
 		SetReflectionOverride(false);
 		SetVisualizeBuffer(false);
 		SetVisualizeNanite(false);
+		SetVisualizeLumen(false);
 		SetVectorFields(false);
 		SetGBufferHints(false);
 		SetCompositeEditorPrimitives(InitMode == ESFIM_Editor || InitMode == ESFIM_VREditing);
@@ -471,6 +472,14 @@ private:
 		SetVisualizeStrataMaterial(false);
 		SetDrawOnlyVSMInvalidatingGeo(false);
 		SetSingleLayerWaterRefractionFullPrecision(false);
+
+		SetVisualizeLumen(false);
+		SetLumenScreenTraces(false);
+		SetLumenDetailTraces(false);
+		SetLumenGlobalTraces(false);
+		SetLumenFarFieldTraces(false);
+		SetLumenSecondaryBounces(false);
+		SetLumenScreenSpaceDirectionalOcclusion(false);
 	}
 
 
