Fixed Lumen default show flag values (all tracing methods should be enabled by default)

diff --git a/Engine/Source/Runtime/Engine/Public/ShowFlags.h b/Engine/Source/Runtime/Engine/Public/ShowFlags.h
--- a/Engine/Source/Runtime/Engine/Public/ShowFlags.h
+++ b/Engine/Source/Runtime/Engine/Public/ShowFlags.h
@@ -473,13 +473,12 @@ private:
 		SetDrawOnlyVSMInvalidatingGeo(false);
 		SetSingleLayerWaterRefractionFullPrecision(false);
 
-		SetVisualizeLumen(false);
-		SetLumenScreenTraces(false);
-		SetLumenDetailTraces(false);
-		SetLumenGlobalTraces(false);
-		SetLumenFarFieldTraces(false);
-		SetLumenSecondaryBounces(false);
-		SetLumenScreenSpaceDirectionalOcclusion(false);
+		SetLumenScreenTraces(true);
+		SetLumenDetailTraces(true);
+		SetLumenGlobalTraces(true);
+		SetLumenFarFieldTraces(true);
+		SetLumenSecondaryBounces(true);
+		SetLumenScreenSpaceDirectionalOcclusion(true);
 	}
 
 
