Fix cook issue with morphtarget, we prevent creating morphtarget during the cook serialize save.
#rnx

diff --git a/Engine/Source/Runtime/Engine/Classes/Animation/MorphTarget.h b/Engine/Source/Runtime/Engine/Classes/Animation/MorphTarget.h
--- a/Engine/Source/Runtime/Engine/Classes/Animation/MorphTarget.h
+++ b/Engine/Source/Runtime/Engine/Classes/Animation/MorphTarget.h
@@ -160,7 +160,7 @@ public:
 	ENGINE_API virtual void LoadFromMemoryArchive(FMemoryArchive & Ar);
 	
 	/** Apply serialized data to skeletal mesh in game thread */
-	ENGINE_API virtual void ApplyEditorData(USkeletalMesh * SkeletalMesh) const;
+	ENGINE_API virtual void ApplyEditorData(USkeletalMesh * SkeletalMesh, bool bIsSerializeSaving) const;
 	
 protected:
 	bool bApplyMorphTargetsData = false;
diff --git a/Engine/Source/Runtime/Engine/Classes/Engine/SkeletalMesh.h b/Engine/Source/Runtime/Engine/Classes/Engine/SkeletalMesh.h
--- a/Engine/Source/Runtime/Engine/Classes/Engine/SkeletalMesh.h
+++ b/Engine/Source/Runtime/Engine/Classes/Engine/SkeletalMesh.h
@@ -135,6 +135,9 @@ public:
 	// Movable
 	FSkeletalMeshCompilationContext(FSkeletalMeshCompilationContext&&) = default;
 	FSkeletalMeshCompilationContext& operator=(FSkeletalMeshCompilationContext&&) = default;
+	
+	//True if this compilation context is start from a serialize save
+	bool bIsSerializeSaving = false;
 };
 
 class FSkeletalMeshPostLoadContext : public FSkeletalMeshCompilationContext
diff --git a/Engine/Source/Runtime/Engine/Private/MorphTools.cpp b/Engine/Source/Runtime/Engine/Private/MorphTools.cpp
--- a/Engine/Source/Runtime/Engine/Private/MorphTools.cpp
+++ b/Engine/Source/Runtime/Engine/Private/MorphTools.cpp
@@ -11,6 +11,7 @@
 #include "Rendering/SkeletalMeshModel.h"
 #include "Rendering/SkeletalMeshLODModel.h"
 #include "UObject/GarbageCollection.h"
+#include "Algo/AnyOf.h"
 
 /** compare based on base mesh source vertex indices */
 struct FCompareMorphTargetDeltas
@@ -167,7 +168,7 @@ TUniquePtr<FFinishBuildMorphTargetData> UMorphTarget::CreateFinishBuildMorphTarg
 	return MakeUnique<FFinishBuildMorphTargetData>();
 }
 
-void FFinishBuildMorphTargetData::ApplyEditorData(USkeletalMesh * SkeletalMesh) const
+void FFinishBuildMorphTargetData::ApplyEditorData(USkeletalMesh * SkeletalMesh, bool bIsSerializeSaving) const
 {
 	//Return if we do not need to apply data
 	if (!bApplyMorphTargetsData)
@@ -195,11 +196,30 @@ void FFinishBuildMorphTargetData::ApplyEditorData(USkeletalMesh * SkeletalMesh)
 	for (const TPair<FName, TArray<FMorphTargetLODModel>>& TargetNameAndMorphLODModels : MorphLODModelsPerTargetName)
 	{
 		FName MorphTargetName = TargetNameAndMorphLODModels.Key;
+		const TArray<FMorphTargetLODModel>& MorphTargetLODModels = TargetNameAndMorphLODModels.Value;
+		int32 MorphLODModelNumber = MorphTargetLODModels.Num();
+
 		UMorphTarget * MorphTarget = ExistingMorphTargets.FindRef(MorphTargetName);
 		if (!MorphTarget)
 		{
-			MorphTarget = NewObject<UMorphTarget>(SkeletalMesh, MorphTargetName);
-			check(MorphTarget);
+			if (!Algo::AnyOf(MorphTargetLODModels, [](const FMorphTargetLODModel& Model) { return Model.Vertices.Num() > 0;}))
+			{
+				//Skip this empty morphtarget
+				continue;
+			}
+			//When we save the cook result we should never have to build a new morph target
+			//When saving cook build, we call GetPlatformSkeletalMeshRenderData in USkeletalMesh::BeginCacheForCookedPlatformData
+			//which happen before the serialization of that cook skeletalmesh
+			if (!bIsSerializeSaving)
+			{
+				MorphTarget = NewObject<UMorphTarget>(SkeletalMesh, MorphTargetName);
+				check(MorphTarget);
+			}
+			else
+			{
+				UE_ASSET_LOG(LogSkeletalMesh, Error, SkeletalMesh, TEXT("Cannot cache a skeletalmesh during a serialize if some morph targets need to be created. The solution is to Pre cache the skeletalmesh before the serialization so no morph target get created."));
+				continue;
+			}
 		}
 		else
 		{
@@ -207,8 +227,7 @@ void FFinishBuildMorphTargetData::ApplyEditorData(USkeletalMesh * SkeletalMesh)
 		}
 		MorphTarget->EmptyMorphLODModels();
 		SkeletalMesh->GetMorphTargets().Add(MorphTarget);
-		const TArray<FMorphTargetLODModel>&MorphTargetLODModels = TargetNameAndMorphLODModels.Value;
-		int32 MorphLODModelNumber = MorphTargetLODModels.Num();
+		
 		MorphTarget->GetMorphLODModels().AddDefaulted(MorphLODModelNumber);
 		for (int32 MorphDataIndex = 0; MorphDataIndex < MorphLODModelNumber; ++MorphDataIndex)
 		{
diff --git a/Engine/Source/Runtime/Engine/Private/SkeletalMesh.cpp b/Engine/Source/Runtime/Engine/Private/SkeletalMesh.cpp
--- a/Engine/Source/Runtime/Engine/Private/SkeletalMesh.cpp
+++ b/Engine/Source/Runtime/Engine/Private/SkeletalMesh.cpp
@@ -1627,17 +1627,18 @@ bool USkeletalMesh::IsReadyForFinishDestroy()
 }
 
 #if WITH_EDITOR
-void CachePlatform(USkeletalMesh* Mesh, const ITargetPlatform* TargetPlatform, FSkeletalMeshRenderData* PlatformRenderData)
+void CachePlatform(USkeletalMesh* Mesh, const ITargetPlatform* TargetPlatform, FSkeletalMeshRenderData* PlatformRenderData, const bool bIsSerializeSaving)
 {
 	//Cache the platform, dcc should be valid so it will be fast
 	check(TargetPlatform);
 	FSkeletalMeshBuildContext Context;
+	Context.bIsSerializeSaving = bIsSerializeSaving;
 	PlatformRenderData->Cache(TargetPlatform, Mesh, &Context);
 }
 
 FString BuildSkeletalMeshDerivedDataKey(const ITargetPlatform* TargetPlatform, USkeletalMesh* SkelMesh);
 
-static FSkeletalMeshRenderData& GetPlatformSkeletalMeshRenderData(USkeletalMesh* Mesh, const ITargetPlatform* TargetPlatform)
+static FSkeletalMeshRenderData& GetPlatformSkeletalMeshRenderData(USkeletalMesh* Mesh, const ITargetPlatform* TargetPlatform, const bool bIsSerializeSaving)
 {
 	const FString PlatformDerivedDataKey = BuildSkeletalMeshDerivedDataKey(TargetPlatform, Mesh);
 	FSkeletalMeshRenderData* PlatformRenderData = Mesh->GetResourceForRendering();
@@ -1656,7 +1657,7 @@ static FSkeletalMeshRenderData& GetPlatformSkeletalMeshRenderData(USkeletalMesh*
 	{
 		// Cache render data for this platform and insert it in to the linked list.
 		PlatformRenderData = new FSkeletalMeshRenderData();
-		CachePlatform(Mesh, TargetPlatform, PlatformRenderData);
+		CachePlatform(Mesh, TargetPlatform, PlatformRenderData, bIsSerializeSaving);
 		check(PlatformRenderData->DerivedDataKey == PlatformDerivedDataKey);
 		Swap(PlatformRenderData->NextCachedRenderData, Mesh->GetResourceForRendering()->NextCachedRenderData);
 		Mesh->GetResourceForRendering()->NextCachedRenderData = TUniquePtr<FSkeletalMeshRenderData>(PlatformRenderData);
@@ -1750,15 +1751,16 @@ void USkeletalMesh::Serialize( FArchive& Ar )
 				FSkeletalMeshRenderData* LocalSkeletalMeshRenderData = GetSkeletalMeshRenderData();
 #if WITH_EDITORONLY_DATA
 				const ITargetPlatform* ArchiveCookingTarget = Ar.CookingTarget();
+				constexpr bool bIsSerializeSaving = true;
 				if (ArchiveCookingTarget)
 				{
-					LocalSkeletalMeshRenderData = &GetPlatformSkeletalMeshRenderData(this, ArchiveCookingTarget);
+					LocalSkeletalMeshRenderData = &GetPlatformSkeletalMeshRenderData(this, ArchiveCookingTarget, bIsSerializeSaving);
 				}
 				else
 				{
 					//Fall back in case we use an archive that the cooking target has not been set (i.e. Duplicate archive)
 					check(RunningPlatform != NULL);
-					LocalSkeletalMeshRenderData = &GetPlatformSkeletalMeshRenderData(this, RunningPlatform);
+					LocalSkeletalMeshRenderData = &GetPlatformSkeletalMeshRenderData(this, RunningPlatform, bIsSerializeSaving);
 				}
 #endif
 				if (bCooked)
@@ -1784,22 +1786,6 @@ void USkeletalMesh::Serialize( FArchive& Ar )
 					}
 				}
 				LocalSkeletalMeshRenderData->Serialize(Ar, this);
-
-#if WITH_EDITOR
-				//We need to get the skeletalmesh ddc cache with the running platform to retrieve the ddc editor data LODModel which can be different because of chunking, reduction and morph targets.
-				//Normally it should just take back the ddc for the running platform, since the ddc was cache when we have load the asset to cook it.
-				if(ArchiveCookingTarget && ArchiveCookingTarget != RunningPlatform)
-				{
-					check(RunningPlatform != NULL);
-					const FString CookPlatformDerivedDataKey = BuildSkeletalMeshDerivedDataKey(ArchiveCookingTarget, this);
-					const FString RunningPlatformDerivedDataKey = BuildSkeletalMeshDerivedDataKey(RunningPlatform, this);
-					if (CookPlatformDerivedDataKey != RunningPlatformDerivedDataKey)
-					{
-						FSkeletalMeshRenderData RunningPlatformSkeletalMeshRenderData;
-						CachePlatform(this, RunningPlatform, &RunningPlatformSkeletalMeshRenderData);
-					}
-				}
-#endif
 			}
 		}
 	}
@@ -4660,13 +4646,26 @@ void USkeletalMesh::ValidateBoneWeights(const ITargetPlatform* TargetPlatform)
 void USkeletalMesh::BeginCacheForCookedPlatformData(const ITargetPlatform* TargetPlatform)
 {
 	// Make sure to cache platform data so it doesn't happen lazily during serialization of the skeletal mesh
-	GetPlatformSkeletalMeshRenderData(this, TargetPlatform);
+	constexpr bool bIsSerializeSaving = false;
+	GetPlatformSkeletalMeshRenderData(this, TargetPlatform, bIsSerializeSaving);
 	ValidateBoneWeights(TargetPlatform);
 }
 
 void USkeletalMesh::ClearAllCachedCookedPlatformData()
 {
 	GetResourceForRendering()->NextCachedRenderData.Reset();
+	
+	if (FApp::CanEverRender())
+	{
+		// We need to keep the ddc editor data LODModel for rendering; it can be different (The number of sections, the number of vertices, the number of morphtargets) because of chunking, build or reduction setting that are or will be per platform.
+		//Normally this call should be able to read values out of ddc rather than rebuilding, because the ddc for the running platform was cached when we loaded the asset.
+		ITargetPlatform* RunningPlatform = GetTargetPlatformManagerRef().GetRunningTargetPlatform();
+		check(RunningPlatform != NULL);
+		const FString RunningPlatformDerivedDataKey = BuildSkeletalMeshDerivedDataKey(RunningPlatform, this);
+		FSkeletalMeshRenderData RunningPlatformSkeletalMeshRenderData;
+		constexpr bool bIsSerializeSaving = false;
+		CachePlatform(this, RunningPlatform, &RunningPlatformSkeletalMeshRenderData, bIsSerializeSaving);
+	}
 }
 
 extern FString BuildSkeletalMeshDerivedDataKey(const ITargetPlatform* TargetPlatform, USkeletalMesh* SkelMesh);
diff --git a/Engine/Source/Runtime/Engine/Private/SkeletalMeshRenderData.cpp b/Engine/Source/Runtime/Engine/Private/SkeletalMeshRenderData.cpp
--- a/Engine/Source/Runtime/Engine/Private/SkeletalMeshRenderData.cpp
+++ b/Engine/Source/Runtime/Engine/Private/SkeletalMeshRenderData.cpp
@@ -412,7 +412,7 @@ void FSkeletalMeshRenderData::Cache(const ITargetPlatform* TargetPlatform, USkel
 			//Apply the morphtargets change if any
 			if (FinishBuildMorphTargetData.IsValid())
 			{
-				FinishBuildMorphTargetData->ApplyEditorData(Owner);
+				FinishBuildMorphTargetData->ApplyEditorData(Owner, ContextPtr->bIsSerializeSaving);
 			}
 
 			int32 T1 = FPlatformTime::Cycles();
