WorldPartitionActorDescView.h error: inline function is not defined

Removed inline attribute from FWorldPartitionActorDescView::GetHLODLayer()

diff --git a/Engine/Source/Runtime/Engine/Public/WorldPartition/WorldPartitionActorDescView.h b/Engine/Source/Runtime/Engine/Public/WorldPartition/WorldPartitionActorDescView.h
--- a/Engine/Source/Runtime/Engine/Public/WorldPartition/WorldPartitionActorDescView.h
+++ b/Engine/Source/Runtime/Engine/Public/WorldPartition/WorldPartitionActorDescView.h
@@ -30,7 +30,7 @@ public:
 	bool GetActorIsEditorOnly() const;
 	bool GetLevelBoundsRelevant() const;
 	bool GetActorIsHLODRelevant() const;
-	inline FName GetHLODLayer() const;
+	FName GetHLODLayer() const;
 	const TArray<FName>& GetDataLayers() const;
 	FName GetActorPackage() const;
 	FName GetActorPath() const;
