Added GenerateObjectsClusters that creates a list of object clusters based on the provided references (type independant) + added several unit tests.
This will be used to replace custom clustering code present in the engine.

diff --git a/Engine/Source/Runtime/Engine/Private/ReferenceCluster.cpp b/Engine/Source/Runtime/Engine/Private/ReferenceCluster.cpp
new file mode 100644
--- /dev/null
+++ b/Engine/Source/Runtime/Engine/Private/ReferenceCluster.cpp
@@ -0,0 +1,59 @@
+// Copyright Epic Games, Inc. All Rights Reserved.
+
+#include "ReferenceCluster.h"
+
+TArray<TArray<FGuid>> GenerateObjectsClusters(TArray<TPair<FGuid, TArray<FGuid>>> InObjects)
+{
+	TMap<FGuid, FGuid> ObjectToClusters;
+	TMap<FGuid, TSet<FGuid>> Clusters;
+
+	for (const auto& Object : InObjects)
+	{
+		const FGuid& ObjectGuid = Object.Key;
+		FGuid ClusterGuid = ObjectToClusters.FindRef(ObjectGuid);
+
+		if (!ClusterGuid.IsValid())
+		{
+			ClusterGuid = FGuid::NewGuid();
+			TSet<FGuid>& Cluster = Clusters.Add(ClusterGuid);
+			ObjectToClusters.Add(ObjectGuid, ClusterGuid);
+			Cluster.Add(ObjectGuid);
+		}
+
+		for (const FGuid& ReferenceGuid : Object.Value)
+		{
+			FGuid ReferenceClusterGuid = ObjectToClusters.FindRef(ReferenceGuid);
+
+			if (ReferenceClusterGuid.IsValid())
+			{
+				if (ReferenceClusterGuid != ClusterGuid)
+				{
+					TSet<FGuid> ReferenceCluster = Clusters.FindAndRemoveChecked(ReferenceClusterGuid);
+					TSet<FGuid>& Cluster = Clusters.FindChecked(ClusterGuid);
+					Cluster.Append(ReferenceCluster);
+					for (const FGuid& OtherReferenceGuid : ReferenceCluster)
+					{
+						ObjectToClusters.FindChecked(OtherReferenceGuid) = ClusterGuid;
+					}
+				}
+			}
+			else
+			{
+				TSet<FGuid>& Cluster = Clusters.FindChecked(ClusterGuid);
+				Cluster.Add(ReferenceGuid);
+			}
+
+			ObjectToClusters.Add(ReferenceGuid, ClusterGuid);
+		}
+	}
+	
+	TArray<TArray<FGuid>> Result;
+	Result.Reserve(Clusters.Num());
+
+	for (const auto& Cluster : Clusters)
+	{
+		Result.AddDefaulted_GetRef() = Cluster.Value.Array();
+	}
+
+	return MoveTemp(Result);
+}
\ No newline at end of file
diff --git a/Engine/Source/Runtime/Engine/Private/Tests/ReferenceClusterTest.cpp b/Engine/Source/Runtime/Engine/Private/Tests/ReferenceClusterTest.cpp
new file mode 100644
--- /dev/null
+++ b/Engine/Source/Runtime/Engine/Private/Tests/ReferenceClusterTest.cpp
@@ -0,0 +1,117 @@
+// Copyright Epic Games, Inc. All Rights Reserved.
+
+#include "Misc/AutomationTest.h"
+#include "ReferenceCluster.h"
+
+#if WITH_DEV_AUTOMATION_TESTS
+
+#define TEST_NAME_ROOT "System.Engine.ReferenceCluster"
+
+namespace ReferenceClusterTests
+{
+	constexpr const uint32 TestFlags = EAutomationTestFlags::ApplicationContextMask | EAutomationTestFlags::EngineFilter;
+
+	IMPLEMENT_SIMPLE_AUTOMATION_TEST(FReferenceClusterTest, TEST_NAME_ROOT, TestFlags)
+	bool FReferenceClusterTest::RunTest(const FString& Parameters)
+	{
+#if WITH_EDITOR
+		TArray<TPair<FGuid, TArray<FGuid>>> ObjectsWithRefs;
+		ObjectsWithRefs.AddDefaulted(10);
+
+		for (auto& ObjectWithRefs : ObjectsWithRefs)
+		{
+			ObjectWithRefs.Key = FGuid::NewGuid();
+		}
+
+		// Basic test
+		{
+			TArray<TArray<FGuid>> Clusters = GenerateObjectsClusters(ObjectsWithRefs);
+
+			TestTrue(TEXT("Reference Cluster Basic"), Clusters.Num() == 10);
+			for (const TArray<FGuid>& Cluster : Clusters)
+			{
+				TestTrue(TEXT("Reference Cluster Basic"), Cluster.Num() == 1);
+			}
+		}
+
+		//
+		// This forms 6 clusters:
+		// 
+		// 0 --> 1     4     6     7     8     9
+		// |     ^     |
+		// v     |     v
+		// 2 --> 3     5
+		ObjectsWithRefs[0].Value.Add(ObjectsWithRefs[1].Key);
+		ObjectsWithRefs[0].Value.Add(ObjectsWithRefs[2].Key);
+		ObjectsWithRefs[2].Value.Add(ObjectsWithRefs[3].Key);
+		ObjectsWithRefs[3].Value.Add(ObjectsWithRefs[1].Key);
+		ObjectsWithRefs[4].Value.Add(ObjectsWithRefs[5].Key);
+
+		// Simple clustering test
+		{
+			TArray<TArray<FGuid>> Clusters = GenerateObjectsClusters(ObjectsWithRefs);
+			TestTrue(TEXT("Reference Cluster Simple"), Clusters.Num() == 6);
+		}
+
+		//
+		// This forms 5 clusters:
+		// 
+		// 0 --> 1     4     6     7     8     9
+		// |     ^     |
+		// v     |     v
+		// 2 --> 3 <-- 5
+		ObjectsWithRefs[5].Value.Add(ObjectsWithRefs[3].Key);
+
+		// Simple clusters merging test
+		{
+			TArray<TArray<FGuid>> Clusters = GenerateObjectsClusters(ObjectsWithRefs);
+			TestTrue(TEXT("Reference Cluster Simple"), Clusters.Num() == 5);
+		}
+
+		//
+		// This forms 2 clusters:
+		// 
+		//                   /-----.-----------\
+		//                   |     |           |
+		//                   v     v           |
+		// 0 --> 1     4     6     7     8 <-- 9
+		// |     ^     |                  
+		// v     |     v                  
+		// 2 --> 3 <-- 5                  
+		ObjectsWithRefs[9].Value.Add(ObjectsWithRefs[8].Key);
+		ObjectsWithRefs[9].Value.Add(ObjectsWithRefs[7].Key);
+		ObjectsWithRefs[9].Value.Add(ObjectsWithRefs[6].Key);
+
+		// Double clusters test
+		{
+			TArray<TArray<FGuid>> Clusters = GenerateObjectsClusters(ObjectsWithRefs);
+			TestTrue(TEXT("Reference Cluster Double"), Clusters.Num() == 2);
+		}
+
+		//
+		// This forms a single cluster:
+		// 
+		//                   /-----.-----------\
+		//                   |     |           |
+		//                   v     v           |
+		// 0 --> 1     4     6     7     8 <-- 9
+		// |     ^     |                 |
+		// v     |     v                 |
+		// 2 --> 3 <-- 5 <---------------/
+		ObjectsWithRefs[8].Value.Add(ObjectsWithRefs[5].Key);
+
+		// Complex clusters merging test
+		{
+			TArray<TArray<FGuid>> Clusters = GenerateObjectsClusters(ObjectsWithRefs);
+
+			TestTrue(TEXT("Reference Cluster Complex Merging"), Clusters.Num() == 1);
+			TestTrue(TEXT("Reference Cluster Complex Merging"), Clusters[0].Num() == 10);
+		}
+#endif
+		return true;
+	}
+}
+
+#undef TEST_NAME_ROOT
+
+#endif 
diff --git a/Engine/Source/Runtime/Engine/Public/ReferenceCluster.h b/Engine/Source/Runtime/Engine/Public/ReferenceCluster.h
new file mode 100644
--- /dev/null
+++ b/Engine/Source/Runtime/Engine/Public/ReferenceCluster.h
@@ -0,0 +1,17 @@
+// Copyright Epic Games, Inc. All Rights Reserved.
+
+#pragma once
+
+#include "CoreMinimal.h"
+
+/**
+ *	Creates a list of object clusters based on the provided references.
+ *
+ *	@param	InObjects Array of object GUIDs with their references. Each pair represents the GUID of an
+ *			object and its list of references. References are expected to be part of the root object
+ *			list, forming a closed set of n clusters.
+ *
+ *	@return	Array of reference clusters (e.g. group of objects referencing each other). For example, an
+ *			input of n objects without any references will generate be n clusters of a single object each.
+ */
+ENGINE_API TArray<TArray<FGuid>> GenerateObjectsClusters(TArray<TPair<FGuid, TArray<FGuid>>> InObjects);
\ No newline at end of file
