Additive Anim Sequence: Selecting 'Frame from this animation' as Base Pose Type crashes editor
#fix handle ABPT_LocalAnimFrame case for generating additive attributes correctly

diff --git a/Engine/Source/Runtime/Engine/Private/Animation/AnimSequence.cpp b/Engine/Source/Runtime/Engine/Private/Animation/AnimSequence.cpp
--- a/Engine/Source/Runtime/Engine/Private/Animation/AnimSequence.cpp
+++ b/Engine/Source/Runtime/Engine/Private/Animation/AnimSequence.cpp
@@ -4829,8 +4829,11 @@ void UAnimSequence::SynchronousAnimatedBoneAttributesCompression()
 	const bool bShouldSampleBasePose = IsValidAdditive() && RefPoseType != ABPT_RefPose;
 	if (bShouldSampleBasePose)
 	{
+		// Select which AnimSequence to sample according to additive type
+		const UAnimSequence* BasePoseSequence = (RefPoseType == ABPT_LocalAnimFrame) ? this : RefPoseSeq.Get();
+		
 		// Behaviour for determining the time to sample the base pose attributes
-		auto GetBasePoseTimeToSample = [this](float InTime) -> float
+		auto GetBasePoseTimeToSample = [this, BasePoseSequence](float InTime) -> float
 		{
 			float BasePoseTime = 0.f;
 
@@ -4838,12 +4841,12 @@ void UAnimSequence::SynchronousAnimatedBoneAttributesCompression()
 			{
 				const float CurrentSequenceLength = GetPlayLength();
 				const float Fraction = (CurrentSequenceLength > 0.f) ? FMath::Clamp<float>(InTime / CurrentSequenceLength, 0.f, 1.f) : 0.f;
-				BasePoseTime = RefPoseSeq->GetPlayLength() * Fraction;
+				BasePoseTime = BasePoseSequence->GetPlayLength() * Fraction;
 			}
-			else if (RefPoseType == ABPT_AnimFrame)
+			else if (RefPoseType == ABPT_AnimFrame || RefPoseType == ABPT_LocalAnimFrame)
 			{
-				const float Fraction = (RefPoseSeq->GetNumberOfSampledKeys() > 0) ? FMath::Clamp<float>((float)RefFrameIndex / (float)RefPoseSeq->GetNumberOfSampledKeys(), 0.f, 1.f) : 0.f;
-				BasePoseTime = RefPoseSeq->GetPlayLength() * Fraction;
+				const float Fraction = (BasePoseSequence->GetNumberOfSampledKeys() > 0) ? FMath::Clamp<float>((float)RefFrameIndex / (float)BasePoseSequence->GetNumberOfSampledKeys(), 0.f, 1.f) : 0.f;
+				BasePoseTime = BasePoseSequence->GetPlayLength() * Fraction;
 			}
 
 			return BasePoseTime;
@@ -4861,7 +4864,7 @@ void UAnimSequence::SynchronousAnimatedBoneAttributesCompression()
 			ensure(!AttributeCurves.Contains(AdditiveAttribute.Identifier));
 			FAttributeCurve& AttributeCurve = AttributeCurves.Add(AdditiveAttribute.Identifier);
 
-			if (const FAnimatedBoneAttribute* RefPoseAttributePtr = RefPoseSeq->GetDataModel()->FindAttribute(AdditiveAttribute.Identifier))
+			if (const FAnimatedBoneAttribute* RefPoseAttributePtr = BasePoseSequence->GetDataModel()->FindAttribute(AdditiveAttribute.Identifier))
 			{
 				AttributeCurve.SetScriptStruct(AdditiveAttribute.Identifier.GetType());
 				const FAnimatedBoneAttribute& RefAttribute = *RefPoseAttributePtr;
