r.ShaderCompiler.JobCacheDDCEnableRemotePolicy to control whether or not we cache single shaders from the Job Cache remotely or Locally.

diff --git a/Engine/Source/Runtime/Engine/Private/ShaderCompiler/ShaderCompiler.cpp b/Engine/Source/Runtime/Engine/Private/ShaderCompiler/ShaderCompiler.cpp
--- a/Engine/Source/Runtime/Engine/Private/ShaderCompiler/ShaderCompiler.cpp
+++ b/Engine/Source/Runtime/Engine/Private/ShaderCompiler/ShaderCompiler.cpp
@@ -120,6 +120,17 @@ static TAutoConsoleVariable<bool> CVarJobCacheDDC(
 	TEXT("Skips compilation of all shaders on Material and Material Instance PostLoad and relies on on-demand shader compilation to compile what is needed."),
 	ECVF_ReadOnly);
 
+static TAutoConsoleVariable<bool> CVarJobCacheDDCPolicy(
+	TEXT("r.ShaderCompiler.JobCacheDDCEnableRemotePolicy"),
+	false,
+	TEXT("Whether to cache shaders in the job cache to your local machine or remotely to the network.\n"),
+	ECVF_ReadOnly);
+
+static bool IsShaderJobCacheDDCRemotePolicyEnabled()
+{
+	return CVarJobCacheDDCPolicy.GetValueOnAnyThread();
+}
+
 bool IsShaderJobCacheDDCEnabled()
 {
 	// For now we only support the editor and not commandlets like the cooker.
@@ -7693,7 +7704,7 @@ FShaderJobCache::FJobCachedOutput* FShaderJobCache::Find(const FJobInputHash& Ha
 				// Create key.
 				Request.Key.Bucket = ShaderJobCacheDDCBucket;
 				Request.Key.Hash = Hash;
-				Request.Policy = UE::DerivedData::ECachePolicy::QueryLocal;
+				Request.Policy = IsShaderJobCacheDDCRemotePolicyEnabled() ? UE::DerivedData::ECachePolicy::Default : UE::DerivedData::ECachePolicy::Local;
 
 				UE::DerivedData::GetCache().Get(
 					{ Request },
@@ -7825,7 +7836,7 @@ void FShaderJobCache::Add(const FJobInputHash& Hash, const FJobCachedOutput& Con
 			UE::DerivedData::FRequestOwner RequestOwner(UE::DerivedData::EPriority::Normal);
 			RequestOwner.KeepAlive();
 			UE::DerivedData::GetCache().Put(
-				{ {{TEXT("FShaderJobCache")}, RecordBuilder.Build(), UE::DerivedData::ECachePolicy::StoreLocal} },
+				{ {{TEXT("FShaderJobCache")}, RecordBuilder.Build(), IsShaderJobCacheDDCRemotePolicyEnabled() ? UE::DerivedData::ECachePolicy::Default : UE::DerivedData::ECachePolicy::Local } },
 				RequestOwner
 			);
 		}
