Add new defines for LARGE_WORLD_MAX, allow incrementally changing from old->LWC WORLD_MAX value
Change to new LARGE_WORLD_MAX for scene otrees used for object culling

diff --git a/Engine/Source/Runtime/Engine/Public/EngineDefines.h b/Engine/Source/Runtime/Engine/Public/EngineDefines.h
--- a/Engine/Source/Runtime/Engine/Public/EngineDefines.h
+++ b/Engine/Source/Runtime/Engine/Public/EngineDefines.h
@@ -28,19 +28,23 @@
 /*-----------------------------------------------------------------------------
 	Size of the world.
 -----------------------------------------------------------------------------*/
-#define UE_OLD_WORLD_MAX			2097152.0					/* UE4 maximum world size */
-#define UE_OLD_HALF_WORLD_MAX		(UE_OLD_WORLD_MAX * 0.5)	/* UE4 half maximum world size */
-#define UE_OLD_HALF_WORLD_MAX1		(UE_OLD_HALF_WORLD_MAX - 1)	/* UE4 half maximum world size minus one */
+#define UE_OLD_WORLD_MAX			2097152.0						/* UE4 maximum world size */
+#define UE_OLD_HALF_WORLD_MAX		(UE_OLD_WORLD_MAX * 0.5)		/* UE4 half maximum world size */
+#define UE_OLD_HALF_WORLD_MAX1		(UE_OLD_HALF_WORLD_MAX - 1)		/* UE4 half maximum world size minus one */
+
+#define UE_LARGE_WORLD_MAX			8796093022208.0					/* LWC maximum world size, Approx 87,960,930.2 km across, or 43,980,465.1 km from the origin */
+#define UE_LARGE_HALF_WORLD_MAX		(UE_LARGE_WORLD_MAX * 0.5)		/* LWC half maximum world size */
+#define UE_LARGE_HALF_WORLD_MAX1	(UE_LARGE_HALF_WORLD_MAX - 1)	/* LWC half maximum world size minus one */
 
 #ifndef UE_USE_UE4_WORLD_MAX
 #define UE_USE_UE4_WORLD_MAX		1 //(UE_LARGE_WORLD_COORDINATES_DISABLED)			// Force UE4 WORLD_MAX for converted UE4 titles that explicitly rely on it.
 #endif
 
+// Note: Modifying WORLD_MAX affects UE_LWC_RENDER_TILE_SIZE in Engine\Source\Runtime\Core\Private\Misc\LargeWorldRenderPosition.cpp and may introduce precision issues in shaders using world coordinates.
 #if UE_USE_UE4_WORLD_MAX
-#define WORLD_MAX					(UE_OLD_WORLD_MAX)		/* Maximum size of the world */
+#define WORLD_MAX					(UE_OLD_WORLD_MAX)
 #else
-// Note: Modifying WORLD_MAX affects UE_LWC_RENDER_TILE_SIZE in Engine\Source\Runtime\Core\Private\Misc\LargeWorldRenderPosition.cpp and may introduce precision issues in shaders using world coordinates.
-#define WORLD_MAX					8796093022208.0			/* Maximum size of the world  - Approx 87,960,930.2 km across, or 43,980,465.1 km from the origin */
+#define WORLD_MAX					(UE_LARGE_WORLD_MAX)
 #endif
 
 #define HALF_WORLD_MAX				(WORLD_MAX * 0.5)		/* Half the maximum size of the world */
diff --git a/Engine/Source/Runtime/Renderer/Private/RendererScene.cpp b/Engine/Source/Runtime/Renderer/Private/RendererScene.cpp
--- a/Engine/Source/Runtime/Renderer/Private/RendererScene.cpp
+++ b/Engine/Source/Runtime/Renderer/Private/RendererScene.cpp
@@ -1166,8 +1166,8 @@ FScene::FScene(UWorld* InWorld, bool bInRequiresHitProxies, bool bInIsEditorScen
 ,	SkyAtmosphere(NULL)
 ,	VolumetricCloud(NULL)
 ,	PrecomputedVisibilityHandler(NULL)
-,	LocalShadowCastingLightOctree(FVector::ZeroVector,HALF_WORLD_MAX)
-,	PrimitiveOctree(FVector::ZeroVector,HALF_WORLD_MAX)
+,	LocalShadowCastingLightOctree(FVector::ZeroVector, UE_LARGE_HALF_WORLD_MAX)
+,	PrimitiveOctree(FVector::ZeroVector, UE_LARGE_HALF_WORLD_MAX)
 ,	bRequiresHitProxies(bInRequiresHitProxies)
 ,	bIsEditorScene(bInIsEditorScene)
 ,	NumUncachedStaticLightingInteractions(0)
