World Partition: fix crash when saving a level BP during PIE, by making the streaming policy transient in PIE.

diff --git a/Engine/Source/Runtime/Engine/Private/WorldPartition/WorldPartitionStreamingGeneration.cpp b/Engine/Source/Runtime/Engine/Private/WorldPartition/WorldPartitionStreamingGeneration.cpp
--- a/Engine/Source/Runtime/Engine/Private/WorldPartition/WorldPartitionStreamingGeneration.cpp
+++ b/Engine/Source/Runtime/Engine/Private/WorldPartition/WorldPartitionStreamingGeneration.cpp
@@ -568,7 +568,7 @@ bool UWorldPartition::GenerateStreaming(TArray<FString>* OutPackagesToGenerate)
 
 	// Generate streaming
 	check(!StreamingPolicy);
-	StreamingPolicy = NewObject<UWorldPartitionStreamingPolicy>(const_cast<UWorldPartition*>(this), WorldPartitionStreamingPolicyClass.Get());
+	StreamingPolicy = NewObject<UWorldPartitionStreamingPolicy>(const_cast<UWorldPartition*>(this), WorldPartitionStreamingPolicyClass.Get(), NAME_None, bIsPIE ? RF_Transient : RF_NoFlags);
 
 	check(RuntimeHash);
 	if (RuntimeHash->GenerateStreaming(StreamingPolicy, ActorClusterContext, OutPackagesToGenerate))
