Change some FVector4 to FVector4f in streaming system to save memory

diff --git a/Engine/Source/Runtime/Engine/Private/Streaming/TextureInstanceView.cpp b/Engine/Source/Runtime/Engine/Private/Streaming/TextureInstanceView.cpp
--- a/Engine/Source/Runtime/Engine/Private/Streaming/TextureInstanceView.cpp
+++ b/Engine/Source/Runtime/Engine/Private/Streaming/TextureInstanceView.cpp
@@ -353,6 +353,11 @@ void FRenderAssetInstanceAsyncView::UpdateBoundSizes_Async(
 	{
 		const FRenderAssetInstanceView::FBounds4& CurrentBounds4 = View->GetBounds4(Bounds4Index);
 
+		// LWC_TODO - Origin values are loaded from doubles, the remaining values are loaded from floats
+		// Could potentially perform some of these operations with float VectorRegisters, which could potentially be more efficient
+		// (Otherwise we're paying cost to convert these values to double VectorRegisters on load)
+		// Tricky to manage precision though, as with large worlds, distance between object and view origin can potentially overflow float capacity
+
 		// Calculate distance of viewer to bounding sphere.
 		const VectorRegister OriginX = VectorLoadAligned( &CurrentBounds4.OriginX );
 		const VectorRegister OriginY = VectorLoadAligned( &CurrentBounds4.OriginY );
diff --git a/Engine/Source/Runtime/Engine/Private/Streaming/TextureInstanceView.h b/Engine/Source/Runtime/Engine/Private/Streaming/TextureInstanceView.h
--- a/Engine/Source/Runtime/Engine/Private/Streaming/TextureInstanceView.h
+++ b/Engine/Source/Runtime/Engine/Private/Streaming/TextureInstanceView.h
@@ -47,27 +47,27 @@ public:
 		FVector4 RangeOriginZ;
 
 		/** X size of the bounds box extent of 4 texture/mesh instances */
-		FVector4 ExtentX;
+		FVector4f ExtentX;
 		/** Y size of the bounds box extent of 4 texture/mesh instances */
-		FVector4 ExtentY;
+		FVector4f ExtentY;
 		/** Z size of the bounds box extent of 4 texture/mesh instances */
-		FVector4 ExtentZ;
+		FVector4f ExtentZ;
 
 		/** Sphere radii for the bounding sphere of 4 texture/mesh instances */
-		FVector4 Radius;
+		FVector4f Radius;
 
 		/** The relative box the bound was computed with. Aligned to be interpreted as FVector4  */
 		MS_ALIGN(16) FUintVector4 PackedRelativeBox;
 
 		/** Minimal distance (between the bounding sphere origin and the view origin) for which this entry is valid */
-		FVector4 MinDistanceSq;
+		FVector4f MinDistanceSq;
 		/** Minimal range distance (between the bounding sphere origin and the view origin) for which this entry is valid */
-		FVector4 MinRangeSq;
+		FVector4f MinRangeSq;
 		/** Maximal range distance (between the bounding sphere origin and the view origin) for which this entry is valid */
-		FVector4 MaxRangeSq;
+		FVector4f MaxRangeSq;
 
 		/** Last visibility time for this bound, used for priority */
-		FVector4 LastRenderTime; //(FApp::GetCurrentTime() - Component->LastRenderTime);
+		FVector4f LastRenderTime; //(FApp::GetCurrentTime() - Component->LastRenderTime);
 
 
 		void Set(int32 Index, const FBoxSphereBounds& Bounds, uint32 InPackedRelativeBox, float LastRenderTime, const FVector& RangeOrigin, float MinDistanceSq, float MinRangeSq, float MaxRangeSq);
