Preserve AActor::ActorGuid in UEngine::CopyPropertiesForUnrelatedObjects, avoiding having actors with the same GUIDs after the operation.

diff --git a/Engine/Source/Runtime/Engine/Classes/GameFramework/Actor.h b/Engine/Source/Runtime/Engine/Classes/GameFramework/Actor.h
--- a/Engine/Source/Runtime/Engine/Classes/GameFramework/Actor.h
+++ b/Engine/Source/Runtime/Engine/Classes/GameFramework/Actor.h
@@ -3909,6 +3909,7 @@ private:
 		InActor->ActorGuid = InActorGuid;
 	}
 	friend class UWorld;
+	friend class UEngine;
 	friend class UWorldPartitionConvertCommandlet;
 };
 #endif
diff --git a/Engine/Source/Runtime/Engine/Private/UnrealEngine.cpp b/Engine/Source/Runtime/Engine/Private/UnrealEngine.cpp
--- a/Engine/Source/Runtime/Engine/Private/UnrealEngine.cpp
+++ b/Engine/Source/Runtime/Engine/Private/UnrealEngine.cpp
@@ -15651,6 +15651,8 @@ void UEngine::CopyPropertiesForUnrelatedObjects(UObject* OldObject, UObject* New
 
 	check(OldObject && NewObject);
 
+	FGuid NewActorGuid;
+
 	// Bad idea to write data to an actor while its components are registered
 	AActor* NewActor = Cast<AActor>(NewObject);
 	if (NewActor != nullptr)
@@ -15659,6 +15661,8 @@ void UEngine::CopyPropertiesForUnrelatedObjects(UObject* OldObject, UObject* New
 		{
 			ensure(Component == nullptr || !Component->IsRegistered());
 		}
+
+		NewActorGuid = NewActor->GetActorGuid();
 	}
 
 	// If the new object is an Actor, save the root component reference, to be restored later
@@ -15808,7 +15812,7 @@ void UEngine::CopyPropertiesForUnrelatedObjects(UObject* OldObject, UObject* New
 		}
 	}
 
-	// Restore the root component reference
+	// Restore the root component reference and other per-instance properties
 	if (NewActor != nullptr && Params.bPreserveRootComponent)
 	{
 		if (RootComponentProperty != nullptr)
@@ -15817,6 +15821,11 @@ void UEngine::CopyPropertiesForUnrelatedObjects(UObject* OldObject, UObject* New
 		}
 
 		NewActor->ResetOwnedComponents();
+
+		if (NewActorGuid.IsValid())
+		{
+			FSetActorGuid SetActorGuid(NewActor, NewActorGuid);
+		}
 	}
 
 	bool bDumpProperties = CVarDumpCopyPropertiesForUnrelatedObjects.GetValueOnAnyThread() != 0;
