Fix non-editor targets.

diff --git a/Engine/Source/Runtime/Engine/Private/UnrealEngine.cpp b/Engine/Source/Runtime/Engine/Private/UnrealEngine.cpp
--- a/Engine/Source/Runtime/Engine/Private/UnrealEngine.cpp
+++ b/Engine/Source/Runtime/Engine/Private/UnrealEngine.cpp
@@ -15651,7 +15651,9 @@ void UEngine::CopyPropertiesForUnrelatedObjects(UObject* OldObject, UObject* New
 
 	check(OldObject && NewObject);
 
+#if WITH_EDITOR
 	FGuid NewActorGuid;
+#endif
 
 	// Bad idea to write data to an actor while its components are registered
 	AActor* NewActor = Cast<AActor>(NewObject);
@@ -15662,7 +15664,9 @@ void UEngine::CopyPropertiesForUnrelatedObjects(UObject* OldObject, UObject* New
 			ensure(Component == nullptr || !Component->IsRegistered());
 		}
 
+#if WITH_EDITOR
 		NewActorGuid = NewActor->GetActorGuid();
+#endif
 	}
 
 	// If the new object is an Actor, save the root component reference, to be restored later
@@ -15822,10 +15826,12 @@ void UEngine::CopyPropertiesForUnrelatedObjects(UObject* OldObject, UObject* New
 
 		NewActor->ResetOwnedComponents();
 
+#if WITH_EDITOR
 		if (NewActorGuid.IsValid())
 		{
 			FSetActorGuid SetActorGuid(NewActor, NewActorGuid);
 		}
+#endif
 	}
 
 	bool bDumpProperties = CVarDumpCopyPropertiesForUnrelatedObjects.GetValueOnAnyThread() != 0;
