Fix for broken animations when using Distance Matching in cooked builds. FCachedFloatCurve and FDistanceCurve needed to be updated to use FAnimCurveBufferAccess.

diff --git a/Engine/Source/Runtime/Engine/Private/Animation/AnimCurveTypes.cpp b/Engine/Source/Runtime/Engine/Private/Animation/AnimCurveTypes.cpp
--- a/Engine/Source/Runtime/Engine/Private/Animation/AnimCurveTypes.cpp
+++ b/Engine/Source/Runtime/Engine/Private/Animation/AnimCurveTypes.cpp
@@ -5,6 +5,7 @@
 #include "UObject/AnimObjectVersion.h"
 #include "Math/RandomStream.h"
 #include "Animation/AnimSequenceBase.h"
+#include "Animation/AnimCurveCompressionCodec_UniformIndexable.h"
 
 DECLARE_CYCLE_STAT(TEXT("EvalRawCurveData"), STAT_EvalRawCurveData, STATGROUP_Anim);
 
@@ -356,17 +357,12 @@ FVectorCurve* FTransformCurve::GetVectorCurveByIndex(int32 Index)
 
 bool FCachedFloatCurve::IsValid(const UAnimSequenceBase* InAnimSequence) const
 {
-	return ((CurveName != NAME_None) && (GetFloatCurve(InAnimSequence) != nullptr));
+	return ((CurveName != NAME_None) && InAnimSequence->HasCurveData(GetAnimCurveUID(InAnimSequence)));
 }
 
 float FCachedFloatCurve::GetValueAtPosition(const UAnimSequenceBase* InAnimSequence, const float& InPosition) const
 {
-	if (const FFloatCurve* DistanceCurve = GetFloatCurve(InAnimSequence))
-	{
-		return DistanceCurve->Evaluate(InPosition);
-	}
-
-	return 0.f;
+	return InAnimSequence->EvaluateCurveData(GetAnimCurveUID(InAnimSequence), InPosition);
 }
 
 USkeleton::AnimCurveUID FCachedFloatCurve::GetAnimCurveUID(const UAnimSequenceBase* InAnimSequence) const
@@ -407,11 +403,13 @@ const FFloatCurve* FCachedFloatCurve::GetFloatCurve(const UAnimSequenceBase* InA
 
 float FDistanceCurve::GetDistanceRange(const UAnimSequenceBase* InAnimSequence) const
 {
-	if (const FFloatCurve* DistanceCurve = GetFloatCurve(InAnimSequence))
+	FAnimCurveBufferAccess BufferCurveAccess(InAnimSequence, GetAnimCurveUID(InAnimSequence));
+	if (BufferCurveAccess.IsValid())
 	{
-		if (DistanceCurve->FloatCurve.GetNumKeys() >= 2)
+		const int32 NumSamples = BufferCurveAccess.GetNumSamples();
+		if (NumSamples >= 2)
 		{
-			return (DistanceCurve->FloatCurve.GetLastKey().Value - DistanceCurve->FloatCurve.GetFirstKey().Value);
+			return (BufferCurveAccess.GetValue(NumSamples - 1) - BufferCurveAccess.GetValue(0));
 		}
 	}
 	return 0.f;
@@ -419,11 +417,10 @@ float FDistanceCurve::GetDistanceRange(const UAnimSequenceBase* InAnimSequence)
 
 float FDistanceCurve::GetAnimPositionFromDistance(const UAnimSequenceBase* InAnimSequence, const float& InDistance) const
 {
-	if (const FFloatCurve* DistanceCurve = GetFloatCurve(InAnimSequence))
+	FAnimCurveBufferAccess BufferCurveAccess(InAnimSequence, GetAnimCurveUID(InAnimSequence));
+	if (BufferCurveAccess.IsValid())
 	{
-		const TArray<FRichCurveKey>& Keys = DistanceCurve->FloatCurve.GetConstRefOfKeys();
-
-		const int32 NumKeys = Keys.Num();
+		const int32 NumKeys = BufferCurveAccess.GetNumSamples();
 		if (NumKeys < 2)
 		{
 			return 0.f;
@@ -442,7 +439,7 @@ float FDistanceCurve::GetAnimPositionFromDistance(const UAnimSequenceBase* InAni
 			int32 step = count / 2;
 			int32 middle = first + step;
 
-			if (InDistance > Keys[middle].Value)
+			if (InDistance > BufferCurveAccess.GetValue(middle))
 			{
 				first = middle + 1;
 				count -= step + 1;
@@ -453,11 +450,14 @@ float FDistanceCurve::GetAnimPositionFromDistance(const UAnimSequenceBase* InAni
 			}
 		}
 
-		const FRichCurveKey& KeyA = Keys[first - 1];
-		const FRichCurveKey& KeyB = Keys[first];
-		const float Diff = KeyB.Value - KeyA.Value;
-		const float Alpha = !FMath::IsNearlyZero(Diff) ? ((InDistance - KeyA.Value) / Diff) : 0.f;
-		return FMath::Lerp(KeyA.Time, KeyB.Time, Alpha);
+		const float KeyAValue = BufferCurveAccess.GetValue(first - 1);
+		const float KeyBValue = BufferCurveAccess.GetValue(first);
+		const float Diff = KeyBValue - KeyAValue;
+		const float Alpha = !FMath::IsNearlyZero(Diff) ? ((InDistance - KeyAValue) / Diff) : 0.f;
+
+		const float KeyATime = BufferCurveAccess.GetTime(first - 1);
+		const float KeyBTime = BufferCurveAccess.GetTime(first);
+		return FMath::Lerp(KeyATime, KeyBTime, Alpha);
 	}
 
 	return 0.f;
