Allow SetParentEditorOnly to initialize parent even if GIsEditor isn't set

diff --git a/Engine/Source/Runtime/Engine/Private/Materials/MaterialInstanceConstant.cpp b/Engine/Source/Runtime/Engine/Private/Materials/MaterialInstanceConstant.cpp
--- a/Engine/Source/Runtime/Engine/Private/Materials/MaterialInstanceConstant.cpp
+++ b/Engine/Source/Runtime/Engine/Private/Materials/MaterialInstanceConstant.cpp
@@ -80,7 +80,10 @@ void UMaterialInstanceConstant::PostEditChangeProperty(FPropertyChangedEvent& Pr
 
 void UMaterialInstanceConstant::SetParentEditorOnly(UMaterialInterface* NewParent, bool RecacheShader)
 {
-	check(GIsEditor || IsRunningCommandlet());
+	checkf(!Parent || GIsEditor || IsRunningCommandlet(), TEXT("SetParentEditorOnly() may only be used to initialize (not change) the parent outside of the editor, GIsEditor=%d, IsRunningCommandlet()=%d"),
+		GIsEditor ? 1 : 0,
+		IsRunningCommandlet() ? 1 : 0);
+
 	if (SetParentInternal(NewParent, RecacheShader))
 	{
 		UpdateCachedData();
