Change the rules for spawning CheatManager to match expectations and update comments
In Shipping CheatManager is now disabled entirely by default due to UE_WITH_CHEAT_MANAGER being 0
In Test, the EnableCheats command can now be used to force them on in multiplayer like in dev/debug
Debug and Development work the same as before

diff --git a/Engine/Source/Runtime/Engine/Classes/GameFramework/CheatManager.h b/Engine/Source/Runtime/Engine/Classes/GameFramework/CheatManager.h
--- a/Engine/Source/Runtime/Engine/Classes/GameFramework/CheatManager.h
+++ b/Engine/Source/Runtime/Engine/Classes/GameFramework/CheatManager.h
@@ -19,6 +19,7 @@
 
 #include "CheatManager.generated.h"
 
+/** If set to 0, cheat manager functionality will be entirely disabled. If cheats are desired in shipping mode this can be overridden in the target definitions */
 #ifndef UE_WITH_CHEAT_MANAGER
 #define UE_WITH_CHEAT_MANAGER (1 && !UE_BUILD_SHIPPING)
 #endif
diff --git a/Engine/Source/Runtime/Engine/Classes/GameFramework/GameModeBase.h b/Engine/Source/Runtime/Engine/Classes/GameFramework/GameModeBase.h
--- a/Engine/Source/Runtime/Engine/Classes/GameFramework/GameModeBase.h
+++ b/Engine/Source/Runtime/Engine/Classes/GameFramework/GameModeBase.h
@@ -466,7 +466,7 @@ public:
 	//~=============================================================================
 	// Engine hooks
 
-	/** @return true if player is allowed to access the cheats */
+	/** Return true if player should be allowed to use cheats by default */
 	virtual bool AllowCheats(APlayerController* P);
 
 	/** Returns true if replays will start/stop during gameplay starting/stopping */
diff --git a/Engine/Source/Runtime/Engine/Classes/GameFramework/PlayerController.h b/Engine/Source/Runtime/Engine/Classes/GameFramework/PlayerController.h
--- a/Engine/Source/Runtime/Engine/Classes/GameFramework/PlayerController.h
+++ b/Engine/Source/Runtime/Engine/Classes/GameFramework/PlayerController.h
@@ -312,8 +312,9 @@ public:
 	 * Object that manages "cheat" commands.
 	 *
 	 * By default:
-	 *   - Debug and Development builds will force it to be instantiated (@see APlayerController::EnableCheats).
-	 *   - Test and Shipping builds will only instantiate it if the authoritative game mode allows cheats (@see AGameModeBase::AllowCheats).
+	 *	 - In Shipping configurations, the manager is always disabled because UE_WITH_CHEAT_MANAGER is 0
+	 *   - When playing in the editor, cheats are always enabled
+	 *   - In other cases, cheats are enabled by default in single player games but can be forced on with the EnableCheats console command
 	 * 
 	 * This behavior can be changed either by overriding APlayerController::EnableCheats or AGameModeBase::AllowCheats.
 	 */
@@ -553,7 +554,7 @@ public:
 	uint16 LastCompletedSeamlessTravelCount;
 
 public:
-	/** Enables cheats within the game */
+	/** Run from the console to try and manually enable cheats which are disabled by default in multiplayer, games can override this */
 	UFUNCTION(exec)
 	virtual void EnableCheats();
 
@@ -1834,7 +1835,7 @@ public:
 	/** Clears out 'left-over' audio components. */
 	virtual void CleanUpAudioComponents();
 
-	/** Notifies the server that the client has ticked gameplay code, and should no longer get the extended "still loading" timeout grace period */
+	/** Called to try and enable cheats for this player, happens during initialization or from AllowCheats command */
 	virtual void AddCheats(bool bForce = false);
 
 	/** Spawn a HUD (make sure that PlayerController always has valid HUD, even if ClientSetHUD() hasn't been called */
diff --git a/Engine/Source/Runtime/Engine/Private/PlayerController.cpp b/Engine/Source/Runtime/Engine/Private/PlayerController.cpp
--- a/Engine/Source/Runtime/Engine/Private/PlayerController.cpp
+++ b/Engine/Source/Runtime/Engine/Private/PlayerController.cpp
@@ -1012,6 +1012,8 @@ void APlayerController::ServerShortTimeout_Implementation()
 
 void APlayerController::AddCheats(bool bForce)
 {
+	// Cheat manager is completely disabled in shipping by default
+#if UE_WITH_CHEAT_MANAGER
 	UWorld* World = GetWorld();
 	check(World);
 
@@ -1027,18 +1029,19 @@ void APlayerController::AddCheats(bool bForce)
 		CheatManager = NewObject<UCheatManager>(this, CheatClass);
 		CheatManager->InitCheatManager();
 	}
+#endif
 }
 
 void APlayerController::EnableCheats()
 {
-#if !(UE_BUILD_SHIPPING || UE_BUILD_TEST)
+	// In non-shipping builds this can be called to enable cheats in multiplayer and override AllowCheats
+#if !UE_BUILD_SHIPPING
 	AddCheats(true);
 #else
 	AddCheats();
 #endif
 }
 
-
 void APlayerController::SpawnDefaultHUD()
 {
 	if ( Cast<ULocalPlayer>(Player) == NULL )
