Fix for missing decrement in audio Streaming Cache (causing ensure on shutdown)
#rnx

diff --git a/Engine/Source/Runtime/AudioMixerCore/Public/AudioMixer.h b/Engine/Source/Runtime/AudioMixerCore/Public/AudioMixer.h
--- a/Engine/Source/Runtime/AudioMixerCore/Public/AudioMixer.h
+++ b/Engine/Source/Runtime/AudioMixerCore/Public/AudioMixer.h
@@ -114,8 +114,10 @@ namespace EAudioMixerChannel
 
 class USoundWave;
 class FSoundWaveProxy;
+class FSoundWaveData;
 class ICompressedAudioInfo;
 using FSoundWaveProxyPtr = TSharedPtr<FSoundWaveProxy, ESPMode::ThreadSafe>;
+using FSoundWavePtr = TSharedPtr<FSoundWaveData, ESPMode::ThreadSafe>;
 
 
 namespace Audio
diff --git a/Engine/Source/Runtime/Engine/Classes/Sound/SoundWave.h b/Engine/Source/Runtime/Engine/Classes/Sound/SoundWave.h
--- a/Engine/Source/Runtime/Engine/Classes/Sound/SoundWave.h
+++ b/Engine/Source/Runtime/Engine/Classes/Sound/SoundWave.h
@@ -24,6 +24,7 @@
 #include "IAudioProxyInitializer.h"
 #include "SoundWave.generated.h"
 
+class FSoundWaveData;
 class ITargetPlatform;
 struct FActiveSound;
 struct FSoundParseParameters;
@@ -316,39 +317,6 @@
 	virtual void OnFinishDestroy(class USoundWave* Wave) = 0;
 };
 
-struct FSoundWave
-{
-	UE_NONCOPYABLE(FSoundWave);
 
-	FSoundWave() = default;
-
-	/** Zeroth Chunk of audio for sources that use Load On Demand. */
-	FBulkDataBuffer<uint8> ZerothChunkData;
-
-	/** The streaming derived data for this sound on this platform. */
-	FStreamedAudioPlatformData RunningPlatformData;
-
-	FAudioChunkHandle FirstChunk;
-	FFormatContainer CompressedFormatData;
-
-	int32 ResourceSize;
-	FBulkDataBuffer<uint8> ResourceData;
-
-	ESoundWaveLoadingBehavior LoadingBehavior = ESoundWaveLoadingBehavior::Uninitialized;
-
-	// shared flags
-	std::atomic<bool> bIsStreaming;
-	std::atomic<bool> bUseBinkAudio;
-	std::atomic<bool> bSeekableStreaming;
-	std::atomic<bool> bShouldUseStreamCaching;
-	std::atomic<bool> bLoadingBehaviorOverridden;
-
-#if WITH_EDITOR
-	std::atomic<int32> CurrentChunkRevision;
-#endif // #if WITH_EDITOR
-
-}; // struct FSoundWave
-
-
 UCLASS(hidecategories=Object, editinlinenew, BlueprintType, meta= (LoadBehavior = "LazyOnDemand"))
 class ENGINE_API USoundWave : public USoundBase, public IAudioProxyDataFactory
@@ -529,7 +497,7 @@
 	TArray<FSoundWaveClientPtr> SourcesPlaying;
 
 	// This is the sample rate retrieved from platform settings.
-	float CachedSampleRateOverride;
+	mutable float CachedSampleRateOverride;
 
 	// We cache a soundwave's loading behavior on the first call to USoundWave::GetLoadingBehaviorForWave(true);
 	// Caches resolved loading behavior from the SoundClass graph. Must be called on the game thread.
@@ -586,7 +554,7 @@
 private:
 
 	// This is set to false on initialization, then set to true on non-editor platforms when we cache appropriate sample rate.
-	uint8 bCachedSampleRateFromPlatformSettings : 1;
+	mutable uint8 bCachedSampleRateFromPlatformSettings : 1;
 
 	// This is set when SetSampleRate is called to invalidate our cached sample rate while not re-parsing project settings.
 	uint8 bSampleRateManuallyReset : 1;
@@ -668,7 +636,7 @@
 	/** Resource index to cross reference with buffers */
 	int32 ResourceID;
 
-	int32 GetResourceSize() { check(SoundWaveDataPtr);  return SoundWaveDataPtr->ResourceSize; }
+	int32 GetResourceSize() const;
 
 	/** Cache the total used memory recorded for this SoundWave to keep INC/DEC consistent */
 	int32 TrackedMemoryUsage;
@@ -734,7 +702,7 @@
 	/** Memory containing the data copied from the compressed bulk data */
 
 public:
-	const uint8* GetResourceData() { check(SoundWaveDataPtr);  return SoundWaveDataPtr->ResourceData.GetView().GetData(); }
+	const uint8* GetResourceData() const;
 
 	/** Uncompressed wav data 16 bit in mono or stereo - stereo not allowed for multichannel data */
 	FByteBulkData RawData;
@@ -1013,7 +981,7 @@
 	/*
 	* Returns a sample rate if there is a specific sample rate override for this platform, -1.0 otherwise.
 	*/
-	float GetSampleRateForCurrentPlatform();
+	float GetSampleRateForCurrentPlatform() const;
 
 	/**
 	* Return the platform compression overrides set for the current platform.
@@ -1094,29 +1062,33 @@
 		return (ESoundWavePrecacheState)PrecacheState.GetValue();
 	}
 
-	TSharedPtr<FSoundWave, ESPMode::ThreadSafe> SoundWaveDataPtr{ MakeShared<FSoundWave>() };
+	TSharedPtr<FSoundWaveData, ESPMode::ThreadSafe> SoundWaveDataPtr{ MakeShared<FSoundWaveData>() };
 
 private:
 	friend class FSoundWaveProxy;
 };
 
 
-using FSoundWaveProxyPtr = TSharedPtr<FSoundWaveProxy, ESPMode::ThreadSafe>;
-class ENGINE_API FSoundWaveProxy : public Audio::TProxyData<FSoundWaveProxy>
+class ENGINE_API FSoundWaveData
 {
 public:
-	IMPL_AUDIOPROXY_CLASS(FSoundWaveProxy);
+	UE_NONCOPYABLE(FSoundWaveData);
 
-	explicit FSoundWaveProxy(USoundWave* InWave);
-	FSoundWaveProxy(const FSoundWaveProxy& Other) = default;
-
-	Audio::IProxyDataPtr Clone() const override
+	struct MaxChunkSizeResults
 	{
-		LLM_SCOPE(ELLMTag::AudioSoundWaveProxies);
-		return MakeUnique<FSoundWaveProxy>(*this);
-	}
+		uint32 MaxUnevictableSize = 0;
+		uint32 MaxSizeInCache = 0;
+	};
+	
+	// ctor
+	FSoundWaveData() = default;
 
-	// UObject Interface cached data getters
+	// dtor
+	~FSoundWaveData();
+
+	// non-const arg for internal call to FAudioDevice::GetRuntimeFormat
+	void InitializeDataFromSoundWave(USoundWave& InWave);
+
 	const FName& GetFName() const { return NameCached; }
 	const FName& GetPackageName() const { return PackageNameCached; }
 	const FName& GetRuntimeFormat() const { return RuntimeFormat; }
@@ -1125,6 +1097,9 @@
 	float GetSampleRate() const { return SampleRate; }
 	uint32 GetNumChannels() const { return NumChannels; }
 	const TArray<FSoundWaveCuePoint>& GetCuePoints() const { return CuePoints; }
+
+	MaxChunkSizeResults GetMaxChunkSizeResults() const;
+
 	uint32 GetNumChunks() const;
 	uint32 GetSizeOfChunk(uint32 ChunkIndex) const;
 	int32 GetNumFrames() const { return NumFrames; }
@@ -1132,55 +1107,165 @@
 
 	void ReleaseCompressedAudio();
 
-	bool UseBinkAudio() const { check(SoundWaveDataPtr); return SoundWaveDataPtr->bUseBinkAudio; }
-	bool IsStreaming() const { check(SoundWaveDataPtr); return SoundWaveDataPtr->bIsStreaming; }
+	bool UseBinkAudio() const { return bUseBinkAudio; }
+	bool IsStreaming() const { return bIsStreaming; }
 	bool IsLooping() const { return bIsLooping; }
-	bool IsSeekableStreaming() const { check(SoundWaveDataPtr);  return SoundWaveDataPtr->bIsStreaming && SoundWaveDataPtr->bSeekableStreaming; }
-	bool IsRetainingAudio() const { check(SoundWaveDataPtr);  return SoundWaveDataPtr->FirstChunk.IsValid(); }
-	bool WasLoadingBehaviorOverridden() const { check(SoundWaveDataPtr); return SoundWaveDataPtr->bLoadingBehaviorOverridden; }
-	ESoundWaveLoadingBehavior GetLoadingBehavior() const { check(SoundWaveDataPtr); return SoundWaveDataPtr->LoadingBehavior; }
+	bool ShouldUseStreamCaching() const { return bShouldUseStreamCaching; }
+	bool IsSeekableStreaming() const {  return bIsStreaming && bSeekableStreaming; }
+	bool IsRetainingAudio() const {  return FirstChunk.IsValid(); }
+	bool WasLoadingBehaviorOverridden() const { return bLoadingBehaviorOverridden; }
+	ESoundWaveLoadingBehavior GetLoadingBehavior() const { return LoadingBehavior; }
 
 	bool IsTemplate() const { return bIsTemplate; }
 
-	virtual bool HasCompressedData(FName Format, ITargetPlatform* TargetPlatform = USoundWave::GetRunningPlatform()) const;
-	virtual FByteBulkData* GetCompressedData(FName Format, const FPlatformAudioCookOverrides* CompressionOverrides = USoundWave::GetPlatformCompressionOverridesForCurrentPlatform());
+ 	bool HasCompressedData(FName Format, ITargetPlatform* TargetPlatform = USoundWave::GetRunningPlatform()) const;
+ 	FByteBulkData* GetCompressedData(FName Format, const FPlatformAudioCookOverrides* CompressionOverrides = USoundWave::GetPlatformCompressionOverridesForCurrentPlatform());
+ 
+	bool GetChunkData(int32 ChunkIndex, uint8** OutChunkData, bool bMakeSureChunkIsLoaded = false);
+ 
+ 	bool IsZerothChunkDataLoaded() const;
+ 	const TArrayView<uint8> GetZerothChunkDataView() const;
+ 	void EnsureZerothChunkIsLoaded();
+ 
+ #if WITH_EDITOR
+ 	int32 GetCurrentChunkRevision() const;
+ #endif // #if WITH_EDITOR
+ 
+ 	FStreamedAudioChunk& GetChunk(uint32 ChunkIndex);
+ 	int32 GetChunkFromDDC(int32 ChunkIndex, uint8** OutChunkData, bool bMakeSureChunkIsLoaded);
 
-	static TArrayView<const uint8> GetZerothChunk(const FSoundWaveProxyPtr& SoundWaveProxy, bool bForImmediatePlayback = false);
-	bool GetChunkData(int32 ChunkIndex, uint8** OutChunkData, bool bMakeSureChunkIsLoaded = false);
+#if WITH_EDITORONLY_DATA
+ 	FString GetDerivedDataKey() const;
+#endif // #if WITH_EDITORONLY_DATA
 
-	bool IsZerothChunkDataLoaded() const;
-	const TArrayView<uint8> GetZerothChunkDataView() const;
-	void EnsureZerothChunkIsLoaded();
+	int32 GetResourceSize() { return ResourceSize; }
+	const uint8* GetResourceData() const { return ResourceData.GetView().GetData(); }
+	uint32 GetNumChannels() { return NumChannels; }
+
+
+private:
+	/** Zeroth Chunk of audio for sources that use Load On Demand. */
+	FBulkDataBuffer<uint8> ZerothChunkData;
+
+	/** The streaming derived data for this sound on this platform. */
+	FStreamedAudioPlatformData RunningPlatformData;
+
+	FAudioChunkHandle FirstChunk;
+
+	FFormatContainer CompressedFormatData;
+
+	int32 ResourceSize;
+
+	FBulkDataBuffer<uint8> ResourceData;
+
+	ESoundWaveLoadingBehavior LoadingBehavior = ESoundWaveLoadingBehavior::Uninitialized;
 
 #if WITH_EDITOR
-	int32 GetCurrentChunkRevision() const;
+	std::atomic<int32> CurrentChunkRevision;
 #endif // #if WITH_EDITOR
 
-	const FStreamedAudioChunk& GetChunk(uint32 ChunkIndex) const;
-	int32 GetChunkFromDDC(int32 ChunkIndex, uint8** OutChunkData, bool bMakeSureChunkIsLoaded);
-	FString GetDerivedDataKey() const;
-
-	int32 GetResourceSize() { check(SoundWaveDataPtr);  return SoundWaveDataPtr->ResourceSize; }
-	uint8* GetResourceData() { check(SoundWaveDataPtr); return SoundWaveDataPtr->ResourceData.GetView().GetData(); }
-
-	uint32 GetNumChannels() { return NumChannels; }	
-
-	TSharedPtr<FSoundWave, ESPMode::ThreadSafe> SoundWaveDataPtr;
-
-private:
 	FName NameCached;
 	FName PackageNameCached;
-	FObjectKey SoundWaveKeyCached;
 	FName RuntimeFormat{ "FSoundWaveProxy_InvalidFormat" };
+	FObjectKey SoundWaveKeyCached;
+	TArray<FSoundWaveCuePoint> CuePoints;
 
 	float SampleRate;
+	float Duration;
+
 	uint32 NumChannels;
-	TArray<FSoundWaveCuePoint> CuePoints;
 	uint32 NumChunks;
-	float Duration;
 	int32 NumFrames;
 
-	bool bIsLooping { false };
-	bool bIsTemplate { false };
+	// shared flags
+	uint8 bIsLooping : 1;
+	uint8 bIsTemplate : 1;
+	uint8 bIsStreaming : 1;
+	uint8 bUseBinkAudio : 1;
+	uint8 bSeekableStreaming : 1;
+	uint8 bShouldUseStreamCaching : 1;
+	uint8 bLoadingBehaviorOverridden : 1;
+
+	friend class USoundWave;
+}; // class FSoundWaveData
+
+using FSoundWaveProxyPtr = TSharedPtr<FSoundWaveProxy, ESPMode::ThreadSafe>;
+class ENGINE_API FSoundWaveProxy : public Audio::TProxyData<FSoundWaveProxy>
+{
+public:
+	IMPL_AUDIOPROXY_CLASS(FSoundWaveProxy);
+
+	explicit FSoundWaveProxy(USoundWave* InWave);
+
+	FSoundWaveProxy(const FSoundWaveProxy& Other) = default;
+
+	~FSoundWaveProxy();
+
+	Audio::IProxyDataPtr Clone() const override
+	{
+		LLM_SCOPE(ELLMTag::AudioSoundWaveProxies);
+		return MakeUnique<FSoundWaveProxy>(*this);
+	}
+
+	// USoundWave Interface
+	void ReleaseCompressedAudio();
+	void EnsureZerothChunkIsLoaded();
+	bool GetChunkData(int32 ChunkIndex, uint8** OutChunkData, bool bMakeSureChunkIsLoaded = false);
+
+	// Getters
+	const FName& GetFName() const;
+	const FName& GetPackageName() const;
+	const FName& GetRuntimeFormat() const;
+	const FObjectKey& GetFObjectKey() const;
+
+	float GetDuration() const;
+	float GetSampleRate() const;
+
+	int32 GetNumFrames() const;
+	uint32 GetNumChunks(); const
+	uint32 GetNumChannels() const;
+	uint32 GetSizeOfChunk(uint32 ChunkIndex) const;
+	const TArray<FSoundWaveCuePoint>& GetCuePoints() const;
+
+	FSoundWaveData::MaxChunkSizeResults GetMaxChunkSizeResults() const;
+
+	bool IsLooping() const;
+	bool IsTemplate() const;
+	bool IsStreaming() const;
+	bool UseBinkAudio() const;
+	bool IsRetainingAudio() const;
+	bool ShouldUseStreamCaching() const;
+	bool IsSeekableStreaming() const;
+	bool IsZerothChunkDataLoaded() const;
+	bool WasLoadingBehaviorOverridden() const;
+	bool HasCompressedData(FName Format, ITargetPlatform* TargetPlatform = USoundWave::GetRunningPlatform()) const;
+
+	ESoundWaveLoadingBehavior GetLoadingBehavior() const;
+
+	const TArrayView<uint8> GetZerothChunkDataView() const;
+
+	FByteBulkData* GetCompressedData(FName Format, const FPlatformAudioCookOverrides* CompressionOverrides = USoundWave::GetPlatformCompressionOverridesForCurrentPlatform());
+
+	static TArrayView<const uint8> GetZerothChunk(const FSoundWaveProxyPtr& SoundWaveProxy, bool bForImmediatePlayback = false);
+
+
+#if WITH_EDITOR
+	int32 GetCurrentChunkRevision() const;
+#endif // #if WITH_EDITOR
+
+	FStreamedAudioChunk& GetChunk(uint32 ChunkIndex);
+	int32 GetChunkFromDDC(int32 ChunkIndex, uint8** OutChunkData, bool bMakeSureChunkIsLoaded);
+
+#if WITH_EDITORONLY_DATA
+	FString GetDerivedDataKey() const;
+#endif // #if WITH_EDITORONLY_DATA
+
+	int32 GetResourceSize() const;
+	const uint8* GetResourceData() const;
+
+	const FSoundWavePtr GetSoundWaveData();
+
+private:
+	TSharedPtr<FSoundWaveData, ESPMode::ThreadSafe> SoundWaveDataPtr;
 };
 
diff --git a/Engine/Source/Runtime/Engine/Private/AudioDevice.cpp b/Engine/Source/Runtime/Engine/Private/AudioDevice.cpp
--- a/Engine/Source/Runtime/Engine/Private/AudioDevice.cpp
+++ b/Engine/Source/Runtime/Engine/Private/AudioDevice.cpp
@@ -1735,18 +1735,11 @@ bool FAudioDevice::HandleAudioMemoryInfo(const TCHAR* Cmd, FOutputDevice& Ar)
 					break;
 			}
 
-			uint32 MaxUnevictableSize = 0;
-			uint32 MaxSizeInCache = 0;
-			
 			check(SoundWave->SoundWaveDataPtr);
-			for (auto& Chunk : SoundWave->SoundWaveDataPtr->RunningPlatformData.Chunks)
-			{
-				MaxUnevictableSize = FMath::Max<uint32>(MaxUnevictableSize, Chunk.AudioDataSize);
-				MaxSizeInCache += Chunk.AudioDataSize;
-			}
+			FSoundWaveData::MaxChunkSizeResults MaxChunkSizes = SoundWave->SoundWaveDataPtr->GetMaxChunkSizeResults();
 
 			// Add the info to the SoundWaveObjects array
-			SoundWaveObjects.Add(FSoundWaveInfo(SoundWave, TrueResourceSize, SoundGroupName, SoundWave->Duration, LoadType, MaxUnevictableSize, MaxSizeInCache));
+			SoundWaveObjects.Add(FSoundWaveInfo(SoundWave, TrueResourceSize, SoundGroupName, SoundWave->Duration, LoadType, MaxChunkSizes.MaxUnevictableSize, MaxChunkSizes.MaxSizeInCache));
 
 			// Track total resource usage
 			TotalResourceSize += TrueResourceSize;
diff --git a/Engine/Source/Runtime/Engine/Private/AudioStreamingCache.cpp b/Engine/Source/Runtime/Engine/Private/AudioStreamingCache.cpp
--- a/Engine/Source/Runtime/Engine/Private/AudioStreamingCache.cpp
+++ b/Engine/Source/Runtime/Engine/Private/AudioStreamingCache.cpp
@@ -227,24 +227,23 @@ static FAutoConsoleCommand GDisableProfilingAudioCacheCommand(
 
 
 FAudioChunkCache::FChunkKey::FChunkKey(const FChunkKey& Other)
-	: SoundWaveProxyWeakPtr(Other.SoundWaveProxyWeakPtr)
+	: SoundWaveWeakPtr(Other.SoundWaveWeakPtr)
 	, ChunkIndex(Other.ChunkIndex)
 	, ObjectKey(Other.ObjectKey)
 #if WITH_EDITOR
 	, ChunkRevision(Other.ChunkRevision)
 #endif // #if WITH_EDITOR
 {
-	if (FSoundWaveProxyPtr ProxyPtr = SoundWaveProxyWeakPtr.Pin())
+	if (FSoundWavePtr SoundWaveDataPtr = SoundWaveWeakPtr.Pin())
 	{
-		SoundWaveName = ProxyPtr->GetFName();
-		ProxyPtr->ReleaseCompressedAudio();
+		SoundWaveName = SoundWaveDataPtr->GetFName();
 	}
 }
 
 
 FAudioChunkCache::FChunkKey& FAudioChunkCache::FChunkKey::operator=(const FChunkKey& Other)
 {
-	SoundWaveProxyWeakPtr = Other.SoundWaveProxyWeakPtr;
+	SoundWaveWeakPtr = Other.SoundWaveWeakPtr;
 	SoundWaveName = Other.SoundWaveName;
 	ChunkIndex = Other.ChunkIndex;
 	ObjectKey = Other.ObjectKey;
@@ -257,22 +256,22 @@ FAudioChunkCache::FChunkKey& FAudioChunkCache::FChunkKey::operator=(const FChunk
 	return *this;
 }
 
-FAudioChunkCache::FChunkKey::FChunkKey(const FSoundWaveProxyPtr& InSoundWave, uint32 InChunkIndex
+FAudioChunkCache::FChunkKey::FChunkKey(const FSoundWavePtr& InSoundWave, uint32 InChunkIndex
 #if WITH_EDITOR
 	, uint32 InChunkRevision
 #endif // #if WITH_EDITOR
 )
 
-	: SoundWaveProxyWeakPtr(InSoundWave)
+	: SoundWaveWeakPtr(InSoundWave)
 	, ChunkIndex(InChunkIndex)
 #if WITH_EDITOR
 	, ChunkRevision(InChunkRevision)
 #endif // #if WITH_EDITOR
 {
-	if (FSoundWaveProxyPtr ProxyPtr = SoundWaveProxyWeakPtr.Pin())
+	if (FSoundWavePtr SoundWavePtr = SoundWaveWeakPtr.Pin())
 	{
-		SoundWaveName = ProxyPtr->GetFName();
-		ObjectKey = ProxyPtr->GetFObjectKey();
+		SoundWaveName = SoundWavePtr->GetFName();
+		ObjectKey = SoundWavePtr->GetFObjectKey();
 	}
 }
 
@@ -288,11 +287,11 @@ bool FAudioChunkCache::FChunkKey::operator==(const FChunkKey& Other) const
 
 bool FAudioChunkCache::FChunkKey::IsChunkStale()
 {
-	if(FSoundWaveProxyPtr ProxyPtr = SoundWaveProxyWeakPtr.Pin())
+	if(FSoundWavePtr SoundWaveDataPtr = SoundWaveWeakPtr.Pin())
 	{
-		return !ProxyPtr.IsValid()
+		return !SoundWaveDataPtr.IsValid()
 #if WITH_EDITOR
-		|| (ProxyPtr->GetCurrentChunkRevision() != ChunkRevision)
+			|| (SoundWaveDataPtr->GetCurrentChunkRevision() != ChunkRevision)
 #endif // #if WITH_EDITOR
 		;
 
@@ -303,9 +302,9 @@ bool FAudioChunkCache::FChunkKey::IsChunkStale()
 
 ESoundWaveLoadingBehavior FAudioChunkCache::FChunkKey::GetLoadingBehavior() const
 {
-	if(FSoundWaveProxyPtr ProxyPtr = SoundWaveProxyWeakPtr.Pin())
+	if(FSoundWavePtr SoundWaveDataPtr = SoundWaveWeakPtr.Pin())
 	{
-		return ProxyPtr->GetLoadingBehavior();
+		return SoundWaveDataPtr->GetLoadingBehavior();
 	}
 
 	return ESoundWaveLoadingBehavior::Uninitialized;
@@ -313,26 +312,25 @@ ESoundWaveLoadingBehavior FAudioChunkCache::FChunkKey::GetLoadingBehavior() cons
 
 FStreamedAudioChunk& FAudioChunkCache::FChunkKey::GetChunk(uint32 InChunkIndex) const
 {
-	FSoundWaveProxyPtr ProxyPtr = SoundWaveProxyWeakPtr.Pin();
+	FSoundWavePtr SoundWaveDataPtr = SoundWaveWeakPtr.Pin();
 
 	// the Weakptr should be valid here since it's from a shared ptr up the stack
-	check(ProxyPtr.IsValid());
-	check(ProxyPtr->SoundWaveDataPtr);
+	check(SoundWaveDataPtr.IsValid());
 
 	// This function shouldn't be called on audio marked "ForceInline."
-	ensureMsgf(ProxyPtr->GetLoadingBehavior() != ESoundWaveLoadingBehavior::ForceInline, TEXT("Calling GetNumChunks on a FSoundWaveProxy that is Force-Inlined is not allowed! SoundWave: %s - %s")
-		, *ProxyPtr->GetFName().ToString(), EnumToString(ProxyPtr->GetLoadingBehavior()));
+	ensureMsgf(SoundWaveDataPtr->GetLoadingBehavior() != ESoundWaveLoadingBehavior::ForceInline, TEXT("Calling GetNumChunks on a FSoundWaveProxy that is Force-Inlined is not allowed! SoundWave: %s - %s")
+		, *SoundWaveDataPtr->GetFName().ToString(), EnumToString(SoundWaveDataPtr->GetLoadingBehavior()));
 
-	check((ChunkIndex < (uint32)ProxyPtr->GetNumChunks()));
+	check((ChunkIndex < (uint32)SoundWaveDataPtr->GetNumChunks()));
 
-	return ProxyPtr->SoundWaveDataPtr->RunningPlatformData.Chunks[ChunkIndex];
+	return SoundWaveDataPtr->GetChunk(ChunkIndex);
 }
 
 uint32 FAudioChunkCache::FChunkKey::GetNumChunks() const
 {
-	if (FSoundWaveProxyPtr ProxyPtr = SoundWaveProxyWeakPtr.Pin())
+	if (FSoundWavePtr SoundWaveDataPtr = SoundWaveWeakPtr.Pin())
 	{
-		return ProxyPtr->GetNumChunks();
+		return SoundWaveDataPtr->GetNumChunks();
 	}
 	
 	return 0;
@@ -340,9 +338,9 @@ uint32 FAudioChunkCache::FChunkKey::GetNumChunks() const
 
 bool FAudioChunkCache::FChunkKey::IsRetainingAudio() const
 {
-	if (FSoundWaveProxyPtr ProxyPtr = SoundWaveProxyWeakPtr.Pin())
+	if (FSoundWavePtr SoundWaveDataPtr = SoundWaveWeakPtr.Pin())
 	{
-		return ProxyPtr->IsRetainingAudio();
+		return SoundWaveDataPtr->IsRetainingAudio();
 	}
 
 	return false;
@@ -350,18 +348,18 @@ bool FAudioChunkCache::FChunkKey::IsRetainingAudio() const
 
 void FAudioChunkCache::FChunkKey::ReleaseCompressedAudio()
 {
-	if (FSoundWaveProxyPtr ProxyPtr = SoundWaveProxyWeakPtr.Pin())
+	if (FSoundWavePtr SoundWaveDataPtr = SoundWaveWeakPtr.Pin())
 	{
-		check(ProxyPtr.IsValid());
-		ProxyPtr->ReleaseCompressedAudio();
+		check(SoundWaveDataPtr.IsValid());
+		SoundWaveDataPtr->ReleaseCompressedAudio();
 	}
 }
 
 bool FAudioChunkCache::FChunkKey::WasLoadingBehaviorOverridden() const
 {
-	if(FSoundWaveProxyPtr ProxyPtr = SoundWaveProxyWeakPtr.Pin())
+	if(FSoundWavePtr SoundWaveDataPtr = SoundWaveWeakPtr.Pin())
 	{
-		return ProxyPtr->WasLoadingBehaviorOverridden();
+		return SoundWaveDataPtr->WasLoadingBehaviorOverridden();
 	}
 
 	return false;
@@ -517,7 +515,7 @@ FAudioChunkHandle FCachedAudioStreamingManager::GetLoadedChunk(const FSoundWaveP
 		checkf(ChunkIndex != 0, TEXT("Decoder tried to access the zeroth chunk through the streaming manager. Use USoundWave::GetZerothChunk() instead."));
 
 		FAudioChunkCache::FChunkKey ChunkKey(
-			  SoundWave
+			  SoundWave->GetSoundWaveData()
 			, ChunkIndex
 #if WITH_EDITOR
 			, (uint32)SoundWave->GetCurrentChunkRevision()
@@ -548,7 +546,7 @@ FAudioChunkHandle FCachedAudioStreamingManager::GetLoadedChunk(const FSoundWaveP
 		if (NextChunk != INDEX_NONE)
 		{
 			FAudioChunkCache::FChunkKey NextChunkKey(
-				SoundWave
+				SoundWave->GetSoundWaveData()
 				, ((uint32)NextChunk)
 #if WITH_EDITOR
 				, (uint32)SoundWave->GetCurrentChunkRevision()
@@ -684,7 +682,7 @@ bool FCachedAudioStreamingManager::RequestChunk(const FSoundWaveProxyPtr& SoundW
 	if (Cache && ensure(SoundWave.IsValid()))
 	{
 		FAudioChunkCache::FChunkKey ChunkKey(
-			SoundWave
+			  SoundWave->GetSoundWaveData()
 			, ChunkIndex
 #if WITH_EDITOR
 			, (uint32)SoundWave->GetCurrentChunkRevision()
diff --git a/Engine/Source/Runtime/Engine/Private/ContentStreaming.cpp b/Engine/Source/Runtime/Engine/Private/ContentStreaming.cpp
--- a/Engine/Source/Runtime/Engine/Private/ContentStreaming.cpp
+++ b/Engine/Source/Runtime/Engine/Private/ContentStreaming.cpp
@@ -1453,7 +1453,7 @@ FAudioChunkHandle::FAudioChunkHandle()
 FAudioChunkHandle::FAudioChunkHandle(const uint8* InData, uint32 NumBytes, const FSoundWaveProxyPtr&  InSoundWave, const FName& SoundWaveName, uint32 InChunkIndex, uint64 InCacheLookupID)
 	: CachedData(InData)
 	, CachedDataNumBytes(NumBytes)
-	, CorrespondingWave(InSoundWave)
+	, CorrespondingWave(InSoundWave->GetSoundWaveData())
 	, CorrespondingWaveName(SoundWaveName)
 	, ChunkIndex(InChunkIndex)
 #if WITH_EDITOR
@@ -1557,18 +1557,18 @@ uint32 FAudioChunkHandle::Num() const
 
 bool FAudioChunkHandle::IsValid() const
 {
-	return GetData() != nullptr;
+	return GetData() && CorrespondingWave.Pin().IsValid();
 }
 
 #if WITH_EDITOR
 bool FAudioChunkHandle::IsStale() const
 {
-	FSoundWaveProxyPtr ProxyPtr = CorrespondingWave.Pin();
+	TSharedPtr<FSoundWaveData, ESPMode::ThreadSafe> SoundWaveDataPtr = CorrespondingWave.Pin();
 
-	if (ProxyPtr.IsValid())
+	if (SoundWaveDataPtr.IsValid())
 	{
 		// NOTE: While this is currently safe in editor, there's no guarantee the USoundWave will be kept alive during the lifecycle of this chunk handle.
-		return ChunkGeneration != ProxyPtr->GetCurrentChunkRevision();
+		return ChunkGeneration != SoundWaveDataPtr->GetCurrentChunkRevision();
 	}
 	else
 	{
diff --git a/Engine/Source/Runtime/Engine/Private/SoundWave.cpp b/Engine/Source/Runtime/Engine/Private/SoundWave.cpp
--- a/Engine/Source/Runtime/Engine/Private/SoundWave.cpp
+++ b/Engine/Source/Runtime/Engine/Private/SoundWave.cpp
@@ -136,6 +136,310 @@
 }
 #endif
 
+FSoundWaveData::~FSoundWaveData()
+{
+	// at this point, first chunk data should have been cleared by the last reference holder
+	// (see USoundWave::BeginDestroy and ~FSoundWaveProxy implmentations)
+	ensure(FirstChunk.GetData() == nullptr);
+}
+
+void FSoundWaveData::InitializeDataFromSoundWave(USoundWave& InWave)
+{
+	LLM_SCOPE(ELLMTag::AudioSoundWaves);
+
+	NameCached = InWave.GetFName();
+
+	// cache package name
+	if (UPackage* Package = InWave.GetPackage())
+	{
+		PackageNameCached = Package->GetFName();
+	}
+
+	// cache the runtime format for the wave
+	if (GEngine)
+	{
+		FAudioDevice* LocalAudioDevice = GEngine->GetMainAudioDeviceRaw();
+		if (LocalAudioDevice)
+		{
+			RuntimeFormat = LocalAudioDevice->GetRuntimeFormat(&InWave);
+		}
+	}
+
+	SoundWaveKeyCached = FObjectKey(&InWave);
+	CuePoints = InWave.CuePoints;
+	SampleRate = InWave.GetSampleRateForCurrentPlatform();
+	Duration = InWave.Duration;
+	NumChannels = InWave.NumChannels;
+	
+	// GetNumChunks() can't be called w/o RunningPlateformData
+	NumChunks = 0;
+	if (RunningPlatformData.NumChunks > 0)
+	{
+		NumChunks = InWave.GetNumChunks();
+	}
+
+	NumFrames = (int32)(Duration * (float)SampleRate);
+
+	// update shared flags
+	bIsLooping = InWave.IsLooping();
+	bIsTemplate = InWave.IsTemplate();
+	bIsStreaming = InWave.IsStreaming(nullptr);
+	bUseBinkAudio = InWave.bUseBinkAudio;
+	bSeekableStreaming = InWave.bSeekableStreaming;
+	bShouldUseStreamCaching = InWave.ShouldUseStreamCaching();
+}
+
+uint32 FSoundWaveData::GetNumChunks() const
+{
+	// This function shouldn't be called on audio marked "ForceInline."
+	ensureMsgf(LoadingBehavior != ESoundWaveLoadingBehavior::ForceInline, TEXT("Calling GetNumChunks without RunnigPlatformData is not allowed! SoundWave: %s - %s")
+		, *GetFName().ToString(), EnumToString(LoadingBehavior));
+
+	return RunningPlatformData.Chunks.Num();
+}
+
+FSoundWaveData::MaxChunkSizeResults FSoundWaveData::GetMaxChunkSizeResults() const
+{
+	MaxChunkSizeResults Results;
+
+	for (auto& Chunk : RunningPlatformData.Chunks)
+	{
+		Results.MaxUnevictableSize = FMath::Max<uint32>(Results.MaxUnevictableSize, Chunk.AudioDataSize);
+		Results.MaxSizeInCache += Chunk.AudioDataSize;
+	}
+
+	return Results;
+}
+
+uint32 FSoundWaveData::GetSizeOfChunk(uint32 ChunkIndex) const
+{
+	check(ChunkIndex < GetNumChunks());
+
+	// This function shouldn't be called on audio marked "ForceInline."
+	ensureMsgf(RunningPlatformData.NumChunks
+		, TEXT("Calling GetNumChunks() without RunnigPlatformData is not allowed! SoundWave: %s - %s")
+		, *GetFName().ToString()
+		, EnumToString(LoadingBehavior));
+
+	check((ChunkIndex < (uint32)GetNumChunks()));
+	return RunningPlatformData.Chunks[ChunkIndex].AudioDataSize;
+}
+
+void FSoundWaveData::ReleaseCompressedAudio()
+{
+	LLM_SCOPE(ELLMTag::AudioSoundWaves);
+
+	FirstChunk = FAudioChunkHandle();
+}
+
+bool FSoundWaveData::HasCompressedData(FName Format, ITargetPlatform* TargetPlatform) const
+{
+	if (IsTemplate() || IsRunningDedicatedServer())
+	{
+		return false;
+	}
+
+	const FPlatformAudioCookOverrides* CompressionOverrides = nullptr;
+
+	if (GIsEditor)
+	{
+		if (TargetPlatform)
+		{
+			CompressionOverrides = FPlatformCompressionUtilities::GetCookOverrides(*TargetPlatform->IniPlatformName());
+		}
+	}
+	else
+	{
+		// TargetPlatform is not available on consoles/mobile, so we have to grab it ourselves:
+		CompressionOverrides = FPlatformCompressionUtilities::GetCookOverrides();
+	}
+
+	if (CompressionOverrides)
+	{
+#if WITH_EDITOR
+		FName PlatformSpecificFormat;
+		FString HashedString = *Format.ToString();
+		FPlatformAudioCookOverrides::GetHashSuffix(CompressionOverrides, HashedString);
+		PlatformSpecificFormat = *HashedString;
+#else
+		// on non-editor builds, we cache the concatenated format in a static FName.
+		static FName PlatformSpecificFormat;
+		static FName CachedFormat;
+		if (!Format.IsEqual(CachedFormat))
+		{
+			FString HashedString = *Format.ToString();
+			FPlatformAudioCookOverrides::GetHashSuffix(CompressionOverrides, HashedString);
+			PlatformSpecificFormat = *HashedString;
+
+			CachedFormat = Format;
+		}
+#endif // WITH_EDITOR
+		return CompressedFormatData.Contains(PlatformSpecificFormat);
+	}
+	else
+	{
+		return CompressedFormatData.Contains(Format);
+	}
+}
+
+FByteBulkData* FSoundWaveData::GetCompressedData(FName Format, const FPlatformAudioCookOverrides* CompressionOverrides)
+{
+	if (IsTemplate() || IsRunningDedicatedServer())
+	{
+		return nullptr;
+	}
+
+	TRACE_CPUPROFILER_EVENT_SCOPE(FSoundWaveProxy::GetCompressedData);
+	FName PlatformSpecificFormat = USoundWave::GetPlatformSpecificFormat(Format, CompressionOverrides);
+
+	bool bContainedValidData = CompressedFormatData.Contains(PlatformSpecificFormat);
+	if (!bContainedValidData)
+	{
+		// the proxy does not support interaction w/ FDerivedAudioDataCompressor
+		return nullptr;
+	}
+
+	FByteBulkData* Result = &CompressedFormatData.GetFormat(PlatformSpecificFormat);
+	check(Result);
+	return Result->GetBulkDataSize() > 0 ? Result : NULL; // we don't return empty bulk data...but we save it to avoid thrashing the DDC
+}
+
+bool FSoundWaveData::GetChunkData(int32 ChunkIndex, uint8** OutChunkData, bool bMakeSureChunkIsLoaded)
+{
+	LLM_SCOPE(ELLMTag::AudioSoundWaves);
+	if (GetChunkFromDDC(ChunkIndex, OutChunkData, bMakeSureChunkIsLoaded) == 0)
+	{
+#if WITH_EDITORONLY_DATA
+		// Unable to load chunks from the cache. Rebuild the sound and attempt to precache it.
+		UE_LOG(LogAudio, Display, TEXT("GetChunkData failed, rebuilding %s"), *GetPackageName().ToString());
+
+		//		ForceRebuildPlatformData();
+		if (GetChunkFromDDC(ChunkIndex, OutChunkData, bMakeSureChunkIsLoaded) == 0)
+		{
+			UE_LOG(LogAudio, Warning, TEXT("Failed to build sound %s."), *GetPackageName().ToString());
+		}
+		else
+		{
+			// Succeeded after rebuilding platform data
+			return true;
+		}
+#else
+		// Failed to find the SoundWave chunk in the cooked package.
+		UE_LOG(LogAudio, Warning, TEXT("GetChunkData failed while streaming. Ensure the following package is cooked: %s"), *GetPackageName().ToString());
+#endif // #if WITH_EDITORONLY_DATA
+		return false;
+	}
+	return true;
+}
+
+bool FSoundWaveData::IsZerothChunkDataLoaded() const
+{
+	return ZerothChunkData.GetView().Num() > 0;
+}
+
+const TArrayView<uint8> FSoundWaveData::GetZerothChunkDataView() const
+{
+	FBulkDataBuffer<uint8>::ViewType View = ZerothChunkData.GetView();
+	check(View.Num() <= MAX_int32);
+	return TArrayView<uint8>(View.GetData(), static_cast<int32>(View.Num()));
+}
+
+void FSoundWaveData::EnsureZerothChunkIsLoaded()
+{
+	LLM_SCOPE(ELLMTag::AudioSoundWaves);
+
+	// If the zeroth chunk is already loaded, early exit.
+	if (ZerothChunkData.GetView().Num() > 0 || !bShouldUseStreamCaching)
+	{
+		return;
+	}
+
+	// #if WITH_EDITOR
+#ifdef WITH_EDITOR
+	// If we're running the editor, we'll need to retrieve the chunked audio from the DDC:
+	uint8* TempChunkBuffer = nullptr;
+	int32 ChunkSizeInBytes = GetChunkFromDDC(0, &TempChunkBuffer, true);
+	// Since we block for the DDC in the previous call we should always have the chunk loaded.
+	if (ChunkSizeInBytes == 0)
+	{
+		return;
+	}
+
+	ZerothChunkData.Reset(TempChunkBuffer, ChunkSizeInBytes);
+
+#else // WITH_EDITOR
+	// Otherwise, the zeroth chunk is cooked out to SharedRuntimeData::RunningPlatformData, and we just need to retrieve it.
+	check(GetNumChunks() > 0);
+	const FStreamedAudioChunk& ZerothChunk = GetChunk(0);
+	// Some sanity checks to ensure that the bulk size set up
+	UE_CLOG(ZerothChunk.BulkData.GetBulkDataSize() != ZerothChunk.DataSize, LogAudio, Warning
+		, TEXT("Bulk data serialized out had a mismatched size with the DataSize field. Soundwave: %s Bulk Data Reported Size: %d Bulk Data Actual Size: %ld")
+		, *GetFullName(), ZerothChunk.DataSize, ZerothChunk.BulkData.GetBulkDataSize());
+
+	ZerothChunkData = ZerothChunk.BulkData.GetCopyAsBuffer(ZerothChunk.AudioDataSize, true);
+#endif // WITH_EDITOR
+}
+
+#if WITH_EDITOR
+int32 FSoundWaveData::GetCurrentChunkRevision() const
+{
+	return CurrentChunkRevision;
+}
+#endif // #if WITH_EDITOR
+
+FStreamedAudioChunk& FSoundWaveData::GetChunk(uint32 ChunkIndex)
+{
+	// This function shouldn't be called on audio marked "ForceInline."
+	ensureMsgf(RunningPlatformData.NumChunks
+		, TEXT("Calling GetNumChunks() without RunnigPlatformData is not allowed! SoundWave: %s - %s")
+		, *GetFName().ToString()
+		, EnumToString(LoadingBehavior));
+
+	check((ChunkIndex < (uint32)GetNumChunks()));
+	return RunningPlatformData.Chunks[ChunkIndex];
+}
+
+int32 FSoundWaveData::GetChunkFromDDC(int32 ChunkIndex, uint8** OutChunkData, bool bMakeSureChunkIsLoaded)
+{
+	LLM_SCOPE(ELLMTag::AudioSoundWaves);
+	// This function shouldn't be called on audio marked "ForceInline."
+	ensureMsgf(RunningPlatformData.NumChunks
+		, TEXT("Calling GetNumChunks on a FSoundWaveProxy without RunnigPlatformData is not allowed! SoundWave: %s - %s")
+		, *GetFName().ToString()
+		, EnumToString(LoadingBehavior));
+
+	return RunningPlatformData.GetChunkFromDDC(ChunkIndex, OutChunkData, bMakeSureChunkIsLoaded);
+}
+
+#if WITH_EDITORONLY_DATA
+FString FSoundWaveData::GetDerivedDataKey() const
+{
+	check(RunningPlatformData.NumChunks);
+
+	// This function shouldn't be called on audio marked "ForceInline."
+	ensureMsgf(RunningPlatformData.NumChunks
+		, TEXT("Calling GetNumChunks on a FSoundWaveProxy without RunnigPlatformData is not allowed! SoundWave: %s - %s")
+		, *GetFName().ToString()
+		, EnumToString(LoadingBehavior));
+
+	return RunningPlatformData.DerivedDataKey;
+}
+#endif // #if WITH_EDITORONLY_DATA
+
+
+int32 USoundWave::GetResourceSize() const
+{
+	check(SoundWaveDataPtr);
+	return SoundWaveDataPtr->ResourceSize;
+}
+
+const uint8* USoundWave::GetResourceData() const
+{
+	check(SoundWaveDataPtr);
+	return SoundWaveDataPtr->GetResourceData();
+}
+
+
 ITargetPlatform* USoundWave::GetRunningPlatform()
 {
 	ITargetPlatformManagerModule* TPM = GetTargetPlatformManager();
@@ -283,7 +587,6 @@
 	bCachedSampleRateFromPlatformSettings = false;
 	bSampleRateManuallyReset = false;
 	CachedSampleRateOverride = 0.0f;
-	SoundWaveDataPtr->LoadingBehavior = ESoundWaveLoadingBehavior::Uninitialized;
 
 	bProcedural = false;
 	bRequiresStopFade = false;
@@ -1227,16 +1530,10 @@
 	INC_FLOAT_STAT_BY(STAT_AudioBufferTime, Duration);
 	INC_FLOAT_STAT_BY(STAT_AudioBufferTimeChannels, NumChannels * Duration);
 
-	// cache current state as a proxy if loading behavior dictates we'll be using the streaming cache
+	// update shared state if loading behavior dictates we'll be using the streaming cache
 	if (ActualLoadingBehavior != ESoundWaveLoadingBehavior::ForceInline)
 	{
-		// cache current state as a proxy
-		InternalProxy = CreateSoundWaveProxy();
-		if (InternalProxy.IsValid())
-		{
-			// release dupe handle already held by 'this'
-			InternalProxy->ReleaseCompressedAudio();
-		}
+		SoundWaveDataPtr->InitializeDataFromSoundWave(*this);
 	}
 }
 
@@ -1343,11 +1640,19 @@
 	FlushAsyncLoadingDataFormats();
 #endif
 
-	ReleaseCompressedAudio();
+	// the FSoundWaveDataPtr is not able to release the handle itself
+	// (weakptr that is already invalid in ~FSoundWaveData)
+	// so we need to Release our handle to the chunk in the streaming cache
+	// if we are the last one holding a refernce to SoundWaveDataPtr
+	if (SoundWaveDataPtr.GetSharedReferenceCount() == 1)
+	{
+		ReleaseCompressedAudio();
+	}
 }
 
 void USoundWave::InitAudioResource(FByteBulkData& CompressedData)
 {
+	LLM_SCOPE(ELLMTag::AudioSoundWaves);
 	check(SoundWaveDataPtr);
 
 	if (!SoundWaveDataPtr->ResourceSize)
@@ -1396,6 +1701,7 @@
 
 bool USoundWave::InitAudioResource(FName Format)
 {
+	LLM_SCOPE(ELLMTag::AudioSoundWaves);
 	check(SoundWaveDataPtr);
 
 	if (!SoundWaveDataPtr->ResourceSize && (!FPlatformProperties::SupportsAudioStreaming() || !IsStreaming(nullptr)))
@@ -1868,6 +2174,8 @@
 
 void USoundWave::PostEditChangeProperty(FPropertyChangedEvent& PropertyChangedEvent)
 {
+	// maxtodo: test this, stale data? looping? channel count?
+
 	Super::PostEditChangeProperty(PropertyChangedEvent);
 	check(SoundWaveDataPtr);
 
@@ -2150,15 +2458,6 @@
 		ClearAllCachedCookedPlatformData();
 	}
 #endif
-
-	if (FApp::CanEverRenderAudio())
-	{
-		if (!InternalProxy.IsValid())
-		{
-			InternalProxy = CreateSoundWaveProxy();
-		}
-		IStreamingManager::Get().GetAudioStreamingManager().RemoveStreamingSoundWave(InternalProxy);
-	}
 }
 
 void USoundWave::Parse(FAudioDevice* AudioDevice, const UPTRINT NodeWaveInstanceHash, FActiveSound& ActiveSound, const FSoundParseParameters& ParseParams, TArray<FWaveInstance*>& WaveInstances)
@@ -2663,18 +2962,14 @@
 
 FSoundWaveProxyPtr USoundWave::CreateSoundWaveProxy()
 {
-	check(SoundWaveDataPtr);
-
 #if WITH_EDITORONLY_DATA
 	EnsureZerothChunkIsLoaded();
 #endif // #if WITH_EDITORONLY_DATA
 
 	LLM_SCOPE(ELLMTag::AudioSoundWaveProxies);
 
-	SoundWaveDataPtr->bUseBinkAudio = bUseBinkAudio;
-	SoundWaveDataPtr->bSeekableStreaming = bSeekableStreaming;
-	SoundWaveDataPtr->bIsStreaming = IsStreaming(nullptr);
-	SoundWaveDataPtr->bShouldUseStreamCaching = ShouldUseStreamCaching();
+	check(SoundWaveDataPtr);
+	SoundWaveDataPtr->InitializeDataFromSoundWave(*this);
 
 	return MakeShared<FSoundWaveProxy, ESPMode::ThreadSafe>(this);
 }
@@ -2689,10 +2984,8 @@
 
 	LLM_SCOPE(ELLMTag::AudioSoundWaveProxies);
 
-	SoundWaveDataPtr->bUseBinkAudio = bUseBinkAudio;
-	SoundWaveDataPtr->bSeekableStreaming = bSeekableStreaming;
-	SoundWaveDataPtr->bIsStreaming = IsStreaming(nullptr);
-	SoundWaveDataPtr->bShouldUseStreamCaching = ShouldUseStreamCaching();
+	check(SoundWaveDataPtr);
+	SoundWaveDataPtr->InitializeDataFromSoundWave(*this);
 
 	return MakeUnique<FSoundWaveProxy>(this);
 }
@@ -2752,7 +3045,7 @@
 	}
 }
 
-float USoundWave::GetSampleRateForCurrentPlatform()
+float USoundWave::GetSampleRateForCurrentPlatform() const
 {
 	if (bProcedural)
 	{
@@ -3355,145 +3648,246 @@
 
 FSoundWaveProxy::FSoundWaveProxy(USoundWave* InWave)
 	: SoundWaveDataPtr(InWave->SoundWaveDataPtr)
-	, NameCached(InWave->GetFName())
-	, SoundWaveKeyCached(FObjectKey(InWave))
-	, SampleRate(InWave->GetSampleRateForCurrentPlatform())
-	, NumChannels(InWave->NumChannels)
-	, CuePoints(InWave->CuePoints)
-	, NumChunks(0)
-	, Duration(InWave->Duration)
-	, NumFrames((int32)(Duration* (float)InWave->SampleRate))
-	, bIsLooping(InWave->IsLooping())
-	, bIsTemplate(InWave->IsTemplate())
 {
 	LLM_SCOPE(ELLMTag::AudioSoundWaveProxies);
 
-	if (UPackage* Package = InWave->GetPackage())
-	{
-		PackageNameCached = Package->GetFName();
-	}
-
 	// this should have been allocated by the USoundWave and should always be valid
 	check(SoundWaveDataPtr);
 
-
-	// cache the runtime format for this wave
-	if (GEngine)
+	// if we are force inline, we need to make sure the shared data is pulled from the DDC
+	// before we attempt to use a decoder on the proxy (not using stream caching)
+	if (SoundWaveDataPtr->GetLoadingBehavior() == ESoundWaveLoadingBehavior::ForceInline && !InWave->GetResourceData())
 	{
-		FAudioDevice* LocalAudioDevice = GEngine->GetMainAudioDeviceRaw();
-		if (LocalAudioDevice)
-		{
-			RuntimeFormat = LocalAudioDevice->GetRuntimeFormat(InWave);
-
-			// if we are force inline, we need to make sure the shared data is pulled from the DDC
-			// before we attempt to use a decoder on the proxy (not using stream caching)
-			if (SoundWaveDataPtr->LoadingBehavior == ESoundWaveLoadingBehavior::ForceInline && !InWave->GetResourceData())
-			{
-				InWave->InitAudioResource(RuntimeFormat);
-				check(SoundWaveDataPtr->ResourceSize > 0);
-			}
-		}
+		InWave->InitAudioResource(SoundWaveDataPtr->GetRuntimeFormat());
+		check(SoundWaveDataPtr->GetResourceSize() > 0);
 	}
-
-	if (InWave->SoundWaveDataPtr->RunningPlatformData.NumChunks > 0)
-	{
-		NumChunks = InWave->GetNumChunks();
-	}
-
-	// force updates of flags shared w/ the sound wave
-	// (assignment is to avoid these calls to const functions w/ mutable side effects being from being optimized out)
-	SoundWaveDataPtr->bIsStreaming = InWave->IsStreaming(nullptr);
-	SoundWaveDataPtr->bShouldUseStreamCaching = InWave->ShouldUseStreamCaching();
 }
 
+FSoundWaveProxy::~FSoundWaveProxy()
+{
+	// the FSoundWaveDataPtr is not able to release the handle itself
+	// (weakptr that is already invalid in ~FSoundWaveData)
+	// so we need to Release our handle to the chunk in the streaming cache
+	// if we are the last one holding a refernce to SoundWaveDataPtr
+	if (SoundWaveDataPtr.GetSharedReferenceCount() == 1)
+	{
+		ReleaseCompressedAudio();
+	}
+}
+
+
+#if WITH_EDITOR
+int32 FSoundWaveProxy::GetCurrentChunkRevision() const
+{
+	check(SoundWaveDataPtr);
+	return SoundWaveDataPtr->GetCurrentChunkRevision();
+}
+#endif // #if WITH_EDITOR
+
+FStreamedAudioChunk& FSoundWaveProxy::GetChunk(uint32 ChunkIndex)
+{
+	check(SoundWaveDataPtr);
+	return SoundWaveDataPtr->GetChunk(ChunkIndex);
+}
+
+int32 FSoundWaveProxy::GetChunkFromDDC(int32 ChunkIndex, uint8** OutChunkData, bool bMakeSureChunkIsLoaded)
+{
+	check(SoundWaveDataPtr);
+	return SoundWaveDataPtr->GetChunkFromDDC(ChunkIndex, OutChunkData, bMakeSureChunkIsLoaded);
+}
+
+#if WITH_EDITORONLY_DATA
+FString FSoundWaveProxy::GetDerivedDataKey() const
+{
+	check(SoundWaveDataPtr);
+	return SoundWaveDataPtr->GetDerivedDataKey();
+}
+#endif // #if WITH_EDITORONLY_DATA
+
+int32 FSoundWaveProxy::GetResourceSize() const
+{
+	check(SoundWaveDataPtr);
+	return SoundWaveDataPtr->GetResourceSize();
+}
+
+const uint8* FSoundWaveProxy::GetResourceData() const
+{
+	check(SoundWaveDataPtr);
+	return SoundWaveDataPtr->GetResourceData();
+}
+
+const FSoundWavePtr FSoundWaveProxy::GetSoundWaveData()
+{
+	check(SoundWaveDataPtr);
+	return SoundWaveDataPtr;
+}
+
 void FSoundWaveProxy::ReleaseCompressedAudio()
 {
-	LLM_SCOPE(ELLMTag::AudioSoundWaveProxies);
 	check(SoundWaveDataPtr);
-	SoundWaveDataPtr->FirstChunk = FAudioChunkHandle();
+	return SoundWaveDataPtr->ReleaseCompressedAudio();
 }
 
+void FSoundWaveProxy::EnsureZerothChunkIsLoaded()
+{
+	check(SoundWaveDataPtr);
+	return SoundWaveDataPtr->EnsureZerothChunkIsLoaded();
+}
+
+bool FSoundWaveProxy::GetChunkData(int32 ChunkIndex, uint8** OutChunkData, bool bMakeSureChunkIsLoaded)
+{
+	check(SoundWaveDataPtr);
+	return SoundWaveDataPtr->GetChunkData(ChunkIndex, OutChunkData, bMakeSureChunkIsLoaded);
+}
+
+const FName& FSoundWaveProxy::GetFName() const
+{
+	check(SoundWaveDataPtr);
+	return SoundWaveDataPtr->GetFName();
+}
+
+const FName& FSoundWaveProxy::GetPackageName() const
+{
+	check(SoundWaveDataPtr);
+	return SoundWaveDataPtr->GetPackageName();
+}
+
+const FName& FSoundWaveProxy::GetRuntimeFormat() const
+{
+	check(SoundWaveDataPtr);
+	return SoundWaveDataPtr->GetRuntimeFormat();
+}
+
+const FObjectKey& FSoundWaveProxy::GetFObjectKey() const
+{
+	check(SoundWaveDataPtr);
+	return SoundWaveDataPtr->GetFObjectKey();
+}
+
+float FSoundWaveProxy::GetDuration() const
+{
+	check(SoundWaveDataPtr);
+	return SoundWaveDataPtr->GetDuration();
+}
+
+float FSoundWaveProxy::GetSampleRate() const
+{
+	check(SoundWaveDataPtr);
+	return SoundWaveDataPtr->GetSampleRate();
+}
+
+int32 FSoundWaveProxy::GetNumFrames() const
+{
+	check(SoundWaveDataPtr);
+	return SoundWaveDataPtr->GetNumFrames();
+}
+
+uint32 FSoundWaveProxy::GetNumChunks()
+{
+	check(SoundWaveDataPtr);
+	return SoundWaveDataPtr->GetNumChunks();
+}
+
+const uint32 FSoundWaveProxy::GetNumChannels() const
+{
+	check(SoundWaveDataPtr);
+	return SoundWaveDataPtr->GetNumChannels();
+}
+
 uint32 FSoundWaveProxy::GetSizeOfChunk(uint32 ChunkIndex) const
 {
-	check(ChunkIndex < GetNumChunks());
-	return GetChunk(ChunkIndex).AudioDataSize;
+	check(SoundWaveDataPtr);
+	return SoundWaveDataPtr->GetSizeOfChunk(ChunkIndex);
 }
 
+const TArray<FSoundWaveCuePoint>& FSoundWaveProxy::GetCuePoints() const
+{
+	check(SoundWaveDataPtr);
+	return SoundWaveDataPtr->GetCuePoints();
+}
+
+FSoundWaveData::MaxChunkSizeResults FSoundWaveProxy::GetMaxChunkSizeResults() const
+{
+	check(SoundWaveDataPtr);
+	return SoundWaveDataPtr->GetMaxChunkSizeResults();
+}
+
+bool FSoundWaveProxy::IsLooping() const
+{
+	check(SoundWaveDataPtr);
+	return SoundWaveDataPtr->IsLooping();
+}
+
+bool FSoundWaveProxy::IsTemplate() const
+{
+	check(SoundWaveDataPtr);
+	return SoundWaveDataPtr->IsTemplate();
+}
+
+bool FSoundWaveProxy::IsStreaming() const
+{
+	check(SoundWaveDataPtr);
+	return SoundWaveDataPtr->IsStreaming();
+}
+
+bool FSoundWaveProxy::UseBinkAudio() const
+{
+	check(SoundWaveDataPtr);
+	return SoundWaveDataPtr->UseBinkAudio();
+}
+
+bool FSoundWaveProxy::IsRetainingAudio() const
+{
+	check(SoundWaveDataPtr);
+	return SoundWaveDataPtr->IsRetainingAudio();
+}
+
+bool FSoundWaveProxy::ShouldUseStreamCaching() const
+{
+	check(SoundWaveDataPtr);
+	return SoundWaveDataPtr->ShouldUseStreamCaching();
+}
+
+bool FSoundWaveProxy::IsSeekableStreaming() const
+{
+	check(SoundWaveDataPtr);
+	return SoundWaveDataPtr->IsSeekableStreaming();
+}
+
+bool FSoundWaveProxy::IsZerothChunkDataLoaded() const
+{
+	check(SoundWaveDataPtr);
+	return SoundWaveDataPtr->IsZerothChunkDataLoaded();
+}
+
+bool FSoundWaveProxy::WasLoadingBehaviorOverridden() const
+{
+	check(SoundWaveDataPtr);
+	return SoundWaveDataPtr->WasLoadingBehaviorOverridden();
+}
+
 bool FSoundWaveProxy::HasCompressedData(FName Format, ITargetPlatform* TargetPlatform) const
 {
 	check(SoundWaveDataPtr);
-
-	if (IsTemplate() || IsRunningDedicatedServer())
-	{
-		return false;
-	}
-
-	const FPlatformAudioCookOverrides* CompressionOverrides = nullptr;
-
-	if (GIsEditor)
-	{
-		if (TargetPlatform)
-		{
-			CompressionOverrides = FPlatformCompressionUtilities::GetCookOverrides(*TargetPlatform->IniPlatformName());
-		}
-	}
-	else
-	{
-		// TargetPlatform is not available on consoles/mobile, so we have to grab it ourselves:
-		CompressionOverrides = FPlatformCompressionUtilities::GetCookOverrides();
-	}
-
-	if (CompressionOverrides)
-	{
-#if WITH_EDITOR
-		FName PlatformSpecificFormat;
-		FString HashedString = *Format.ToString();
-		FPlatformAudioCookOverrides::GetHashSuffix(CompressionOverrides, HashedString);
-		PlatformSpecificFormat = *HashedString;
-#else
-		// on non-editor builds, we cache the concatenated format in a static FName.
-		static FName PlatformSpecificFormat;
-		static FName CachedFormat;
-		if (!Format.IsEqual(CachedFormat))
-		{
-			FString HashedString = *Format.ToString();
-			FPlatformAudioCookOverrides::GetHashSuffix(CompressionOverrides, HashedString);
-			PlatformSpecificFormat = *HashedString;
-
-			CachedFormat = Format;
-		}
-#endif // WITH_EDITOR
-		return SoundWaveDataPtr->CompressedFormatData.Contains(PlatformSpecificFormat);
-	}
-	else
-	{
-		return SoundWaveDataPtr->CompressedFormatData.Contains(Format);
-	}
-
+	return SoundWaveDataPtr->HasCompressedData(Format, TargetPlatform);
 }
 
+ESoundWaveLoadingBehavior FSoundWaveProxy::GetLoadingBehavior() const
+{
+	check(SoundWaveDataPtr);
+	return SoundWaveDataPtr->GetLoadingBehavior();
+}
+
+const TArrayView<uint8> FSoundWaveProxy::GetZerothChunkDataView() const
+{
+	check(SoundWaveDataPtr);
+	return SoundWaveDataPtr->GetZerothChunkDataView();
+}
+
 FByteBulkData* FSoundWaveProxy::GetCompressedData(FName Format, const FPlatformAudioCookOverrides* CompressionOverrides)
 {
 	check(SoundWaveDataPtr);
-
-	if (IsTemplate() || IsRunningDedicatedServer())
-	{
-		return nullptr;
-	}
-
-	TRACE_CPUPROFILER_EVENT_SCOPE(FSoundWaveProxy::GetCompressedData);
-	FName PlatformSpecificFormat = USoundWave::GetPlatformSpecificFormat(Format, CompressionOverrides);
-
-	bool bContainedValidData = SoundWaveDataPtr->CompressedFormatData.Contains(PlatformSpecificFormat);
-	if (!bContainedValidData)
-	{
-		// the proxy does not support interaction w/ FDerivedAudioDataCompressor
-		return nullptr;
-	}
-
-	FByteBulkData* Result = &SoundWaveDataPtr->CompressedFormatData.GetFormat(PlatformSpecificFormat);
-	check(Result);
-	return Result->GetBulkDataSize() > 0 ? Result : NULL; // we don't return empty bulk data...but we save it to avoid thrashing the DDC
+	return SoundWaveDataPtr->GetCompressedData(Format, CompressionOverrides);
 }
 
 TArrayView<const uint8> FSoundWaveProxy::GetZerothChunk(const FSoundWaveProxyPtr& SoundWaveProxy, bool bForImmediatePlayback)
@@ -3507,7 +3901,7 @@
 
 		if (ensure(SoundWaveProxy->IsZerothChunkDataLoaded()))
 		{
-			if (SoundWaveProxy->SoundWaveDataPtr->bShouldUseStreamCaching)
+			if (SoundWaveProxy->SoundWaveDataPtr->ShouldUseStreamCaching())
 			{
 				if (SoundWaveProxy->GetNumChunks() > 1)
 				{
@@ -3527,142 +3921,3 @@
 
 	return {};
 }
-
-
-bool FSoundWaveProxy::GetChunkData(int32 ChunkIndex, uint8** OutChunkData, bool bMakeSureChunkIsLoaded)
-{
-	LLM_SCOPE(ELLMTag::AudioSoundWaveProxies);
-	if (GetChunkFromDDC(ChunkIndex, OutChunkData, bMakeSureChunkIsLoaded) == 0)
-	{
-#if WITH_EDITORONLY_DATA
-		// Unable to load chunks from the cache. Rebuild the sound and attempt to precache it.
-		UE_LOG(LogAudio, Display, TEXT("GetChunkData failed, rebuilding %s"), *GetPackageName().ToString());
-
-		//		ForceRebuildPlatformData();
-		if (GetChunkFromDDC(ChunkIndex, OutChunkData, bMakeSureChunkIsLoaded) == 0)
-		{
-			UE_LOG(LogAudio, Warning, TEXT("Failed to build sound %s."), *GetPackageName().ToString());
-		}
-		else
-		{
-			// Succeeded after rebuilding platform data
-			return true;
-		}
-#else
-		// Failed to find the SoundWave chunk in the cooked package.
-		UE_LOG(LogAudio, Warning, TEXT("GetChunkData failed while streaming. Ensure the following package is cooked: %s"), *GetPackageName().ToString());
-#endif // #if WITH_EDITORONLY_DATA
-		return false;
-	}
-	return true;
-}
-
-bool FSoundWaveProxy::IsZerothChunkDataLoaded() const
-{
-	check(SoundWaveDataPtr);
-	return SoundWaveDataPtr->ZerothChunkData.GetView().Num() > 0;
-}
-
-const TArrayView<uint8> FSoundWaveProxy::GetZerothChunkDataView() const
-{
-	check(SoundWaveDataPtr);
-	FBulkDataBuffer<uint8>::ViewType View = SoundWaveDataPtr->ZerothChunkData.GetView();
-	check(View.Num() <= MAX_int32);
-	return TArrayView<uint8>(View.GetData(), static_cast<int32>(View.Num()));
-}
-
-void FSoundWaveProxy::EnsureZerothChunkIsLoaded()
-{
-	LLM_SCOPE(ELLMTag::AudioSoundWaveProxies);
-	check(SoundWaveDataPtr);
-
-	// If the zeroth chunk is already loaded, early exit.
-	if (SoundWaveDataPtr->ZerothChunkData.GetView().Num() > 0 || !SoundWaveDataPtr->bShouldUseStreamCaching)
-	{
-		return;
-	}
-
-// #if WITH_EDITOR
-#ifdef WITH_EDITOR
-	// If we're running the editor, we'll need to retrieve the chunked audio from the DDC:
-	uint8* TempChunkBuffer = nullptr;
-	int32 ChunkSizeInBytes = GetChunkFromDDC(0, &TempChunkBuffer, true);
-	// Since we block for the DDC in the previous call we should always have the chunk loaded.
-	if (ChunkSizeInBytes == 0)
-	{
-		return;
-	}
-
-	SoundWaveDataPtr->ZerothChunkData.Reset(TempChunkBuffer, ChunkSizeInBytes);
-
-#else // WITH_EDITOR
-	// Otherwise, the zeroth chunk is cooked out to SharedRuntimeData::RunningPlatformData, and we just need to retrieve it.
-	check(GetNumChunks() > 0);
-	const FStreamedAudioChunk& ZerothChunk = GetChunk(0);
-	// Some sanity checks to ensure that the bulk size set up
-	UE_CLOG(ZerothChunk.BulkData.GetBulkDataSize() != ZerothChunk.DataSize, LogAudio, Warning, TEXT("Bulk data serialized out had a mismatched size with the DataSize field. Soundwave: %s Bulk Data Reported Size: %d Bulk Data Actual Size: %ld"), *GetFullName(), ZerothChunk.DataSize, ZerothChunk.BulkData.GetBulkDataSize());
-
-	ZerothChunkData = ZerothChunk.BulkData.GetCopyAsBuffer(ZerothChunk.AudioDataSize, true);
-#endif // WITH_EDITOR
-}
-
-uint32 FSoundWaveProxy::GetNumChunks() const
-{
-	check(SoundWaveDataPtr);
-
-	// This function shouldn't be called on audio marked "ForceInline."
-	ensureMsgf(SoundWaveDataPtr->LoadingBehavior != ESoundWaveLoadingBehavior::ForceInline, TEXT("Calling GetNumChunks on a FSoundWaveProxy without RunnigPlatformData is not allowed! SoundWave: %s - %s")
-		, *GetFName().ToString(), EnumToString(SoundWaveDataPtr->LoadingBehavior));
-
-	return SoundWaveDataPtr->RunningPlatformData.Chunks.Num();
-}
-
-#if WITH_EDITOR
-int32 FSoundWaveProxy::GetCurrentChunkRevision() const
-{
-	check(SoundWaveDataPtr);
-	return SoundWaveDataPtr->CurrentChunkRevision;
-}
-#endif // #if WITH_EDITOR
-
-const FStreamedAudioChunk& FSoundWaveProxy::GetChunk(uint32 ChunkIndex) const
-{
-	// This function shouldn't be called on audio marked "ForceInline."
-	check(SoundWaveDataPtr);
-	ensureMsgf(SoundWaveDataPtr->RunningPlatformData.NumChunks
-		, TEXT("Calling GetNumChunks on a FSoundWaveProxy without RunnigPlatformData is not allowed! SoundWave: %s - %s")
-		, *GetFName().ToString()
-		, EnumToString(SoundWaveDataPtr->LoadingBehavior));
-
-	check((ChunkIndex < (uint32)GetNumChunks()));
- 	return SoundWaveDataPtr->RunningPlatformData.Chunks[ChunkIndex];
-}
-
-int32 FSoundWaveProxy::GetChunkFromDDC(int32 ChunkIndex, uint8** OutChunkData, bool bMakeSureChunkIsLoaded)
-{
-	LLM_SCOPE(ELLMTag::AudioSoundWaveProxies);
-	// This function shouldn't be called on audio marked "ForceInline."
-	check(SoundWaveDataPtr);
-	ensureMsgf(SoundWaveDataPtr->RunningPlatformData.NumChunks
-		, TEXT("Calling GetNumChunks on a FSoundWaveProxy without RunnigPlatformData is not allowed! SoundWave: %s - %s")
-		, *GetFName().ToString()
-		, EnumToString(SoundWaveDataPtr->LoadingBehavior));
-
-	return SoundWaveDataPtr->RunningPlatformData.GetChunkFromDDC(ChunkIndex, OutChunkData, bMakeSureChunkIsLoaded);
-}
-
-#if WITH_EDITORONLY_DATA
-FString FSoundWaveProxy::GetDerivedDataKey() const
-{
-	check(SoundWaveDataPtr);
-	check(SoundWaveDataPtr->RunningPlatformData.NumChunks);
-
-	// This function shouldn't be called on audio marked "ForceInline."
-	ensureMsgf(SoundWaveDataPtr->RunningPlatformData.NumChunks
-		, TEXT("Calling GetNumChunks on a FSoundWaveProxy without RunnigPlatformData is not allowed! SoundWave: %s - %s")
-		, *GetFName().ToString()
-		, EnumToString(SoundWaveDataPtr->LoadingBehavior));
-
-	return SoundWaveDataPtr->RunningPlatformData.DerivedDataKey;
-}
-#endif // #if WITH_EDITORONLY_DATA
diff --git a/Engine/Source/Runtime/Engine/Public/AudioStreamingCache.h b/Engine/Source/Runtime/Engine/Public/AudioStreamingCache.h
--- a/Engine/Source/Runtime/Engine/Public/AudioStreamingCache.h
+++ b/Engine/Source/Runtime/Engine/Public/AudioStreamingCache.h
@@ -37,7 +37,7 @@ public:
 		FChunkKey(const FChunkKey& Other);
 
 		FChunkKey(
-			  const FSoundWaveProxyPtr& InSoundWave
+			  const FSoundWavePtr& InSoundWave
 			, uint32 InChunkIndex
 #if WITH_EDITOR
 			, uint32 InChunkRevision = 0
@@ -48,7 +48,7 @@ public:
 
 
 	private:
-		TWeakPtr<FSoundWaveProxy, ESPMode::ThreadSafe> SoundWaveProxyWeakPtr;
+		TWeakPtr<FSoundWaveData, ESPMode::ThreadSafe> SoundWaveWeakPtr;
 
 	public:
 		FName SoundWaveName = FName();
@@ -62,7 +62,7 @@ public:
 		inline bool operator==(const FChunkKey& Other) const;
 
 		bool IsChunkStale();
-		bool IsSoundWaveValid() const { return SoundWaveProxyWeakPtr.IsValid(); }
+		bool IsSoundWaveValid() const { return SoundWaveWeakPtr.Pin().IsValid(); }
 		// Forward FSoundWaveProxy interface
 		ESoundWaveLoadingBehavior GetLoadingBehavior() const;
 		FStreamedAudioChunk& GetChunk(uint32 InChunkIndex) const;
diff --git a/Engine/Source/Runtime/Engine/Public/ContentStreaming.h b/Engine/Source/Runtime/Engine/Public/ContentStreaming.h
--- a/Engine/Source/Runtime/Engine/Public/ContentStreaming.h
+++ b/Engine/Source/Runtime/Engine/Public/ContentStreaming.h
@@ -14,6 +14,7 @@
 class AActor;
 class FSoundSource;
 class UPrimitiveComponent;
+class FSoundWaveData;
 class FSoundWaveProxy;
 class ICompressedAudioInfo;
 class UTexture2D;
@@ -128,7 +129,7 @@ private:
 	const uint8*  CachedData;
 	int32 CachedDataNumBytes;
 
-	TWeakPtr<FSoundWaveProxy, ESPMode::ThreadSafe> CorrespondingWave;
+	TWeakPtr<FSoundWaveData, ESPMode::ThreadSafe> CorrespondingWave;
 	FName CorrespondingWaveName;
 
 	// The index of this chunk in the sound wave's full set of chunks of compressed audio.
