Adding ForceUnits metadata to various properties in UCharacterMovementComponent (not an exaustive pass, e.g., we don't support acceleration units yet)

diff --git a/Engine/Source/Runtime/Engine/Classes/GameFramework/CharacterMovementComponent.h b/Engine/Source/Runtime/Engine/Classes/GameFramework/CharacterMovementComponent.h
--- a/Engine/Source/Runtime/Engine/Classes/GameFramework/CharacterMovementComponent.h
+++ b/Engine/Source/Runtime/Engine/Classes/GameFramework/CharacterMovementComponent.h
@@ -155,11 +155,11 @@ public:
 	float GravityScale;
 
 	/** Maximum height character can step up */
-	UPROPERTY(Category="Character Movement: Walking", EditAnywhere, BlueprintReadWrite, meta=(ClampMin="0", UIMin="0"))
+	UPROPERTY(Category="Character Movement: Walking", EditAnywhere, BlueprintReadWrite, meta=(ClampMin="0", UIMin="0", ForceUnits="cm"))
 	float MaxStepHeight;
 
 	/** Initial velocity (instantaneous vertical acceleration) when jumping. */
-	UPROPERTY(Category="Character Movement: Jumping / Falling", EditAnywhere, BlueprintReadWrite, meta=(DisplayName="Jump Z Velocity", ClampMin="0", UIMin="0"))
+	UPROPERTY(Category="Character Movement: Jumping / Falling", EditAnywhere, BlueprintReadWrite, meta=(DisplayName="Jump Z Velocity", ClampMin="0", UIMin="0", ForceUnits="cm/s"))
 	float JumpZVelocity;
 
 	/** Fraction of JumpZVelocity to use when automatically "jumping off" of a base actor that's not allowed to be a base for a character. (For example, if you're not allowed to stand on other players.) */
@@ -179,7 +179,7 @@ private:
 	/**
 	 * Max angle in degrees of a walkable surface. Any greater than this and it is too steep to be walkable.
 	 */
-	UPROPERTY(Category="Character Movement: Walking", EditAnywhere, meta=(ClampMin="0.0", ClampMax="90.0", UIMin = "0.0", UIMax = "90.0"))
+	UPROPERTY(Category="Character Movement: Walking", EditAnywhere, meta=(ClampMin="0.0", ClampMax="90.0", UIMin = "0.0", UIMax = "90.0", ForceUnits="degrees"))
 	float WalkableFloorAngle;
 
 	/**
@@ -232,23 +232,23 @@ public:
 	FVector OldBaseLocation;
 
 	/** The maximum ground speed when walking. Also determines maximum lateral speed when falling. */
-	UPROPERTY(Category="Character Movement: Walking", EditAnywhere, BlueprintReadWrite, meta=(ClampMin="0", UIMin="0"))
+	UPROPERTY(Category="Character Movement: Walking", EditAnywhere, BlueprintReadWrite, meta=(ClampMin="0", UIMin="0", ForceUnits="cm/s"))
 	float MaxWalkSpeed;
 
 	/** The maximum ground speed when walking and crouched. */
-	UPROPERTY(Category="Character Movement: Walking", EditAnywhere, BlueprintReadWrite, meta=(ClampMin="0", UIMin="0"))
+	UPROPERTY(Category="Character Movement: Walking", EditAnywhere, BlueprintReadWrite, meta=(ClampMin="0", UIMin="0", ForceUnits="cm/s"))
 	float MaxWalkSpeedCrouched;
 
 	/** The maximum swimming speed. */
-	UPROPERTY(Category="Character Movement: Swimming", EditAnywhere, BlueprintReadWrite, meta=(ClampMin="0", UIMin="0"))
+	UPROPERTY(Category="Character Movement: Swimming", EditAnywhere, BlueprintReadWrite, meta=(ClampMin="0", UIMin="0", ForceUnits="cm/s"))
 	float MaxSwimSpeed;
 
 	/** The maximum flying speed. */
-	UPROPERTY(Category="Character Movement: Flying", EditAnywhere, BlueprintReadWrite, meta=(ClampMin="0", UIMin="0"))
+	UPROPERTY(Category="Character Movement: Flying", EditAnywhere, BlueprintReadWrite, meta=(ClampMin="0", UIMin="0", ForceUnits="cm/s"))
 	float MaxFlySpeed;
 
 	/** The maximum speed when using Custom movement mode. */
-	UPROPERTY(Category="Character Movement: Custom Movement", EditAnywhere, BlueprintReadWrite, meta=(ClampMin="0", UIMin="0"))
+	UPROPERTY(Category="Character Movement: Custom Movement", EditAnywhere, BlueprintReadWrite, meta=(ClampMin="0", UIMin="0", ForceUnits="cm/s"))
 	float MaxCustomMovementSpeed;
 
 	/** Max Acceleration (rate of change of velocity) */
@@ -256,7 +256,7 @@ public:
 	float MaxAcceleration;
 
 	/** The ground speed that we should accelerate up to when walking at minimum analog stick tilt */
-	UPROPERTY(Category = "Character Movement: Walking", EditAnywhere, BlueprintReadWrite, meta = (ClampMin = "0", UIMin = "0"))
+	UPROPERTY(Category = "Character Movement: Walking", EditAnywhere, BlueprintReadWrite, meta = (ClampMin = "0", UIMin = "0", ForceUnits="cm/s"))
 	float MinAnalogWalkSpeed;
 
 	/**
@@ -344,7 +344,7 @@ public:
 
 	/** Collision half-height when crouching (component scale is applied separately) */
 	UE_DEPRECATED_FORGAME(5.0, "Public access to this property is deprecated, and it will become private in a future release. Please use SetCrouchedHalfHeight and GetCrouchedHalfHeight instead.")
-	UPROPERTY(Category="Character Movement (General Settings)", EditAnywhere, BlueprintReadWrite, BlueprintSetter=SetCrouchedHalfHeight, BlueprintGetter=GetCrouchedHalfHeight, meta=(ClampMin="0", UIMin="0"))
+	UPROPERTY(Category="Character Movement (General Settings)", EditAnywhere, BlueprintReadWrite, BlueprintSetter=SetCrouchedHalfHeight, BlueprintGetter=GetCrouchedHalfHeight, meta=(ClampMin="0", UIMin="0", ForceUnits=cm))
 	float CrouchedHalfHeight;
 
 	/** Water buoyancy. A ratio (1.0 = neutral buoyancy, 0.0 = no buoyancy) */
@@ -355,7 +355,7 @@ public:
 	 * Don't allow the character to perch on the edge of a surface if the contact is this close to the edge of the capsule.
 	 * Note that characters will not fall off if they are within MaxStepHeight of a walkable surface below.
 	 */
-	UPROPERTY(Category="Character Movement: Walking", EditAnywhere, BlueprintReadWrite, AdvancedDisplay, meta=(ClampMin="0", UIMin="0"))
+	UPROPERTY(Category="Character Movement: Walking", EditAnywhere, BlueprintReadWrite, AdvancedDisplay, meta=(ClampMin="0", UIMin="0", ForceUnits=cm))
 	float PerchRadiusThreshold;
 
 	/**
@@ -363,7 +363,7 @@ public:
 	 * Note that we still enforce MaxStepHeight to start the step up; this just allows the character to hang off the edge or step slightly higher off the floor.
 	 * (@see PerchRadiusThreshold)
 	 */
-	UPROPERTY(Category="Character Movement: Walking", EditAnywhere, BlueprintReadWrite, AdvancedDisplay, meta=(ClampMin="0", UIMin="0"))
+	UPROPERTY(Category="Character Movement: Walking", EditAnywhere, BlueprintReadWrite, AdvancedDisplay, meta=(ClampMin="0", UIMin="0", ForceUnits=cm))
 	float PerchAdditionalHeight;
 
 	/** Change in rotation per second, used when UseControllerDesiredRotation or OrientRotationToMovement are true. Set a negative value for infinite rotation rate and instant turns. */
@@ -533,7 +533,7 @@ public:
 	TObjectPtr<USceneComponent> DeferredUpdatedMoveComponent;
 
 	/** Maximum step height for getting out of water */
-	UPROPERTY(Category="Character Movement: Swimming", EditAnywhere, BlueprintReadWrite, AdvancedDisplay, meta=(ClampMin="0", UIMin="0"))
+	UPROPERTY(Category="Character Movement: Swimming", EditAnywhere, BlueprintReadWrite, AdvancedDisplay, meta=(ClampMin="0", UIMin="0", ForceUnits=cm))
 	float MaxOutOfWaterStepHeight;
 
 	/** Z velocity applied when pawn tries to get out of water */
@@ -724,7 +724,7 @@ public:
 	* This is generally more tolerant than with Pawns, because other geometry is either not moving, or is moving predictably with a bit of delay compared to on the server.
 	* @see MaxDepenetrationWithGeometryAsProxy, MaxDepenetrationWithPawn, MaxDepenetrationWithPawnAsProxy
 	*/
-	UPROPERTY(Category="Character Movement (General Settings)", EditAnywhere, BlueprintReadWrite, AdvancedDisplay, meta=(ClampMin="0", UIMin="0"))
+	UPROPERTY(Category="Character Movement (General Settings)", EditAnywhere, BlueprintReadWrite, AdvancedDisplay, meta=(ClampMin="0", UIMin="0", ForceUnits=cm))
 	float MaxDepenetrationWithGeometry;
 
 	/**
@@ -732,14 +732,14 @@ public:
 	* This is generally more tolerant than with Pawns, because other geometry is either not moving, or is moving predictably with a bit of delay compared to on the server.
 	* @see MaxDepenetrationWithGeometry, MaxDepenetrationWithPawn, MaxDepenetrationWithPawnAsProxy
 	*/
-	UPROPERTY(Category="Character Movement (General Settings)", EditAnywhere, BlueprintReadWrite, AdvancedDisplay, meta=(ClampMin="0", UIMin="0"))
+	UPROPERTY(Category="Character Movement (General Settings)", EditAnywhere, BlueprintReadWrite, AdvancedDisplay, meta=(ClampMin="0", UIMin="0", ForceUnits=cm))
 	float MaxDepenetrationWithGeometryAsProxy;
 
 	/**
 	* Max distance we are allowed to depenetrate when moving out of other Pawns.
 	* @see MaxDepenetrationWithGeometry, MaxDepenetrationWithGeometryAsProxy, MaxDepenetrationWithPawnAsProxy
 	*/
-	UPROPERTY(Category="Character Movement (General Settings)", EditAnywhere, BlueprintReadWrite, AdvancedDisplay, meta=(ClampMin="0", UIMin="0"))
+	UPROPERTY(Category="Character Movement (General Settings)", EditAnywhere, BlueprintReadWrite, AdvancedDisplay, meta=(ClampMin="0", UIMin="0", ForceUnits=cm))
 	float MaxDepenetrationWithPawn;
 
 	/**
@@ -747,61 +747,61 @@ public:
 	 * Typically we don't want a large value, because we receive a server authoritative position that we should not then ignore by pushing them out of the local player.
 	 * @see MaxDepenetrationWithGeometry, MaxDepenetrationWithGeometryAsProxy, MaxDepenetrationWithPawn
 	 */
-	UPROPERTY(Category="Character Movement (General Settings)", EditAnywhere, BlueprintReadWrite, AdvancedDisplay, meta=(ClampMin="0", UIMin="0"))
+	UPROPERTY(Category="Character Movement (General Settings)", EditAnywhere, BlueprintReadWrite, AdvancedDisplay, meta=(ClampMin="0", UIMin="0", ForceUnits=cm))
 	float MaxDepenetrationWithPawnAsProxy;
 
 	/**
 	 * How long to take to smoothly interpolate from the old pawn position on the client to the corrected one sent by the server. Not used by Linear smoothing.
 	 */
-	UPROPERTY(Category="Character Movement (Networking)", EditDefaultsOnly, AdvancedDisplay, meta=(ClampMin="0.0", ClampMax="1.0", UIMin="0.0", UIMax="1.0"))
+	UPROPERTY(Category="Character Movement (Networking)", EditDefaultsOnly, AdvancedDisplay, meta=(ClampMin="0.0", ClampMax="1.0", UIMin="0.0", UIMax="1.0", ForceUnits=s))
 	float NetworkSimulatedSmoothLocationTime;
 
 	/**
 	 * How long to take to smoothly interpolate from the old pawn rotation on the client to the corrected one sent by the server. Not used by Linear smoothing.
 	 */
-	UPROPERTY(Category="Character Movement (Networking)", EditDefaultsOnly, AdvancedDisplay, meta=(ClampMin="0.0", ClampMax="1.0", UIMin="0.0", UIMax="1.0"))
+	UPROPERTY(Category="Character Movement (Networking)", EditDefaultsOnly, AdvancedDisplay, meta=(ClampMin="0.0", ClampMax="1.0", UIMin="0.0", UIMax="1.0", ForceUnits=s))
 	float NetworkSimulatedSmoothRotationTime;
 
 	/**
 	* Similar setting as NetworkSimulatedSmoothLocationTime but only used on Listen servers.
 	*/
-	UPROPERTY(Category="Character Movement (Networking)", EditDefaultsOnly, AdvancedDisplay, meta=(ClampMin="0.0", ClampMax="1.0", UIMin="0.0", UIMax="1.0"))
+	UPROPERTY(Category="Character Movement (Networking)", EditDefaultsOnly, AdvancedDisplay, meta=(ClampMin="0.0", ClampMax="1.0", UIMin="0.0", UIMax="1.0", ForceUnits=s))
 	float ListenServerNetworkSimulatedSmoothLocationTime;
 
 	/**
 	* Similar setting as NetworkSimulatedSmoothRotationTime but only used on Listen servers.
 	*/
-	UPROPERTY(Category="Character Movement (Networking)", EditDefaultsOnly, AdvancedDisplay, meta=(ClampMin="0.0", ClampMax="1.0", UIMin="0.0", UIMax="1.0"))
+	UPROPERTY(Category="Character Movement (Networking)", EditDefaultsOnly, AdvancedDisplay, meta=(ClampMin="0.0", ClampMax="1.0", UIMin="0.0", UIMax="1.0", ForceUnits=s))
 	float ListenServerNetworkSimulatedSmoothRotationTime;
 
 	/**
 	 * Shrink simulated proxy capsule radius by this amount, to account for network rounding that may cause encroachment. Changing during gameplay is not supported.
 	 * @see AdjustProxyCapsuleSize()
 	 */
-	UPROPERTY(Category="Character Movement (Networking)", EditDefaultsOnly, AdvancedDisplay, meta=(ClampMin="0.0", UIMin="0.0"))
+	UPROPERTY(Category="Character Movement (Networking)", EditDefaultsOnly, AdvancedDisplay, meta=(ClampMin="0.0", UIMin="0.0", ForceUnits=cm))
 	float NetProxyShrinkRadius;
 
 	/**
 	 * Shrink simulated proxy capsule half height by this amount, to account for network rounding that may cause encroachment. Changing during gameplay is not supported.
 	 * @see AdjustProxyCapsuleSize()
 	 */
-	UPROPERTY(Category="Character Movement (Networking)", EditDefaultsOnly, AdvancedDisplay, meta=(ClampMin="0.0", UIMin="0.0"))
+	UPROPERTY(Category="Character Movement (Networking)", EditDefaultsOnly, AdvancedDisplay, meta=(ClampMin="0.0", UIMin="0.0", ForceUnits=cm))
 	float NetProxyShrinkHalfHeight;
 
 	/** Maximum distance character is allowed to lag behind server location when interpolating between updates. */
-	UPROPERTY(Category="Character Movement (Networking)", EditDefaultsOnly, meta=(ClampMin="0.0", UIMin="0.0"))
+	UPROPERTY(Category="Character Movement (Networking)", EditDefaultsOnly, meta=(ClampMin="0.0", UIMin="0.0", ForceUnits=cm))
 	float NetworkMaxSmoothUpdateDistance;
 
 	/**
 	 * Maximum distance beyond which character is teleported to the new server location without any smoothing.
 	 */
-	UPROPERTY(Category="Character Movement (Networking)", EditDefaultsOnly, meta=(ClampMin="0.0", UIMin="0.0"))
+	UPROPERTY(Category="Character Movement (Networking)", EditDefaultsOnly, meta=(ClampMin="0.0", UIMin="0.0", ForceUnits=cm))
 	float NetworkNoSmoothUpdateDistance;
 
 	/**
 	 * Minimum time on the server between acknowledging good client moves. This can save on bandwidth. Set to 0 to disable throttling.
 	 */
-	UPROPERTY(Category="Character Movement (Networking)", EditDefaultsOnly, meta=(ClampMin="0.0", UIMin="0.0"))
+	UPROPERTY(Category="Character Movement (Networking)", EditDefaultsOnly, meta=(ClampMin="0.0", UIMin="0.0", ForceUnits=s))
 	float NetworkMinTimeBetweenClientAckGoodMoves;
 
 	/**
@@ -810,7 +810,7 @@ public:
   	 * This can save on bandwidth. Set to 0 to disable throttling.
 	 * @see ServerLastClientAdjustmentTime
 	 */
-	UPROPERTY(Category="Character Movement (Networking)", EditDefaultsOnly, meta=(ClampMin="0.0", UIMin="0.0"))
+	UPROPERTY(Category="Character Movement (Networking)", EditDefaultsOnly, meta=(ClampMin="0.0", UIMin="0.0", ForceUnits=s))
 	float NetworkMinTimeBetweenClientAdjustments;
 
 	/**
@@ -818,18 +818,18 @@ public:
 	* Should be <= NetworkMinTimeBetweenClientAdjustments (the smaller value is used regardless).
 	* @see NetworkMinTimeBetweenClientAdjustments
 	*/
-	UPROPERTY(Category="Character Movement (Networking)", EditDefaultsOnly, meta=(ClampMin="0.0", UIMin="0.0"))
+	UPROPERTY(Category="Character Movement (Networking)", EditDefaultsOnly, meta=(ClampMin="0.0", UIMin="0.0", ForceUnits=s))
 	float NetworkMinTimeBetweenClientAdjustmentsLargeCorrection;
 
 	/**
 	* If client error is larger than this, sets bNetworkLargeClientCorrection to reduce delay between client adjustments.
 	* @see NetworkMinTimeBetweenClientAdjustments, NetworkMinTimeBetweenClientAdjustmentsLargeCorrection
 	*/
-	UPROPERTY(Category="Character Movement (Networking)", EditDefaultsOnly, meta=(ClampMin="0.0", UIMin="0.0"))
+	UPROPERTY(Category="Character Movement (Networking)", EditDefaultsOnly, meta=(ClampMin="0.0", UIMin="0.0", ForceUnits=cm))
 	float NetworkLargeClientCorrectionDistance;
 
 	/** Used in determining if pawn is going off ledge.  If the ledge is "shorter" than this value then the pawn will be able to walk off it. **/
-	UPROPERTY(Category="Character Movement: Walking", EditAnywhere, BlueprintReadWrite, AdvancedDisplay)
+	UPROPERTY(Category="Character Movement: Walking", EditAnywhere, BlueprintReadWrite, AdvancedDisplay, meta=(ForceUnits=cm))
 	float LedgeCheckThreshold;
 
 	/** When exiting water, jump if control pitch angle is this high or above. */
@@ -1058,7 +1058,7 @@ protected:
 
 public:
 
-	UPROPERTY(Category="Character Movement: Avoidance", EditAnywhere, BlueprintReadOnly)
+	UPROPERTY(Category="Character Movement: Avoidance", EditAnywhere, BlueprintReadOnly, meta=(ForceUnits=cm))
 	float AvoidanceConsiderationRadius;
 
 	/**
