Add RuntimeVirtualTextureSample mip mode that ignores the WorldPostion input.
This will calculate mip based on derivatives of actual world position instead.
It gives the same result as 'Default' when no WorldPosition input is used.
The use case is when a constant value is being piped into WorldPostion.
That happens for things like picking a constant color from an RVT for grass.
Without the new mip mode we always compute a mip value of 0 which can be very heavy on the physical texture pool.

diff --git a/Engine/Source/Runtime/Engine/Classes/Materials/MaterialExpressionRuntimeVirtualTextureSample.h b/Engine/Source/Runtime/Engine/Classes/Materials/MaterialExpressionRuntimeVirtualTextureSample.h
--- a/Engine/Source/Runtime/Engine/Classes/Materials/MaterialExpressionRuntimeVirtualTextureSample.h
+++ b/Engine/Source/Runtime/Engine/Classes/Materials/MaterialExpressionRuntimeVirtualTextureSample.h
@@ -9,21 +9,35 @@
 #include "MaterialExpressionRuntimeVirtualTextureSample.generated.h"
 
 /**
- * Defines how MipValue is used.
+ * Set how Mip levels are calculated.
  * Internally we will convert to ETextureMipValueMode which is used by internal APIs.
- * ETextureMipValueMode has more options then are valid for runtime virtual texture.
  */
 UENUM()
 enum ERuntimeVirtualTextureMipValueMode
 {
-	/* Use hardware computed sample's mip level with automatic anisotropic filtering support. */
-	RVTMVM_None UMETA(DisplayName = "None (use computed mip level)"),
-
-	/* Explicitly compute the sample's mip level. Disables anisotropic filtering. */
-	RVTMVM_MipLevel UMETA(DisplayName = "MipLevel (absolute, 0 is full resolution)"),
-
-	/* Bias the hardware computed sample's mip level. Disables anisotropic filtering. */
-	RVTMVM_MipBias UMETA(DisplayName = "MipBias (relative to the computed mip level)"),
+	/* 
+	 * Use default computed mip level. Takes into account UV scaling from using the WorldPosition pin.
+	 */
+	RVTMVM_None UMETA(DisplayName = "Default"),
+
+	/* 
+	 * Use an absolute mip level from the MipValue pin. 
+	 * 0 is full resolution.
+	 */
+	RVTMVM_MipLevel UMETA(DisplayName = "Mip Level"),
+
+	/* 
+	 * Bias the default computed mip level using the MipValue pin. 
+	 * Negative values increase resolution.
+	 */
+	RVTMVM_MipBias UMETA(DisplayName = "Mip Bias"),
+
+	/* 
+	 * This is like 'Default' but it ignores the WorldPosition pin when computing the mip level.
+	 * It uses the actual pixel WorldPosition instead.
+	 * This can prevent sampling mip 0 if the WorldPosition pin gives a constant value.
+	 */
+	RVTMVM_RecalculateDerivatives UMETA(DisplayName = "Ignore Input WorldPosition "),
 
 	RVTMVM_MAX,
 };
@@ -76,7 +90,7 @@ class ENGINE_API UMaterialExpressionRuntimeVirtualTextureSample : public UMateri
 	UPROPERTY(EditAnywhere, Category = VirtualTexture, meta = (DisplayName = "Enable adaptive page table"))
 	bool bAdaptive = false;
 
-	/** Defines how the MipValue property is applied to the virtual texture lookup. */
+	/** Defines how the mip level is calculated for the virtual texture lookup. */
 	UPROPERTY(EditAnywhere, Category = TextureSample)
 	TEnumAsByte<enum ERuntimeVirtualTextureMipValueMode> MipValueMode = RVTMVM_None;
 
diff --git a/Engine/Source/Runtime/Engine/Private/Materials/MaterialExpressions.cpp b/Engine/Source/Runtime/Engine/Private/Materials/MaterialExpressions.cpp
--- a/Engine/Source/Runtime/Engine/Private/Materials/MaterialExpressions.cpp
+++ b/Engine/Source/Runtime/Engine/Private/Materials/MaterialExpressions.cpp
@@ -2607,18 +2607,32 @@ int32 UMaterialExpressionRuntimeVirtualTextureSample::Compile(class FMaterialCom
 	
 	// Compile the mip level for the current mip value mode
 	ETextureMipValueMode TextureMipLevelMode = TMVM_None;
-	int32 MipValueIndex = INDEX_NONE;
-	if (MipValue.GetTracedInput().Expression != nullptr)
+	int32 MipValue0Index = INDEX_NONE;
+	int32 MipValue1Index = INDEX_NONE;
+	const bool bMipValueExpressionValid = MipValue.GetTracedInput().Expression != nullptr;
+	if (MipValueMode == RVTMVM_MipLevel && bMipValueExpressionValid)
 	{
-		switch (MipValueMode)
-		{
-		case RVTMVM_MipLevel: TextureMipLevelMode = TMVM_MipLevel; break;
-		case RVTMVM_MipBias: TextureMipLevelMode = TMVM_MipBias; break;
-		}
-		if (TextureMipLevelMode != TMVM_None)
-		{
-			MipValueIndex = MipValue.Compile(Compiler);
-		}
+		TextureMipLevelMode = TMVM_MipLevel;
+		MipValue0Index = MipValue.Compile(Compiler);
+	}
+	else if (MipValueMode == RVTMVM_MipBias && bMipValueExpressionValid)
+	{
+		TextureMipLevelMode = TMVM_MipBias;
+		MipValue0Index = MipValue.Compile(Compiler);
+	}
+	else if (MipValueMode == RVTMVM_RecalculateDerivatives)
+	{
+		// Calculate derivatives from world position.
+		TextureMipLevelMode = TMVM_Derivative;
+		const int32 WorldPos = Compiler->WorldPosition(WPT_CameraRelative);
+		const int32 WorldPositionDdx = Compiler->DDX(WorldPos);
+		const int32 UDdx = Compiler->Dot(WorldPositionDdx, Uniforms[ERuntimeVirtualTextureShaderUniform_WorldToUVTransform1]);
+		const int32 VDdx = Compiler->Dot(WorldPositionDdx, Uniforms[ERuntimeVirtualTextureShaderUniform_WorldToUVTransform2]);
+		MipValue0Index = Compiler->AppendVector(UDdx, VDdx);
+		const int32 WorldPositionDdy = Compiler->DDY(WorldPos);
+		const int32 UDdy = Compiler->Dot(WorldPositionDdy, Uniforms[ERuntimeVirtualTextureShaderUniform_WorldToUVTransform1]);
+		const int32 VDdy = Compiler->Dot(WorldPositionDdy, Uniforms[ERuntimeVirtualTextureShaderUniform_WorldToUVTransform2]);
+		MipValue1Index = Compiler->AppendVector(UDdy, VDdy);
 	}
 
 	// Convert texture address mode to matching sampler source mode.
@@ -2643,7 +2657,7 @@ int32 UMaterialExpressionRuntimeVirtualTextureSample::Compile(class FMaterialCom
 			TextureCodeIndex[TexureLayerIndex],
 			CoordinateIndex, 
 			SAMPLERTYPE_VirtualMasks,
-			MipValueIndex, INDEX_NONE, TextureMipLevelMode, SamplerSourceMode,
+			MipValue0Index, MipValue1Index, TextureMipLevelMode, SamplerSourceMode,
 			TextureReferenceIndex[TexureLayerIndex],
 			bAutomaticMipViewBias, bAdaptive);
 	}
