Add ability to reference Android configrules variables via Android's device profile selector.
Add ability for DP rules to dereference operator arguments. e.g. "$(SRC_GPUFamily)" will be processed as "Adreno 660 (TM)"
Add option dereference propertiies when setting user variables.
Add '[at]' first character option to escape property dereferencing.

diff --git a/Engine/Plugins/Runtime/AndroidDeviceProfileSelector/Source/AndroidDeviceProfileSelectorRuntime/Private/AndroidDeviceProfileSelectorRuntimeModule.cpp b/Engine/Plugins/Runtime/AndroidDeviceProfileSelector/Source/AndroidDeviceProfileSelectorRuntime/Private/AndroidDeviceProfileSelectorRuntimeModule.cpp
--- a/Engine/Plugins/Runtime/AndroidDeviceProfileSelector/Source/AndroidDeviceProfileSelectorRuntime/Private/AndroidDeviceProfileSelectorRuntimeModule.cpp
+++ b/Engine/Plugins/Runtime/AndroidDeviceProfileSelector/Source/AndroidDeviceProfileSelectorRuntime/Private/AndroidDeviceProfileSelectorRuntimeModule.cpp
@@ -77,6 +77,17 @@ static const TMap<FName,FString>& GetDeviceSelectorParams()
 		AndroidParams.Add(FAndroidProfileSelectorSourceProperties::SRC_Chipset, GetParam(FString(TEXT("unknown")), TEXT("chipset")));
 		AndroidParams.Add(FAndroidProfileSelectorSourceProperties::SRC_HMDSystemName, HMDRequestedProfileName);
 		AndroidParams.Add(FAndroidProfileSelectorSourceProperties::SRC_TotalPhysicalGB, FString::Printf(TEXT("%d"), TotalPhysicalGB));
+
+#if PLATFORM_ANDROID
+		// allow ConfigRules to override cvars first
+		const TMap<FString, FString>& ConfigRules = FAndroidMisc::GetConfigRulesTMap();
+		for (const TPair<FString, FString>& Pair : ConfigRules)
+		{
+			const FString& VariableName = Pair.Key;
+			const FString& VariableValue = Pair.Value;
+			AndroidParams.Add(FName(FString::Printf(TEXT("SRC_ConfigRuleVar[%s]"), *VariableName)), *VariableValue);
+		}
+#endif
 	}
 	return AndroidParams;
 }
@@ -135,6 +146,15 @@ bool FAndroidDeviceProfileSelectorRuntimeModule::GetSelectorPropertyValue(const
 		PropertyValueOUT = *Found;
 		return true;
 	}
+	// Special case for non-existent config rule variables
+	// they should return true and a value of '[null]'
+	// this prevents configrule issues from throwing errors.
+	if (PropertyType.ToString().StartsWith(TEXT("SRC_ConfigRuleVar[")))
+	{
+		PropertyValueOUT = TEXT("[null]");
+		return true;
+	}
+
 	return false;
 }
 
diff --git a/Engine/Source/Runtime/Engine/Classes/DeviceProfiles/DeviceProfileMatching.h b/Engine/Source/Runtime/Engine/Classes/DeviceProfiles/DeviceProfileMatching.h
--- a/Engine/Source/Runtime/Engine/Classes/DeviceProfiles/DeviceProfileMatching.h
+++ b/Engine/Source/Runtime/Engine/Classes/DeviceProfiles/DeviceProfileMatching.h
@@ -76,7 +76,7 @@ struct FDPMatchingRulestructBase
 	FString AppendFragments;
 
 	UPROPERTY()
-	FString SetSrc;
+	FString SetUserVar;
 
 	virtual const FDPMatchingRulestructBase* GetOnTrue() const { return nullptr; };
 	virtual const FDPMatchingRulestructBase* GetOnFalse() const { return nullptr; };
diff --git a/Engine/Source/Runtime/Engine/Private/DeviceProfiles/DeviceProfileManager.cpp b/Engine/Source/Runtime/Engine/Private/DeviceProfiles/DeviceProfileManager.cpp
--- a/Engine/Source/Runtime/Engine/Private/DeviceProfiles/DeviceProfileManager.cpp
+++ b/Engine/Source/Runtime/Engine/Private/DeviceProfiles/DeviceProfileManager.cpp
@@ -186,7 +186,7 @@ void UDeviceProfileManager::ProcessDeviceProfileIniSettings(const FString& Devic
 #if !UE_BUILD_SHIPPING
 #if PLATFORM_ANDROID
 		// allow ConfigRules to override cvars first
-		TMap<FString, FString> ConfigRules = FAndroidMisc::GetConfigRulesTMap();
+		const TMap<FString, FString>& ConfigRules = FAndroidMisc::GetConfigRulesTMap();
 		for (const TPair<FString, FString>& Pair : ConfigRules)
 		{
 			FString Key = Pair.Key;
diff --git a/Engine/Source/Runtime/Engine/Private/DeviceProfiles/DeviceProfileMatching.cpp b/Engine/Source/Runtime/Engine/Private/DeviceProfiles/DeviceProfileMatching.cpp
--- a/Engine/Source/Runtime/Engine/Private/DeviceProfiles/DeviceProfileMatching.cpp
+++ b/Engine/Source/Runtime/Engine/Private/DeviceProfiles/DeviceProfileMatching.cpp
@@ -22,8 +22,6 @@ static FName SRC_MakeAndModel(TEXT("SRC_MakeAndModel"));
 static FName SRC_OSVersion(TEXT("SRC_OSVersion"));
 static FName SRC_CommandLine(TEXT("SRC_CommandLine"));
 static FName SRC_PrimaryGPUDesc(TEXT("SRC_PrimaryGPUDesc"));
-static FName SRC_False(TEXT("false"));
-static FName SRC_True(TEXT("true"));
 static FName SRC_PreviousRegexMatch(TEXT("SRC_PreviousRegexMatch"));
 static FName SRC_PreviousRegexMatch1(TEXT("SRC_PreviousRegexMatch1"));
 
@@ -45,249 +43,279 @@ static FName CMP_Hash(TEXT("CMP_Hash"));
 static FName CMP_Or(TEXT("OR"));
 static FName CMP_And(TEXT("AND"));
 
-class FRuleMatchRunner
+class FRulePropertiesContainer
 {
-	const TMap<FName, FString>* UserDefinedSrcs;
-	FOutputDevice* ErrorDevice;
-	FString PreviousRegexMatches[2];
-public:
-	struct FDeviceProfileMatchCriterion
+	TMap<FName, FString> UserDefinedValues;
+	bool RetrievePropertyValue(IDeviceProfileSelectorModule* DPSelectorModule, const FName& PropertyType, FString& PropertyValueOUT)
 	{
-		FName SourceType;
-		FString SourceArg;
-		FName CompareType;
-		FString MatchString;
-	};
+		static const FString CommandLine = FCommandLine::Get();
+		// universal properties
+		if (PropertyType == SRC_Chipset) { PropertyValueOUT = FPlatformMisc::GetCPUChipset(); }
+		else if (PropertyType == SRC_MakeAndModel) { PropertyValueOUT = FPlatformMisc::GetDeviceMakeAndModel(); }
+		else if (PropertyType == SRC_OSVersion) { PropertyValueOUT = FPlatformMisc::GetOSVersion(); }
+		else if (PropertyType == SRC_PrimaryGPUDesc) { PropertyValueOUT = FPlatformMisc::GetGPUDriverInfo(FPlatformMisc::GetPrimaryGPUBrand()).DeviceDescription; }
+		else if (PropertyType == SRC_PreviousRegexMatch) { PropertyValueOUT = PreviousRegexMatches[0]; }
+		else if (PropertyType == SRC_PreviousRegexMatch1) { PropertyValueOUT = PreviousRegexMatches[1]; }
+		else if (PropertyType == SRC_CommandLine) { PropertyValueOUT = CommandLine; }
+		// Selector module source data
+		else if (DPSelectorModule && DPSelectorModule->GetSelectorPropertyValue(PropertyType, PropertyValueOUT)) { ; }
+		// SetUserVar defined properties.
+		else if (const FString* Value = UserDefinedValues.Find(PropertyType)) { PropertyValueOUT = *Value; }
+		else
+		{
+			return false;
+		}
+		return true;
+	}
+public:
+	FString PreviousRegexMatches[2];
 
-	FRuleMatchRunner(const TMap<FName, FString>* UserDefinedSrcsIN, FOutputDevice* ErrorDeviceIn) : UserDefinedSrcs(UserDefinedSrcsIN), ErrorDevice(ErrorDeviceIn){}
-	bool ProcessRules(IDeviceProfileSelectorModule* DPSelectorModule, TArray<FDeviceProfileMatchCriterion>& MatchingCriteria, const FString& RuleName)
+	// ExpandVariables takes ParamString and expands any variable references $(...) into their values.
+	// i.e. '$(SRC_GPUFamily)' becomes 'Adreno (TM) 640'
+	// returns true if all variables were successfully expanded or if no variables were referenced. ExpandedResultOUT contains the result.
+	// returns false if an error was encountered, ErrorsOUT will contain the error encountered. ExpandedResultOUT is untouched.
+	// use '@' at the beginning of ParamString to treat the remainder of the string as a literal.
+	bool ExpandVariables(IDeviceProfileSelectorModule* DPSelectorModule, const FString& ParamsString, FString& ExpandedResultOUT, FString& ErrorsOUT)
 	{
-		bool bResult = false;
-		bool bFoundMatch = true;
-		for (const FDeviceProfileMatchCriterion& DeviceProfileMatchCriterion : MatchingCriteria)
-		{
-			bool bCurrentMatch = true;
-			FName SourceType = DeviceProfileMatchCriterion.SourceType;
-			const FString& SourceArg = DeviceProfileMatchCriterion.SourceArg;
-			const FString& MatchString = DeviceProfileMatchCriterion.MatchString;
-			FName CompareType = DeviceProfileMatchCriterion.CompareType;
-			FString CommandLine = FCommandLine::Get();
-
-			FString SourceString("");
-			FString PlatformSourceValue("");
-
-			// Selector module gets first dibs on retrieving source data
-			if (DPSelectorModule && DPSelectorModule->GetSelectorPropertyValue(SourceType, PlatformSourceValue)) { SourceString = PlatformSourceValue; }
-			// universal properties
-			else if (SourceType == SRC_Chipset) { SourceString = FPlatformMisc::GetCPUChipset(); }
-			else if (SourceType == SRC_MakeAndModel) { SourceString = FPlatformMisc::GetDeviceMakeAndModel(); }
-			else if (SourceType == SRC_OSVersion) { SourceString = FPlatformMisc::GetOSVersion(); }
-			else if (SourceType == SRC_PrimaryGPUDesc) {SourceString = FPlatformMisc::GetGPUDriverInfo(FPlatformMisc::GetPrimaryGPUBrand()).DeviceDescription;}
-			else if (SourceType == SRC_PreviousRegexMatch) { SourceString = PreviousRegexMatches[0]; }
-			else if (SourceType == SRC_PreviousRegexMatch1) { SourceString = PreviousRegexMatches[1]; }
-			else if (SourceType == SRC_CommandLine) { SourceString = CommandLine; }
-			else if (SourceType == SRC_False) { SourceString = TEXT("false"); }
-			else if (SourceType == SRC_True) { SourceString = TEXT("true"); }
-			// SetSrc defined properties.
-			else if (const FString* Value = UserDefinedSrcs ? UserDefinedSrcs->Find(SourceType) : nullptr) { SourceString = *Value; }
-			else
-			{
-				// SourceType wasn't found. 
-				ErrorDevice->Logf(TEXT("source type '%s' was not defined for matchingrule %s. (%s, %s, %s)"), *SourceType.ToString(), *RuleName, *SourceType.ToString(), *CompareType.ToString(), *MatchString);
-			}
+		if (ParamsString.IsValidIndex(0) && ParamsString.GetCharArray()[0] == TEXT('@'))
+		{
+			ExpandedResultOUT = ParamsString.RightChop(1);
+			return true;
+		}
+
+		ErrorsOUT.Reset();
+		FString Result = ParamsString;
+
+		int32 VarTokenIndex = INDEX_NONE;
+		while ((VarTokenIndex = ParamsString.Find(TEXT("$("), ESearchCase::CaseSensitive, ESearchDir::FromEnd, VarTokenIndex)) != INDEX_NONE)
+		{
+			int32 VarTokenEndIndex = ParamsString.Find(TEXT(")"), ESearchCase::CaseSensitive, ESearchDir::FromStart, VarTokenIndex + 2);
 
-			const bool bNumericOperands = SourceString.IsNumeric() && MatchString.IsNumeric();
-			if (CompareType == CMP_Equal)
+			if (VarTokenEndIndex != INDEX_NONE)
 			{
-				if (SourceType == SRC_CommandLine)
+				FString SourceVarWithTokens = ParamsString.Mid(VarTokenIndex, (VarTokenEndIndex - VarTokenIndex) + 1);
+				FString SourceVar = SourceVarWithTokens.Mid(2, (VarTokenEndIndex - VarTokenIndex) - 2);
+				int32 FindIdxDummy;
+				FString RetrievedValue;
+				if (SourceVar.FindChar(TEXT('$'), FindIdxDummy) || SourceVar.FindChar(TEXT('('), FindIdxDummy) || !RetrievePropertyValue(DPSelectorModule, FName(SourceVar), RetrievedValue))
 				{
-					if (!FParse::Param(*SourceString, *MatchString))
-					{
-						bCurrentMatch = false;
-					}
+					ErrorsOUT = FString::Printf(TEXT("could not find %s (from %s)"), *SourceVarWithTokens, *ParamsString);
+					return false;
 				}
 				else
 				{
-					if (SourceString != MatchString)
-					{
-						bCurrentMatch = false;
-					}
+					Result.RemoveAt(VarTokenIndex, (VarTokenEndIndex - VarTokenIndex)+1);
+					Result.InsertAt(VarTokenIndex, *RetrievedValue);
 				}
 			}
-			else if (CompareType == CMP_Less)
+			else
 			{
-				if ((bNumericOperands && FCString::Atof(*SourceString) >= FCString::Atof(*MatchString)) || (!bNumericOperands && SourceString >= MatchString))
-				{
-					bCurrentMatch = false;
-				}
+				ErrorsOUT = FString::Printf(TEXT("no closing parenthesis for %s"), *ParamsString);
+				return false;
 			}
-			else if (CompareType == CMP_LessEqual)
+		}
+		ExpandedResultOUT = Result;
+		return true;
+	}
+
+	void SetUserDefinedValue(FName PropertyName, const FString& PropertyValue)
+	{
+		UserDefinedValues.Add(PropertyName, PropertyValue);
+	}
+};
+
+class FRuleMatchRunner
+{
+	FRulePropertiesContainer& RuleProperties;
+	FOutputDevice* ErrorDevice;
+
+public:
+	struct FDeviceProfileMatchCriterion
+	{
+		FString Arg1;
+		FName Operator;
+		FString Arg2;
+	};
+
+	FRuleMatchRunner(FRulePropertiesContainer& RulePropertiesIN, FOutputDevice* ErrorDeviceIn) : RuleProperties(RulePropertiesIN), ErrorDevice(ErrorDeviceIn){}
+	bool ProcessRule(IDeviceProfileSelectorModule* DPSelectorModule, const FDeviceProfileMatchCriterion& DeviceProfileMatchCriterion, const FString& RuleName)
+	{
+		bool bCurrentMatch = true;
+
+		const FString& SourceString = DeviceProfileMatchCriterion.Arg1;
+		FName CompareType = DeviceProfileMatchCriterion.Operator;
+		const FString& MatchString = DeviceProfileMatchCriterion.Arg2;
+
+		const bool bNumericOperands = SourceString.IsNumeric() && MatchString.IsNumeric();
+
+		if (CompareType == CMP_Equal)
+		{
+			if (SourceString != MatchString)
 			{
-				if ((bNumericOperands && FCString::Atof(*SourceString) > FCString::Atof(*MatchString)) || (!bNumericOperands && SourceString > MatchString))
-				{
-					bCurrentMatch = false;
-				}
+				bCurrentMatch = false;
 			}
-			else if (CompareType == CMP_Greater)
+		}
+		else if (CompareType == CMP_Less)
+		{
+			if ((bNumericOperands && FCString::Atof(*SourceString) >= FCString::Atof(*MatchString)) || (!bNumericOperands && SourceString >= MatchString))
 			{
-				if ((bNumericOperands && FCString::Atof(*SourceString) <= FCString::Atof(*MatchString)) || (!bNumericOperands && SourceString <= MatchString))
-				{
-					bCurrentMatch = false;
-				}
+				bCurrentMatch = false;
 			}
-			else if (CompareType == CMP_GreaterEqual)
+		}
+		else if (CompareType == CMP_LessEqual)
+		{
+			if ((bNumericOperands && FCString::Atof(*SourceString) > FCString::Atof(*MatchString)) || (!bNumericOperands && SourceString > MatchString))
 			{
-				if ((bNumericOperands && FCString::Atof(*SourceString) < FCString::Atof(*MatchString)) || (!bNumericOperands && SourceString < MatchString))
-				{
-					bCurrentMatch = false;
-				}
+				bCurrentMatch = false;
 			}
-			else if (CompareType == CMP_NotEqual)
+		}
+		else if (CompareType == CMP_Greater)
+		{
+			if ((bNumericOperands && FCString::Atof(*SourceString) <= FCString::Atof(*MatchString)) || (!bNumericOperands && SourceString <= MatchString))
 			{
-				if (SourceType == SRC_CommandLine)
-				{
-					if (FParse::Param(*CommandLine, *MatchString))
-					{
-						bCurrentMatch = false;
-					}
-				}
-				else
-				{
-					if (*SourceString == MatchString)
-					{
-						bCurrentMatch = false;
-					}
-				}
+				bCurrentMatch = false;
 			}
-			else if (CompareType == CMP_Or || CompareType == CMP_And)
+		}
+		else if (CompareType == CMP_GreaterEqual)
+		{
+			if ((bNumericOperands && FCString::Atof(*SourceString) < FCString::Atof(*MatchString)) || (!bNumericOperands && SourceString < MatchString))
 			{
-				bool bArg1, bArg2;
+				bCurrentMatch = false;
+			}
+		}
+		else if (CompareType == CMP_NotEqual)
+		{
+			if (*SourceString == MatchString)
+			{
+				bCurrentMatch = false;
+			}
+		}
+		else if (CompareType == CMP_Or || CompareType == CMP_And)
+		{
+			bool bArg1, bArg2;
 
-				if (bNumericOperands)
-				{
-					bArg1 = FCString::Atoi(*SourceString) != 0;
-					bArg2 = FCString::Atoi(*MatchString) != 0;
-				}
-				else
-				{
-					static const FString TrueString("true");
-					bArg1 = SourceString == TrueString;
-					bArg2 = MatchString == TrueString;
-				}
+			if (bNumericOperands)
+			{
+				bArg1 = FCString::Atoi(*SourceString) != 0;
+				bArg2 = FCString::Atoi(*MatchString) != 0;
+			}
+			else
+			{
+				static const FString TrueString("true");
+				bArg1 = SourceString == TrueString;
+				bArg2 = MatchString == TrueString;
+			}
 
-				if (CompareType == CMP_Or)
-				{
-					bCurrentMatch = (bArg1 || bArg2);
-				}
-				else
-				{
-					bCurrentMatch = (bArg1 && bArg2);
-				}
+			if (CompareType == CMP_Or)
+			{
+				bCurrentMatch = (bArg1 || bArg2);
 			}
-			else if (CompareType == CMP_EqualIgnore)
+			else
 			{
-				if (SourceString.ToLower() != MatchString.ToLower())
-				{
-					bCurrentMatch = false;
-				}
+				bCurrentMatch = (bArg1 && bArg2);
 			}
-			else if (CompareType == CMP_LessIgnore)
+		}
+		else if (CompareType == CMP_EqualIgnore)
+		{
+			if (SourceString.ToLower() != MatchString.ToLower())
 			{
-				if (SourceString.ToLower() >= MatchString.ToLower())
-				{
-					bCurrentMatch = false;
-				}
+				bCurrentMatch = false;
+			}
+		}
+		else if (CompareType == CMP_LessIgnore)
+		{
+			if (SourceString.ToLower() >= MatchString.ToLower())
+			{
+				bCurrentMatch = false;
 			}
-			else if (CompareType == CMP_LessEqualIgnore)
+		}
+		else if (CompareType == CMP_LessEqualIgnore)
+		{
+			if (SourceString.ToLower() > MatchString.ToLower())
 			{
-				if (SourceString.ToLower() > MatchString.ToLower())
-				{
-					bCurrentMatch = false;
-				}
+				bCurrentMatch = false;
 			}
-			else if (CompareType == CMP_GreaterIgnore)
+		}
+		else if (CompareType == CMP_GreaterIgnore)
+		{
+			if (SourceString.ToLower() <= MatchString.ToLower())
 			{
-				if (SourceString.ToLower() <= MatchString.ToLower())
-				{
-					bCurrentMatch = false;
-				}
-				else if (CompareType == CMP_GreaterEqualIgnore)
-				{
-					if (SourceString.ToLower() < MatchString.ToLower())
-					{
-						bCurrentMatch = false;
-					}
-				}
+				bCurrentMatch = false;
 			}
-			else if (CompareType == CMP_NotEqualIgnore)
+			else if (CompareType == CMP_GreaterEqualIgnore)
 			{
-				if (SourceString.ToLower() == MatchString.ToLower())
+				if (SourceString.ToLower() < MatchString.ToLower())
 				{
 					bCurrentMatch = false;
 				}
 			}
-			else if (CompareType == CMP_Regex)
+		}
+		else if (CompareType == CMP_NotEqualIgnore)
+		{
+			if (SourceString.ToLower() == MatchString.ToLower())
 			{
-				const FRegexPattern RegexPattern(MatchString);
-				FRegexMatcher RegexMatcher(RegexPattern, SourceString);
-				if (RegexMatcher.FindNext())
-				{
-					PreviousRegexMatches[0] = RegexMatcher.GetCaptureGroup(1);
-					PreviousRegexMatches[1] = RegexMatcher.GetCaptureGroup(2);
-				}
-				else
-				{
-					PreviousRegexMatches[0].Empty();
-					PreviousRegexMatches[1].Empty();
+				bCurrentMatch = false;
+			}
+		}
+		else if (CompareType == CMP_Regex)
+		{
+			const FRegexPattern RegexPattern(MatchString);
+			FRegexMatcher RegexMatcher(RegexPattern, SourceString);
+			if (RegexMatcher.FindNext())
+			{
+				RuleProperties.PreviousRegexMatches[0] = RegexMatcher.GetCaptureGroup(1);
+				RuleProperties.PreviousRegexMatches[1] = RegexMatcher.GetCaptureGroup(2);
+			}
+			else
+			{
+				RuleProperties.PreviousRegexMatches[0].Empty();
+				RuleProperties.PreviousRegexMatches[1].Empty();
 
-					bCurrentMatch = false;
-				}
+				bCurrentMatch = false;
 			}
-			else if (CompareType == CMP_Hash)
+		}
+		else if (CompareType == CMP_Hash)
+		{
+			// Salt string is concatenated onto the end of the input text.
+			// For example the input string "PhoneModel" with salt "Salt" and pepper "Pepper" can be computed with
+			// % printf "PhoneModelSaltPepper" | openssl dgst -sha1 -hex
+			// resulting in d9e5cbd6b0e4dba00edd9de92cf64ee4c3f3a2db and would be stored in the matching rules as 
+			// "Salt|d9e5cbd6b0e4dba00edd9de92cf64ee4c3f3a2db". Salt is optional.
+			FString MatchHashString;
+			FString SaltString;
+			if (!MatchString.Split(TEXT("|"), &SaltString, &MatchHashString))
 			{
-				// Salt string is concatenated onto the end of the input text.
-				// For example the input string "PhoneModel" with salt "Salt" and pepper "Pepper" can be computed with
-				// % printf "PhoneModelSaltPepper" | openssl dgst -sha1 -hex
-				// resulting in d9e5cbd6b0e4dba00edd9de92cf64ee4c3f3a2db and would be stored in the matching rules as 
-				// "Salt|d9e5cbd6b0e4dba00edd9de92cf64ee4c3f3a2db". Salt is optional.
-				FString MatchHashString;
-				FString SaltString;
-				if (!MatchString.Split(TEXT("|"), &SaltString, &MatchHashString))
-				{
-					MatchHashString = MatchString;
-				}
-				FString HashInputString = SourceString + SaltString
+				MatchHashString = MatchString;
+			}
+			FString HashInputString = SourceString + SaltString
 #ifdef HASH_PEPPER_SECRET_GUID
-					+ HASH_PEPPER_SECRET_GUID.ToString()
+				+ HASH_PEPPER_SECRET_GUID.ToString()
 #endif
-					;
+				;
 
-				FSHAHash SourceHash;
-				FSHA1::HashBuffer(TCHAR_TO_ANSI(*HashInputString), HashInputString.Len(), SourceHash.Hash);
-				if (SourceHash.ToString() != MatchHashString.ToUpper())
-				{
-					bCurrentMatch = false;
-				}
-			}
-			else
+			FSHAHash SourceHash;
+			FSHA1::HashBuffer(TCHAR_TO_ANSI(*HashInputString), HashInputString.Len(), SourceHash.Hash);
+			if (SourceHash.ToString() != MatchHashString.ToUpper())
 			{
 				bCurrentMatch = false;
 			}
-
-			bFoundMatch = bCurrentMatch;
 		}
-		return bFoundMatch;
+		else
+		{
+			ErrorDevice->Logf(TEXT("rule %s unknown operator type, %s"), *RuleName, *CompareType.ToString());
+	
+			bCurrentMatch = false;
+		}
+
+		return bCurrentMatch;
 	}
 };
 
-static bool EvaluateMatch(IDeviceProfileSelectorModule* DPSelectorModule, const FDPMatchingRulestructBase& rule, TMap<FName, FString>& UserDefinedSrc, FString& SelectedFragmentsOUT, FOutputDevice* Errors)
+static bool EvaluateMatch(IDeviceProfileSelectorModule* DPSelectorModule, const FString& RuleName, const FDPMatchingRulestructBase& rule, FRulePropertiesContainer& RuleProperties, FString& SelectedFragmentsOUT, FOutputDevice* Errors)
 {
-	static const FString NoName("");
-	const FString& RuleName = rule.RuleName.IsEmpty() ? NoName : rule.RuleName;
-
-	if(!rule.AppendFragments.IsEmpty())
+	if (!rule.AppendFragments.IsEmpty())
 	{
-		if(!SelectedFragmentsOUT.IsEmpty())
+		if (!SelectedFragmentsOUT.IsEmpty())
 		{
 			SelectedFragmentsOUT += TEXT(",");
 		}
@@ -302,17 +330,26 @@ static bool EvaluateMatch(IDeviceProfileSelectorModule* DPSelectorModule, const
 		}
 	}
 
-	if (!rule.SetSrc.IsEmpty())
+	if (!rule.SetUserVar.IsEmpty())
 	{
 		TArray<FString> NewSRCs;
-		rule.SetSrc.ParseIntoArray(NewSRCs, TEXT(","), true);
+		rule.SetUserVar.ParseIntoArray(NewSRCs, TEXT(","), true);
 		for (FString& SRCentry : NewSRCs)
 		{
 			FString CSRCType, CSRCValue;
 			if (SRCentry.Split(TEXT("="), &CSRCType, &CSRCValue))
 			{
-				UserDefinedSrc.Add(FName(CSRCType), CSRCValue);
-				UE_LOG(LogInit, Verbose, TEXT("MatchesRules: Adding source %s : %s"), *CSRCType, *CSRCValue);
+				FString ExpandedValue;
+				FString ErrorStr;
+				if (!RuleProperties.ExpandVariables(DPSelectorModule, CSRCValue, ExpandedValue, ErrorStr))
+				{
+					Errors->Logf(TEXT("MatchesRules: rule %s failed. SetSRC %s= could not find %s"), *RuleName, *CSRCType, *ErrorStr);
+				}
+				else
+				{
+					RuleProperties.SetUserDefinedValue(FName(CSRCType), ExpandedValue);
+					UE_LOG(LogInit, Verbose, TEXT("MatchesRules: Adding source %s : %s"), *CSRCType, *ExpandedValue);
+				}
 			}
 		}
 	}
@@ -320,7 +357,7 @@ static bool EvaluateMatch(IDeviceProfileSelectorModule* DPSelectorModule, const
 	TArray<FDPMatchingIfCondition> Expression = rule.IfConditions;
 	if (Expression.Num() == 0)
 	{
-		UE_LOG(LogInit, Verbose, TEXT("MatchesRules: %s, no match criteria."), *RuleName);
+		UE_LOG(LogInit, Verbose, TEXT("MatchesRules: %s = true (no rules specified)"), *RuleName);
 		return true;
 	}
 
@@ -375,10 +412,8 @@ static bool EvaluateMatch(IDeviceProfileSelectorModule* DPSelectorModule, const
 		if (!element.Arg1.IsEmpty())
 		{
 			PushOperand(element.Arg1);
-			//if (!element.Arg2.IsEmpty())
-			{
-				PushOperand(element.Arg2);
-			}
+			PushOperand(element.Arg2);
+
 			if (!element.Operator.IsNone())
 			{
 				PushOperator(element.Operator.ToString());
@@ -417,13 +452,8 @@ static bool EvaluateMatch(IDeviceProfileSelectorModule* DPSelectorModule, const
 		}
 		PushOperator(PoppedOperator);
 	}
-	UE_LOG(LogInit, Verbose, TEXT("MatchesRules: rule %s : "), *RuleName);
-	for (int i = 0; i < RPNOutput.Num(); i++)
-	{
-		UE_LOG(LogInit, Verbose, TEXT("MatchesRules: (%d - %s)"), i, *RPNOutput[i].Value);
-	}
 
-	FRuleMatchRunner MatchMe(&UserDefinedSrc, Errors);
+	FRuleMatchRunner MatchMe(RuleProperties, Errors);
 	TArray<FExpressionItem> ExpressionStack;
 	for (int i = 0; i < RPNOutput.Num(); i++)
 	{
@@ -435,19 +465,27 @@ static bool EvaluateMatch(IDeviceProfileSelectorModule* DPSelectorModule, const
 		{
 			FExpressionItem B = ExpressionStack.Pop();
 			FExpressionItem A = ExpressionStack.Pop();
-			TArray<FRuleMatchRunner::FDeviceProfileMatchCriterion> MatchingCriteria;
-
 			FRuleMatchRunner::FDeviceProfileMatchCriterion crit;
-			crit.CompareType = FName(RPNOutput[i].Value);
-			crit.SourceType = FName(A.Value);
-			crit.MatchString = B.Value;
+			crit.Operator = FName(RPNOutput[i].Value);
 
-			MatchingCriteria.Add(crit);
-			bool bResult = MatchMe.ProcessRules(DPSelectorModule, MatchingCriteria, RuleName);
+			// process the arguments and dereference any $(variables) to their content.
+			FString ErrorStr;
+			if (!RuleProperties.ExpandVariables(DPSelectorModule, A.Value, crit.Arg1, ErrorStr))
+			{
+				Errors->Logf(TEXT("rule %s failed to expand Arg1, %s"), *RuleName, *ErrorStr);
+				crit.Arg1 = TEXT("[UnknownArg1Variable]");
+			}
+			if (!RuleProperties.ExpandVariables(DPSelectorModule, B.Value, crit.Arg2, ErrorStr))
+			{
+				Errors->Logf(TEXT("rule %s failed to expand Arg2, %s"), *RuleName, *ErrorStr);
+				crit.Arg2 = TEXT("[UnknownArg2Variable]");
+			}
+
+			bool bResult = MatchMe.ProcessRule(DPSelectorModule, crit, RuleName);
 			FExpressionItem C(bResult ? TEXT("true") : TEXT("false"), false);
 			ExpressionStack.Push(C);
 
-			UE_LOG(LogInit, Verbose, TEXT("MatchesRules: rule %s evaluating (%s %s %s) = %s"), *RuleName, *A.Value, *RPNOutput[i].Value, *B.Value, *C.Value);
+			UE_LOG(LogInit, Verbose, TEXT("MatchesRules: rule %s evaluating {%s %s %s} = {%s %s %s} = %s"), *RuleName, *A.Value, *RPNOutput[i].Value, *B.Value, *crit.Arg1, *RPNOutput[i].Value, *crit.Arg2, *C.Value);
 		}
 	}
 	FExpressionItem Result = ExpressionStack.Pop();
@@ -459,7 +497,11 @@ static bool EvaluateMatch(IDeviceProfileSelectorModule* DPSelectorModule, const
 	bool bSuccess = true;
 	if (Next)
 	{
-		bSuccess = EvaluateMatch(DPSelectorModule, *Next, UserDefinedSrc, SelectedFragmentsOUT, Errors);
+		// Generate a name from the parent and the new rule to help trace problems.
+		static const FString NoName("<unnamed>");
+		const FString NewRuleName = RuleName + TEXT("::")+(Next->RuleName.IsEmpty() ? NoName : Next->RuleName);
+
+		bSuccess = EvaluateMatch(DPSelectorModule, NewRuleName, *Next, RuleProperties, SelectedFragmentsOUT, Errors);
 	}
 
 	return bSuccess;
@@ -533,7 +575,7 @@ static FString LoadAndProcessMatchingRulesFromConfig(const FString& ParentDP, ID
 	TArray<FString> MatchingRulesArray;
 	ConfigSystem->GetArray(*CurrentSectionName, *ArrayName, MatchingRulesArray, GDeviceProfilesIni);
 
-	TMap<FName, FString> UserDefinedSrc;
+	FRulePropertiesContainer RuleProperties;
 	FString SelectedFragments;
 	FDeviceProfileMatchingErrorContext DPMatchingErrorOutput;
 	int Count = 0;
@@ -543,14 +585,18 @@ static FString LoadAndProcessMatchingRulesFromConfig(const FString& ParentDP, ID
 		Count++;
 		FDPMatchingRulestruct RuleStruct;
 		FDPMatchingRulestruct::StaticStruct()->ImportText(*RuleText, &RuleStruct, nullptr, EPropertyPortFlags::PPF_None, &DPMatchingErrorOutput, FDPMatchingRulestruct::StaticStruct()->GetName(), true);
-		EvaluateMatch(DPSelector, RuleStruct, UserDefinedSrc, SelectedFragments, &DPMatchingErrorOutput);
+		static const FString NoName("<unnamed>");
+		const FString& RuleName = RuleStruct.RuleName.IsEmpty() ? NoName : RuleStruct.RuleName;
+		EvaluateMatch(DPSelector, RuleName, RuleStruct, RuleProperties, SelectedFragments, &DPMatchingErrorOutput);
 	}
-#if UE_BUILD_SHIPPING
+#if 1//UE_BUILD_SHIPPING
 	UE_CLOG(DPMatchingErrorOutput.NumErrors > 0, LogInit, Error, TEXT("DeviceProfileMatching: %d Error(s) encountered while processing MatchedRules for %s"), DPMatchingErrorOutput.NumErrors, *ParentDP);
 #else
 	UE_CLOG(DPMatchingErrorOutput.NumErrors > 0, LogInit, Fatal, TEXT("DeviceProfileMatching: %d Error(s) encountered while processing MatchedRules for %s"), DPMatchingErrorOutput.NumErrors, *ParentDP);
 #endif
 
+	FGenericCrashContext::SetEngineData(TEXT("MatchingDPStatus_")+ParentDP, DPMatchingErrorOutput.NumErrors > 0 ? TEXT("Error") : TEXT("No errors"));
+
 	return SelectedFragments;
 }
 
