Fix build issue

diff --git a/Engine/Source/Runtime/Engine/Classes/DeviceProfiles/DeviceProfileMatching.h b/Engine/Source/Runtime/Engine/Classes/DeviceProfiles/DeviceProfileMatching.h
--- a/Engine/Source/Runtime/Engine/Classes/DeviceProfiles/DeviceProfileMatching.h
+++ b/Engine/Source/Runtime/Engine/Classes/DeviceProfiles/DeviceProfileMatching.h
@@ -86,7 +86,7 @@ struct FDPMatchingRulestructBase
 		return IfConditions == Other.IfConditions
 			&& RuleName.Equals(Other.RuleName, ESearchCase::CaseSensitive)
 			&& AppendFragments.Equals(Other.AppendFragments, ESearchCase::CaseSensitive)
-			&& SetSrc.Equals(Other.SetSrc, ESearchCase::CaseSensitive);
+			&& SetUserVar.Equals(Other.SetUserVar, ESearchCase::CaseSensitive);
 	}
 };
 
