Array index out of bounds when Building Lighting Only on an empty level with a BSP box

diff --git a/Engine/Source/Editor/UnrealEd/Private/StaticLightingSystem/StaticLightingSystem.cpp b/Engine/Source/Editor/UnrealEd/Private/StaticLightingSystem/StaticLightingSystem.cpp
--- a/Engine/Source/Editor/UnrealEd/Private/StaticLightingSystem/StaticLightingSystem.cpp
+++ b/Engine/Source/Editor/UnrealEd/Private/StaticLightingSystem/StaticLightingSystem.cpp
@@ -1639,7 +1639,7 @@ void FStaticLightingSystem::AddBSPStaticLightingInfo(ULevel* Level, bool bBuildL
 	{
 		FNodeGroup* NodeGroup = It.Value();
 
-		if (NodeGroup->Nodes.Num())
+		if (NodeGroup->Nodes.Num() && Level->ModelComponents.Num())
 		{
 			// get one of the surfaces/components from the NodeGroup
 			// @todo: Remove need for GetSurfaceLightMapResolution to take a surfaceindex, or a ModelComponent :)
diff --git a/Engine/Source/Runtime/Engine/Private/ModelLight.cpp b/Engine/Source/Runtime/Engine/Private/ModelLight.cpp
--- a/Engine/Source/Runtime/Engine/Private/ModelLight.cpp
+++ b/Engine/Source/Runtime/Engine/Private/ModelLight.cpp
@@ -786,7 +786,7 @@ void UModel::GroupAllNodes(ULevel* Level, const TArray<ULightComponentBase*>& Li
 	{
 		const FBspNode& Node = Nodes[NodeIndex];
 
-		if (Node.NumVertices > 0 && HasStaticLightingCache[Node.ComponentIndex])
+		if (Node.NumVertices > 0 && HasStaticLightingCache.Num() && HasStaticLightingCache[Node.ComponentIndex])
 		{
 			const FBspSurf& Surf = Surfs[Node.iSurf];
 			PlaneMap.AddPlane(FPlane(Surf.Plane), NodeIndex);	// LWC_TODO: Perf pessimization
@@ -974,7 +974,7 @@ void UModel::GroupAllNodes(ULevel* Level, const TArray<ULightComponentBase*>& Li
 			// add the relevant lights to the nodegroup
 			TArray<ULightComponent*>* RelevantLights = ComponentRelevantLights.Find(Nodes[NodeIndex].ComponentIndex);
 			check(RelevantLights);
-			for (int32 LightIndex = 0; LightIndex < RelevantLights->Num(); LightIndex++)
+			for (int32 LightIndex = 0; RelevantLights && LightIndex < RelevantLights->Num(); LightIndex++)
 			{
 				NodeGroup->RelevantLights.AddUnique((*RelevantLights)[LightIndex]);
 			}
