Align texture legacy cache key to float precision (not double) to avoid change affecting projects that have DDC cached results with the float precision key.

diff --git a/Engine/Source/Runtime/Engine/Private/TextureDerivedData.cpp b/Engine/Source/Runtime/Engine/Private/TextureDerivedData.cpp
--- a/Engine/Source/Runtime/Engine/Private/TextureDerivedData.cpp
+++ b/Engine/Source/Runtime/Engine/Private/TextureDerivedData.cpp
@@ -95,8 +95,8 @@ static void SerializeForKey(FArchive& Ar, const FTextureBuildSettings& Settings)
 	float TempFloat;
 	uint8 TempByte;
 	FColor TempColor;
-	FVector2d TempVector2d;
-	FVector4d TempVector4d;
+	FVector2f TempVector2f;
+	FVector4f TempVector4f;
 	UE::Color::FColorSpace TempColorSpace;
 	FGuid TempGuid;
 
@@ -129,10 +129,10 @@ static void SerializeForKey(FArchive& Ar, const FTextureBuildSettings& Settings)
 		TempUint32 = UE::Color::COLORSPACE_VER; Ar << TempUint32;
 		TempColorSpace = UE::Color::FColorSpace::GetWorking(); Ar << TempColorSpace;
 
-		TempVector2d = Settings.RedChromaticityCoordinate; Ar << TempVector2d;
-		TempVector2d = Settings.GreenChromaticityCoordinate; Ar << TempVector2d;
-		TempVector2d = Settings.BlueChromaticityCoordinate; Ar << TempVector2d;
-		TempVector2d = Settings.WhiteChromaticityCoordinate; Ar << TempVector2d;
+		TempVector2f = Settings.RedChromaticityCoordinate; Ar << TempVector2f;
+		TempVector2f = Settings.GreenChromaticityCoordinate; Ar << TempVector2f;
+		TempVector2f = Settings.BlueChromaticityCoordinate; Ar << TempVector2f;
+		TempVector2f = Settings.WhiteChromaticityCoordinate; Ar << TempVector2f;
 		TempByte = Settings.ChromaticAdaptationMethod; Ar << TempByte;
 	}
 
@@ -141,7 +141,7 @@ static void SerializeForKey(FArchive& Ar, const FTextureBuildSettings& Settings)
 
 	if (Settings.AlphaCoverageThresholds != FVector4(0, 0, 0, 0))
 	{
-		TempVector4d = Settings.AlphaCoverageThresholds; Ar << TempVector4d;
+		TempVector4f = Settings.AlphaCoverageThresholds; Ar << TempVector4f;
 	}
 	
 	TempByte = Settings.bComputeBokehAlpha; Ar << TempByte;
