Ensure that texture mips can be streamed in whether there was an async request issued for it or not.

diff --git a/Engine/Source/Runtime/Engine/Private/Streaming/Texture2DStreamIn_DDC.cpp b/Engine/Source/Runtime/Engine/Private/Streaming/Texture2DStreamIn_DDC.cpp
--- a/Engine/Source/Runtime/Engine/Private/Streaming/Texture2DStreamIn_DDC.cpp
+++ b/Engine/Source/Runtime/Engine/Private/Streaming/Texture2DStreamIn_DDC.cpp
@@ -102,7 +102,7 @@ FTexture2DStreamIn_DDC::FTexture2DStreamIn_DDC(UTexture2D* InTexture)
 	, DDCRequestOwner(UE::DerivedData::EPriority::Normal)
 {
 	DDCHandles.AddZeroed(ResourceState.MaxNumLODs);
-	DDCBuffers.AddZeroed(ResourceState.MaxNumLODs);
+	DDCMipRequestStatus.AddZeroed(ResourceState.MaxNumLODs);
 }
 
 FTexture2DStreamIn_DDC::~FTexture2DStreamIn_DDC()
@@ -171,12 +171,13 @@ void FTexture2DStreamIn_DDC::DoCreateAsyncDDCRequests(const FContext& Context)
 				if (MipMap.IsPagedToDerivedData())
 				{
 					FCacheGetChunkRequest& Request = MipKeys.AddDefaulted_GetRef();
-					MipNameBuilder.Appendf(TEXT(" [MIP 0]"), MipIndex + LODBias);
+					MipNameBuilder.Appendf(TEXT(" [MIP %d]"), MipIndex + LODBias);
 					Request.Name = MipNameBuilder;
 					Request.Key = Key;
 					Request.Id = FTexturePlatformData::MakeMipId(MipIndex + LODBias);
 					Request.UserData = MipIndex;
 					MipNameBuilder.RemoveSuffix(MipNameBuilder.Len() - TextureNameLen);
+					DDCMipRequestStatus[MipIndex].bRequestIssued = true;
 				}
 			}
 
@@ -187,8 +188,8 @@ void FTexture2DStreamIn_DDC::DoCreateAsyncDDCRequests(const FContext& Context)
 					if (Response.Status == EStatus::Ok)
 					{
 						const int32 MipIndex = int32(Response.UserData);
-						check(!DDCBuffers[MipIndex]);
-						DDCBuffers[MipIndex] = MoveTemp(Response.RawData);
+						check(!DDCMipRequestStatus[MipIndex].Buffer);
+						DDCMipRequestStatus[MipIndex].Buffer = MoveTemp(Response.RawData);
 					}
 				});
 			}
@@ -283,6 +284,12 @@ void FTexture2DStreamIn_DDC::DoLoadNewMipsFromDDC(const FContext& Context)
 		}
 		else if (PlatformData->DerivedDataKey.IsType<UE::DerivedData::FCacheKeyProxy>())
 		{
+			using namespace UE::DerivedData;
+			const FCacheKey& Key = *PlatformData->DerivedDataKey.Get<FCacheKeyProxy>().AsCacheKey();
+			TStringBuilder<256> MipNameBuilder;
+			Context.Texture->GetPathName(nullptr, MipNameBuilder);
+			const int32 TextureNameLen = MipNameBuilder.Len();
+
 			for (int32 MipIndex = PendingFirstLODIdx; MipIndex < CurrentFirstLODIdx && !IsCancelled(); ++MipIndex)
 			{
 				const FTexture2DMipMap& MipMap = *Context.MipsView[MipIndex];
@@ -290,19 +297,47 @@ void FTexture2DStreamIn_DDC::DoLoadNewMipsFromDDC(const FContext& Context)
 
 				if (MipMap.IsPagedToDerivedData())
 				{
-					if (DDCBuffers[MipIndex])
+					FSharedBuffer MipResult;
+					FMipRequestStatus& MipRequestStatus = DDCMipRequestStatus[MipIndex];
+					if (MipRequestStatus.bRequestIssued)
+					{
+						MipResult = MoveTemp(MipRequestStatus.Buffer);
+						MipRequestStatus.bRequestIssued = false;
+					}
+					else
+					{
+						FCacheGetChunkRequest Request;
+						MipNameBuilder.Appendf(TEXT(" [MIP %d]"), MipIndex + LODBias);
+						Request.Name = MipNameBuilder;
+						Request.Key = Key;
+						Request.Id = FTexturePlatformData::MakeMipId(MipIndex + LODBias);
+						Request.UserData = MipIndex;
+						MipNameBuilder.RemoveSuffix(MipNameBuilder.Len() - TextureNameLen);
+
+						FRequestOwner BlockingRequestOwner(EPriority::Blocking);
+						GetCache().GetChunks({Request}, DDCRequestOwner, [MipIndex, &MipResult](FCacheGetChunkResponse&& Response)
+						{
+							if (Response.Status == EStatus::Ok)
+							{
+								MipResult = MoveTemp(Response.RawData);
+							}
+						});
+						BlockingRequestOwner.Wait();
+
+					}
+
+					if (MipResult)
 					{
 						const int32 ExpectedMipSize = CalcTextureMipMapSize(MipMap.SizeX, MipMap.SizeY, Context.Resource->GetPixelFormat(), 0);
-						if (DDCBuffers[MipIndex].GetSize() == ExpectedMipSize)
+						if (MipResult.GetSize() == ExpectedMipSize)
 						{
-							FMemory::Memcpy(MipData[MipIndex], DDCBuffers[MipIndex].GetData(), DDCBuffers[MipIndex].GetSize());
+							FMemory::Memcpy(MipData[MipIndex], MipResult.GetData(), MipResult.GetSize());
 						}
 						else
 						{
-							UE_LOG(LogTexture, Error, TEXT("DDC mip size (%d) not as expected (%d) for mip %d of %s."), static_cast<int32>(DDCBuffers[MipIndex].GetSize()), ExpectedMipSize, MipIndex, *Context.Texture->GetPathName());
+							UE_LOG(LogTexture, Error, TEXT("DDC mip size (%d) not as expected (%d) for mip %d of %s."), static_cast<int32>(MipResult.GetSize()), ExpectedMipSize, MipIndex, *Context.Texture->GetPathName());
 							MarkAsCancelled();
 						}
-						DDCBuffers[MipIndex].Reset();
 					}
 					else
 					{
diff --git a/Engine/Source/Runtime/Engine/Private/Streaming/Texture2DStreamIn_DDC.h b/Engine/Source/Runtime/Engine/Private/Streaming/Texture2DStreamIn_DDC.h
--- a/Engine/Source/Runtime/Engine/Private/Streaming/Texture2DStreamIn_DDC.h
+++ b/Engine/Source/Runtime/Engine/Private/Streaming/Texture2DStreamIn_DDC.h
@@ -24,9 +24,15 @@ public:
 
 protected:
 
+	struct FMipRequestStatus
+	{
+		FSharedBuffer Buffer;
+		bool bRequestIssued = false;
+	};
+
 	// StreamIn_Default : Locked mips of the intermediate textures, used as disk load destination.
 	TArray<uint32, TInlineAllocator<MAX_TEXTURE_MIP_COUNT> > DDCHandles;
-	TArray<FSharedBuffer, TInlineAllocator<MAX_TEXTURE_MIP_COUNT> > DDCBuffers;
+	TArray<FMipRequestStatus, TInlineAllocator<MAX_TEXTURE_MIP_COUNT> > DDCMipRequestStatus;
 	UE::DerivedData::FRequestOwner DDCRequestOwner;
 
 	// ****************************
