World Partition Cleanup
- Removed actor descriptors OnRegister/OnUnegister events.
- Actor descriptors now get their container set when added to a container.
- World Outliner now gets the unloaded actor associated world through the actor descriptor container.
- Moved all code for World Outliner actors pinning from actor descriptor containers to world partition.
- Level Instance actor descriptors now register their level container when associated to an actor descriptor container.

diff --git a/Engine/Source/Editor/SceneOutliner/Private/ActorDescTreeItem.cpp b/Engine/Source/Editor/SceneOutliner/Private/ActorDescTreeItem.cpp
--- a/Engine/Source/Editor/SceneOutliner/Private/ActorDescTreeItem.cpp
+++ b/Engine/Source/Editor/SceneOutliner/Private/ActorDescTreeItem.cpp
@@ -1,6 +1,7 @@
 // Copyright Epic Games, Inc. All Rights Reserved.
 
 #include "ActorDescTreeItem.h"
+#include "WorldPartition/WorldPartition.h"
 #include "WorldPartition/WorldPartitionActorDesc.h"
 #include "Widgets/DeclarativeSyntaxSupport.h"
 #include "Layout/WidgetPath.h"
@@ -273,7 +274,8 @@ bool FActorDescTreeItem::GetPinnedState() const
 {
 	if (ActorDescHandle.IsValid())
 	{
-		return ActorDescHandle->GetContainer()->IsActorPinned(GetGuid());
+		UWorldPartition* WorldPartition = Cast<UWorldPartition>(ActorDescHandle->GetContainer());
+		return WorldPartition ? WorldPartition->IsActorPinned(GetGuid()) : false;
 	}
 	return false;
 }
diff --git a/Engine/Source/Editor/SceneOutliner/Private/SceneOutlinerModule.cpp b/Engine/Source/Editor/SceneOutliner/Private/SceneOutlinerModule.cpp
--- a/Engine/Source/Editor/SceneOutliner/Private/SceneOutlinerModule.cpp
+++ b/Engine/Source/Editor/SceneOutliner/Private/SceneOutlinerModule.cpp
@@ -245,16 +245,15 @@ void FSceneOutlinerModule::CreateActorInfoColumns(FSceneOutlinerInitializationOp
 		}
 		else if (const FActorDescTreeItem* ActorDescItem = Item.CastTo<FActorDescTreeItem>())
 		{
-			const FWorldPartitionActorDesc* ActorDesc = ActorDescItem->ActorDescHandle.Get();
-
-			if (ActorDesc)
+			if (const FWorldPartitionActorDesc* ActorDesc = ActorDescItem->ActorDescHandle.Get(); ActorDesc && !ActorDesc->GetDataLayers().IsEmpty())
 			{
-				if (UWorld* World = ActorDesc->GetWorld())
+				const UActorDescContainer* ActorDescContainer = ActorDescItem->ActorDescHandle.Container.Get();
+				const UWorld* World = ActorDescContainer ? ActorDescContainer->GetWorld() : nullptr;
+				const AWorldDataLayers* WorldDataLayers = World ? World->GetWorldDataLayers() : nullptr;
+			
+				if (WorldDataLayers)
 				{
-					if (AWorldDataLayers* WorldDataLayers = World->GetWorldDataLayers())
-					{
-						DataLayerObjects = WorldDataLayers->GetDataLayerObjects(ActorDesc->GetDataLayers());
-					}
+					DataLayerObjects = WorldDataLayers->GetDataLayerObjects(ActorDesc->GetDataLayers());
 				}
 			}
 		}
diff --git a/Engine/Source/Runtime/Engine/Private/WorldPartition/ActorDescContainer.cpp b/Engine/Source/Runtime/Engine/Private/WorldPartition/ActorDescContainer.cpp
--- a/Engine/Source/Runtime/Engine/Private/WorldPartition/ActorDescContainer.cpp
+++ b/Engine/Source/Runtime/Engine/Private/WorldPartition/ActorDescContainer.cpp
@@ -20,7 +20,7 @@ UActorDescContainer::UActorDescContainer(const FObjectInitializer& ObjectInitial
 #endif
 {}
 
-void UActorDescContainer::Initialize(UWorld* InWorld, FName InPackageName, TFunctionRef<void(FWorldPartitionActorDesc*)> PreRegister)
+void UActorDescContainer::Initialize(UWorld* InWorld, FName InPackageName)
 {
 	check(!World || World == InWorld);
 	World = InWorld;
@@ -106,17 +106,10 @@ void UActorDescContainer::Initialize(UWorld* InWorld, FName InPackageName, TFunc
 
 		if (ActorDesc.IsValid())
 		{
-			ActorDesc->SetContainer(this);
 			AddActorDescriptor(ActorDesc.Release());
 		}
 	}
 
-	for (FActorDescList::TIterator<> ActorDescIterator(this); ActorDescIterator; ++ActorDescIterator)
-	{
-		PreRegister(*ActorDescIterator);
-		ActorDescIterator->OnRegister(World);
-	}
-
 	RegisterEditorDelegates();
 
 	bContainerInitialized = true;
@@ -126,9 +119,6 @@ void UActorDescContainer::Initialize(UWorld* InWorld, FName InPackageName, TFunc
 void UActorDescContainer::Uninitialize()
 {
 #if WITH_EDITOR
-	PinnedActors.Empty();
-	PinnedActorRefs.Empty();
-
 	if (bContainerInitialized)
 	{
 		UnregisterEditorDelegates();
@@ -139,7 +129,7 @@ void UActorDescContainer::Uninitialize()
 	{
 		if (FWorldPartitionActorDesc* ActorDesc = ActorDescPtr.Get())
 		{
-			ActorDesc->OnUnregister();
+			RemoveActorDescriptor(ActorDesc);
 		}
 		ActorDescPtr.Reset();
 	}
@@ -164,6 +154,18 @@ void UActorDescContainer::BeginDestroy()
 }
 
 #if WITH_EDITOR
+void UActorDescContainer::AddActorDescriptor(FWorldPartitionActorDesc* ActorDesc)
+{
+	FActorDescList::AddActorDescriptor(ActorDesc);
+	ActorDesc->SetContainer(this);
+}
+
+void UActorDescContainer::RemoveActorDescriptor(FWorldPartitionActorDesc* ActorDesc)
+{
+	ActorDesc->SetContainer(nullptr);
+	FActorDescList::RemoveActorDescriptor(ActorDesc);
+}
+
 
 void UActorDescContainer::OnWorldRenamed(UWorld* RenamedWorld)
 {
@@ -214,7 +216,6 @@ void UActorDescContainer::OnObjectPreSave(UObject* Object, FObjectPreSaveContext
 				else
 				{
 					FWorldPartitionActorDesc* const AddedActorDesc = AddActor(Actor);
-					AddedActorDesc->SetContainer(this);
 					OnActorDescAdded(AddedActorDesc);
 				}
 			}
@@ -256,64 +257,10 @@ void UActorDescContainer::RemoveActor(const FGuid& ActorGuid)
 	{
 		OnActorDescRemoved(ExistingActorDesc->Get());
 		RemoveActorDescriptor(ExistingActorDesc->Get());
-		ExistingActorDesc->Get()->SetContainer(nullptr);
 		ExistingActorDesc->Reset();
 	}
 }
 
-AActor* UActorDescContainer::PinActor(const FGuid& ActorGuid)
-{
-	if (FWorldPartitionReference* PinnedActor = PinnedActors.Find(ActorGuid))
-	{
-		return (*PinnedActor).IsValid() ? (*PinnedActor)->GetActor() : nullptr;
-	}
-
-	TFunction<void(const FGuid&, TMap<FGuid, FWorldPartitionReference>&)> AddReferences = [this, &AddReferences](const FGuid& ActorGuid, TMap<FGuid, FWorldPartitionReference>& ReferenceMap)
-	{
-		if (ReferenceMap.Contains(ActorGuid))
-		{
-			return;
-		}
-
-		if (TUniquePtr<FWorldPartitionActorDesc>* const ActorDescPtr = GetActorDescriptor(ActorGuid); ActorDescPtr != nullptr && ActorDescPtr->IsValid())
-		{
-			ReferenceMap.Emplace(ActorGuid, ActorDescPtr);
-			
-			const FWorldPartitionActorDesc* const ActorDesc = ActorDescPtr->Get(); 
-			for (const FGuid& ReferencedActor : ActorDesc->GetReferences())
-			{
-				AddReferences(ReferencedActor, ReferenceMap);
-			}
-		}
-	};
-	
-	if (TUniquePtr<FWorldPartitionActorDesc>* const ActorDescPtr = GetActorDescriptor(ActorGuid))
-	{
-		if (const FWorldPartitionActorDesc* const ActorDesc = ActorDescPtr->Get())
-		{
-			FWorldPartitionReference& PinnedActor = PinnedActors.Emplace(ActorGuid, ActorDescPtr);
-
-			// If the pinned actor has references, we must also create references to those to ensure they are added
-			TMap<FGuid, FWorldPartitionReference>& References = PinnedActorRefs.Emplace(ActorGuid);
-
-			for (const FGuid& ReferencedActor : ActorDesc->GetReferences())
-			{
-				AddReferences(ReferencedActor, References);
-			}
-
-			return PinnedActor.IsValid() ? PinnedActor->GetActor() : nullptr;
-		}
-	}
-
-	return nullptr;
-}
-
-void UActorDescContainer::UnpinActor(const FGuid& ActorGuid)
-{
-	PinnedActors.Remove(ActorGuid);
-	PinnedActorRefs.Remove(ActorGuid);
-}
-
 void UActorDescContainer::LoadAllActors(TArray<FWorldPartitionReference>& OutReferences)
 {
 	OutReferences.Reserve(OutReferences.Num() + GetActorDescCount());
@@ -348,15 +295,11 @@ void UActorDescContainer::UnregisterEditorDelegates()
 
 void UActorDescContainer::OnActorDescAdded(FWorldPartitionActorDesc* NewActorDesc)
 {
-	NewActorDesc->OnRegister(World);
-
 	OnActorDescAddedEvent.Broadcast(NewActorDesc);
 }
 
 void UActorDescContainer::OnActorDescRemoved(FWorldPartitionActorDesc* ActorDesc)
 {
 	OnActorDescRemovedEvent.Broadcast(ActorDesc);
-
-	ActorDesc->OnUnregister();
 }
 #endif
diff --git a/Engine/Source/Runtime/Engine/Private/WorldPartition/LevelInstance/LevelInstanceActorDesc.cpp b/Engine/Source/Runtime/Engine/Private/WorldPartition/LevelInstance/LevelInstanceActorDesc.cpp
--- a/Engine/Source/Runtime/Engine/Private/WorldPartition/LevelInstance/LevelInstanceActorDesc.cpp
+++ b/Engine/Source/Runtime/Engine/Private/WorldPartition/LevelInstance/LevelInstanceActorDesc.cpp
@@ -74,30 +74,30 @@ bool FLevelInstanceActorDesc::Equals(const FWorldPartitionActorDesc* Other) cons
 	return false;
 }
 
-void FLevelInstanceActorDesc::OnRegister(UWorld* InWorld)
+void FLevelInstanceActorDesc::SetContainer(UActorDescContainer* InContainer)
 {
-	FWorldPartitionActorDesc::OnRegister(InWorld);
+	FWorldPartitionActorDesc::SetContainer(InContainer);
 
-	check(!LevelInstanceContainer);
-
-	if (DesiredRuntimeBehavior == ELevelInstanceRuntimeBehavior::Partitioned && !GLevelInstanceDebugForceLevelStreaming)
+	if (Container)
 	{
-		if (!LevelPackage.IsNone() && ULevel::GetIsLevelUsingExternalActorsFromPackage(LevelPackage) && !ULevel::GetIsLevelPartitionedFromPackage(LevelPackage))
+		check(!LevelInstanceContainer);
+
+		if (DesiredRuntimeBehavior == ELevelInstanceRuntimeBehavior::Partitioned && !GLevelInstanceDebugForceLevelStreaming)
 		{
-			LevelInstanceContainer = RegisterActorDescContainer(LevelPackage, InWorld);
-			check(LevelInstanceContainer);
+			if (!LevelPackage.IsNone() && ULevel::GetIsLevelUsingExternalActorsFromPackage(LevelPackage) && !ULevel::GetIsLevelPartitionedFromPackage(LevelPackage))
+			{
+				LevelInstanceContainer = RegisterActorDescContainer(LevelPackage, Container->GetWorld());
+				check(LevelInstanceContainer);
+			}
 		}
 	}
-}
-
-void FLevelInstanceActorDesc::OnUnregister()
-{
-	FWorldPartitionActorDesc::OnUnregister();
-
-	if (LevelInstanceContainer)
+	else
 	{
-		UnregisterActorDescContainer(LevelInstanceContainer);
-		LevelInstanceContainer = nullptr;
+		if (LevelInstanceContainer)
+		{
+			UnregisterActorDescContainer(LevelInstanceContainer);
+			LevelInstanceContainer = nullptr;
+		}
 	}
 }
 
diff --git a/Engine/Source/Runtime/Engine/Private/WorldPartition/WorldPartition.cpp b/Engine/Source/Runtime/Engine/Private/WorldPartition/WorldPartition.cpp
--- a/Engine/Source/Runtime/Engine/Private/WorldPartition/WorldPartition.cpp
+++ b/Engine/Source/Runtime/Engine/Private/WorldPartition/WorldPartition.cpp
@@ -426,26 +426,30 @@ void UWorldPartition::Initialize(UWorld* InWorld, const FTransform& InTransform)
 
 		{
 			TRACE_CPUPROFILER_EVENT_SCOPE(UActorDescContainer::Initialize);
-			UActorDescContainer::Initialize(World, PackageName, [&](FWorldPartitionActorDesc* ActorDesc)
+			UActorDescContainer::Initialize(World, PackageName);
+			check(bContainerInitialized);
+		}
+
+		{
+			TRACE_CPUPROFILER_EVENT_SCOPE(UActorDescContainer::Hash);
+			for (FActorDescList::TIterator<> ActorDescIterator(this); ActorDescIterator; ++ActorDescIterator)
 			{
-				// This will get called by UActorDescContainer after having initialized all ActorDescs and before calling OnRegister on the ActorDesc
 				if (bIsInstanced)
 				{
-					const FString LongActorPackageName = ActorDesc->ActorPackage.ToString();
+					const FString LongActorPackageName = ActorDescIterator->ActorPackage.ToString();
 					const FString ActorPackageName = FPaths::GetBaseFilename(LongActorPackageName);
 					const FString InstancedName = FString::Printf(TEXT("%s_InstanceOf_%s"), *LevelPackage->GetName(), *ActorPackageName);
 
 					InstancingContext.AddMapping(*LongActorPackageName, *InstancedName);
 
-					ActorDesc->TransformInstance(SourceWorldPath, RemappedWorldPath);
+					ActorDescIterator->TransformInstance(SourceWorldPath, RemappedWorldPath);
 				}
 
 				if (bIsEditor && !bIsCooking)
 				{
-					HashActorDesc(ActorDesc);
+					HashActorDesc(*ActorDescIterator);
 				}
-			});
-			check(bContainerInitialized);
+			}
 		}
 	}
 
@@ -558,6 +562,9 @@ void UWorldPartition::Uninitialize()
 
 		WorldDataLayersActor = FWorldPartitionReference();
 
+		PinnedActors.Empty();
+		PinnedActorRefs.Empty();
+
 		EditorHash = nullptr;
 #endif		
 
@@ -1336,6 +1343,64 @@ void UWorldPartition::SetEditorWantedCellSize(uint32 InCellSize)
 	EditorHash->SetEditorWantedCellSize(InCellSize);
 }
 
+AActor* UWorldPartition::PinActor(const FGuid& ActorGuid)
+{
+	if (FWorldPartitionReference* PinnedActor = PinnedActors.Find(ActorGuid))
+	{
+		return (*PinnedActor).IsValid() ? (*PinnedActor)->GetActor() : nullptr;
+	}
+
+	TFunction<void(const FGuid&, TMap<FGuid, FWorldPartitionReference>&)> AddReferences = [this, &AddReferences](const FGuid& ActorGuid, TMap<FGuid, FWorldPartitionReference>& ReferenceMap)
+	{
+		if (ReferenceMap.Contains(ActorGuid))
+		{
+			return;
+		}
+
+		if (TUniquePtr<FWorldPartitionActorDesc>* const ActorDescPtr = GetActorDescriptor(ActorGuid); ActorDescPtr != nullptr && ActorDescPtr->IsValid())
+		{
+			ReferenceMap.Emplace(ActorGuid, ActorDescPtr);
+			
+			const FWorldPartitionActorDesc* const ActorDesc = ActorDescPtr->Get(); 
+			for (const FGuid& ReferencedActor : ActorDesc->GetReferences())
+			{
+				AddReferences(ReferencedActor, ReferenceMap);
+			}
+		}
+	};
+	
+	if (TUniquePtr<FWorldPartitionActorDesc>* const ActorDescPtr = GetActorDescriptor(ActorGuid))
+	{
+		if (const FWorldPartitionActorDesc* const ActorDesc = ActorDescPtr->Get())
+		{
+			FWorldPartitionReference& PinnedActor = PinnedActors.Emplace(ActorGuid, ActorDescPtr);
+
+			// If the pinned actor has references, we must also create references to those to ensure they are added
+			TMap<FGuid, FWorldPartitionReference>& References = PinnedActorRefs.Emplace(ActorGuid);
+
+			for (const FGuid& ReferencedActor : ActorDesc->GetReferences())
+			{
+				AddReferences(ReferencedActor, References);
+			}
+
+			return PinnedActor.IsValid() ? PinnedActor->GetActor() : nullptr;
+		}
+	}
+
+	return nullptr;
+}
+
+void UWorldPartition::UnpinActor(const FGuid& ActorGuid)
+{
+	PinnedActors.Remove(ActorGuid);
+	PinnedActorRefs.Remove(ActorGuid);
+}
+
+bool UWorldPartition::IsActorPinned(const FGuid& ActorGuid) const
+{
+	return PinnedActors.Contains(ActorGuid);
+}
+
 void UWorldPartition::OnWorldRenamed()
 {
 	UActorDescContainer::OnWorldRenamed();
diff --git a/Engine/Source/Runtime/Engine/Private/WorldPartition/WorldPartitionActorDesc.cpp b/Engine/Source/Runtime/Engine/Private/WorldPartition/WorldPartitionActorDesc.cpp
--- a/Engine/Source/Runtime/Engine/Private/WorldPartition/WorldPartitionActorDesc.cpp
+++ b/Engine/Source/Runtime/Engine/Private/WorldPartition/WorldPartitionActorDesc.cpp
@@ -153,11 +153,6 @@ void FWorldPartitionActorDesc::SerializeTo(TArray<uint8>& OutData)
 	OutData.Append(PayloadData);
 }
 
-UWorld* FWorldPartitionActorDesc::GetWorld() const
-{
-	return Container ? Container->GetWorld() : nullptr;
-}
-
 void FWorldPartitionActorDesc::TransformInstance(const FString& From, const FString& To)
 {
 	check(!HardRefCount);
diff --git a/Engine/Source/Runtime/Engine/Private/WorldPartition/WorldPartitionRuntimeHash.cpp b/Engine/Source/Runtime/Engine/Private/WorldPartition/WorldPartitionRuntimeHash.cpp
--- a/Engine/Source/Runtime/Engine/Private/WorldPartition/WorldPartitionRuntimeHash.cpp
+++ b/Engine/Source/Runtime/Engine/Private/WorldPartition/WorldPartitionRuntimeHash.cpp
@@ -34,10 +34,6 @@ void UWorldPartitionRuntimeHash::OnEndPlay()
 	// Release references (will unload actors that were not already loaded in the Editor)
 	AlwaysLoadedActorsForPIE.Empty();
 
-	for (FActorDescList::TIterator<> ActorDescIterator(&ModifiedActorDescListForPIE); ActorDescIterator; ++ActorDescIterator)
-	{
-		ActorDescIterator->OnUnregister();
-	}
 	ModifiedActorDescListForPIE.Empty();
 }
 
diff --git a/Engine/Source/Runtime/Engine/Private/WorldPartition/WorldPartitionStreamingGeneration.cpp b/Engine/Source/Runtime/Engine/Private/WorldPartition/WorldPartitionStreamingGeneration.cpp
--- a/Engine/Source/Runtime/Engine/Private/WorldPartition/WorldPartitionStreamingGeneration.cpp
+++ b/Engine/Source/Runtime/Engine/Private/WorldPartition/WorldPartitionStreamingGeneration.cpp
@@ -61,8 +61,6 @@ class FWorldPartitionStreamingGenerator
 					{
 						// Dirty, unsaved actor for PIE
 						FWorldPartitionActorDesc* ActorDesc = ModifiedActorsDescList->AddActor(Actor);
-						ActorDesc->SetContainer(const_cast<UActorDescContainer*>(InContainer));
-						ActorDesc->OnRegister(Actor->GetWorld());
 						OutActorDescViewMap.Emplace(ActorDescIt->GetGuid(), ActorDesc);
 						continue;
 					}
@@ -81,8 +79,6 @@ class FWorldPartitionStreamingGenerator
 				if (IsValid(Actor) && Actor->IsPackageExternal() && Actor->IsMainPackageActor() && !Actor->IsEditorOnly() && !InContainer->GetActorDesc(Actor->GetActorGuid()))
 				{
 					FWorldPartitionActorDesc* ActorDesc = ModifiedActorsDescList->AddActor(Actor);
-					ActorDesc->SetContainer(const_cast<UActorDescContainer*>(InContainer));
-					ActorDesc->OnRegister(Actor->GetWorld());
 					OutActorDescViewMap.Emplace(ActorDesc->GetGuid(), ActorDesc);
 				}
 			}
diff --git a/Engine/Source/Runtime/Engine/Public/WorldPartition/ActorDescContainer.h b/Engine/Source/Runtime/Engine/Public/WorldPartition/ActorDescContainer.h
--- a/Engine/Source/Runtime/Engine/Public/WorldPartition/ActorDescContainer.h
+++ b/Engine/Source/Runtime/Engine/Public/WorldPartition/ActorDescContainer.h
@@ -18,7 +18,7 @@ class ENGINE_API UActorDescContainer : public UObject, public FActorDescList
 	friend class FWorldPartitionActorDesc;
 
 public:
-	void Initialize(UWorld* World, FName InPackageName, TFunctionRef<void(FWorldPartitionActorDesc*)> PreRegister = [](FWorldPartitionActorDesc*) {});
+	void Initialize(UWorld* World, FName InPackageName);
 	virtual void Uninitialize();
 	
 	virtual UWorld* GetWorld() const override;
@@ -33,10 +33,6 @@ public:
 	/** Removes an actor desc without the need to load a package */
 	virtual void RemoveActor(const FGuid& ActorGuid);
 
-	AActor* PinActor(const FGuid& ActorGuid);
-	void UnpinActor(const FGuid& ActorGuid);
-	bool IsActorPinned(const FGuid& ActorGuid) const { return PinnedActors.Contains(ActorGuid); }
-
 	void LoadAllActors(TArray<FWorldPartitionReference>& OutReferences);
 
 	/** Instancing support */
@@ -59,6 +55,11 @@ protected:
 	//~ End UObject Interface
 
 #if WITH_EDITOR
+	//~ Begin FActorDescList Interface
+	virtual void AddActorDescriptor(FWorldPartitionActorDesc* ActorDesc) override;
+	virtual void RemoveActorDescriptor(FWorldPartitionActorDesc* ActorDesc) override;
+	//~ End FActorDescList Interface
+
 	// Events
 	void OnWorldRenamed(UWorld* RenamedWorld);
 	virtual void OnWorldRenamed();
@@ -73,9 +74,6 @@ protected:
 
 	bool ShouldHandleActorEvent(const AActor* Actor);
 
-	TMap<FGuid, FWorldPartitionReference> PinnedActors;
-	TMap<FGuid, TMap<FGuid, FWorldPartitionReference>> PinnedActorRefs;
-	
 	bool bContainerInitialized;
 
 	FName ContainerPackageName;
diff --git a/Engine/Source/Runtime/Engine/Public/WorldPartition/ActorDescList.h b/Engine/Source/Runtime/Engine/Public/WorldPartition/ActorDescList.h
--- a/Engine/Source/Runtime/Engine/Public/WorldPartition/ActorDescList.h
+++ b/Engine/Source/Runtime/Engine/Public/WorldPartition/ActorDescList.h
@@ -13,6 +13,7 @@ class ENGINE_API FActorDescList
 
 public:
 	FActorDescList() {}
+	virtual ~FActorDescList() {}
 
 	// Non-copyable
 	FActorDescList(const FActorDescList&) = delete;
@@ -141,8 +142,8 @@ public:
 	};
 
 protected:
-	void AddActorDescriptor(FWorldPartitionActorDesc* ActorDesc);
-	void RemoveActorDescriptor(FWorldPartitionActorDesc* ActorDesc);
+	virtual void AddActorDescriptor(FWorldPartitionActorDesc* ActorDesc);
+	virtual void RemoveActorDescriptor(FWorldPartitionActorDesc* ActorDesc);
 	TUniquePtr<FWorldPartitionActorDesc>* GetActorDescriptor(const FGuid& ActorGuid);
 
 	TChunkedArray<TUniquePtr<FWorldPartitionActorDesc>> ActorDescList;
diff --git a/Engine/Source/Runtime/Engine/Public/WorldPartition/LevelInstance/LevelInstanceActorDesc.h b/Engine/Source/Runtime/Engine/Public/WorldPartition/LevelInstance/LevelInstanceActorDesc.h
--- a/Engine/Source/Runtime/Engine/Public/WorldPartition/LevelInstance/LevelInstanceActorDesc.h
+++ b/Engine/Source/Runtime/Engine/Public/WorldPartition/LevelInstance/LevelInstanceActorDesc.h
@@ -36,12 +36,8 @@ protected:
 	virtual void TransferFrom(const FWorldPartitionActorDesc* From) override;
 	virtual void Serialize(FArchive& Ar) override;
 	virtual void AddReferencedObjects(FReferenceCollector& Collector);
-	virtual FString GetReferencerName() const override
-	{
-		return TEXT("FLevelInstanceActorDesc");
-	}
-	virtual void OnRegister(UWorld* InWorld) override;
-	virtual void OnUnregister() override;
+	virtual FString GetReferencerName() const override { return TEXT("FLevelInstanceActorDesc"); }
+	virtual void SetContainer(UActorDescContainer* InContainer) override;
 
 	FName LevelPackage;
 	FTransform LevelInstanceTransform;
diff --git a/Engine/Source/Runtime/Engine/Public/WorldPartition/WorldPartition.h b/Engine/Source/Runtime/Engine/Public/WorldPartition/WorldPartition.h
--- a/Engine/Source/Runtime/Engine/Public/WorldPartition/WorldPartition.h
+++ b/Engine/Source/Runtime/Engine/Public/WorldPartition/WorldPartition.h
@@ -154,6 +154,11 @@ public:
 
 	uint32 GetWantedEditorCellSize() const;
 	void SetEditorWantedCellSize(uint32 InCellSize);
+
+	// Actors pinning
+	AActor* PinActor(const FGuid& ActorGuid);
+	void UnpinActor(const FGuid& ActorGuid);
+	bool IsActorPinned(const FGuid& ActorGuid) const;
 #endif
 
 public:
@@ -241,6 +246,9 @@ private:
 	bool UpdateEditorCells(TFunctionRef<bool(TArray<UWorldPartitionEditorCell*>&)> GetCellsToProcess, bool bIsCellShouldBeLoaded, bool bIsFromUserChange);
 
 	void ApplyActorTransform(AActor* Actor, const FTransform& InTransform);
+
+	TMap<FGuid, FWorldPartitionReference> PinnedActors;
+	TMap<FGuid, TMap<FGuid, FWorldPartitionReference>> PinnedActorRefs;
 #endif
 
 #if !UE_BUILD_SHIPPING
diff --git a/Engine/Source/Runtime/Engine/Public/WorldPartition/WorldPartitionActorDesc.h b/Engine/Source/Runtime/Engine/Public/WorldPartition/WorldPartitionActorDesc.h
--- a/Engine/Source/Runtime/Engine/Public/WorldPartition/WorldPartitionActorDesc.h
+++ b/Engine/Source/Runtime/Engine/Public/WorldPartition/WorldPartitionActorDesc.h
@@ -125,7 +125,7 @@ protected:
 		return HardRefCount;
 	}
 
-	void SetContainer(UActorDescContainer* InContainer)
+	virtual void SetContainer(UActorDescContainer* InContainer)
 	{
 		check(!Container || !InContainer);
 		Container = InContainer;
@@ -160,8 +160,6 @@ public:
 
 	void SerializeTo(TArray<uint8>& OutData);
 
-	UWorld* GetWorld() const;
-
 protected:
 	FWorldPartitionActorDesc();
 
@@ -176,9 +174,6 @@ protected:
 
 	virtual void Serialize(FArchive& Ar);
 
-	virtual void OnRegister(UWorld* InWorld) {}
-	virtual void OnUnregister() {}
-
 	// Persistent
 	FGuid							Guid;
 	FName							Class;
