FEM - Fix LevelTick hang

diff --git a/Engine/Source/Runtime/Engine/Private/LevelTick.cpp b/Engine/Source/Runtime/Engine/Private/LevelTick.cpp
--- a/Engine/Source/Runtime/Engine/Private/LevelTick.cpp
+++ b/Engine/Source/Runtime/Engine/Private/LevelTick.cpp
@@ -1110,19 +1110,16 @@ void UWorld::SendAllEndOfFrameUpdates()
 		LocalComponentsThatNeedEndOfFrameUpdate.Append(ComponentsThatNeedEndOfFrameUpdate);
 	}
 
-#if WITH_EDITOR
-	std::atomic<uint32> ParallelTasksRemaining(0);
-#endif
-
-	auto ParallelWork = 
-		[
-#if WITH_EDITOR
-		&ParallelTasksRemaining
-#endif
-		](int32 Index)
+	auto ParallelWork = [](int32 Index)
 		{
 			TRACE_CPUPROFILER_EVENT_SCOPE(DeferredRenderUpdates);
 			FOptionalTaskTagScope Scope(ETaskTag::EParallelGameThread);
+#if WITH_EDITOR
+			if (!IsInParallelGameThread() && IsInGameThread())
+			{
+				FObjectCacheEventSink::ProcessQueuedNotifyEvents();
+			}
+#endif
 			UActorComponent* NextComponent = LocalComponentsThatNeedEndOfFrameUpdate[Index];
 			if (NextComponent)
 			{
@@ -1133,17 +1130,9 @@ void UWorld::SendAllEndOfFrameUpdates()
 				check(!IsValid(NextComponent) || NextComponent->GetMarkedForEndOfFrameUpdateState() == EComponentMarkedForEndOfFrameUpdateState::Marked);
 				FMarkComponentEndOfFrameUpdateState::Set(NextComponent, INDEX_NONE, EComponentMarkedForEndOfFrameUpdateState::Unmarked);
 			}
-#if WITH_EDITOR
-			ParallelTasksRemaining--;
-#endif
 		};
 
-	auto GTWork = 
-		[this
-#if WITH_EDITOR
-		, &ParallelTasksRemaining
-#endif
-		]()
+	auto GTWork = [this]()
 		{
 			QUICK_SCOPE_CYCLE_COUNTER(STAT_PostTickComponentUpdate_ForcedGameThread);
 
@@ -1173,13 +1162,6 @@ void UWorld::SendAllEndOfFrameUpdates()
 			{
 				Component->DoDeferredRenderUpdates_Concurrent();
 			}
-
-#if WITH_EDITOR
-			while (ParallelTasksRemaining > 0)
-			{
-				FObjectCacheEventSink::ProcessQueuedNotifyEvents();
-			}
-#endif
 	};
 
 	if (CVarAllowAsyncRenderThreadUpdatesDuringGamethreadUpdates.GetValueOnGameThread() > 0 && 
@@ -1187,11 +1169,9 @@ void UWorld::SendAllEndOfFrameUpdates()
 		FTaskGraphInterface::Get().GetNumWorkerThreads() > 2)
 	{
 #if WITH_EDITOR
-		ParallelTasksRemaining = LocalComponentsThatNeedEndOfFrameUpdate.Num();
 		FObjectCacheEventSink::BeginQueueNotifyEvents();
 #endif
 		ParallelForWithPreWork(LocalComponentsThatNeedEndOfFrameUpdate.Num(), ParallelWork, GTWork);
-
 #if WITH_EDITOR
 		// Any remaining events will be flushed with this call
 		FObjectCacheEventSink::EndQueueNotifyEvents();
