Fix for Noise function driven from LWC positions

diff --git a/Engine/Shaders/Private/MaterialTemplate.ush b/Engine/Shaders/Private/MaterialTemplate.ush
--- a/Engine/Shaders/Private/MaterialTemplate.ush
+++ b/Engine/Shaders/Private/MaterialTemplate.ush
@@ -23,7 +23,7 @@
 
 // Update this GUID to force all material shaders to recompile
 // Merge conflicts on this line should be resolved by generating a new GUID
-// GUID = 4E3A3611-FB00-4498-AF79-80F57CBC06A8
+// GUID = 5F290441-5611-4232-B26D-613F5E421262
 
 #if USES_SPEEDTREE
 	#include "/Engine/Private/SpeedTreeCommon.ush"
diff --git a/Engine/Source/Runtime/Engine/Private/Materials/HLSLMaterialTranslator.cpp b/Engine/Source/Runtime/Engine/Private/Materials/HLSLMaterialTranslator.cpp
--- a/Engine/Source/Runtime/Engine/Private/Materials/HLSLMaterialTranslator.cpp
+++ b/Engine/Source/Runtime/Engine/Private/Materials/HLSLMaterialTranslator.cpp
@@ -8818,7 +8818,15 @@ int32 FHLSLMaterialTranslator::Noise(int32 Position, float Scale, int32 Quality,
 	int32 TilingConst = Constant(bTiling);
 	int32 RepeatSizeConst = Constant(RepeatSize);
 
-	// LWC_TODO: LWC-enabled noise?
+	const EMaterialValueType PositionType = GetParameterType(Position);
+	if (IsLWCType(PositionType))
+	{
+		// If Noise is driven by a LWC position, just take the offset within the current tile
+		// Will generate discontinuity in noise at tile boudaries
+		// Could potentially add noise functions that operate directly on LWC values, but that would be very expensive
+		Position = AddCodeChunk(MCT_Float3, TEXT("LWCNormalizeTile(%s).Offset"), *CoerceParameter(Position, MCT_LWCVector3));
+	}
+
 	return AddCodeChunk(MCT_Float, 
 		TEXT("MaterialExpressionNoise(%s,%s,%s,%s,%s,%s,%s,%s,%s,%s,%s,%s)"), 
 		*CoerceParameter(Position, MCT_Float3),
