FEM - More LevelTick fixes

diff --git a/Engine/Source/Runtime/Engine/Private/LevelTick.cpp b/Engine/Source/Runtime/Engine/Private/LevelTick.cpp
--- a/Engine/Source/Runtime/Engine/Private/LevelTick.cpp
+++ b/Engine/Source/Runtime/Engine/Private/LevelTick.cpp
@@ -1110,63 +1110,65 @@ void UWorld::SendAllEndOfFrameUpdates()
 		LocalComponentsThatNeedEndOfFrameUpdate.Append(ComponentsThatNeedEndOfFrameUpdate);
 	}
 
-	auto ParallelWork = [](int32 Index)
-		{
-			TRACE_CPUPROFILER_EVENT_SCOPE(DeferredRenderUpdates);
-			FOptionalTaskTagScope Scope(ETaskTag::EParallelGameThread);
+	const bool IsUsingParallelNotifyEvents = CVarAllowAsyncRenderThreadUpdatesDuringGamethreadUpdates.GetValueOnGameThread() > 0 && 
+		LocalComponentsThatNeedEndOfFrameUpdate.Num() > FTaskGraphInterface::Get().GetNumWorkerThreads() &&
+		FTaskGraphInterface::Get().GetNumWorkerThreads() > 2;
+
+	auto ParallelWork = [IsUsingParallelNotifyEvents](int32 Index)
+	{
+		TRACE_CPUPROFILER_EVENT_SCOPE(DeferredRenderUpdates);
+		FOptionalTaskTagScope Scope(ETaskTag::EParallelGameThread);
 #if WITH_EDITOR
-			if (!IsInParallelGameThread() && IsInGameThread())
-			{
-				FObjectCacheEventSink::ProcessQueuedNotifyEvents();
-			}
+		if (!IsInParallelGameThread() && IsInGameThread() && IsUsingParallelNotifyEvents)
+		{
+			FObjectCacheEventSink::ProcessQueuedNotifyEvents();
+		}
 #endif
-			UActorComponent* NextComponent = LocalComponentsThatNeedEndOfFrameUpdate[Index];
-			if (NextComponent)
+		UActorComponent* NextComponent = LocalComponentsThatNeedEndOfFrameUpdate[Index];
+		if (NextComponent)
+		{
+			if (NextComponent->IsRegistered() && !NextComponent->IsTemplate() && IsValid(NextComponent))
 			{
-				if (NextComponent->IsRegistered() && !NextComponent->IsTemplate() && IsValid(NextComponent))
-				{
-					NextComponent->DoDeferredRenderUpdates_Concurrent();
-				}
-				check(!IsValid(NextComponent) || NextComponent->GetMarkedForEndOfFrameUpdateState() == EComponentMarkedForEndOfFrameUpdateState::Marked);
-				FMarkComponentEndOfFrameUpdateState::Set(NextComponent, INDEX_NONE, EComponentMarkedForEndOfFrameUpdateState::Unmarked);
+				NextComponent->DoDeferredRenderUpdates_Concurrent();
 			}
-		};
+			check(!IsValid(NextComponent) || NextComponent->GetMarkedForEndOfFrameUpdateState() == EComponentMarkedForEndOfFrameUpdateState::Marked);
+			FMarkComponentEndOfFrameUpdateState::Set(NextComponent, INDEX_NONE, EComponentMarkedForEndOfFrameUpdateState::Unmarked);
+		}
+	};
 
 	auto GTWork = [this]()
-		{
-			QUICK_SCOPE_CYCLE_COUNTER(STAT_PostTickComponentUpdate_ForcedGameThread);
+	{
+		QUICK_SCOPE_CYCLE_COUNTER(STAT_PostTickComponentUpdate_ForcedGameThread);
 
-			// To avoid any problems in case of reentrancy during the deferred update pass, we gather everything and clears the buffers first
-			// Reentrancy can occur if a render update need to force wait on an async resource and a progress bar ticks the game-thread during that time.
-			TArray< UActorComponent*> DeferredUpdates;
-			DeferredUpdates.Reserve(ComponentsThatNeedEndOfFrameUpdate_OnGameThread.Num());
+		// To avoid any problems in case of reentrancy during the deferred update pass, we gather everything and clears the buffers first
+		// Reentrancy can occur if a render update need to force wait on an async resource and a progress bar ticks the game-thread during that time.
+		TArray< UActorComponent*> DeferredUpdates;
+		DeferredUpdates.Reserve(ComponentsThatNeedEndOfFrameUpdate_OnGameThread.Num());
 
-			for (UActorComponent* Component : ComponentsThatNeedEndOfFrameUpdate_OnGameThread)
+		for (UActorComponent* Component : ComponentsThatNeedEndOfFrameUpdate_OnGameThread)
+		{
+			if (Component)
 			{
-				if (Component)
+				if (Component->IsRegistered() && !Component->IsTemplate() && IsValid(Component))
 				{
-					if (Component->IsRegistered() && !Component->IsTemplate() && IsValid(Component))
-					{
-						DeferredUpdates.Add(Component);
-					}
-
-					check(!IsValid(Component) || Component->GetMarkedForEndOfFrameUpdateState() == EComponentMarkedForEndOfFrameUpdateState::MarkedForGameThread);
-					FMarkComponentEndOfFrameUpdateState::Set(Component, INDEX_NONE, EComponentMarkedForEndOfFrameUpdateState::Unmarked);
+					DeferredUpdates.Add(Component);
 				}
+
+				check(!IsValid(Component) || Component->GetMarkedForEndOfFrameUpdateState() == EComponentMarkedForEndOfFrameUpdateState::MarkedForGameThread);
+				FMarkComponentEndOfFrameUpdateState::Set(Component, INDEX_NONE, EComponentMarkedForEndOfFrameUpdateState::Unmarked);
 			}
+		}
 
-			ComponentsThatNeedEndOfFrameUpdate_OnGameThread.Reset();
-			ComponentsThatNeedEndOfFrameUpdate.Reset();
+		ComponentsThatNeedEndOfFrameUpdate_OnGameThread.Reset();
+		ComponentsThatNeedEndOfFrameUpdate.Reset();
 
-			for (UActorComponent* Component : DeferredUpdates)
-			{
-				Component->DoDeferredRenderUpdates_Concurrent();
-			}
+		for (UActorComponent* Component : DeferredUpdates)
+		{
+			Component->DoDeferredRenderUpdates_Concurrent();
+		}
 	};
 
-	if (CVarAllowAsyncRenderThreadUpdatesDuringGamethreadUpdates.GetValueOnGameThread() > 0 && 
-		LocalComponentsThatNeedEndOfFrameUpdate.Num() > FTaskGraphInterface::Get().GetNumWorkerThreads() &&
-		FTaskGraphInterface::Get().GetNumWorkerThreads() > 2)
+	if (IsUsingParallelNotifyEvents)
 	{
 #if WITH_EDITOR
 		FObjectCacheEventSink::BeginQueueNotifyEvents();
