Fix non-editor

diff --git a/Engine/Source/Runtime/Engine/Private/Components/RuntimeVirtualTextureComponent.cpp b/Engine/Source/Runtime/Engine/Private/Components/RuntimeVirtualTextureComponent.cpp
--- a/Engine/Source/Runtime/Engine/Private/Components/RuntimeVirtualTextureComponent.cpp
+++ b/Engine/Source/Runtime/Engine/Private/Components/RuntimeVirtualTextureComponent.cpp
@@ -111,10 +111,12 @@ FBoxSphereBounds URuntimeVirtualTextureComponent::CalcBounds(const FTransform& L
 	return FBoxSphereBounds(FBox(FVector(0.f, 0.f, 0.f), FVector(1.f, 1.f, 1.f))).TransformBy(LocalToWorld);
 }
 
+#if WITH_EDITOR
 void URuntimeVirtualTextureComponent::SetBoundsAlignActor(AActor* InActor)
 {
 	BoundsAlignActor = InActor;
 }
+#endif
 
 FTransform URuntimeVirtualTextureComponent::GetTexelSnapTransform() const
 {
