World Partition: report actor descriptors with duplicated GUIDs and skip their registration + better invalid actor files logging.

diff --git a/Engine/Source/Runtime/Engine/Private/WorldPartition/ActorDescContainer.cpp b/Engine/Source/Runtime/Engine/Private/WorldPartition/ActorDescContainer.cpp
--- a/Engine/Source/Runtime/Engine/Private/WorldPartition/ActorDescContainer.cpp
+++ b/Engine/Source/Runtime/Engine/Private/WorldPartition/ActorDescContainer.cpp
@@ -5,6 +5,7 @@
 #if WITH_EDITOR
 #include "Editor.h"
 #include "AssetRegistryModule.h"
+#include "WorldPartition/WorldPartitionLog.h"
 #include "WorldPartition/WorldPartitionHandle.h"
 #include "Misc/Base64.h"
 #include "UObject/ObjectSaveContext.h"
@@ -86,10 +87,9 @@ void UActorDescContainer::Initialize(UWorld* InWorld, FName InPackageName)
 
 				NewActorDesc->Init(ActorDescInitData);
 			
-				//check(bIsValidClass || (NewActorDesc->GetActorIsEditorOnly() && IsRunningGame()));
-
 				if (!bIsValidClass)
 				{
+					UE_LOG(LogWorldPartition, Warning, TEXT("Invalid class `%s` for actor guid `%s` ('%s') from package '%s'"), *NewClassName.ToString(), *NewActorDesc->GetGuid().ToString(), *NewActorDesc->GetActorName().ToString(), *NewActorDesc->GetActorPackage().ToString());
 					return nullptr;
 				}
 
diff --git a/Engine/Source/Runtime/Engine/Private/WorldPartition/ActorDescList.cpp b/Engine/Source/Runtime/Engine/Private/WorldPartition/ActorDescList.cpp
--- a/Engine/Source/Runtime/Engine/Private/WorldPartition/ActorDescList.cpp
+++ b/Engine/Source/Runtime/Engine/Private/WorldPartition/ActorDescList.cpp
@@ -1,6 +1,7 @@
 // Copyright Epic Games, Inc. All Rights Reserved.
 
 #include "WorldPartition/ActorDescList.h"
+#include "WorldPartition/WorldPartitionLog.h"
 
 #if WITH_EDITOR
 FWorldPartitionActorDesc* FActorDescList::AddActor(const AActor* InActor)
@@ -69,9 +70,17 @@ void FActorDescList::Empty()
 void FActorDescList::AddActorDescriptor(FWorldPartitionActorDesc* ActorDesc)
 {
 	check(ActorDesc);
-	TUniquePtr<FWorldPartitionActorDesc>* NewActorDesc = new(ActorDescList) TUniquePtr<FWorldPartitionActorDesc>(ActorDesc);
-	ActorsByGuid.Add(ActorDesc->GetGuid(), NewActorDesc);
-	ActorsByName.Add(ActorDesc->GetActorName(), NewActorDesc);
+
+	if (const TUniquePtr<FWorldPartitionActorDesc>* const* ExistingActorDesc = ActorsByGuid.Find(ActorDesc->GetGuid()))
+	{
+		UE_LOG(LogWorldPartition, Warning, TEXT("Duplicated actor descriptor guid '%s' detected: `%s` (existing: `%s`)"), *ActorDesc->GetGuid().ToString(), *ActorDesc->GetActorName().ToString(), *(*ExistingActorDesc)->Get()->GetActorName().ToString());
+	}
+	else
+	{
+		TUniquePtr<FWorldPartitionActorDesc>* NewActorDesc = new(ActorDescList) TUniquePtr<FWorldPartitionActorDesc>(ActorDesc);
+		ActorsByGuid.Add(ActorDesc->GetGuid(), NewActorDesc);
+		ActorsByName.Add(ActorDesc->GetActorName(), NewActorDesc);
+	}
 }
 
 void FActorDescList::RemoveActorDescriptor(FWorldPartitionActorDesc* ActorDesc)
diff --git a/Engine/Source/Runtime/Engine/Private/WorldPartition/WorldPartitionActorDesc.cpp b/Engine/Source/Runtime/Engine/Private/WorldPartition/WorldPartitionActorDesc.cpp
--- a/Engine/Source/Runtime/Engine/Private/WorldPartition/WorldPartitionActorDesc.cpp
+++ b/Engine/Source/Runtime/Engine/Private/WorldPartition/WorldPartitionActorDesc.cpp
@@ -322,7 +322,7 @@ AActor* FWorldPartitionActorDesc::Load() const
 			}
 			else
 			{
-				UE_LOG(LogWorldPartition, Warning, TEXT("Can't load actor '%s' from package '%s'"), *GetActorName().ToString(), *ActorPackage.ToString());
+				UE_LOG(LogWorldPartition, Warning, TEXT("Can't load actor guid `%s` ('%s') from package '%s'"), *Guid.ToString(), *GetActorName().ToString(), *ActorPackage.ToString());
 			}
 		}
 	}
