RichCurve evaluation is bad when we have weighted tangents
#fix ensure that compressed weighted-tangent key data is indexed correctly

diff --git a/Engine/Source/Runtime/Engine/Private/Curves/RichCurve.cpp b/Engine/Source/Runtime/Engine/Private/Curves/RichCurve.cpp
--- a/Engine/Source/Runtime/Engine/Private/Curves/RichCurve.cpp
+++ b/Engine/Source/Runtime/Engine/Private/Curves/RichCurve.cpp
@@ -2000,7 +2000,7 @@ struct WeightedKeyDataAdapter
 
 	KeyDataHandle GetKeyDataHandle(int32 KeyIndexToQuery) const
 	{
-		return KeyIndexToQuery * 3;
+		return KeyIndexToQuery * 5;
 	};
 
 	constexpr float GetKeyValue(KeyDataHandle Handle) const
