Fix apparent non-determinism by bumping DDC key for MESHDATAKEY_STATICMESH. A previous format change was resulting in different data being generated.

diff --git a/Engine/Source/Runtime/Engine/Private/StaticMeshSourceData.cpp b/Engine/Source/Runtime/Engine/Private/StaticMeshSourceData.cpp
--- a/Engine/Source/Runtime/Engine/Private/StaticMeshSourceData.cpp
+++ b/Engine/Source/Runtime/Engine/Private/StaticMeshSourceData.cpp
@@ -360,7 +360,7 @@ void FStaticMeshSourceModel::FillMaterialName(TMap<int32, FName>& OutMaterialMap
 // This is the key which is used for DDC data when legacy RawMesh is converted to MeshDescription.
 // If static mesh derived data needs to be rebuilt (new format, serialization differences, etc.) replace the version GUID below with a new one.
 // In case of merge conflicts with DDC versions, you MUST generate a new GUID and set this new GUID as the version.
-#define MESHDATAKEY_STATICMESH_DERIVEDDATA_VER TEXT("2BF8C1219912409497797709F70B5105")
+#define MESHDATAKEY_STATICMESH_DERIVEDDATA_VER TEXT("EACA97B8341141EA92DE0231AC534A71")
 
 static bool GetMeshDataKey(FString& OutKey, const FRawMeshBulkData* RawMeshBulkData, int LodIndex)
 {
