Fix InitializeSubsystems not calling PostInitializeSubsystems because of out of data SubsystemArrayMap

This fixes path where world was loaded as Inactive and then loaded in Editor. The result was it's Subsystems not being post-initialized (source operation duplicate of world)

diff --git a/Engine/Source/Runtime/Engine/Private/Subsystems/SubsystemCollection.cpp b/Engine/Source/Runtime/Engine/Private/Subsystems/SubsystemCollection.cpp
--- a/Engine/Source/Runtime/Engine/Private/Subsystems/SubsystemCollection.cpp
+++ b/Engine/Source/Runtime/Engine/Private/Subsystems/SubsystemCollection.cpp
@@ -72,14 +72,7 @@ const TArray<USubsystem*>& FSubsystemCollectionBase::GetSubsystemArrayInternal(U
 	{
 		TArray<USubsystem*>& NewList = SubsystemArrayMap.Add(SubsystemClass);
 
-		for (auto Iter = SubsystemMap.CreateConstIterator(); Iter; ++Iter)
-		{
-			UClass* KeyClass = Iter.Key();
-			if (KeyClass->IsChildOf(SubsystemClass))
-			{
-				NewList.Add(Iter.Value());
-			}
-		}
+		UpdateSubsystemArrayInternal(SubsystemClass, NewList);
 
 		return NewList;
 	}
@@ -88,6 +81,18 @@ const TArray<USubsystem*>& FSubsystemCollectionBase::GetSubsystemArrayInternal(U
 	return List;
 }
 
+void FSubsystemCollectionBase::UpdateSubsystemArrayInternal(UClass* SubsystemClass, TArray<USubsystem*>& SubsystemArray) const
+{
+	for (auto Iter = SubsystemMap.CreateConstIterator(); Iter; ++Iter)
+	{
+		UClass* KeyClass = Iter.Key();
+		if (KeyClass->IsChildOf(SubsystemClass))
+		{
+			SubsystemArray.Add(Iter.Value());
+		}
+	}
+}
+
 void FSubsystemCollectionBase::Initialize(UObject* NewOuter)
 {
 	if (Outer != nullptr)
@@ -136,6 +141,13 @@ void FSubsystemCollectionBase::Initialize(UObject* NewOuter)
 			}
 		}
 
+		// Update Internal Arrays without emptying it so that existing refs remain valid
+		for (auto& Pair : SubsystemArrayMap)
+		{
+			Pair.Value.Empty();
+			UpdateSubsystemArrayInternal(Pair.Key, Pair.Value);
+		}
+
 		// Statically track collections
 		GlobalSubsystemCollections.Add(this);
 	}
diff --git a/Engine/Source/Runtime/Engine/Public/Subsystems/SubsystemCollection.h b/Engine/Source/Runtime/Engine/Public/Subsystems/SubsystemCollection.h
--- a/Engine/Source/Runtime/Engine/Public/Subsystems/SubsystemCollection.h
+++ b/Engine/Source/Runtime/Engine/Public/Subsystems/SubsystemCollection.h
@@ -74,6 +74,8 @@ private:
 
 	void RemoveAndDeinitializeSubsystem(USubsystem* Subsystem);
 
+	void UpdateSubsystemArrayInternal(UClass* SubsystemClass, TArray<USubsystem*>& SubsystemArray) const;
+
 	TMap<UClass*, USubsystem*> SubsystemMap;
 
 	mutable TMap<UClass*, TArray<USubsystem*>> SubsystemArrayMap;
