Fix Strata material conversion with material attributes material used by VT and using material functions.

* Changed Strata material conversion to plug masked opacity only if blend mode is set to masked opacity, otherwise it fails VT material validation.
 * Changed Strata material conversion to happen after the material expresions have been cached out, so that materials expression inputs/outputs are always valid.

diff --git a/Engine/Source/Runtime/Engine/Private/Materials/Material.cpp b/Engine/Source/Runtime/Engine/Private/Materials/Material.cpp
--- a/Engine/Source/Runtime/Engine/Private/Materials/Material.cpp
+++ b/Engine/Source/Runtime/Engine/Private/Materials/Material.cpp
@@ -2779,9 +2779,13 @@ void UMaterial::ConvertMaterialToStrataMaterial()
 		// * Remove support for material attribute
 		// * explicitely connect the Strata node to the root node
 		// * Forward inputs to the root node (Do not reconnect the Opacity as we handle the opacity by internally within the conversion node)
+		// * Forward masked opacity only if blend mode is set to masked, as certain material (e.g., FlattenVT requires to have no OpacityMask plugged)
 		bUseMaterialAttributes = false;
 		FrontMaterial.Connect(0, ConvertNode);
-		OpacityMask.Connect(7, BreakMatAtt);
+		if (BlendMode == BLEND_Masked)
+		{
+			OpacityMask.Connect(7, BreakMatAtt);
+		}
 		WorldPositionOffset.Connect(10, BreakMatAtt);
 		AmbientOcclusion.Connect(14, BreakMatAtt);
 		PixelDepthOffset.Connect(24, BreakMatAtt);
@@ -3135,7 +3139,6 @@ void UMaterial::PostLoad()
 	BackwardsCompatibilityInputConversion();
 	BackwardsCompatibilityVirtualTextureOutputConversion();
 	BackwardsCompatibilityDecalConversion();
-	ConvertMaterialToStrataMaterial();
 
 #if WITH_EDITOR
 	if ( GMaterialsThatNeedSamplerFixup.Get( this ) )
@@ -3184,6 +3187,10 @@ void UMaterial::PostLoad()
 	}
 #endif // WITH_EDITOR
 
+	// Strata materials conversion needs to be done after expressions are cached, otherwise material function won't have 
+	// valid inputs in certain cases
+	ConvertMaterialToStrataMaterial();
+
 	checkf(CachedExpressionData, TEXT("Missing cached expression data for material, should have been either serialized or created during PostLoad"));
 
 	for (const FMaterialParameterCollectionInfo& CollectionInfo : CachedExpressionData->ParameterCollectionInfos)
