Add a IsStrataEnable function per module where it is used.

diff --git a/Engine/Source/Runtime/Engine/Private/Materials/HLSLMaterialTranslator.cpp b/Engine/Source/Runtime/Engine/Private/Materials/HLSLMaterialTranslator.cpp
--- a/Engine/Source/Runtime/Engine/Private/Materials/HLSLMaterialTranslator.cpp
+++ b/Engine/Source/Runtime/Engine/Private/Materials/HLSLMaterialTranslator.cpp
@@ -46,6 +46,8 @@ static inline bool IsDebugTextureSampleEnabled()
 	return IsAnalyticDerivEnabled() && (GDebugTextureSampleEnabled != 0);
 }
 
+bool Engine_IsStrataEnabled();
+
 /** @return the vector type containing a given number of components. */
 static inline EMaterialValueType GetVectorType(uint32 NumComponents)
 {
@@ -1045,8 +1047,7 @@ bool FHLSLMaterialTranslator::Translate()
 				Errorf(TEXT("SingleLayerWater materials cannot be combined with other shading models.")); // Simply untested for now
 			}
 
-			static const auto CVarStrata = IConsoleManager::Get().FindTConsoleVariableDataInt(TEXT("r.Strata"));
-			const bool bStrata = CVarStrata ? CVarStrata->GetValueOnAnyThread() > 0 : false;
+			const bool bStrata = Engine_IsStrataEnabled();
 			if (Material->GetMaterialInterface() && !Material->GetMaterialInterface()->GetMaterial()->HasAnyExpressionsInMaterialAndFunctionsOfType<UMaterialExpressionSingleLayerWaterMaterialOutput>()
 				&& !bStrata
 				)
diff --git a/Engine/Source/Runtime/Engine/Private/Materials/Material.cpp b/Engine/Source/Runtime/Engine/Private/Materials/Material.cpp
--- a/Engine/Source/Runtime/Engine/Private/Materials/Material.cpp
+++ b/Engine/Source/Runtime/Engine/Private/Materials/Material.cpp
@@ -139,6 +139,12 @@ static TAutoConsoleVariable<int32> CVarMaterialEnableNewHLSLGenerator(
 	TEXT("Enables the new (WIP) material HLSL generator.\n"),
 	ECVF_RenderThreadSafe | ECVF_ReadOnly);
 
+bool Engine_IsStrataEnabled()
+{
+	static const auto CVar = IConsoleManager::Get().FindTConsoleVariableDataInt(TEXT("r.Strata"));
+	return CVar && CVar->GetValueOnAnyThread() > 0;
+}
+
 #if WITH_EDITOR
 const FMaterialsWithDirtyUsageFlags FMaterialsWithDirtyUsageFlags::DefaultAnnotation;
 
@@ -2724,9 +2730,7 @@ static void AddStrataShadingModelFromMaterialShadingModel(FStrataMaterialInfo& O
 void UMaterial::ConvertMaterialToStrataMaterial()
 {
 #if WITH_EDITOR
-	static const auto CVarStrata = IConsoleManager::Get().FindTConsoleVariableDataInt(TEXT("r.Strata"));
-	const bool bStrataEnabled = CVarStrata ? CVarStrata->GetValueOnAnyThread() > 0 : false;
-
+	const bool bStrataEnabled = Engine_IsStrataEnabled();
 	if (!bStrataEnabled)
 	{
 		return;
@@ -3875,9 +3879,7 @@ void UMaterial::UpdateExpressionParameterName(UMaterialExpression* Expression)
 void UMaterial::RebuildShadingModelField()
 {
 	ShadingModels.ClearShadingModels();
-
-	static const auto CVar = IConsoleManager::Get().FindTConsoleVariableDataInt(TEXT("r.Strata"));
-	const bool bStrataEnabled = CVar && CVar->GetValueOnAnyThread() > 0;
+	const bool bStrataEnabled = Engine_IsStrataEnabled();
 	if (bStrataEnabled && FrontMaterial.IsConnected())
 	{
 		FStrataMaterialInfo StrataMaterialInfo;
@@ -5303,8 +5305,7 @@ static bool IsPropertyActive_Internal(EMaterialProperty InProperty,
 	bool bUsesShadingModelFromMaterialExpression,
 	bool bIsTranslucencyWritingVelocity)
 {
-	static const auto CVar = IConsoleManager::Get().FindTConsoleVariableDataInt(TEXT("r.Strata"));
-	const bool bStrataEnabled = CVar && CVar->GetValueOnAnyThread() > 0;
+	const bool bStrataEnabled = Engine_IsStrataEnabled();
 
 	if (Domain == MD_PostProcess)
 	{
diff --git a/Engine/Source/Runtime/Engine/Private/Materials/MaterialExpressions.cpp b/Engine/Source/Runtime/Engine/Private/Materials/MaterialExpressions.cpp
--- a/Engine/Source/Runtime/Engine/Private/Materials/MaterialExpressions.cpp
+++ b/Engine/Source/Runtime/Engine/Private/Materials/MaterialExpressions.cpp
@@ -303,6 +303,8 @@ FUObjectAnnotationSparseBool GMaterialFunctionsThatNeedCommentFix;
 FUObjectAnnotationSparseBool GMaterialFunctionsThatNeedSamplerFixup;
 #endif // #if WITH_EDITOR
 
+bool Engine_IsStrataEnabled();
+
 /** Returns whether the given expression class is allowed. */
 bool IsAllowedExpressionType(const UClass* const Class, const bool bMaterialFunction)
 {
@@ -18626,8 +18628,7 @@ int32  UMaterialExpressionClearCoatNormalCustomOutput::Compile(class FMaterialCo
 	}
 	else
 	{
-		static const auto CVarStrata = IConsoleManager::Get().FindTConsoleVariableDataInt(TEXT("r.Strata"));
-		const bool bStrata = CVarStrata ? CVarStrata->GetValueOnAnyThread() > 0 : false;
+		const bool bStrata = Engine_IsStrataEnabled();
 		if (!bStrata)
 		{
 			return CompilerError(Compiler, TEXT("Input missing"));
diff --git a/Engine/Source/Runtime/Engine/Private/ShaderCompiler/ShaderGenerationUtil.cpp b/Engine/Source/Runtime/Engine/Private/ShaderCompiler/ShaderGenerationUtil.cpp
--- a/Engine/Source/Runtime/Engine/Private/ShaderCompiler/ShaderGenerationUtil.cpp
+++ b/Engine/Source/Runtime/Engine/Private/ShaderCompiler/ShaderGenerationUtil.cpp
@@ -44,6 +44,9 @@ bool NeedsVelocityDepth(EShaderPlatform TargetPlatform)
 		|| FDataDrivenShaderPlatformInfo::GetSupportsRayTracing(TargetPlatform);
 }
 
+// Strata::IsEnabled is only accessible in the Renderer module
+bool Engine_IsStrataEnabled();
+
 #if WITH_EDITOR
 
 static int FetchCompileInt(FShaderCompilerEnvironment& OutEnvironment, const char* SrcName)
@@ -251,13 +254,6 @@ void FShaderCompileUtilities::ApplyFetchEnvironment(FShaderCompilerDefines& SrcD
 	FETCH_COMPILE_BOOL(PLATFORM_SUPPORTS_DEVELOPMENT_SHADERS);
 }
 
-// Strata::IsEnabled is only accessible in the Renderer module
-static bool IsStrataEnabled()
-{
-	static IConsoleVariable* CVar = IConsoleManager::Get().FindConsoleVariable(TEXT("r.Strata"));
-	return CVar->GetInt() > 0;
-}
-
 // if we change the logic, increment this number to force a DDC key change
 static const int32 GBufferGeneratorVersion = 4;
 
@@ -1110,7 +1106,7 @@ static FString CreateGBufferDecodeFunctionDirect(const FGBufferInfo& BufferInfo)
 	// Default initialization in case no gbuffer data are generated when Strata is enabled to 
 	// prevent shader compiler error (division by zero, variable not-initialized) with passes 
 	// not converted to Strata and still using Gbuffer data
-	const bool bStrata = IsStrataEnabled();
+	const bool bStrata = Engine_IsStrataEnabled();
 	if (bStrata)
 	{
 		FullStr += TEXT("\tRet.WorldNormal = float3(0,0,1);\n");
diff --git a/Engine/Source/Runtime/RenderCore/Private/GBufferInfo.cpp b/Engine/Source/Runtime/RenderCore/Private/GBufferInfo.cpp
--- a/Engine/Source/Runtime/RenderCore/Private/GBufferInfo.cpp
+++ b/Engine/Source/Runtime/RenderCore/Private/GBufferInfo.cpp
@@ -19,7 +19,7 @@ static const int32 EGBufferFormat_Force16BitsPerChannel = 5;
 
 
 // Strata::IsEnabled is only accessible in the Renderer module
-static bool IsStrataEnabled()
+bool RenderCore_IsStrataEnabled()
 {
 	static IConsoleVariable* CVar = IConsoleManager::Get().FindConsoleVariable(TEXT("r.Strata"));
 	return CVar->GetInt() > 0;
@@ -238,7 +238,7 @@ FGBufferInfo RENDERCORE_API FetchLegacyGBufferInfo(const FGBufferParams& Params)
 	int32 TargetVelocity = -1;
 
 	// Strata ouputs material data through UAV. Only SceneColor, PrecalcShadow & Velocity data are still emitted through RenderTargets
-	const bool bStrata = IsStrataEnabled();
+	const bool bStrata = RenderCore_IsStrataEnabled();
 	if (bStrata)
 	{
 		TargetGBufferA = -1;
diff --git a/Engine/Source/Runtime/RenderCore/Private/Shader.cpp b/Engine/Source/Runtime/RenderCore/Private/Shader.cpp
--- a/Engine/Source/Runtime/RenderCore/Private/Shader.cpp
+++ b/Engine/Source/Runtime/RenderCore/Private/Shader.cpp
@@ -112,6 +112,8 @@ static TLinkedList<FShaderPipelineType*>*	GShaderPipelineList = nullptr;
 
 static FSHAHash ShaderSourceDefaultHash; //will only be read (never written) for the cooking case
 
+bool RenderCore_IsStrataEnabled();
+
 /**
  * Find the shader pipeline type with the given name.
  * @return NULL if no type matched.
@@ -1768,8 +1770,7 @@ void ShaderMapAppendKeyString(EShaderPlatform Platform, FString& KeyString)
 	}
 
 	{
-		static const auto CVar = IConsoleManager::Get().FindTConsoleVariableDataInt(TEXT("r.Strata"));
-		const bool bStrataEnabled = CVar && CVar->GetValueOnAnyThread() > 0;
+		const bool bStrataEnabled = RenderCore_IsStrataEnabled();
 		if (bStrataEnabled)
 		{
 			KeyString += TEXT("_STRATA");
