Move FMaterialShaderParameters constructor definition in private.

diff --git a/Engine/Source/Runtime/Engine/Private/Materials/MaterialShared.cpp b/Engine/Source/Runtime/Engine/Private/Materials/MaterialShared.cpp
--- a/Engine/Source/Runtime/Engine/Private/Materials/MaterialShared.cpp
+++ b/Engine/Source/Runtime/Engine/Private/Materials/MaterialShared.cpp
@@ -5189,4 +5189,68 @@ bool FMemoryImageMaterialParameterInfo::RemapLayerIndex(TArrayView<const int32>
 	return RemapParameterLayerIndex(IndexRemap, *this, OutResult);
 }
 
+FMaterialShaderParameters::FMaterialShaderParameters(const FMaterial* InMaterial)
+{
+	// Make sure to zero-initialize so we get consistent hashes
+	FMemory::Memzero(*this);
+
+	MaterialDomain = InMaterial->GetMaterialDomain();
+	ShadingModels = InMaterial->GetShadingModels();
+	BlendMode = InMaterial->GetBlendMode();
+	FeatureLevel = InMaterial->GetFeatureLevel();
+	QualityLevel = InMaterial->GetQualityLevel();
+	BlendableLocation = InMaterial->GetBlendableLocation();
+	NumCustomizedUVs = InMaterial->GetNumCustomizedUVs();
+	StencilCompare = InMaterial->GetStencilCompare();
+	bIsDefaultMaterial = InMaterial->IsDefaultMaterial();
+	bIsSpecialEngineMaterial = InMaterial->IsSpecialEngineMaterial();
+	bIsMasked = InMaterial->IsMasked();
+	bIsTwoSided = InMaterial->IsTwoSided();
+	bIsDistorted = InMaterial->IsDistorted();
+	bShouldCastDynamicShadows = InMaterial->ShouldCastDynamicShadows();
+	bWritesEveryPixel = InMaterial->WritesEveryPixel(false);
+	bWritesEveryPixelShadowPass = InMaterial->WritesEveryPixel(true);
+	bHasBaseColorConnected = InMaterial->HasBaseColorConnected();
+	bHasNormalConnected = InMaterial->HasNormalConnected();
+	bHasRoughnessConnected = InMaterial->HasRoughnessConnected();
+	bHasSpecularConnected = InMaterial->HasSpecularConnected();
+	bHasMetallicConnected = InMaterial->HasMetallicConnected();
+	bHasEmissiveColorConnected = InMaterial->HasEmissiveColorConnected();
+	bHasAmbientOcclusionConnected = InMaterial->HasAmbientOcclusionConnected();
+	bHasAnisotropyConnected = InMaterial->HasAnisotropyConnected();
+	bHasVertexPositionOffsetConnected = InMaterial->HasVertexPositionOffsetConnected();
+	bHasPixelDepthOffsetConnected = InMaterial->HasPixelDepthOffsetConnected();
+	bMaterialMayModifyMeshPosition = InMaterial->MaterialMayModifyMeshPosition();
+	bIsUsedWithStaticLighting = InMaterial->IsUsedWithStaticLighting();
+	bIsUsedWithParticleSprites = InMaterial->IsUsedWithParticleSprites();
+	bIsUsedWithMeshParticles = InMaterial->IsUsedWithMeshParticles();
+	bIsUsedWithNiagaraSprites = InMaterial->IsUsedWithNiagaraSprites();
+	bIsUsedWithNiagaraMeshParticles = InMaterial->IsUsedWithNiagaraMeshParticles();
+	bIsUsedWithNiagaraRibbons = InMaterial->IsUsedWithNiagaraRibbons();
+	bIsUsedWithLandscape = InMaterial->IsUsedWithLandscape();
+	bIsUsedWithBeamTrails = InMaterial->IsUsedWithBeamTrails();
+	bIsUsedWithSplineMeshes = InMaterial->IsUsedWithSplineMeshes();
+	bIsUsedWithSkeletalMesh = InMaterial->IsUsedWithSkeletalMesh();
+	bIsUsedWithMorphTargets = InMaterial->IsUsedWithMorphTargets();
+	bIsUsedWithAPEXCloth = InMaterial->IsUsedWithAPEXCloth();
+	bIsUsedWithGeometryCache = InMaterial->IsUsedWithGeometryCache();
+	bIsUsedWithGeometryCollections = InMaterial->IsUsedWithGeometryCollections();
+	bIsUsedWithHairStrands = InMaterial->IsUsedWithHairStrands();
+	bIsUsedWithWater = InMaterial->IsUsedWithWater();
+	bIsTranslucencyWritingVelocity = InMaterial->IsTranslucencyWritingVelocity();
+	bIsTranslucencyWritingCustomDepth = InMaterial->IsTranslucencyWritingCustomDepth();
+	bIsDitheredLODTransition = InMaterial->IsDitheredLODTransition();
+	bIsUsedWithInstancedStaticMeshes = InMaterial->IsUsedWithInstancedStaticMeshes();
+	bHasPerInstanceCustomData = InMaterial->HasPerInstanceCustomData();
+	bHasPerInstanceRandom = InMaterial->HasPerInstanceRandom();
+	bHasVertexInterpolator = InMaterial->HasVertexInterpolator();
+	bHasRuntimeVirtualTextureOutput = InMaterial->HasRuntimeVirtualTextureOutput();
+	bIsUsedWithLidarPointCloud = InMaterial->IsUsedWithLidarPointCloud();
+	bIsUsedWithVirtualHeightfieldMesh = InMaterial->IsUsedWithVirtualHeightfieldMesh();
+	bIsUsedWithNanite = InMaterial->IsUsedWithNanite();
+	bIsStencilTestEnabled = InMaterial->IsStencilTestEnabled();
+	bIsTranslucencySurface = InMaterial->GetTranslucencyLightingMode() == ETranslucencyLightingMode::TLM_Surface || InMaterial->GetTranslucencyLightingMode() == ETranslucencyLightingMode::TLM_SurfacePerPixelLighting;
+	bShouldDisableDepthTest = InMaterial->ShouldDisableDepthTest();
+}
+
 #undef LOCTEXT_NAMESPACE
\ No newline at end of file
diff --git a/Engine/Source/Runtime/Engine/Public/MaterialShared.h b/Engine/Source/Runtime/Engine/Public/MaterialShared.h
--- a/Engine/Source/Runtime/Engine/Public/MaterialShared.h
+++ b/Engine/Source/Runtime/Engine/Public/MaterialShared.h
@@ -3344,69 +3344,7 @@ struct FMaterialShaderParameters
 		};
 	};
 
-	FMaterialShaderParameters(const FMaterial* InMaterial)
-	{
-		// Make sure to zero-initialize so we get consistent hashes
-		FMemory::Memzero(*this);
-
-		MaterialDomain = InMaterial->GetMaterialDomain();
-		ShadingModels = InMaterial->GetShadingModels();
-		BlendMode = InMaterial->GetBlendMode();
-		FeatureLevel = InMaterial->GetFeatureLevel();
-		QualityLevel = InMaterial->GetQualityLevel();
-		BlendableLocation = InMaterial->GetBlendableLocation();
-		NumCustomizedUVs = InMaterial->GetNumCustomizedUVs();
-		StencilCompare = InMaterial->GetStencilCompare();
-		bIsDefaultMaterial = InMaterial->IsDefaultMaterial();
-		bIsSpecialEngineMaterial = InMaterial->IsSpecialEngineMaterial();
-		bIsMasked = InMaterial->IsMasked();
-		bIsTwoSided = InMaterial->IsTwoSided();
-		bIsDistorted = InMaterial->IsDistorted();
-		bShouldCastDynamicShadows = InMaterial->ShouldCastDynamicShadows();
-		bWritesEveryPixel = InMaterial->WritesEveryPixel(false);
-		bWritesEveryPixelShadowPass = InMaterial->WritesEveryPixel(true);
-		bHasBaseColorConnected = InMaterial->HasBaseColorConnected();
-		bHasNormalConnected = InMaterial->HasNormalConnected();
-		bHasRoughnessConnected = InMaterial->HasRoughnessConnected();
-		bHasSpecularConnected = InMaterial->HasSpecularConnected();
-		bHasMetallicConnected = InMaterial->HasMetallicConnected();
-		bHasEmissiveColorConnected = InMaterial->HasEmissiveColorConnected();
-		bHasAmbientOcclusionConnected = InMaterial->HasAmbientOcclusionConnected();
-		bHasAnisotropyConnected = InMaterial->HasAnisotropyConnected();
-		bHasVertexPositionOffsetConnected = InMaterial->HasVertexPositionOffsetConnected();
-		bHasPixelDepthOffsetConnected = InMaterial->HasPixelDepthOffsetConnected();
-		bMaterialMayModifyMeshPosition = InMaterial->MaterialMayModifyMeshPosition();
-		bIsUsedWithStaticLighting = InMaterial->IsUsedWithStaticLighting();
-		bIsUsedWithParticleSprites = InMaterial->IsUsedWithParticleSprites();
-		bIsUsedWithMeshParticles = InMaterial->IsUsedWithMeshParticles();
-		bIsUsedWithNiagaraSprites = InMaterial->IsUsedWithNiagaraSprites();
-		bIsUsedWithNiagaraMeshParticles = InMaterial->IsUsedWithNiagaraMeshParticles();
-		bIsUsedWithNiagaraRibbons = InMaterial->IsUsedWithNiagaraRibbons();
-		bIsUsedWithLandscape = InMaterial->IsUsedWithLandscape();
-		bIsUsedWithBeamTrails = InMaterial->IsUsedWithBeamTrails();
-		bIsUsedWithSplineMeshes = InMaterial->IsUsedWithSplineMeshes();
-		bIsUsedWithSkeletalMesh = InMaterial->IsUsedWithSkeletalMesh();
-		bIsUsedWithMorphTargets = InMaterial->IsUsedWithMorphTargets();
-		bIsUsedWithAPEXCloth = InMaterial->IsUsedWithAPEXCloth();
-		bIsUsedWithGeometryCache = InMaterial->IsUsedWithGeometryCache();
-		bIsUsedWithGeometryCollections = InMaterial->IsUsedWithGeometryCollections();
-		bIsUsedWithHairStrands = InMaterial->IsUsedWithHairStrands();
-		bIsUsedWithWater = InMaterial->IsUsedWithWater();
-		bIsTranslucencyWritingVelocity = InMaterial->IsTranslucencyWritingVelocity();
-		bIsTranslucencyWritingCustomDepth = InMaterial->IsTranslucencyWritingCustomDepth();
-		bIsDitheredLODTransition = InMaterial->IsDitheredLODTransition();
-		bIsUsedWithInstancedStaticMeshes = InMaterial->IsUsedWithInstancedStaticMeshes();
-		bHasPerInstanceCustomData = InMaterial->HasPerInstanceCustomData();
-		bHasPerInstanceRandom = InMaterial->HasPerInstanceRandom();
-		bHasVertexInterpolator = InMaterial->HasVertexInterpolator();
-		bHasRuntimeVirtualTextureOutput = InMaterial->HasRuntimeVirtualTextureOutput();
-		bIsUsedWithLidarPointCloud = InMaterial->IsUsedWithLidarPointCloud();
-		bIsUsedWithVirtualHeightfieldMesh = InMaterial->IsUsedWithVirtualHeightfieldMesh();
-		bIsUsedWithNanite = InMaterial->IsUsedWithNanite();
-		bIsStencilTestEnabled = InMaterial->IsStencilTestEnabled();
-		bIsTranslucencySurface = InMaterial->GetTranslucencyLightingMode() == ETranslucencyLightingMode::TLM_Surface || InMaterial->GetTranslucencyLightingMode() == ETranslucencyLightingMode::TLM_SurfacePerPixelLighting;
-		bShouldDisableDepthTest = InMaterial->ShouldDisableDepthTest();
-	}
+	FMaterialShaderParameters(const FMaterial* InMaterial);
 };
 
 inline bool ShouldIncludeMaterialInDefaultOpaquePass(const FMaterial& Material)
