- Changed how CVar preview filtering happens, now it's only on ECVF_Preview (which is now set on the sg. group cvars, which imply Preview for all expanded cvars from the scalability group
- Instead of filtering on load of DP by Previewability, it will now add everything to the DP, and then kleep a separate list of Preview cvars. This way, code can still ask for the value of a cvar for a DP even if it's not previewable
- Additionally, the non-DP cvars are now also filtered by ECVF_Preview, not just DP-specified CVars

diff --git a/Engine/Source/Runtime/Core/Private/HAL/ConsoleManager.cpp b/Engine/Source/Runtime/Core/Private/HAL/ConsoleManager.cpp
--- a/Engine/Source/Runtime/Core/Private/HAL/ConsoleManager.cpp
+++ b/Engine/Source/Runtime/Core/Private/HAL/ConsoleManager.cpp
@@ -410,6 +410,9 @@ bool IConsoleManager::VisitPlatformCVarsForEmulation(FName PlatformName, bool bV
 					}
 				}
 
+				IConsoleVariable* CVar = IConsoleManager::Get().FindConsoleVariable(*Key);
+				EConsoleVariableFlags PreviewFlag = (CVar != nullptr) ? (EConsoleVariableFlags)(CVar->GetFlags() & ECVF_Preview) : ECVF_Default;
+
 				if (Key.StartsWith(TEXT("sg.")))
 				{
 					// @todo ini: If anything in here was already set, overwrite it or skip it?
@@ -421,12 +424,12 @@ bool IConsoleManager::VisitPlatformCVarsForEmulation(FName PlatformName, bool bV
 
 					for (const auto& ScalabilityPair : ScalabilityCVars)
 					{
-						Visit(ScalabilityPair.Key, ScalabilityPair.Value, ECVF_SetByScalability);
+						Visit(ScalabilityPair.Key, ScalabilityPair.Value, (EConsoleVariableFlags)(ECVF_SetByScalability | PreviewFlag));
 					}
 				}
 				else
 				{
-					Visit(Key, Value, SectionPair.SetBy);
+					Visit(Key, Value, (EConsoleVariableFlags)(SectionPair.SetBy | PreviewFlag));
 				}
 			}
 		}
@@ -462,13 +465,13 @@ static bool GetConfigValueFromRuntimeSources(FName PlatformName, IConsoleVariabl
 	LastSetBy = ECVF_SetByConstructor;
 
 	IConsoleManager::VisitPlatformCVarsForEmulation(PlatformName, true, 
-		[&VariableName, &OutValue, &LastSetBy](const FString& CVarName, const FString& CVarValue, EConsoleVariableFlags SetBy)
+		[&VariableName, &OutValue, &LastSetBy](const FString& CVarName, const FString& CVarValue, EConsoleVariableFlags SetByAndPreview)
 		{
 			// if this key is the variable, set it at the current SetBy level
 			if (CVarName == VariableName)
 			{
 				OutValue = CVarValue;
-				LastSetBy = SetBy;
+				LastSetBy = (EConsoleVariableFlags)(SetByAndPreview & ECVF_SetByMask);
 			}
 		});
 
diff --git a/Engine/Source/Runtime/Engine/Classes/DeviceProfiles/DeviceProfile.h b/Engine/Source/Runtime/Engine/Classes/DeviceProfiles/DeviceProfile.h
--- a/Engine/Source/Runtime/Engine/Classes/DeviceProfiles/DeviceProfile.h
+++ b/Engine/Source/Runtime/Engine/Classes/DeviceProfiles/DeviceProfile.h
@@ -163,10 +163,15 @@ private:
 #if ALLOW_OTHER_PLATFORM_CONFIG
 public:
 	const TMap<FString, FString>& GetAllExpandedCVars();
+	const TMap<FString, FString>& GetAllPreviewCVars();
 	void AddExpandedCVars(const TMap<FString, FString>& CVarsToMerge);
+	void AddPreviewCVars(const TMap<FString, FString>& CVarsToMerge);
 	void ClearAllExpandedCVars();
 private:
 	/** Resolved CVars, including expanded scalability cvars used to properly emulate one platform on another */
 	TMap<FString, FString> AllExpandedCVars;
+
+	/** The set of cvars that can be previewed (a subset of AllExpandedCVars) */
+	TMap<FString, FString> AllPreviewCVars;
 #endif
 };
diff --git a/Engine/Source/Runtime/Engine/Private/DeviceProfiles/DeviceProfile.cpp b/Engine/Source/Runtime/Engine/Private/DeviceProfiles/DeviceProfile.cpp
--- a/Engine/Source/Runtime/Engine/Private/DeviceProfiles/DeviceProfile.cpp
+++ b/Engine/Source/Runtime/Engine/Private/DeviceProfiles/DeviceProfile.cpp
@@ -386,6 +386,17 @@ const TMap<FString, FString>& UDeviceProfile::GetAllExpandedCVars()
 	return AllExpandedCVars;
 }
 
+const TMap<FString, FString>& UDeviceProfile::GetAllPreviewCVars()
+{
+	// expand on first use
+	if (AllPreviewCVars.Num() == 0)
+	{
+		UDeviceProfileManager::Get().ExpandDeviceProfileCVars(this);
+	}
+
+	return AllPreviewCVars;
+}
+
 void UDeviceProfile::AddExpandedCVars(const TMap<FString, FString>& CVarsToMerge)
 {
 	// merge the cvars, overwriting any existing ones
@@ -395,9 +406,18 @@ void UDeviceProfile::AddExpandedCVars(const TMap<FString, FString>& CVarsToMerge
 	}
 }
 
+void UDeviceProfile::AddPreviewCVars(const TMap<FString, FString>& CVarsToMerge)
+{
+	for (const auto& Pair : CVarsToMerge)
+	{
+		AllPreviewCVars.Add(Pair.Key, Pair.Value);
+	}
+}
+
 void UDeviceProfile::ClearAllExpandedCVars()
 {
 	AllExpandedCVars.Empty();
+	AllPreviewCVars.Empty();
 }
 
 
diff --git a/Engine/Source/Runtime/Engine/Private/DeviceProfiles/DeviceProfileManager.cpp b/Engine/Source/Runtime/Engine/Private/DeviceProfiles/DeviceProfileManager.cpp
--- a/Engine/Source/Runtime/Engine/Private/DeviceProfiles/DeviceProfileManager.cpp
+++ b/Engine/Source/Runtime/Engine/Private/DeviceProfiles/DeviceProfileManager.cpp
@@ -175,6 +175,9 @@ void UDeviceProfileManager::ProcessDeviceProfileIniSettings(const FString& Devic
 	// if we are just caching, this also contains the set of all CVars (including expanding scalability groups)
 	TMap<FString, FString> CVarsAlreadySetList;
 
+	// previewable subset of CVarsAlreadySetList
+	TMap<FString, FString> PreviewableCVars;
+
 	// reset some global state for "active DP" mode
 	if (Mode == EDeviceProfileMode::DPM_SetCVars)
 	{
@@ -244,24 +247,6 @@ void UDeviceProfileManager::ProcessDeviceProfileIniSettings(const FString& Devic
 
 	FString SectionSuffix = *FString::Printf(TEXT(" %s"), *UDeviceProfile::StaticClass()->GetName());
 
-#if WITH_EDITOR
-	TArray<FString> PreviewAllowlistCVars;
-	TArray<FString> PreviewDenylistCVars;
-	bool bFoundAllowDeny = false;
-	if (Mode == EDeviceProfileMode::DPM_CacheValues)
-	{
-		// Walk up the device profile tree to find the most specific device profile with a Denylist or Allowlist of cvars to apply, and use those Allow/Denylists.
-		for(FString CurrentProfileName = DeviceProfileName, CurrentSectionName = DeviceProfileName + SectionSuffix;
-			PreviewAllowlistCVars.Num()==0 && PreviewDenylistCVars.Num()==0 && !CurrentProfileName.IsEmpty() && AvailableProfiles.Contains(CurrentSectionName);
-			CurrentProfileName = GConfig->GetStr(*CurrentSectionName, TEXT("BaseProfileName"), GDeviceProfilesIni), CurrentSectionName = CurrentProfileName + SectionSuffix)
-		{
-			ConfigSystem->GetArray(*CurrentSectionName, TEXT("PreviewAllowlistCVars"), PreviewAllowlistCVars, GDeviceProfilesIni);
-			ConfigSystem->GetArray(*CurrentSectionName, TEXT("PreviewDenylistCVars"), PreviewDenylistCVars, GDeviceProfilesIni);
-		}
-	}
-#endif
-
-
 	TSet<FString> FragmentCVarKeys;
 	TArray<FString> SelectedFragmentCVars;
 	TArray<FSelectedFragmentProperties> FragmentsSelected;
@@ -392,31 +377,9 @@ void UDeviceProfileManager::ProcessDeviceProfileIniSettings(const FString& Devic
 					{
 						if (!CVarsAlreadySetList.Find(CVarKey))
 						{
-#if WITH_EDITOR
-							if (Mode == EDeviceProfileMode::DPM_CacheValues)
-							{
-								if (PreviewDenylistCVars.ContainsByPredicate([&CVarKey](const FString& Entry) { return CVarKey.StartsWith(Entry); }))
-								{
-									UE_LOG(LogInit, Log, TEXT("Skipping Device Profile CVar due to PreviewDenylistCVars: [[%s]]"), *CVarKey);
-									continue;
-								}
-								
-								bool bIsPreviewCVar = false;
-								if (IConsoleVariable* CVar = IConsoleManager::Get().FindConsoleVariable(*CVarKey))
-								{
-									bIsPreviewCVar = CVar->TestFlags(ECVF_Preview);
-								}
-								if (bIsPreviewCVar == false && (PreviewAllowlistCVars.Num() > 0 && !PreviewAllowlistCVars.ContainsByPredicate([&CVarKey](const FString& Entry) { return CVarKey.StartsWith(Entry); })))
-								{
-									UE_LOG(LogInit, Log, TEXT("Skipping Device Profile CVar due to PreviewAllowlistCVars: [[%s]]"), *CVarKey);
-									continue;
-								}
-							}
-#endif
-
+							IConsoleVariable* CVar = IConsoleManager::Get().FindConsoleVariable(*CVarKey);
 							if (Mode == EDeviceProfileMode::DPM_SetCVars)
 							{
-								IConsoleVariable* CVar = IConsoleManager::Get().FindConsoleVariable(*CVarKey);
 								if (CVar)
 								{
 									// remember the previous value
@@ -445,6 +408,19 @@ void UDeviceProfileManager::ProcessDeviceProfileIniSettings(const FString& Devic
 
 								// cache key with value
 								CVarsAlreadySetList.Add(CVarKey, CVarValue);
+
+#if WITH_EDITOR
+								// if this is previewable, add it to the list - and also make any expanded scalability cvars previewable
+								if (CVar->TestFlags(EConsoleVariableFlags::ECVF_Preview))
+								{
+									if (bIsScalabilityBucket)
+									{
+										// don't overwrite any existing cvars when expanding
+										ExpandScalabilityCVar(ConfigSystem, CVarKey, CVarValue, PreviewableCVars, false);
+									}
+									PreviewableCVars.Add(CVarKey, CVarValue);
+								}
+#endif
 							}
 							// actually set the cvar if not just caching
 							else
@@ -489,6 +465,9 @@ void UDeviceProfileManager::ProcessDeviceProfileIniSettings(const FString& Devic
 	{
 		UDeviceProfile* Profile = UDeviceProfileManager::Get().FindProfile(DeviceProfileName, false);
 		Profile->AddExpandedCVars(CVarsAlreadySetList);
+
+		// now remember all previewable cvars
+		Profile->AddPreviewCVars(PreviewableCVars);
 	}
 #endif
 }
@@ -566,14 +545,20 @@ void UDeviceProfileManager::InitializeCVarsForActiveDeviceProfile()
 void UDeviceProfileManager::ExpandDeviceProfileCVars(UDeviceProfile* DeviceProfile)
 {
 	TMap<FString, FString> CVarsToAdd;
+	TMap<FString, FString> PreviewCVars;
 	IConsoleManager::VisitPlatformCVarsForEmulation(*DeviceProfile->DeviceType, false,
-		[&CVarsToAdd](const FString& CVarName, const FString& CVarValue, EConsoleVariableFlags SetBy)
+		[&CVarsToAdd, &PreviewCVars](const FString& CVarName, const FString& CVarValue, EConsoleVariableFlags SetByAndPreview)
 		{
 			CVarsToAdd.Add(CVarName, CVarValue);
+			if (SetByAndPreview & EConsoleVariableFlags::ECVF_Preview)
+			{
+				PreviewCVars.Add(CVarName, CVarValue);
+			}
 		});
 
 	// we now have all the cvars for this platform without _any_ device profiles, so now add those into the map
 	DeviceProfile->AddExpandedCVars(CVarsToAdd);
+	DeviceProfile->AddPreviewCVars(PreviewCVars);
 
 	// and finalize them into a cache
 	ProcessDeviceProfileIniSettings(DeviceProfile->GetName(), EDeviceProfileMode::DPM_CacheValues);
@@ -958,7 +943,7 @@ void UDeviceProfileManager::SetPreviewDeviceProfile(UDeviceProfile* DeviceProfil
 
 	UE_LOG(LogDeviceProfileManager, Log, TEXT("SetPreviewDeviceProfile preview to %s"), *DeviceProfile->GetName());
 	// apply the preview DP cvars.
-	for (const auto& Pair : DeviceProfile->GetAllExpandedCVars())
+	for (const auto& Pair : DeviceProfile->GetAllPreviewCVars())
 	{
 		IConsoleVariable* CVar = IConsoleManager::Get().FindConsoleVariable(*Pair.Key);
 		// skip over scalability group cvars (maybe they shouldn't be left in the AllExpandedCVars?)
@@ -1374,6 +1359,18 @@ public:
 				}
 			}
 		}
+		else if (FParse::Command(&Cmd, TEXT("dpdumppreview")))
+		{
+			UDeviceProfile* DeviceProfile = UDeviceProfileManager::Get().FindProfile(Cmd, false);
+			if (DeviceProfile)
+			{
+				Ar.Logf(TEXT("All preview cvars found for deviceprofile %s"), Cmd);
+				for (const auto& Pair : DeviceProfile->GetAllPreviewCVars())
+				{
+					Ar.Logf(TEXT("%s = %s"), *Pair.Key, *Pair.Value);
+				}
+			}
+		}
 		else if (FParse::Command(&Cmd, TEXT("dppreview")))
 		{
 			UDeviceProfile* DeviceProfile = UDeviceProfileManager::Get().FindProfile(Cmd, false);
diff --git a/Engine/Source/Runtime/Engine/Private/Scalability.cpp b/Engine/Source/Runtime/Engine/Private/Scalability.cpp
--- a/Engine/Source/Runtime/Engine/Private/Scalability.cpp
+++ b/Engine/Source/Runtime/Engine/Private/Scalability.cpp
@@ -15,77 +15,77 @@ static TAutoConsoleVariable<float> CVarResolutionQuality(
 	100.0f,
 	TEXT("Scalability quality state (internally used by scalability system, ini load/save or using SCALABILITY console command)\n")
 	TEXT(" 10..100, default: 100"),
-	ECVF_ScalabilityGroup);
+	ECVF_ScalabilityGroup | ECVF_Preview);
 
 static TAutoConsoleVariable<int32> CVarViewDistanceQuality(
 	TEXT("sg.ViewDistanceQuality"),
 	Scalability::DefaultQualityLevel,
 	TEXT("Scalability quality state (internally used by scalability system, ini load/save or using SCALABILITY console command)\n")
 	TEXT(" 0:low, 1:med, 2:high, 3:epic, 4:cinematic, default: 3"),
-	ECVF_ScalabilityGroup);
+	ECVF_ScalabilityGroup | ECVF_Preview);
 
 static TAutoConsoleVariable<int32> CVarAntiAliasingQuality(
 	TEXT("sg.AntiAliasingQuality"),
 	Scalability::DefaultQualityLevel,
 	TEXT("Scalability quality state (internally used by scalability system, ini load/save or using SCALABILITY console command)\n")
 	TEXT(" 0:low, 1:med, 2:high, 3:epic, 4:cinematic, default: 3"),
-	ECVF_ScalabilityGroup);
+	ECVF_ScalabilityGroup | ECVF_Preview);
 
 static TAutoConsoleVariable<int32> CVarShadowQuality(
 	TEXT("sg.ShadowQuality"),
 	Scalability::DefaultQualityLevel,
 	TEXT("Scalability quality state (internally used by scalability system, ini load/save or using SCALABILITY console command)\n")
 	TEXT(" 0:low, 1:med, 2:high, 3:epic, 4:cinematic, default: 3"),
-	ECVF_ScalabilityGroup);
+	ECVF_ScalabilityGroup | ECVF_Preview);
 
 static TAutoConsoleVariable<int32> CVarGlobalIlluminationQuality(
 	TEXT("sg.GlobalIlluminationQuality"),
 	Scalability::DefaultQualityLevel,
 	TEXT("Scalability quality state (internally used by scalability system, ini load/save or using SCALABILITY console command)\n")
 	TEXT(" 0:low, 1:med, 2:high, 3:epic, 4:cinematic, default: 3"),
-	ECVF_ScalabilityGroup);
+	ECVF_ScalabilityGroup | ECVF_Preview);
 
 static TAutoConsoleVariable<int32> CVarReflectionQuality(
 	TEXT("sg.ReflectionQuality"),
 	Scalability::DefaultQualityLevel,
 	TEXT("Scalability quality state (internally used by scalability system, ini load/save or using SCALABILITY console command)\n")
 	TEXT(" 0:low, 1:med, 2:high, 3:epic, 4:cinematic, default: 3"),
-	ECVF_ScalabilityGroup);
+	ECVF_ScalabilityGroup | ECVF_Preview);
 
 static TAutoConsoleVariable<int32> CVarPostProcessQuality(
 	TEXT("sg.PostProcessQuality"),
 	Scalability::DefaultQualityLevel,
 	TEXT("Scalability quality state (internally used by scalability system, ini load/save or using SCALABILITY console command)\n")
 	TEXT(" 0:low, 1:med, 2:high, 3:epic, 4:cinematic, default: 3"),
-	ECVF_ScalabilityGroup);
+	ECVF_ScalabilityGroup | ECVF_Preview);
 
 static TAutoConsoleVariable<int32> CVarTextureQuality(
 	TEXT("sg.TextureQuality"),
 	Scalability::DefaultQualityLevel,
 	TEXT("Scalability quality state (internally used by scalability system, ini load/save or using SCALABILITY console command)\n")
 	TEXT(" 0:low, 1:med, 2:high, 3:epic, 4:cinematic, default: 3"),
-	ECVF_ScalabilityGroup);
+	ECVF_ScalabilityGroup | ECVF_Preview);
 
 static TAutoConsoleVariable<int32> CVarEffectsQuality(
 	TEXT("sg.EffectsQuality"),
 	Scalability::DefaultQualityLevel,
 	TEXT("Scalability quality state (internally used by scalability system, ini load/save or using SCALABILITY console command)\n")
 	TEXT(" 0:low, 1:med, 2:high, 3:epic, 4:cinematic, default: 3"),
-	ECVF_ScalabilityGroup);
+	ECVF_ScalabilityGroup | ECVF_Preview);
 
 static TAutoConsoleVariable<int32> CVarFoliageQuality(
 	TEXT("sg.FoliageQuality"),
 	3,
 	TEXT("Scalability quality state (internally used by scalability system, ini load/save or using SCALABILITY console command)\n")
 	TEXT(" 0:low, 1:med, 2:high, 3:epic, 4:cinematic, default: 3"),
-	ECVF_ScalabilityGroup);
+	ECVF_ScalabilityGroup | ECVF_Preview);
 
 static TAutoConsoleVariable<int32> CVarShadingQuality(
 	TEXT("sg.ShadingQuality"),
 	Scalability::DefaultQualityLevel,
 	TEXT("Scalability quality state (internally used by scalability system, ini load/save or using SCALABILITY console command)\n")
 	TEXT(" 0:low, 1:med, 2:high, 3:epic, 4:cinematic, default: 3"),
-	ECVF_ScalabilityGroup);
+	ECVF_ScalabilityGroup | ECVF_Preview);
 
 static TAutoConsoleVariable<int32> CVarViewDistanceQuality_NumLevels(
 	TEXT("sg.ViewDistanceQuality.NumLevels"),
