- Added a null check for DP previewing

diff --git a/Engine/Source/Runtime/Engine/Private/DeviceProfiles/DeviceProfileManager.cpp b/Engine/Source/Runtime/Engine/Private/DeviceProfiles/DeviceProfileManager.cpp
--- a/Engine/Source/Runtime/Engine/Private/DeviceProfiles/DeviceProfileManager.cpp
+++ b/Engine/Source/Runtime/Engine/Private/DeviceProfiles/DeviceProfileManager.cpp
@@ -411,7 +411,7 @@ void UDeviceProfileManager::ProcessDeviceProfileIniSettings(const FString& Devic
 
 #if WITH_EDITOR
 								// if this is previewable, add it to the list - and also make any expanded scalability cvars previewable
-								if (CVar->TestFlags(EConsoleVariableFlags::ECVF_Preview))
+								if (CVar != nullptr && CVar->TestFlags(EConsoleVariableFlags::ECVF_Preview))
 								{
 									if (bIsScalabilityBucket)
 									{
