We should still load default material's full shader maps from the DDC even if we have single shader DDC turned on.  Since these shaders are required the map will always be complete and it will be much faster.

diff --git a/Engine/Source/Runtime/Engine/Private/Materials/MaterialShared.cpp b/Engine/Source/Runtime/Engine/Private/Materials/MaterialShared.cpp
--- a/Engine/Source/Runtime/Engine/Private/Materials/MaterialShared.cpp
+++ b/Engine/Source/Runtime/Engine/Private/Materials/MaterialShared.cpp
@@ -2235,7 +2235,7 @@ bool FMaterial::CacheShaders(const FMaterialShaderMapId& ShaderMapId, EShaderPla
 			if (!ShaderMap && !FPlatformProperties::RequiresCookedData())
 			{
 				TRefCountPtr<FMaterialShaderMap> LoadedShaderMap;
-				if (bSkipCompilationOnPostLoad == false)
+				if (bSkipCompilationOnPostLoad == false || IsDefaultMaterial())
 				{
 					FMaterialShaderMap::LoadFromDerivedDataCache(this, ShaderMapId, Platform, TargetPlatform, LoadedShaderMap, DDCKeyHash);
 				}
