Updating llm tags for networking
New llm tags for onlinesubsystem/http/websockets

diff --git a/Engine/Plugins/Online/OnlineSubsystem/Source/Private/OnlineAsyncTaskManager.cpp b/Engine/Plugins/Online/OnlineSubsystem/Source/Private/OnlineAsyncTaskManager.cpp
--- a/Engine/Plugins/Online/OnlineSubsystem/Source/Private/OnlineAsyncTaskManager.cpp
+++ b/Engine/Plugins/Online/OnlineSubsystem/Source/Private/OnlineAsyncTaskManager.cpp
@@ -57,7 +57,7 @@ bool FOnlineAsyncTaskManager::Init(void)
 
 uint32 FOnlineAsyncTaskManager::Run(void)
 {
-	LLM_SCOPE(ELLMTag::Networking);
+	LLM_SCOPE_BYTAG(OnlineSubsystem);
 
 	InvocationCount++;
 	// This should not be set yet
diff --git a/Engine/Plugins/Online/OnlineSubsystem/Source/Private/OnlineSubsystem.cpp b/Engine/Plugins/Online/OnlineSubsystem/Source/Private/OnlineSubsystem.cpp
--- a/Engine/Plugins/Online/OnlineSubsystem/Source/Private/OnlineSubsystem.cpp
+++ b/Engine/Plugins/Online/OnlineSubsystem/Source/Private/OnlineSubsystem.cpp
@@ -30,6 +30,8 @@
 #include "Interfaces/OnlineTournamentInterface.h"
 #include "Interfaces/OnlineStatsInterface.h"
 
+LLM_DEFINE_TAG(OnlineSubsystem);
+
 DEFINE_LOG_CATEGORY(LogOnline);
 DEFINE_LOG_CATEGORY(LogOnlineGame);
 DEFINE_LOG_CATEGORY(LogOnlineChat);
diff --git a/Engine/Plugins/Online/OnlineSubsystem/Source/Public/OnlineSubsystem.h b/Engine/Plugins/Online/OnlineSubsystem/Source/Public/OnlineSubsystem.h
--- a/Engine/Plugins/Online/OnlineSubsystem/Source/Public/OnlineSubsystem.h
+++ b/Engine/Plugins/Online/OnlineSubsystem/Source/Public/OnlineSubsystem.h
@@ -10,6 +10,7 @@
 #include "OnlineSubsystemTypes.h"
 #include "OnlineDelegateMacros.h"
 #include "OnlineSubsystemNames.h"
+#include "HAL/LowLevelMemTracker.h"
 
 class FOnlineNotificationHandler;
 class FOnlineNotificationTransportManager;
@@ -43,6 +44,8 @@ class IOnlineGameActivity;
 class IOnlineGameItemStats;
 class IOnlineGameMatches;
 
+LLM_DECLARE_TAG_API(OnlineSubsystem, ONLINESUBSYSTEM_API);
+
 ONLINESUBSYSTEM_API DECLARE_LOG_CATEGORY_EXTERN(LogOnline, Log, All);
 ONLINESUBSYSTEM_API DECLARE_LOG_CATEGORY_EXTERN(LogOnlineGame, Log, All);
 
diff --git a/Engine/Plugins/Online/OnlineSubsystemUtils/Source/OnlineSubsystemUtils/Private/IpNetDriver.cpp b/Engine/Plugins/Online/OnlineSubsystemUtils/Source/OnlineSubsystemUtils/Private/IpNetDriver.cpp
--- a/Engine/Plugins/Online/OnlineSubsystemUtils/Source/OnlineSubsystemUtils/Private/IpNetDriver.cpp
+++ b/Engine/Plugins/Online/OnlineSubsystemUtils/Source/OnlineSubsystemUtils/Private/IpNetDriver.cpp
@@ -1146,7 +1146,7 @@ bool UIpNetDriver::InitListen( FNetworkNotify* InNotify, FURL& LocalURL, bool bR
 
 void UIpNetDriver::TickDispatch(float DeltaTime)
 {
-	LLM_SCOPE(ELLMTag::Networking);
+	LLM_SCOPE_BYTAG(NetDriver);
 
 	Super::TickDispatch( DeltaTime );
 
diff --git a/Engine/Plugins/Runtime/ReplicationGraph/Source/Private/ReplicationGraph.cpp b/Engine/Plugins/Runtime/ReplicationGraph/Source/Private/ReplicationGraph.cpp
--- a/Engine/Plugins/Runtime/ReplicationGraph/Source/Private/ReplicationGraph.cpp
+++ b/Engine/Plugins/Runtime/ReplicationGraph/Source/Private/ReplicationGraph.cpp
@@ -299,6 +299,8 @@ void UReplicationGraph::Serialize(FArchive& Ar)
 
 void UReplicationGraph::InitForNetDriver(UNetDriver* InNetDriver)
 {
+	LLM_SCOPE_BYTAG(NetRepGraph);
+
 	NetDriver = InNetDriver;
 
 	InitGlobalActorClassSettings();
@@ -569,6 +571,8 @@ void UReplicationGraph::InitializeActorsInWorld(UWorld* InWorld)
 
 void UReplicationGraph::InitializeForWorld(UWorld* World)
 {
+	LLM_SCOPE_BYTAG(NetRepGraph);
+
 	ActiveNetworkActors.Reset();
 	GlobalActorReplicationInfoMap.ResetActorMap();
 
@@ -597,6 +601,7 @@ void UReplicationGraph::InitializeForWorld(UWorld* World)
 
 void UReplicationGraph::AddNetworkActor(AActor* Actor)
 {
+	LLM_SCOPE_BYTAG(NetRepGraph);
 	QUICK_SCOPE_CYCLE_COUNTER(UReplicationGraph_AddNetworkActor);
 
 	if (IsActorValidForReplicationGather(Actor) == false)
@@ -864,6 +869,8 @@ FNativeClassAccumulator NoChangeClassAccumulator;
 
 int32 UReplicationGraph::ServerReplicateActors(float DeltaSeconds)
 {
+	LLM_SCOPE_BYTAG(NetRepGraph);
+
 #if !(UE_BUILD_SHIPPING || UE_BUILD_TEST)
 	if (CVar_RepGraph_Pause)
 	{
diff --git a/Engine/Plugins/Runtime/ReplicationGraph/Source/Private/ReplicationGraphTypes.cpp b/Engine/Plugins/Runtime/ReplicationGraph/Source/Private/ReplicationGraphTypes.cpp
--- a/Engine/Plugins/Runtime/ReplicationGraph/Source/Private/ReplicationGraphTypes.cpp
+++ b/Engine/Plugins/Runtime/ReplicationGraph/Source/Private/ReplicationGraphTypes.cpp
@@ -16,14 +16,17 @@
 #include "Net/UnrealNetwork.h"
 #include "Net/NetworkProfiler.h"
 #include "HAL/LowLevelMemTracker.h"
+#include "HAL/LowLevelMemStats.h"
 #include "UObject/UObjectIterator.h"
 #include "Engine/Level.h"
 #include "Templates/UnrealTemplate.h"
 #include "Misc/CoreDelegates.h"
 
-
 DEFINE_LOG_CATEGORY( LogReplicationGraph );
 
+DECLARE_LLM_MEMORY_STAT(TEXT("NetRepGraph"), STAT_NetRepGraphLLM, STATGROUP_LLMFULL);
+LLM_DEFINE_TAG(NetRepGraph, NAME_None, TEXT("Networking"), GET_STATFNAME(STAT_NetRepGraphLLM), GET_STATFNAME(STAT_NetworkingSummaryLLM));
+
 // --------------------------------------------------------------------------------------------------------------------------------------------
 // --------------------------------------------------------------------------------------------------------------------------------------------
 // Actor List Allocator
diff --git a/Engine/Plugins/Runtime/ReplicationGraph/Source/Public/ReplicationGraphTypes.h b/Engine/Plugins/Runtime/ReplicationGraph/Source/Public/ReplicationGraphTypes.h
--- a/Engine/Plugins/Runtime/ReplicationGraph/Source/Public/ReplicationGraphTypes.h
+++ b/Engine/Plugins/Runtime/ReplicationGraph/Source/Public/ReplicationGraphTypes.h
@@ -13,6 +13,7 @@
 #include "UObject/UObjectHash.h"
 #include "ProfilingDebugging/CsvProfiler.h"
 #include "Engine/NetConnection.h"
+#include "HAL/LowLevelMemTracker.h"
 #include "ReplicationGraphTypes.generated.h"
 
 class AActor;
@@ -25,6 +26,8 @@ struct FActorDestructionInfo;
 
 REPLICATIONGRAPH_API DECLARE_LOG_CATEGORY_EXTERN( LogReplicationGraph, Log, All );
 
+LLM_DECLARE_TAG_API(NetRepGraph, REPLICATIONGRAPH_API);
+
 // Check aliases for within the system. The intention is that these can be flipped to checkSlow once the system is stable.
 
 // Macro enable extra logging/bookkeeping. Disabled in Test/Shipping for perf.
diff --git a/Engine/Source/Runtime/Core/Private/HAL/LowLevelMemTracker.cpp b/Engine/Source/Runtime/Core/Private/HAL/LowLevelMemTracker.cpp
--- a/Engine/Source/Runtime/Core/Private/HAL/LowLevelMemTracker.cpp
+++ b/Engine/Source/Runtime/Core/Private/HAL/LowLevelMemTracker.cpp
@@ -77,6 +77,7 @@ DECLARE_LLM_MEMORY_STAT(TEXT("LLM Overhead"), STAT_LLMOverheadTotal, STATGROUP_L
 
 DEFINE_STAT(STAT_EngineSummaryLLM);
 DEFINE_STAT(STAT_ProjectSummaryLLM);
+DEFINE_STAT(STAT_NetworkingSummaryLLM);
 
 /*
  * LLM stats referenced by ELLMTagNames
diff --git a/Engine/Source/Runtime/Core/Public/HAL/LowLevelMemStats.h b/Engine/Source/Runtime/Core/Public/HAL/LowLevelMemStats.h
--- a/Engine/Source/Runtime/Core/Public/HAL/LowLevelMemStats.h
+++ b/Engine/Source/Runtime/Core/Public/HAL/LowLevelMemStats.h
@@ -24,6 +24,7 @@
 
 	DECLARE_LLM_MEMORY_STAT_EXTERN(TEXT("Engine"), STAT_EngineSummaryLLM, STATGROUP_LLM, CORE_API);
 	DECLARE_LLM_MEMORY_STAT_EXTERN(TEXT("Project"), STAT_ProjectSummaryLLM, STATGROUP_LLM, CORE_API);
+	DECLARE_LLM_MEMORY_STAT_EXTERN(TEXT("Networking"), STAT_NetworkingSummaryLLM, STATGROUP_LLM, CORE_API);
 
 #else
 	#define DECLARE_LLM_MEMORY_STAT(CounterName,StatId,GroupId)
diff --git a/Engine/Source/Runtime/Core/Public/HAL/LowLevelMemTracker.h b/Engine/Source/Runtime/Core/Public/HAL/LowLevelMemTracker.h
--- a/Engine/Source/Runtime/Core/Public/HAL/LowLevelMemTracker.h
+++ b/Engine/Source/Runtime/Core/Public/HAL/LowLevelMemTracker.h
@@ -164,7 +164,7 @@ enum class ELLMTagSet : uint8
 	macro(AudioSynthesis,						"AudioSynthesis",				GET_STATFNAME(STAT_AudioSynthesisLLM),						GET_STATFNAME(STAT_AudioSummaryLLM),			ELLMTag::Audio)\
 	macro(RealTimeCommunications,				"RealTimeCommunications",		GET_STATFNAME(STAT_RealTimeCommunicationsLLM),				NAME_None,										-1)\
 	macro(FName,								"FName",						GET_STATFNAME(STAT_FNameLLM),								GET_STATFNAME(STAT_EngineSummaryLLM),			-1)\
-	macro(Networking,							"Networking",					GET_STATFNAME(STAT_NetworkingLLM),							GET_STATFNAME(STAT_EngineSummaryLLM),			-1)\
+	macro(Networking,							"Networking",					GET_STATFNAME(STAT_NetworkingLLM),							GET_STATFNAME(STAT_NetworkingSummaryLLM),		-1)\
 	macro(Meshes,								"Meshes",						GET_STATFNAME(STAT_MeshesLLM),								GET_STATFNAME(STAT_MeshesSummaryLLM),			-1)\
 	macro(Stats,								"Stats",						GET_STATFNAME(STAT_StatsLLM),								GET_STATFNAME(STAT_EngineSummaryLLM),			-1)\
 	macro(Shaders,								"Shaders",						GET_STATFNAME(STAT_ShadersLLM),								GET_STATFNAME(STAT_EngineSummaryLLM),			-1)\
@@ -223,7 +223,7 @@ enum class ELLMTagSet : uint8
 	macro(ConfigSystem,							"ConfigSystem",					GET_STATFNAME(STAT_ConfigSystemLLM),						GET_STATFNAME(STAT_EngineSummaryLLM),			-1)\
 	macro(InitUObject,							"InitUObject",					GET_STATFNAME(STAT_InitUObjectLLM),							GET_STATFNAME(STAT_EngineSummaryLLM),			-1)\
 	macro(VideoRecording,						"VideoRecording",				GET_STATFNAME(STAT_VideoRecordingLLM),						GET_STATFNAME(STAT_EngineSummaryLLM),			-1)\
-	macro(Replays,								"Replays",						GET_STATFNAME(STAT_ReplaysLLM),								GET_STATFNAME(STAT_EngineSummaryLLM),			-1)\
+	macro(Replays,								"Replays",						GET_STATFNAME(STAT_ReplaysLLM),								GET_STATFNAME(STAT_NetworkingSummaryLLM),		ELLMTag::Networking)\
 	macro(MaterialInstance,						"MaterialInstance",				GET_STATFNAME(STAT_MaterialInstanceLLM),					GET_STATFNAME(STAT_EngineSummaryLLM),			-1)\
 	macro(SkeletalMesh,							"SkeletalMesh",					GET_STATFNAME(STAT_SkeletalMeshLLM),						GET_STATFNAME(STAT_EngineSummaryLLM),			ELLMTag::Meshes)\
 	macro(InstancedMesh,						"InstancedMesh",				GET_STATFNAME(STAT_InstancedMeshLLM),						GET_STATFNAME(STAT_EngineSummaryLLM),			ELLMTag::Meshes)\
diff --git a/Engine/Source/Runtime/Engine/Classes/Engine/Channel.h b/Engine/Source/Runtime/Engine/Classes/Engine/Channel.h
--- a/Engine/Source/Runtime/Engine/Classes/Engine/Channel.h
+++ b/Engine/Source/Runtime/Engine/Classes/Engine/Channel.h
@@ -6,8 +6,11 @@
 #include "UObject/ObjectMacros.h"
 #include "UObject/Object.h"
 #include "UObject/CoreNet.h"
+#include "HAL/LowLevelMemTracker.h"
 #include "Channel.generated.h"
 
+LLM_DECLARE_TAG_API(NetChannel, ENGINE_API);
+
 class FInBunch;
 class FOutBunch;
 class UNetConnection;
diff --git a/Engine/Source/Runtime/Engine/Classes/Engine/NetConnection.h b/Engine/Source/Runtime/Engine/Classes/Engine/NetConnection.h
--- a/Engine/Source/Runtime/Engine/Classes/Engine/NetConnection.h
+++ b/Engine/Source/Runtime/Engine/Classes/Engine/NetConnection.h
@@ -31,11 +31,14 @@
 #include "Net/Core/Connection/NetCloseResult.h"
 #include "Net/NetConnectionFaultRecovery.h"
 #include "Net/TrafficControl.h"
+#include "HAL/LowLevelMemTracker.h"
 
 #include "NetConnection.generated.h"
 
 #define NETCONNECTION_HAS_SETENCRYPTIONKEY 1
 
+LLM_DECLARE_TAG_API(NetConnection, ENGINE_API);
+
 class FInternetAddr;
 class FObjectReplicator;
 class StatelessConnectHandlerComponent;
diff --git a/Engine/Source/Runtime/Engine/Classes/Engine/NetDriver.h b/Engine/Source/Runtime/Engine/Classes/Engine/NetDriver.h
--- a/Engine/Source/Runtime/Engine/Classes/Engine/NetDriver.h
+++ b/Engine/Source/Runtime/Engine/Classes/Engine/NetDriver.h
@@ -18,7 +18,7 @@
 #include "IPAddress.h"
 #include "Net/NetAnalyticsTypes.h"
 #include "Net/NetConnectionIdHandler.h"
-
+#include "HAL/LowLevelMemTracker.h"
 #include "NetDriver.generated.h"
 
 /**
@@ -313,6 +313,8 @@
  *
  */
 
+LLM_DECLARE_TAG_API(NetDriver, ENGINE_API);
+
 class Error;
 class FNetGUIDCache;
 struct FNetSyncLoadReport;
diff --git a/Engine/Source/Runtime/Engine/Classes/Engine/PackageMapClient.h b/Engine/Source/Runtime/Engine/Classes/Engine/PackageMapClient.h
--- a/Engine/Source/Runtime/Engine/Classes/Engine/PackageMapClient.h
+++ b/Engine/Source/Runtime/Engine/Classes/Engine/PackageMapClient.h
@@ -29,8 +29,11 @@
 #include "Net/DataBunch.h"
 #include "Net/NetAnalyticsTypes.h"
 #include "ProfilingDebugging/CsvProfiler.h"
+#include "HAL/LowLevelMemTracker.h"
 #include "PackageMapClient.generated.h"
 
+LLM_DECLARE_TAG_API(GuidCache, ENGINE_API);
+
 class UNetConnection;
 class UNetDriver;
 
diff --git a/Engine/Source/Runtime/Engine/Private/DataChannel.cpp b/Engine/Source/Runtime/Engine/Private/DataChannel.cpp
--- a/Engine/Source/Runtime/Engine/Private/DataChannel.cpp
+++ b/Engine/Source/Runtime/Engine/Private/DataChannel.cpp
@@ -25,6 +25,7 @@
 #include "Net/Core/Trace/NetTrace.h"
 #include "Misc/NetworkVersion.h"
 #include "Net/Core/PushModel/PushModel.h"
+#include "HAL/LowLevelMemStats.h"
 
 DEFINE_LOG_CATEGORY(LogNet);
 DEFINE_LOG_CATEGORY(LogRep);
@@ -40,6 +41,9 @@ DECLARE_CYCLE_STAT(TEXT("ActorChan_PostNetInit"), Stat_PostNetInit, STATGROUP_Ne
 DECLARE_CYCLE_STAT(TEXT("Channel ReceivedRawBunch"), Stat_ChannelReceivedRawBunch, STATGROUP_Net);
 DECLARE_CYCLE_STAT(TEXT("ActorChan_FindOrCreateRep"), Stat_ActorChanFindOrCreateRep, STATGROUP_Net);
 
+DECLARE_LLM_MEMORY_STAT(TEXT("NetChannel"), STAT_NetChannelLLM, STATGROUP_LLMFULL);
+LLM_DEFINE_TAG(NetChannel, NAME_None, TEXT("Networking"), GET_STATFNAME(STAT_NetChannelLLM), GET_STATFNAME(STAT_NetworkingSummaryLLM));
+
 extern int32 GDoReplicationContextString;
 extern int32 GNetDormancyValidate;
 extern bool GbNetReuseReplicatorsForDormantObjects;
@@ -1005,6 +1009,8 @@ void UActorChannel::AppendMustBeMappedGuids( FOutBunch* Bunch )
 
 FPacketIdRange UChannel::SendBunch( FOutBunch* Bunch, bool Merge )
 {
+	LLM_SCOPE_BYTAG(NetChannel);
+
 	using namespace UE::Net;
 
 	if (!ensure(ChIndex != -1))
@@ -1758,6 +1764,8 @@ void UControlChannel::ReceivedBunch( FInBunch& Bunch )
 
 void UControlChannel::QueueMessage(const FOutBunch* Bunch)
 {
+	LLM_SCOPE_BYTAG(NetChannel);
+
 	if (QueuedMessages.Num() >= MAX_QUEUED_CONTROL_MESSAGES)
 	{
 		// we're out of room in our extra buffer as well, so kill the connection
@@ -2078,6 +2086,8 @@ void UActorChannel::CleanupReplicators(const bool bKeepReplicators)
 		// so it won't actually check to see whether or not the Object was marked PendingKill.
 		if (bKeepReplicators && CompIt.Value()->GetObject() != nullptr)
 		{
+			LLM_SCOPE_BYTAG(NetConnection);
+
 			// If we want to keep the replication state of the actor/sub-objects around, transfer ownership to the connection
 			// This way, if this actor opens another channel on this connection, we can reclaim or use this replicator to compare state, etc.
 			// For example, we may want to see if any state changed since the actor went dormant, and is now active again. 
@@ -2176,6 +2186,8 @@ bool UActorChannel::CleanUp(const bool bForDestroy, EChannelCloseReason CloseRea
 
 	if (!bIsServer && QueuedBunches.Num() > 0 && ChIndex >= 0 && !bForDestroy)
 	{
+		LLM_SCOPE_BYTAG(NetConnection);
+
 		checkf(ActorNetGUID.IsValid(), TEXT("UActorChannel::Cleanup: ActorNetGUID is invalid! Channel: %i"), ChIndex);
 		
 		TArray<UActorChannel*>& ChannelsStillProcessing = Connection->KeepProcessingActorChannelBunchesMap.FindOrAdd(ActorNetGUID);
@@ -2619,6 +2631,8 @@ bool UActorChannel::ProcessQueuedBunches()
 
 void UActorChannel::ReceivedBunch( FInBunch & Bunch )
 {
+	LLM_SCOPE_BYTAG(NetChannel);
+
 	using namespace UE::Net;
 
 	check(!Closing);
@@ -2964,6 +2978,7 @@ void UActorChannel::ProcessBunch( FInBunch & Bunch )
 
 		if ( bHasUnmapped )
 		{
+			LLM_SCOPE_BYTAG(NetDriver);
 			Connection->Driver->UnmappedReplicators.Add( &Replicator.Get() );
 		}
 	}
@@ -3358,6 +3373,7 @@ int64 UActorChannel::ReplicateActor()
 
 			if (Actor->bNetTemporary)
 			{
+				LLM_SCOPE_BYTAG(NetConnection);
 				Connection->SentTemporaries.Add(Actor);
 			}
 		}
@@ -3917,8 +3933,9 @@ UObject* UActorChannel::ReadContentBlockHeader(FInBunch& Bunch, bool& bObjectDel
 		// Add this sub-object to the ImportedNetGuids list so we can possibly map this object if needed
 		if (ensureMsgf(NetGUID.IsValid(), TEXT("Channel tried to add an invalid GUID to the import list: %s"), *Describe()))
 		{
-		Connection->Driver->GuidCache->ImportedNetGuids.Add( NetGUID );
-	}
+			LLM_SCOPE_BYTAG(GuidCache);
+			Connection->Driver->GuidCache->ImportedNetGuids.Add( NetGUID );
+		}
 	}
 
 	return SubObj;
diff --git a/Engine/Source/Runtime/Engine/Private/DataReplication.cpp b/Engine/Source/Runtime/Engine/Private/DataReplication.cpp
--- a/Engine/Source/Runtime/Engine/Private/DataReplication.cpp
+++ b/Engine/Source/Runtime/Engine/Private/DataReplication.cpp
@@ -21,6 +21,10 @@
 #include "Misc/ScopeExit.h"
 #include "Net/Core/Trace/NetTrace.h"
 #include "ProfilingDebugging/CsvProfiler.h"
+#include "HAL/LowLevelMemStats.h"
+
+DECLARE_LLM_MEMORY_STAT(TEXT("NetObjReplicator"), STAT_NetObjReplicatorLLM, STATGROUP_LLMFULL);
+LLM_DEFINE_TAG(NetObjReplicator, NAME_None, TEXT("Networking"), GET_STATFNAME(STAT_NetObjReplicatorLLM), GET_STATFNAME(STAT_NetworkingSummaryLLM));
 
 DECLARE_CYCLE_STAT(TEXT("Custom Delta Property Rep Time"), STAT_NetReplicateCustomDeltaPropTime, STATGROUP_Game);
 DECLARE_CYCLE_STAT(TEXT("ReceiveRPC"), STAT_NetReceiveRPC, STATGROUP_Game);
@@ -480,6 +484,8 @@ bool FObjectReplicator::ValidateAgainstState( const UObject* ObjectState )
 
 void FObjectReplicator::InitWithObject(UObject* InObject, UNetConnection* InConnection, bool bUseDefaultState)
 {
+	LLM_SCOPE_BYTAG(NetObjReplicator);
+
 	check(GetObject() == nullptr);
 	check(ObjectClass == nullptr);
 	check(bLastUpdateEmpty == false);
@@ -529,7 +535,10 @@ void FObjectReplicator::InitWithObject(UObject* InObject, UNetConnection* InConn
 
 	InitRecentProperties(Source);
 
-	Connection->Driver->AllOwnedReplicators.Add(this);
+	{
+		LLM_SCOPE_BYTAG(NetDriver);
+		Connection->Driver->AllOwnedReplicators.Add(this);
+	}
 
 	if (GbPushModelSkipUndirtiedFastArrays)
 	{
@@ -1384,6 +1393,7 @@ bool FObjectReplicator::ReceivedRPC(FNetBitReader& Reader, const FReplicationFla
 
 void FObjectReplicator::UpdateGuidToReplicatorMap()
 {
+	LLM_SCOPE_BYTAG(NetObjReplicator);
 	SCOPE_CYCLE_COUNTER(STAT_NetUpdateGuidToReplicatorMap);
 
 	if (Connection->Driver->IsServer())
@@ -1430,6 +1440,7 @@ void FObjectReplicator::UpdateGuidToReplicatorMap()
 	{
 		if (!ReferencedGuids.Contains(GUID))
 		{
+			LLM_SCOPE_BYTAG(NetDriver);
 			Connection->Driver->GuidToReplicatorMap.FindOrAdd(GUID).Add(this);
 		}
 	}
@@ -1861,6 +1872,8 @@ void FObjectReplicator::ForceRefreshUnreliableProperties()
 
 void FObjectReplicator::PostSendBunch( FPacketIdRange & PacketRange, uint8 bReliable )
 {
+	LLM_SCOPE_BYTAG(NetObjReplicator);
+
 	const UObject* Object = GetObject();
 
 	if ( Object == nullptr )
@@ -2007,6 +2020,8 @@ void FObjectReplicator::CountBytes(FArchive& Ar) const
 
 void FObjectReplicator::QueueRemoteFunctionBunch( UFunction* Func, FOutBunch &Bunch )
 {
+	LLM_SCOPE_BYTAG(NetObjReplicator);
+
 	if (Connection == nullptr)
 	{
 		return;
@@ -2324,6 +2339,8 @@ void FObjectReplicator::QueuePropertyRepNotify(
 	const int32 ElementIndex,
 	TArray<uint8>& MetaData)
 {
+	LLM_SCOPE_BYTAG(NetObjReplicator);
+
 	if (!Property->HasAnyPropertyFlags(CPF_RepNotify))
 	{
 		return;
diff --git a/Engine/Source/Runtime/Engine/Private/DemoNetDriver.cpp b/Engine/Source/Runtime/Engine/Private/DemoNetDriver.cpp
--- a/Engine/Source/Runtime/Engine/Private/DemoNetDriver.cpp
+++ b/Engine/Source/Runtime/Engine/Private/DemoNetDriver.cpp
@@ -1111,7 +1111,7 @@ double GTickFlushDemoDriverTimeSeconds = 0.0;
 
 void UDemoNetDriver::TickFlushInternal(float DeltaSeconds)
 {
-	LLM_SCOPE(ELLMTag::Networking);
+	LLM_SCOPE(ELLMTag::Replays);
 	CSV_SCOPED_TIMING_STAT_EXCLUSIVE(DemoRecording);
 
 	GTickFlushDemoDriverTimeSeconds = 0.0;
@@ -1198,7 +1198,7 @@ void UDemoNetDriver::TickFlushInternal(float DeltaSeconds)
 
 void UDemoNetDriver::TickDispatch(float DeltaSeconds)
 {
-	LLM_SCOPE(ELLMTag::Networking);
+	LLM_SCOPE(ELLMTag::Replays);
 
 	// Set the context on the world for this driver's level collection.
 	const int32 FoundCollectionIndex = World ? World->GetLevelCollections().IndexOfByPredicate([this](const FLevelCollection& Collection)
@@ -2442,7 +2442,7 @@ private:
 
 void UDemoNetDriver::TickDemoPlayback(float DeltaSeconds)
 {
-	LLM_SCOPE(ELLMTag::Networking);
+	LLM_SCOPE(ELLMTag::Replays);
 	SCOPED_NAMED_EVENT(UDemoNetDriver_TickDemoPlayback, FColor::Purple);
 	if (World && World->IsInSeamlessTravel())
 	{
diff --git a/Engine/Source/Runtime/Engine/Private/GameInstance.cpp b/Engine/Source/Runtime/Engine/Private/GameInstance.cpp
--- a/Engine/Source/Runtime/Engine/Private/GameInstance.cpp
+++ b/Engine/Source/Runtime/Engine/Private/GameInstance.cpp
@@ -1103,7 +1103,7 @@ private:
 
 bool UGameInstance::PlayReplayPlaylist(const FReplayPlaylistParams& PlaylistParams)
 {
-	LLM_SCOPE(ELLMTag::Networking);
+	LLM_SCOPE(ELLMTag::Replays);
 
 	return FGameInstanceReplayPlaylistHelper::StartReplay(PlaylistParams, this);
 }
diff --git a/Engine/Source/Runtime/Engine/Private/NetConnection.cpp b/Engine/Source/Runtime/Engine/Private/NetConnection.cpp
--- a/Engine/Source/Runtime/Engine/Private/NetConnection.cpp
+++ b/Engine/Source/Runtime/Engine/Private/NetConnection.cpp
@@ -39,6 +39,11 @@
 #include "SocketSubsystem.h"
 #include "Math/NumericLimits.h"
 #include "UObject/UnrealNames.h"
+#include "HAL/LowLevelMemStats.h"
+
+DECLARE_LLM_MEMORY_STAT(TEXT("NetConnection"), STAT_NetConnectionLLM, STATGROUP_LLMFULL);
+LLM_DEFINE_TAG(NetConnection, NAME_None, TEXT("Networking"), GET_STATFNAME(STAT_NetConnectionLLM), GET_STATFNAME(STAT_NetworkingSummaryLLM));
+
 
 static TAutoConsoleVariable<int32> CVarPingExcludeFrameTime(TEXT("net.PingExcludeFrameTime"), 0,
 	TEXT("If true, game frame times are subtracted from calculated ping to approximate actual network ping"));
@@ -297,7 +302,7 @@ UNetConnection::UNetConnection(const FObjectInitializer& ObjectInitializer)
 	// This isn't ideal, because it won't capture memory derived classes are creating dynamically.
 	// The allocations could *probably* be moved somewhere else (like InitBase), but that
 	// causes failure to connect for some reason, and for now this is easier.
-	LLM_SCOPE(ELLMTag::Networking);
+	LLM_SCOPE_BYTAG(NetConnection);
 
 	MaxChannelSize = CVarMaxChannelSize.GetValueOnAnyThread();
 	if (MaxChannelSize <= 0)
diff --git a/Engine/Source/Runtime/Engine/Private/NetDriver.cpp b/Engine/Source/Runtime/Engine/Private/NetDriver.cpp
--- a/Engine/Source/Runtime/Engine/Private/NetDriver.cpp
+++ b/Engine/Source/Runtime/Engine/Private/NetDriver.cpp
@@ -13,7 +13,7 @@
 #include "Misc/App.h"
 #include "Misc/MemStack.h"
 #include "HAL/IConsoleManager.h"
-#include "HAL/LowLevelMemTracker.h"
+#include "HAL/LowLevelMemStats.h"
 #include "UObject/ObjectMacros.h"
 #include "UObject/Object.h"
 #include "UObject/Class.h"
@@ -80,6 +80,9 @@
 #include "Editor.h"
 #endif
 
+DECLARE_LLM_MEMORY_STAT(TEXT("NetDriver"), STAT_NetDriverLLM, STATGROUP_LLMFULL);
+LLM_DEFINE_TAG(NetDriver, NAME_None, TEXT("Networking"), GET_STATFNAME(STAT_NetDriverLLM), GET_STATFNAME(STAT_NetworkingSummaryLLM));
+
 // Default net driver stats
 DEFINE_STAT(STAT_Ping);
 DEFINE_STAT(STAT_Channels);
@@ -499,6 +502,8 @@ void UNetDriver::PostReloadConfig(FProperty* PropertyToLoad)
 
 void UNetDriver::LoadChannelDefinitions()
 {
+	LLM_SCOPE_BYTAG(NetDriver);
+
 	if (!HasAnyFlags(RF_ClassDefaultObject))
 	{
 		if (ServerConnection != nullptr)
@@ -687,6 +692,8 @@ bool ShouldEnableScopeSecondsTimers()
 
 void UNetDriver::TickFlush(float DeltaSeconds)
 {
+	LLM_SCOPE_BYTAG(NetDriver);
+
 	TGuardValue<bool> GuardInNetTick(bInTick, true);
 
 	CSV_SCOPED_TIMING_STAT_EXCLUSIVE(NetworkOutgoing);
@@ -3425,6 +3432,8 @@ void UNetDriver::RemoveNetworkActor(AActor* Actor)
 
 void UNetDriver::NotifyActorRenamed(AActor* ThisActor, FName PreviousName)
 {
+	LLM_SCOPE_BYTAG(NetDriver);
+
 	const bool bIsServer = IsServer();
 	const bool bIsActorStatic = !GuidCache->IsDynamicObject(ThisActor);
 	const bool bActorHasRole = ThisActor->GetRemoteRole() != ROLE_None;
@@ -5293,11 +5302,13 @@ UChannel* UNetDriver::GetOrCreateChannelByName(const FName& ChName)
 
 UChannel* UNetDriver::InternalCreateChannelByName(const FName& ChName)
 {
+	LLM_SCOPE_BYTAG(NetChannel);
 	return NewObject<UChannel>(GetTransientPackage(), ChannelDefinitionMap[ChName].ChannelClass);
 }
 
 void UNetDriver::ReleaseToChannelPool(UChannel* Channel)
 {
+	LLM_SCOPE_BYTAG(NetChannel);
 	check(IsValid(Channel));
 	if (Channel->ChName == NAME_Actor && CVarActorChannelPool.GetValueOnAnyThread() != 0)
 	{
@@ -5505,6 +5516,8 @@ bool UNetDriver::NetObjectIsDynamic(const UObject *Object) const
 
 void UNetDriver::AddClientConnection(UNetConnection* NewConnection)
 {
+	LLM_SCOPE_BYTAG(NetDriver);
+
 	SCOPE_CYCLE_COUNTER(Stat_NetDriverAddClientConnection);
 
 	UE_CLOG(!DDoS.CheckLogRestrictions(), LogNet, Log, TEXT("AddClientConnection: Added client connection: %s"),
@@ -5629,6 +5642,8 @@ void UNetDriver::CreateReplicatedStaticActorDestructionInfo(UNetDriver* NetDrive
 
 void UNetDriver::InitDestroyedStartupActors()
 {
+	LLM_SCOPE_BYTAG(NetDriver);
+
 	if (World)
 	{
 		// add startup actors destroyed before the creation of this net driver
diff --git a/Engine/Source/Runtime/Engine/Private/PackageMapClient.cpp b/Engine/Source/Runtime/Engine/Private/PackageMapClient.cpp
--- a/Engine/Source/Runtime/Engine/Private/PackageMapClient.cpp
+++ b/Engine/Source/Runtime/Engine/Private/PackageMapClient.cpp
@@ -22,6 +22,7 @@
 #include "ProfilingDebugging/ScopedTimers.h"
 #include "GameFramework/GameStateBase.h"
 #include "HAL/LowLevelMemTracker.h"
+#include "HAL/LowLevelMemStats.h"
 #include "Net/Core/Trace/NetTrace.h"
 #include "Engine/DemoNetDriver.h"
 #include "Serialization/MemoryReader.h"
@@ -34,6 +35,9 @@
 #include "UObject/ObjectRedirector.h"
 #endif // WITH_EDITOR
 
+DECLARE_LLM_MEMORY_STAT(TEXT("GuidCache"), STAT_GuidCacheLLM, STATGROUP_LLMFULL);
+LLM_DEFINE_TAG(GuidCache, NAME_None, TEXT("Networking"), GET_STATFNAME(STAT_GuidCacheLLM), GET_STATFNAME(STAT_NetworkingSummaryLLM));
+
 // ( OutPacketId == GUID_PACKET_NOT_ACKED ) == NAK'd		(this GUID is not acked, and is not pending either, so sort of waiting)
 // ( OutPacketId == GUID_PACKET_ACKED )		== FULLY ACK'd	(this GUID is fully acked, and we no longer need to send full path)
 // ( OutPacketId > GUID_PACKET_ACKED )		== PENDING		(this GUID needs to be acked, it has been recently reference, and path was sent)
@@ -412,6 +416,7 @@ bool UPackageMapClient::SerializeNewActor(FArchive& Ar, class UActorChannel *Cha
 		{
 			if (ensureMsgf(NetGUID.IsValid(), TEXT("Channel tried to add an invalid GUID to the import list: %s"), *Channel->Describe()))
 			{
+				LLM_SCOPE_BYTAG(GuidCache);
 				GuidCache->ImportedNetGuids.Add( NetGUID );
 			}
 		}
@@ -449,6 +454,7 @@ bool UPackageMapClient::SerializeNewActor(FArchive& Ar, class UActorChannel *Cha
 		{
 			if (ensureMsgf(NetGUID.IsValid(), TEXT("Channel tried to add an invalid GUID to the import list: %s"), *Channel->Describe()))
 			{
+				LLM_SCOPE_BYTAG(GuidCache);
 				GuidCache->ImportedNetGuids.Add(NetGUID);
 			}
 		}
@@ -2763,7 +2769,7 @@ FNetworkGUID FNetGUIDCache::AssignNewNetGUIDFromPath_Server( const FString& Path
 
 void FNetGUIDCache::RegisterNetGUID_Internal( const FNetworkGUID& NetGUID, const FNetGuidCacheObject& CacheObject )
 {
-	LLM_SCOPE(ELLMTag::Networking);
+	LLM_SCOPE_BYTAG(GuidCache);
 
 	// We're pretty strict in this function, we expect everything to have been handled before we get here
 	check( !ObjectLookup.Contains( NetGUID ) );
@@ -3045,6 +3051,8 @@ PRAGMA_ENABLE_DEPRECATION_WARNINGS
 
 void FNetGUIDCache::StartAsyncLoadingPackage(FNetGuidCacheObject& CacheObject, const FNetworkGUID NetGUID, const bool bWasAlreadyAsyncLoading)
 {
+	LLM_SCOPE_BYTAG(GuidCache);
+
 	// Something else is already async loading this package, calling load again will add our callback to the existing load request
 PRAGMA_DISABLE_DEPRECATION_WARNINGS
 	PendingAsyncPackages.Add(CacheObject.PathName, NetGUID);
@@ -3066,6 +3074,8 @@ PRAGMA_ENABLE_DEPRECATION_WARNINGS
 
 void FNetGUIDCache::AsyncPackageCallback(const FName& PackageName, UPackage * Package, EAsyncLoadingResult::Type Result)
 {
+	LLM_SCOPE_BYTAG(GuidCache);
+
 	check(Package == nullptr || Package->IsFullyLoaded());
 
 	if (FPendingAsyncLoadRequest const * const PendingLoadRequest = PendingAsyncLoadRequests.Find(PackageName))
@@ -3155,7 +3165,7 @@ static bool ObjectLevelHasFinishedLoading( UObject* Object, UNetDriver* Driver )
 
 UObject* FNetGUIDCache::GetObjectFromNetGUID( const FNetworkGUID& NetGUID, const bool bIgnoreMustBeMapped )
 {
-	LLM_SCOPE(ELLMTag::Networking);
+	LLM_SCOPE_BYTAG(GuidCache);
 
 	if ( !ensure( NetGUID.IsValid() ) )
 	{
diff --git a/Engine/Source/Runtime/Engine/Private/ReplayHelper.cpp b/Engine/Source/Runtime/Engine/Private/ReplayHelper.cpp
--- a/Engine/Source/Runtime/Engine/Private/ReplayHelper.cpp
+++ b/Engine/Source/Runtime/Engine/Private/ReplayHelper.cpp
@@ -2079,7 +2079,7 @@ void FReplayHelper::ResetState()
 
 void FReplayHelper::OnLevelAddedToWorld(ULevel* InLevel, UWorld* InWorld)
 {
-	LLM_SCOPE(ELLMTag::Networking);
+	LLM_SCOPE(ELLMTag::Replays);
 
 	if (InLevel && !InLevel->bClientOnlyVisible && (World == InWorld) && HasLevelStreamingFixes() && InWorld->IsPlayingReplay())
 	{
diff --git a/Engine/Source/Runtime/Engine/Private/ReplaySubsystem.cpp b/Engine/Source/Runtime/Engine/Private/ReplaySubsystem.cpp
--- a/Engine/Source/Runtime/Engine/Private/ReplaySubsystem.cpp
+++ b/Engine/Source/Runtime/Engine/Private/ReplaySubsystem.cpp
@@ -99,7 +99,7 @@ void UReplaySubsystem::OnSeamlessLevelTransition(UWorld* CurrentWorld)
 
 void UReplaySubsystem::RecordReplay(const FString& Name, const FString& FriendlyName, const TArray<FString>& AdditionalOptions, TSharedPtr<IAnalyticsProvider> AnalyticsProvider)
 {
-	LLM_SCOPE(ELLMTag::Networking);
+	LLM_SCOPE(ELLMTag::Replays);
 
 	if (FParse::Param(FCommandLine::Get(), TEXT("NOREPLAYS")))
 	{
@@ -215,7 +215,7 @@ void UReplaySubsystem::RecordReplay(const FString& Name, const FString& Friendly
 
 bool UReplaySubsystem::PlayReplay(const FString& Name, UWorld* WorldOverride, const TArray<FString>& AdditionalOptions)
 {
-	LLM_SCOPE(ELLMTag::Networking);
+	LLM_SCOPE(ELLMTag::Replays);
 
 	UWorld* CurrentWorld = WorldOverride != nullptr ? WorldOverride : GetWorld();
 
@@ -380,6 +380,8 @@ bool UReplaySubsystem::IsPlaying() const
 
 void UReplaySubsystem::AddEvent(const FString& Group, const FString& Meta, const TArray<uint8>& Data)
 {
+	LLM_SCOPE(ELLMTag::Replays);
+
 	UWorld* CurrentWorld = GetWorld();
 
 	if (CurrentWorld != nullptr && CurrentWorld->GetDemoNetDriver() != nullptr)
@@ -395,6 +397,8 @@ void UReplaySubsystem::AddEvent(const FString& Group, const FString& Meta, const
 
 void UReplaySubsystem::AddOrUpdateEvent(const FString& EventName, const FString& Group, const FString& Meta, const TArray<uint8>& Data)
 {
+	LLM_SCOPE(ELLMTag::Replays);
+
 	UWorld* CurrentWorld = GetWorld();
 
 	if (CurrentWorld != nullptr && CurrentWorld->GetDemoNetDriver() != nullptr)
@@ -457,6 +461,8 @@ void UReplaySubsystem::RequestCheckpoint()
 
 void UReplaySubsystem::SetExternalDataForObject(UObject* OwningObject, const uint8* Src, const int32 NumBits)
 {
+	LLM_SCOPE(ELLMTag::Replays);
+
 	UWorld* CurrentWorld = GetWorld();
 
 	if (CurrentWorld != nullptr && CurrentWorld->GetDemoNetDriver() != nullptr)
diff --git a/Engine/Source/Runtime/Engine/Public/Net/DataReplication.h b/Engine/Source/Runtime/Engine/Public/Net/DataReplication.h
--- a/Engine/Source/Runtime/Engine/Public/Net/DataReplication.h
+++ b/Engine/Source/Runtime/Engine/Public/Net/DataReplication.h
@@ -12,6 +12,9 @@
 #include "Engine/EngineTypes.h"
 #include "UObject/UnrealType.h"
 #include "UObject/GCObject.h"
+#include "HAL/LowLevelMemTracker.h"
+
+LLM_DECLARE_TAG_API(NetObjReplicator, ENGINE_API);
 
 class FNetFieldExportGroup;
 class FOutBunch;
diff --git a/Engine/Source/Runtime/NetworkReplayStreaming/HttpNetworkReplayStreaming/Private/HttpNetworkReplayStreaming.cpp b/Engine/Source/Runtime/NetworkReplayStreaming/HttpNetworkReplayStreaming/Private/HttpNetworkReplayStreaming.cpp
--- a/Engine/Source/Runtime/NetworkReplayStreaming/HttpNetworkReplayStreaming/Private/HttpNetworkReplayStreaming.cpp
+++ b/Engine/Source/Runtime/NetworkReplayStreaming/HttpNetworkReplayStreaming/Private/HttpNetworkReplayStreaming.cpp
@@ -2002,7 +2002,7 @@ bool FHttpNetworkReplayStreamer::DecompressResponse(FHttpResponsePtr HttpRespons
 
 void FHttpNetworkReplayStreamer::HttpDownloadFinished( FHttpRequestPtr HttpRequest, FHttpResponsePtr HttpResponse, bool bSucceeded, int32 RequestedStreamChunkIndex, bool bStreamWasLive )
 {
-	LLM_SCOPE(ELLMTag::Networking);
+	LLM_SCOPE(ELLMTag::Replays);
 
 	RequestFinished( EStreamerState::StreamingDown, EQueuedHttpRequestType::DownloadingStream, HttpRequest );
 
@@ -2482,7 +2482,7 @@ void FHttpNetworkReplayStreamer::ProcessRequestInternal( TSharedPtr< class IHttp
 
 void FHttpNetworkReplayStreamer::Tick( const float DeltaTime )
 {
-	LLM_SCOPE(ELLMTag::Networking);
+	LLM_SCOPE(ELLMTag::Replays);
 
 	// Attempt to process the next http request
 	if ( ProcessNextHttpRequest() )
diff --git a/Engine/Source/Runtime/Online/HTTP/Private/HttpRetrySystem.cpp b/Engine/Source/Runtime/Online/HTTP/Private/HttpRetrySystem.cpp
--- a/Engine/Source/Runtime/Online/HTTP/Private/HttpRetrySystem.cpp
+++ b/Engine/Source/Runtime/Online/HTTP/Private/HttpRetrySystem.cpp
@@ -3,12 +3,15 @@
 #include "HttpRetrySystem.h"
 #include "HAL/PlatformTime.h"
 #include "HAL/PlatformProcess.h"
+#include "HAL/LowLevelMemTracker.h"
 #include "Math/RandomStream.h"
 #include "HttpModule.h"
 #include "Http.h"
 #include "HttpManager.h"
 #include "Stats/Stats.h"
 
+LLM_DEFINE_TAG(HTTP);
+
 namespace FHttpRetrySystem
 {
 	TOptional<double> ReadThrottledTimeFromResponseInSeconds(FHttpResponsePtr Response)
@@ -303,7 +306,7 @@ static FRandomStream TempRandomStream(4435261);
 bool FHttpRetrySystem::FManager::Update(uint32* FileCount, uint32* FailingCount, uint32* FailedCount, uint32* CompletedCount)
 {
 	QUICK_SCOPE_CYCLE_COUNTER(STAT_FHttpRetrySystem_FManager_Update);
-	LLM_SCOPE(ELLMTag::Networking);
+	LLM_SCOPE_BYTAG(HTTP);
 
 	bool bIsGreen = true;
 
diff --git a/Engine/Source/Runtime/Online/WebSockets/Private/Lws/LwsWebSocketsManager.cpp b/Engine/Source/Runtime/Online/WebSockets/Private/Lws/LwsWebSocketsManager.cpp
--- a/Engine/Source/Runtime/Online/WebSockets/Private/Lws/LwsWebSocketsManager.cpp
+++ b/Engine/Source/Runtime/Online/WebSockets/Private/Lws/LwsWebSocketsManager.cpp
@@ -20,6 +20,8 @@
 #include "Misc/Fork.h"
 #include "Stats/Stats.h"
 
+LLM_DEFINE_TAG(WebSockets);
+
 namespace {
 	static const struct lws_extension LwsExtensions[] = {
 		{
@@ -357,7 +359,7 @@ int FLwsWebSocketsManager::CallbackWrapper(lws* Connection, lws_callback_reasons
 void FLwsWebSocketsManager::Tick()
 {
 	QUICK_SCOPE_CYCLE_COUNTER(STAT_FLwsWebSocketsManager_Tick);
-	LLM_SCOPE(ELLMTag::Networking);
+	LLM_SCOPE_BYTAG(WebSockets);
 
 	{
 		FLwsWebSocket* SocketToStart;
