Fix for USoundWave invalidating compressed audio data, but not re-initializing it when Compression Quality parameter changes.

diff --git a/Engine/Source/Runtime/Engine/Private/SoundWave.cpp b/Engine/Source/Runtime/Engine/Private/SoundWave.cpp
--- a/Engine/Source/Runtime/Engine/Private/SoundWave.cpp
+++ b/Engine/Source/Runtime/Engine/Private/SoundWave.cpp
@@ -2229,7 +2229,9 @@ void USoundWave::PostEditChangeProperty(FPropertyChangedEvent& PropertyChangedEv
 				UpdatePlatformData();
 				MarkPackageDirty();
 
-				if (Name == LoadingBehaviorFName && LoadingBehavior == ESoundWaveLoadingBehavior::ForceInline)
+				// if we are force inline, we need to make sure the shared data is pulled from the DDC
+				// before we attempt to use a decoder on the proxy (not using stream caching)
+				if (LoadingBehavior == ESoundWaveLoadingBehavior::ForceInline)
 				{
 					if (GEngine)
 					{
@@ -2238,8 +2240,6 @@ void USoundWave::PostEditChangeProperty(FPropertyChangedEvent& PropertyChangedEv
 						{
 							FName RuntimeFormat = LocalAudioDevice->GetRuntimeFormat(this);
 
-							// if we are force inline, we need to make sure the shared data is pulled from the DDC
-							// before we attempt to use a decoder on the proxy (not using stream caching)
 							if (LoadingBehavior == ESoundWaveLoadingBehavior::ForceInline && !GetResourceData())
 							{
 								InitAudioResource(RuntimeFormat);
