Making ESplitScreenType::Type a BlueprintType

diff --git a/Engine/Source/Runtime/Engine/Classes/Engine/ViewportSplitScreen.h b/Engine/Source/Runtime/Engine/Classes/Engine/ViewportSplitScreen.h
--- a/Engine/Source/Runtime/Engine/Classes/Engine/ViewportSplitScreen.h
+++ b/Engine/Source/Runtime/Engine/Classes/Engine/ViewportSplitScreen.h
@@ -15,6 +15,7 @@
 /**
  * Enum of the different splitscreen types
  */
+UENUM(BlueprintType)
 namespace ESplitScreenType
 {
 	enum Type
@@ -42,10 +43,6 @@ namespace ESplitScreenType
 
 		SplitTypeCount
 	};
-
-	// Deprecated old FourPlayer grid enum value
-	UE_DEPRECATED(4.21, "FourPlayer is now FourPlayer_Grid")
-	const Type FourPlayer = FourPlayer_Grid;
 }
 
 /** Structure to store splitscreen data. */
