Chaos - Remove the static delegate destructor code in the chaos pause handler causing crash on exit.

diff --git a/Engine/Source/Runtime/Engine/Private/PhysicsEngine/Experimental/PhysScene_Chaos.cpp b/Engine/Source/Runtime/Engine/Private/PhysicsEngine/Experimental/PhysScene_Chaos.cpp
--- a/Engine/Source/Runtime/Engine/Private/PhysicsEngine/Experimental/PhysScene_Chaos.cpp
+++ b/Engine/Source/Runtime/Engine/Private/PhysicsEngine/Experimental/PhysScene_Chaos.cpp
@@ -306,12 +306,7 @@ public:
 
 	~FPhysScene_ChaosPauseHandler()
 	{
-		// Remove editor pause/step delegates
-		FEditorDelegates::BeginPIE.RemoveAll(this);
-		FEditorDelegates::EndPIE.RemoveAll(this);
-		FEditorDelegates::PausePIE.RemoveAll(this);
-		FEditorDelegates::ResumePIE.RemoveAll(this);
-		FEditorDelegates::SingleStepPIE.RemoveAll(this);
+		// No need to remove the editor pause/step delegates, it'll be too late for that on exit.
 	}
 
 private:
