Fix a crash in typed element ISM manager selection implementation on AISMPartitionActor when there are multiple client instance managers.

#rnx

diff --git a/Engine/Source/Runtime/Engine/Private/ISMPartition/ISMPartitionActor.cpp b/Engine/Source/Runtime/Engine/Private/ISMPartition/ISMPartitionActor.cpp
--- a/Engine/Source/Runtime/Engine/Private/ISMPartition/ISMPartitionActor.cpp
+++ b/Engine/Source/Runtime/Engine/Private/ISMPartition/ISMPartitionActor.cpp
@@ -685,7 +685,7 @@ void AISMPartitionActor::ForEachClientSMInstance(const FISMClientHandle& Handle,
 	check(Handle.Guid == Clients[Handle.Index]);
 	for (const FISMComponentData& ComponentData : DescriptorComponents)
 	{
-		if (ComponentData.Component)
+		if (ComponentData.Component && ComponentData.ClientInstances.IsValidIndex(Handle.Index))
 		{
 			const FISMClientData& ClientData = ComponentData.ClientInstances[Handle.Index];
 			for (const FISMClientInstance& ClientInstance : ClientData.Instances)
@@ -707,15 +707,18 @@ void AISMPartitionActor::ForEachClientSMInstance(const FISMClientHandle& Handle,
 	check(Handle.Guid == Clients[Handle.Index]);
 	for (const FISMComponentData& ComponentData : DescriptorComponents)
 	{
-		if (ComponentData.Component)
+		if (ComponentData.Component && ComponentData.ClientInstances.IsValidIndex(Handle.Index))
 		{
 			const FISMClientData& ClientData = ComponentData.ClientInstances[Handle.Index];
-			const FISMClientInstance& ClientInstance = ClientData.Instances[InstanceIndex];
-			for (int32 ComponentIndex : ClientInstance.ComponentIndices)
+			if (ClientData.Instances.IsValidIndex(InstanceIndex))
 			{
-				if (!Callback(FSMInstanceId{ ComponentData.Component, ComponentIndex }))
+				const FISMClientInstance& ClientInstance = ClientData.Instances[InstanceIndex];
+				for (int32 ComponentIndex : ClientInstance.ComponentIndices)
 				{
-					return;
+					if (!Callback(FSMInstanceId{ ComponentData.Component, ComponentIndex }))
+					{
+						return;
+					}
 				}
 			}
 		}
