Fix a bug when fetching a non-async-requested texture mip caused by my previous changelist.

diff --git a/Engine/Source/Runtime/Engine/Private/Streaming/Texture2DStreamIn_DDC.cpp b/Engine/Source/Runtime/Engine/Private/Streaming/Texture2DStreamIn_DDC.cpp
--- a/Engine/Source/Runtime/Engine/Private/Streaming/Texture2DStreamIn_DDC.cpp
+++ b/Engine/Source/Runtime/Engine/Private/Streaming/Texture2DStreamIn_DDC.cpp
@@ -315,7 +315,7 @@ void FTexture2DStreamIn_DDC::DoLoadNewMipsFromDDC(const FContext& Context)
 						MipNameBuilder.RemoveSuffix(MipNameBuilder.Len() - TextureNameLen);
 
 						FRequestOwner BlockingRequestOwner(EPriority::Blocking);
-						GetCache().GetChunks({Request}, DDCRequestOwner, [MipIndex, &MipResult](FCacheGetChunkResponse&& Response)
+						GetCache().GetChunks({Request}, BlockingRequestOwner, [MipIndex, &MipResult](FCacheGetChunkResponse&& Response)
 						{
 							if (Response.Status == EStatus::Ok)
 							{
