- Fix code assuming a net id is a v1 net id.

diff --git a/Engine/Plugins/Online/OnlineFramework/Source/Rejoin/Private/RejoinCheck.cpp b/Engine/Plugins/Online/OnlineFramework/Source/Rejoin/Private/RejoinCheck.cpp
--- a/Engine/Plugins/Online/OnlineFramework/Source/Rejoin/Private/RejoinCheck.cpp
+++ b/Engine/Plugins/Online/OnlineFramework/Source/Rejoin/Private/RejoinCheck.cpp
@@ -85,10 +85,10 @@ void URejoinCheck::CheckRejoinStatus(const FOnRejoinCheckComplete& InCompletionD
 			IOnlineIdentityPtr IdentityInt = Online::GetIdentityInterface(World);
 			if (IdentityInt.IsValid())
 			{
-				FUniqueNetIdPtr PrimaryUniqueId = GameInstance->GetPrimaryPlayerUniqueId();
+				FUniqueNetIdRepl PrimaryUniqueId = GameInstance->GetPrimaryPlayerUniqueIdRepl();
 				if (ensure(PrimaryUniqueId.IsValid()))
 				{
-					ULocalPlayer* LP = GameInstance->FindLocalPlayerFromUniqueNetId(*PrimaryUniqueId);
+					ULocalPlayer* LP = GameInstance->FindLocalPlayerFromUniqueNetId(PrimaryUniqueId);
 					if (ensure(LP))
 					{
 						FOnFindFriendSessionCompleteDelegate CompletionDelegate;
diff --git a/Engine/Plugins/Online/OnlineSubsystemUtils/Source/OnlineSubsystemUtils/Private/ShowLoginUICallbackProxy.cpp b/Engine/Plugins/Online/OnlineSubsystemUtils/Source/OnlineSubsystemUtils/Private/ShowLoginUICallbackProxy.cpp
--- a/Engine/Plugins/Online/OnlineSubsystemUtils/Source/OnlineSubsystemUtils/Private/ShowLoginUICallbackProxy.cpp
+++ b/Engine/Plugins/Online/OnlineSubsystemUtils/Source/OnlineSubsystemUtils/Private/ShowLoginUICallbackProxy.cpp
@@ -78,12 +78,12 @@ void UShowLoginUICallbackProxy::OnShowLoginUICompleted(FUniqueNetIdPtr UniqueId,
 		ULocalPlayer* LocalPlayer = MyPlayerController->GetLocalPlayer();
 		if (LocalPlayer != nullptr)
 		{
-			LocalPlayer->SetCachedUniqueNetId(UniqueId);
+			LocalPlayer->SetCachedUniqueNetId(FUniqueNetIdRepl(UniqueId));
 		}
 
 		if (MyPlayerController->PlayerState != nullptr)
 		{
-			MyPlayerController->PlayerState->SetUniqueId(UniqueId);
+			MyPlayerController->PlayerState->SetUniqueId(FUniqueNetIdRepl(UniqueId));
 		}
 	}
 
diff --git a/Engine/Source/Runtime/Engine/Classes/Engine/GameInstance.h b/Engine/Source/Runtime/Engine/Classes/Engine/GameInstance.h
--- a/Engine/Source/Runtime/Engine/Classes/Engine/GameInstance.h
+++ b/Engine/Source/Runtime/Engine/Classes/Engine/GameInstance.h
@@ -353,8 +353,16 @@ public:
 	 *
 	 * @return the unique id of the primary player on this machine
 	 */
+	UE_DEPRECATED(5.0, "Use GetPrimaryPlayerUniqueIdRepl.")
 	FUniqueNetIdPtr GetPrimaryPlayerUniqueId() const;
 
+	/**
+	 * Get the unique id for the primary player on this machine (others are splitscreen children)
+	 *
+	 * @return the unique id of the primary player on this machine
+	 */
+	FUniqueNetIdRepl GetPrimaryPlayerUniqueIdRepl() const;
+
 	void CleanupGameViewport();
 
 	/** Called when demo playback fails for any reason */
diff --git a/Engine/Source/Runtime/Engine/Classes/Engine/LocalPlayer.h b/Engine/Source/Runtime/Engine/Classes/Engine/LocalPlayer.h
--- a/Engine/Source/Runtime/Engine/Classes/Engine/LocalPlayer.h
+++ b/Engine/Source/Runtime/Engine/Classes/Engine/LocalPlayer.h
@@ -473,7 +473,10 @@ public:
 	FUniqueNetIdRepl GetCachedUniqueNetId() const;
 
 	/** Sets the players current cached unique net id */
+	UE_DEPRECATED(5.0, "Use SetCachedUniqueNetId with FUniqueNetIdRepl")
 	void SetCachedUniqueNetId(FUniqueNetIdPtr NewUniqueNetId);
+	/** Sets the players current cached unique net id */
+	void SetCachedUniqueNetId(const FUniqueNetIdRepl& NewUniqueNetId);
 
 	/** 
 	 * Retrieves the preferred unique net id. This is for backwards compatibility for games that don't use the cached unique net id logic
diff --git a/Engine/Source/Runtime/Engine/Classes/Engine/World.h b/Engine/Source/Runtime/Engine/Classes/Engine/World.h
--- a/Engine/Source/Runtime/Engine/Classes/Engine/World.h
+++ b/Engine/Source/Runtime/Engine/Classes/Engine/World.h
@@ -3517,6 +3517,7 @@ public:
 	 * @param InNetPlayerIndex (optional) - the NetPlayerIndex to set on the PlayerController
 	 * @return the PlayerController that was spawned (may fail and return NULL)
 	 */
+	UE_DEPRECATED(5.0, "Use SpawnPlayActor with FUniqueNetIdRepl")
 	APlayerController* SpawnPlayActor(class UPlayer* Player, ENetRole RemoteRole, const FURL& InURL, const FUniqueNetIdPtr& UniqueId, FString& Error, uint8 InNetPlayerIndex = 0);
 	APlayerController* SpawnPlayActor(class UPlayer* Player, ENetRole RemoteRole, const FURL& InURL, const FUniqueNetIdRepl& UniqueId, FString& Error, uint8 InNetPlayerIndex = 0);
 	
diff --git a/Engine/Source/Runtime/Engine/Classes/GameFramework/GameSession.h b/Engine/Source/Runtime/Engine/Classes/GameFramework/GameSession.h
--- a/Engine/Source/Runtime/Engine/Classes/GameFramework/GameSession.h
+++ b/Engine/Source/Runtime/Engine/Classes/GameFramework/GameSession.h
@@ -83,8 +83,17 @@ class ENGINE_API AGameSession : public AInfo
 	 * @param UniqueId uniqueId they sent over on Login
 	 * @param bWasFromInvite was this from an invite
 	 */
+	UE_DEPRECATED(5.0, "Use RegisterPlayer with FUniqueNetIdRepl")
 	virtual void RegisterPlayer(APlayerController* NewPlayer, const FUniqueNetIdPtr& UniqueId, bool bWasFromInvite);
 
+	/**
+	 * Register a player with the online service session
+	 * @param NewPlayer player to register
+	 * @param UniqueId uniqueId they sent over on Login
+	 * @param bWasFromInvite was this from an invite
+	 */
+	virtual void RegisterPlayer(APlayerController* NewPlayer, const FUniqueNetIdRepl& UniqueId, bool bWasFromInvite);
+
 	/**
 	 * Called by GameMode::PostLogin to give session code chance to do work after PostLogin
 	 *
@@ -120,7 +129,16 @@ class ENGINE_API AGameSession : public AInfo
 	 * @param UniqueId id of the player to unregister
 	 */
 	virtual void UnregisterPlayer(FName InSessionName, const FUniqueNetIdRepl& UniqueId);
+
+	/**
+	 * Unregister players from the online service session
+	 *
+	 * @param SessionName name of session to unregister from
+	 * @param Players ids of the players to unregister
+	 */
+	UE_DEPRECATED(5.0, "Use UnregisterPlayers with FUniqueNetIdRepl")
 	virtual void UnregisterPlayers(FName InSessionName, const TArray< FUniqueNetIdRef >& Players);
+	virtual void UnregisterPlayers(FName InSessionName, const TArray< FUniqueNetIdRepl >& Players);
 	
 	/**
 	 * Unregister a player from the online service session
@@ -260,8 +278,14 @@ protected:
  * @param PlayerNetId the id to search for
  * @return the player controller if found, otherwise NULL
  */
+UE_DEPRECATED(5.0, "Use GetPlayerControllerFromNetId with FUniqueNetIdRepl")
 ENGINE_API APlayerController* GetPlayerControllerFromNetId(UWorld* World, const FUniqueNetId& PlayerNetId);
 
-
+/**
+ * Returns the player controller associated with this net id
+ * @param PlayerNetId the id to search for
+ * @return the player controller if found, otherwise NULL
+ */
+ENGINE_API APlayerController* GetPlayerControllerFromNetId(UWorld* World, const FUniqueNetIdRepl& PlayerNetId);
 
 
diff --git a/Engine/Source/Runtime/Engine/Classes/GameFramework/OnlineReplStructs.h b/Engine/Source/Runtime/Engine/Classes/GameFramework/OnlineReplStructs.h
--- a/Engine/Source/Runtime/Engine/Classes/GameFramework/OnlineReplStructs.h
+++ b/Engine/Source/Runtime/Engine/Classes/GameFramework/OnlineReplStructs.h
@@ -30,6 +30,10 @@ struct FUniqueNetIdRepl : public FUniqueNetIdWrapper
 	{
 	}
 
+	FUniqueNetIdRepl(TYPE_OF_NULLPTR)
+	{
+	}
+
 	FUniqueNetIdRepl(const FUniqueNetIdRepl& InWrapper)
 		: FUniqueNetIdWrapper(InWrapper.Variant)
 	{
diff --git a/Engine/Source/Runtime/Engine/Classes/GameFramework/PlayerState.h b/Engine/Source/Runtime/Engine/Classes/GameFramework/PlayerState.h
--- a/Engine/Source/Runtime/Engine/Classes/GameFramework/PlayerState.h
+++ b/Engine/Source/Runtime/Engine/Classes/GameFramework/PlayerState.h
@@ -458,10 +458,19 @@ public:
 	 * Associate an online unique id with this player
 	 * @param InUniqueId the unique id associated with this player
 	 */
+	UE_DEPRECATED(5.0, "Use SetUniqueId(FUniqueNetIdRepl) instead")
 	virtual void SetUniqueId(const FUniqueNetIdPtr& InUniqueId);
 
-	/** Sets the value of UniqueId without causing other side effects to this instance. */
+	/**
+	 * Associate an online unique id with this player
+	 * @param InUniqueId the unique id associated with this player
+	 */
 	void SetUniqueId(const FUniqueNetIdRepl& NewUniqueId);
+	
+	/**
+	 * Associate an online unique id with this player
+	 * @param InUniqueId the unique id associated with this player
+	 */
 	void SetUniqueId(FUniqueNetIdRepl&& NewUniqueId);
 
 	/** Called on both the client and server when unique ID has been modified */
diff --git a/Engine/Source/Runtime/Engine/Private/GameInstance.cpp b/Engine/Source/Runtime/Engine/Private/GameInstance.cpp
--- a/Engine/Source/Runtime/Engine/Private/GameInstance.cpp
+++ b/Engine/Source/Runtime/Engine/Private/GameInstance.cpp
@@ -955,6 +955,12 @@ APlayerController* UGameInstance::GetPrimaryPlayerController(bool bRequiresValid
 }
 
 FUniqueNetIdPtr UGameInstance::GetPrimaryPlayerUniqueId() const
+{
+	FUniqueNetIdRepl UniqueIdRepl = GetPrimaryPlayerUniqueIdRepl();
+	return UniqueIdRepl.GetV1();
+}
+
+FUniqueNetIdRepl UGameInstance::GetPrimaryPlayerUniqueIdRepl() const
 {
 	ULocalPlayer* PrimaryLP = nullptr;
 
@@ -974,7 +980,7 @@ FUniqueNetIdPtr UGameInstance::GetPrimaryPlayerUniqueId() const
 		LocalUserId = PrimaryLP->GetPreferredUniqueNetId();
 	}
 
-	return LocalUserId.GetUniqueNetId();
+	return LocalUserId;
 }
 
 ULocalPlayer* UGameInstance::FindLocalPlayerFromControllerId(const int32 ControllerId) const
diff --git a/Engine/Source/Runtime/Engine/Private/GameMode.cpp b/Engine/Source/Runtime/Engine/Private/GameMode.cpp
--- a/Engine/Source/Runtime/Engine/Private/GameMode.cpp
+++ b/Engine/Source/Runtime/Engine/Private/GameMode.cpp
@@ -708,7 +708,7 @@ bool AGameMode::FindInactivePlayer(APlayerController* PC)
 			OldPlayerState->SetIsInactive(true);
 			// Set the uniqueId to nullptr so it will not kill the player's registration 
 			// in UnregisterPlayerWithSession()
-			OldPlayerState->SetUniqueId(nullptr);
+			OldPlayerState->SetUniqueId(FUniqueNetIdRepl());
 			OldPlayerState->Destroy();
 			PC->PlayerState->OnReactivated();
 			return true;
diff --git a/Engine/Source/Runtime/Engine/Private/GameModeBase.cpp b/Engine/Source/Runtime/Engine/Private/GameModeBase.cpp
--- a/Engine/Source/Runtime/Engine/Private/GameModeBase.cpp
+++ b/Engine/Source/Runtime/Engine/Private/GameModeBase.cpp
@@ -737,7 +737,7 @@ FString AGameModeBase::InitNewPlayer(APlayerController* NewPlayerController, con
 	}
 
 	// Register the player with the session
-	GameSession->RegisterPlayer(NewPlayerController, UniqueId.GetUniqueNetId(), UGameplayStatics::HasOption(Options, TEXT("bIsFromInvite")));
+	GameSession->RegisterPlayer(NewPlayerController, UniqueId, UGameplayStatics::HasOption(Options, TEXT("bIsFromInvite")));
 
 	// Find a starting spot
 	FString ErrorMessage;
diff --git a/Engine/Source/Runtime/Engine/Private/GameSession.cpp b/Engine/Source/Runtime/Engine/Private/GameSession.cpp
--- a/Engine/Source/Runtime/Engine/Private/GameSession.cpp
+++ b/Engine/Source/Runtime/Engine/Private/GameSession.cpp
@@ -18,12 +18,12 @@ DEFINE_LOG_CATEGORY(LogGameSession);
 
 static TAutoConsoleVariable<int32> CVarMaxPlayersOverride( TEXT( "net.MaxPlayersOverride" ), 0, TEXT( "If greater than 0, will override the standard max players count. Useful for testing full servers." ) );
 
-/** 
- * Returns the player controller associated with this net id
- * @param PlayerNetId the id to search for
- * @return the player controller if found, otherwise NULL
- */
 APlayerController* GetPlayerControllerFromNetId(UWorld* World, const FUniqueNetId& PlayerNetId)
+{
+	return GetPlayerControllerFromNetId(World, FUniqueNetIdRepl(PlayerNetId.AsShared()));
+}
+
+APlayerController* GetPlayerControllerFromNetId(UWorld* World, const FUniqueNetIdRepl& PlayerNetId)
 {
 	if (PlayerNetId.IsValid())
 	{
@@ -32,10 +32,10 @@ APlayerController* GetPlayerControllerFromNetId(UWorld* World, const FUniqueNetI
 		{
 			APlayerController* PlayerController = Iterator->Get();
 			// Determine if this is a player with replication
-			if (PlayerController && PlayerController->PlayerState != NULL && PlayerController->PlayerState->GetUniqueId().IsValid())
+			if (PlayerController && PlayerController->PlayerState && PlayerController->PlayerState->GetUniqueId().IsValid())
 			{
 				// If the ids match, then this is the right player.
-				if (*PlayerController->PlayerState->GetUniqueId() == PlayerNetId)
+				if (PlayerController->PlayerState->GetUniqueId() == PlayerNetId)
 				{
 					return PlayerController;
 				}
@@ -228,6 +228,11 @@ int32 AGameSession::GetNextPlayerID()
 }
 
 void AGameSession::RegisterPlayer(APlayerController* NewPlayer, const FUniqueNetIdPtr& UniqueId, bool bWasFromInvite)
+{
+	RegisterPlayer(NewPlayer, FUniqueNetIdRepl(UniqueId), bWasFromInvite);
+}
+
+void AGameSession::RegisterPlayer(APlayerController* NewPlayer, const FUniqueNetIdRepl& UniqueId, bool bWasFromInvite)
 {
 	if (NewPlayer != NULL)
 	{
@@ -243,8 +248,7 @@ void AGameSession::UnregisterPlayer(FName InSessionName, const FUniqueNetIdRepl&
 {
 	UWorld* World = GetWorld();
 	if (GetNetMode() != NM_Standalone &&
-		UniqueId.IsValid() &&
-		UniqueId->IsValid())
+		UniqueId.IsValid())
 	{
 		// Remove the player from the session
 		UOnlineEngineInterface::Get()->UnregisterPlayer(World, InSessionName, UniqueId);
@@ -267,6 +271,22 @@ void AGameSession::UnregisterPlayers(FName InSessionName, const TArray<FUniqueNe
 	}
 }
 
+void AGameSession::UnregisterPlayers(FName InSessionName, const TArray<FUniqueNetIdRepl>& Players)
+{
+	UWorld * World = GetWorld();
+	if (GetNetMode() != NM_Standalone &&
+		Players.Num() > 0)
+	{
+		// Remove the player from the session
+		TArray<FUniqueNetIdWrapper> PlayerIdsAsWrappers;
+		for (const FUniqueNetIdRepl& PlayerId : Players)
+		{
+			PlayerIdsAsWrappers.Emplace(PlayerId);
+		}
+		UOnlineEngineInterface::Get()->UnregisterPlayers(World, InSessionName, PlayerIdsAsWrappers);
+	}
+}
+
 void AGameSession::UnregisterPlayer(const APlayerController* ExitingPlayer)
 {
 	if (GetNetMode() != NM_Standalone &&
diff --git a/Engine/Source/Runtime/Engine/Private/LocalPlayer.cpp b/Engine/Source/Runtime/Engine/Private/LocalPlayer.cpp
--- a/Engine/Source/Runtime/Engine/Private/LocalPlayer.cpp
+++ b/Engine/Source/Runtime/Engine/Private/LocalPlayer.cpp
@@ -1624,6 +1624,11 @@ void ULocalPlayer::SetCachedUniqueNetId(FUniqueNetIdPtr NewUniqueNetId)
 	CachedUniqueNetId = NewUniqueNetId;
 }
 
+void ULocalPlayer::SetCachedUniqueNetId(const FUniqueNetIdRepl& NewUniqueNetId)
+{
+	CachedUniqueNetId = NewUniqueNetId;
+}
+
 FUniqueNetIdRepl ULocalPlayer::GetPreferredUniqueNetId() const
 {
 	// Prefer the cached unique net id (only if it's valid)
diff --git a/Engine/Source/Runtime/Engine/Private/PlayerController.cpp b/Engine/Source/Runtime/Engine/Private/PlayerController.cpp
--- a/Engine/Source/Runtime/Engine/Private/PlayerController.cpp
+++ b/Engine/Source/Runtime/Engine/Private/PlayerController.cpp
@@ -3644,7 +3644,7 @@ bool APlayerController::ServerUnblockPlayer_Validate(FUniqueNetIdRepl PlayerId)
 
 APlayerController* APlayerController::GetPlayerControllerForMuting(const FUniqueNetIdRepl& PlayerNetId)
 {
-	return GetPlayerControllerFromNetId(GetWorld(), *PlayerNetId.GetUniqueNetId());
+	return GetPlayerControllerFromNetId(GetWorld(), PlayerNetId);
 }
 
 bool APlayerController::IsPlayerMuted(const FUniqueNetId& PlayerId)
diff --git a/Engine/Source/Runtime/Engine/Private/PlayerState.cpp b/Engine/Source/Runtime/Engine/Private/PlayerState.cpp
--- a/Engine/Source/Runtime/Engine/Private/PlayerState.cpp
+++ b/Engine/Source/Runtime/Engine/Private/PlayerState.cpp
@@ -157,7 +157,7 @@ void APlayerState::OverrideWith(APlayerState* PlayerState)
 {
 	SetIsSpectator(PlayerState->IsSpectator());
 	SetIsOnlyASpectator(PlayerState->IsOnlyASpectator());
-	SetUniqueId(PlayerState->GetUniqueId().GetUniqueNetId());
+	SetUniqueId(PlayerState->GetUniqueId());
 	SetPlayerNameInternal(PlayerState->GetPlayerName());
 }
 
@@ -168,7 +168,7 @@ void APlayerState::CopyProperties(APlayerState* PlayerState)
 	PlayerState->SetCompressedPing(GetCompressedPing());
 	PlayerState->ExactPing = ExactPing;
 	PlayerState->SetPlayerId(GetPlayerId());
-	PlayerState->SetUniqueId(GetUniqueId().GetUniqueNetId());
+	PlayerState->SetUniqueId(GetUniqueId());
 	PlayerState->SetPlayerNameInternal(GetPlayerName());
 	PlayerState->SetStartTime(GetStartTime());
 	PlayerState->SavedNetworkAddress = SavedNetworkAddress;
