- Fix CIS errors in projects caused by ambiguous calls to SetCachedUniqueNetId(nullptr)

diff --git a/Engine/Source/Runtime/Engine/Classes/Engine/LocalPlayer.h b/Engine/Source/Runtime/Engine/Classes/Engine/LocalPlayer.h
--- a/Engine/Source/Runtime/Engine/Classes/Engine/LocalPlayer.h
+++ b/Engine/Source/Runtime/Engine/Classes/Engine/LocalPlayer.h
@@ -476,6 +476,9 @@ public:
 	UE_DEPRECATED(5.0, "Use SetCachedUniqueNetId with FUniqueNetIdRepl")
 	void SetCachedUniqueNetId(FUniqueNetIdPtr NewUniqueNetId);
 	/** Sets the players current cached unique net id */
+	UE_DEPRECATED(5.0, "Use SetCachedUniqueNetId with FUniqueNetIdRepl")
+	void SetCachedUniqueNetId(TYPE_OF_NULLPTR);
+	/** Sets the players current cached unique net id */
 	void SetCachedUniqueNetId(const FUniqueNetIdRepl& NewUniqueNetId);
 
 	/** 
diff --git a/Engine/Source/Runtime/Engine/Private/LocalPlayer.cpp b/Engine/Source/Runtime/Engine/Private/LocalPlayer.cpp
--- a/Engine/Source/Runtime/Engine/Private/LocalPlayer.cpp
+++ b/Engine/Source/Runtime/Engine/Private/LocalPlayer.cpp
@@ -1624,6 +1624,11 @@ void ULocalPlayer::SetCachedUniqueNetId(FUniqueNetIdPtr NewUniqueNetId)
 	CachedUniqueNetId = NewUniqueNetId;
 }
 
+void ULocalPlayer::SetCachedUniqueNetId(TYPE_OF_NULLPTR)
+{
+	CachedUniqueNetId = FUniqueNetIdRepl(nullptr);
+}
+
 void ULocalPlayer::SetCachedUniqueNetId(const FUniqueNetIdRepl& NewUniqueNetId)
 {
 	CachedUniqueNetId = NewUniqueNetId;
