Fix for Pose Watch/Pose Watch Folders label set to empty

diff --git a/Engine/Source/Editor/Persona/Private/PoseWatchManagerPoseWatchTreeItem.cpp b/Engine/Source/Editor/Persona/Private/PoseWatchManagerPoseWatchTreeItem.cpp
--- a/Engine/Source/Editor/Persona/Private/PoseWatchManagerPoseWatchTreeItem.cpp
+++ b/Engine/Source/Editor/Persona/Private/PoseWatchManagerPoseWatchTreeItem.cpp
@@ -157,7 +157,8 @@ private:
 
 	void OnLabelCommitted(const FText& InLabel, ETextCommit::Type InCommitInfo)
 	{
-		check(WeakPoseWatchPtr->SetLabel(InLabel));
+		bool bRenameSuccessful = WeakPoseWatchPtr->SetLabel(InLabel);
+		check(bRenameSuccessful);
 		WeakPoseWatchManager.Pin()->FullRefresh();
 		WeakPoseWatchManager.Pin()->SetKeyboardFocus();
 	}
diff --git a/Engine/Source/Runtime/Engine/Private/Engine/PoseWatch.cpp b/Engine/Source/Runtime/Engine/Private/Engine/PoseWatch.cpp
--- a/Engine/Source/Runtime/Engine/Private/Engine/PoseWatch.cpp
+++ b/Engine/Source/Runtime/Engine/Private/Engine/PoseWatch.cpp
@@ -167,9 +167,10 @@ void UPoseWatchFolder::MoveTo(UPoseWatchFolder* InFolder)
 
 bool UPoseWatchFolder::SetLabel(const FText& InLabel)
 {
-	if (IsFolderLabelUniqueInFolder(InLabel, Parent.Get()))
+	FText NewLabel = FText::TrimPrecedingAndTrailing(InLabel);
+	if (IsFolderLabelUniqueInFolder(NewLabel, Parent.Get()))
 	{
-		Label = InLabel;
+		Label = NewLabel;
 		return true;
 	}
 	return false;
@@ -288,7 +289,13 @@ bool UPoseWatchFolder::IsAssignedFolder() const
 
 bool UPoseWatchFolder::ValidateLabelRename(const FText& InLabel, FText& OutErrorMessage)
 {
-	if (!IsFolderLabelUniqueInFolder(InLabel, Parent.Get()))
+	FText UseLabel = FText::TrimPrecedingAndTrailing(InLabel);
+	if (UseLabel.IsEmptyOrWhitespace())
+	{
+		OutErrorMessage = LOCTEXT("PoseWatchFolderNameEmpty", "A pose watch folder must have a label");
+		return false;
+	}
+	if (!IsFolderLabelUniqueInFolder(UseLabel, Parent.Get()))
 	{
 		OutErrorMessage = LOCTEXT("PoseWatchFolderNameTaken", "A folder already has this name at this level");
 		return false;
@@ -448,15 +455,10 @@ void UPoseWatch::MoveTo(UPoseWatchFolder* InFolder)
 
 bool UPoseWatch::SetLabel(const FText& InLabel)
 {
-	UPoseWatch* ConflictingPoseWatch = nullptr;
-	if (IsPoseWatchLabelUniqueInFolder(InLabel, Parent.Get()))
+	FText NewLabel = FText::TrimPrecedingAndTrailing(InLabel);
+	if (IsPoseWatchLabelUniqueInFolder(NewLabel, Parent.Get()))
 	{
-		Label = InLabel;
-		return true;
-	}
-	else if (ConflictingPoseWatch == this)
-	{
-		Label = InLabel;
+		Label = NewLabel;
 		return true;
 	}
 	return false;
@@ -499,14 +501,16 @@ bool UPoseWatch::IsAssignedFolder() const
 
 bool UPoseWatch::ValidateLabelRename(const FText& InLabel, FText& OutErrorMessage)
 {
-	UPoseWatch* ConflictingPoseWatch = nullptr;
-	if (!IsPoseWatchLabelUniqueInFolder(InLabel, Parent.Get()))
+	FText UseLabel = FText::TrimPrecedingAndTrailing(InLabel);
+	if (UseLabel.IsEmptyOrWhitespace())
 	{
-		if (ConflictingPoseWatch != this)
-		{
-			OutErrorMessage = LOCTEXT("PoseWatchNameTaken", "A pose watch already has this name at this level");
-			return false;
-		}
+		OutErrorMessage = LOCTEXT("PoseWatchNameEmpty", "A pose watch must have a label");
+		return false;
+	}
+	if (!IsPoseWatchLabelUniqueInFolder(UseLabel, Parent.Get()))
+	{
+		OutErrorMessage = LOCTEXT("PoseWatchNameTaken", "A pose watch already has this name at this level");
+		return false;
 	}
 	return true;
 }
