Allow actors to have default expansion state in the world outliner

diff --git a/Engine/Source/Editor/SceneOutliner/Private/ActorTreeItem.cpp b/Engine/Source/Editor/SceneOutliner/Private/ActorTreeItem.cpp
--- a/Engine/Source/Editor/SceneOutliner/Private/ActorTreeItem.cpp
+++ b/Engine/Source/Editor/SceneOutliner/Private/ActorTreeItem.cpp
@@ -341,6 +341,8 @@ FActorTreeItem::FActorTreeItem(AActor* InActor)
 {
 	check(InActor);
 	ActorLabel = InActor->GetActorLabel();
+
+	Flags.bIsExpanded = InActor->bDefaultOutlinerExpansionState;
 	
 	bExistsInCurrentWorldAndPIE = GEditor->ObjectsThatExistInEditorWorld.Get(InActor);
 }
diff --git a/Engine/Source/Runtime/Engine/Classes/GameFramework/Actor.h b/Engine/Source/Runtime/Engine/Classes/GameFramework/Actor.h
--- a/Engine/Source/Runtime/Engine/Classes/GameFramework/Actor.h
+++ b/Engine/Source/Runtime/Engine/Classes/GameFramework/Actor.h
@@ -968,6 +968,9 @@ public:
 	/** If true during PostEditMove the construction script will be run every time. If false it will only run when the drag finishes. */
 	uint8 bRunConstructionScriptOnDrag:1;
 
+	/** Default expansion state for this actor. Some actors have attached children that we may not want to automatically expand by default */
+	uint8 bDefaultOutlinerExpansionState : 1;
+
 protected:
 	/** If true, prevents the actor from being moved in the editor viewport. */
 	UPROPERTY()
diff --git a/Engine/Source/Runtime/Engine/Private/Actor.cpp b/Engine/Source/Runtime/Engine/Private/Actor.cpp
--- a/Engine/Source/Runtime/Engine/Private/Actor.cpp
+++ b/Engine/Source/Runtime/Engine/Private/Actor.cpp
@@ -123,6 +123,7 @@ void AActor::InitializeDefaults()
 #if WITH_EDITORONLY_DATA
 	bEditable = true;
 	bListedInSceneOutliner = true;
+	bDefaultOutlinerExpansionState = true;
 	bIsEditorPreviewActor = false;
 	bHiddenEdLayer = false;
 	bHiddenEdTemporary = false;
