Fix for ForceCast applied to LWCScalars, should replicate rather than append zeroes, unless an exact match is requested

diff --git a/Engine/Source/Runtime/Engine/Private/Materials/HLSLMaterialTranslator.cpp b/Engine/Source/Runtime/Engine/Private/Materials/HLSLMaterialTranslator.cpp
--- a/Engine/Source/Runtime/Engine/Private/Materials/HLSLMaterialTranslator.cpp
+++ b/Engine/Source/Runtime/Engine/Private/Materials/HLSLMaterialTranslator.cpp
@@ -3835,9 +3835,6 @@ int32 FHLSLMaterialTranslator::ForceCast(int32 Code, EMaterialValueType DestType
 	const EMaterialValueType SourceType = GetParameterType(Code);
 	const bool bExactMatch = (ForceCastFlags & MFCF_ExactMatch) ? true : false;
 
-	EMaterialCastFlags Flags = EMaterialCastFlags::AllowTruncate | EMaterialCastFlags::AllowAppendZeroes;
-	if (ForceCastFlags & MFCF_ReplicateValue) Flags |= EMaterialCastFlags::ReplicateScalar;
-
 	if (bExactMatch ? (SourceType == DestType) : (SourceType & DestType))
 	{
 		return Code;
@@ -3846,14 +3843,23 @@ int32 FHLSLMaterialTranslator::ForceCast(int32 Code, EMaterialValueType DestType
 	{
 		const FDerivInfo CodeDerivInfo = GetDerivInfo(Code);
 
-		FString FiniteCode = CastValue(GetParameterCode(Code), SourceType, DestType, Flags);
+		EMaterialCastFlags CastFlags = EMaterialCastFlags::AllowTruncate | EMaterialCastFlags::AllowAppendZeroes;
+		if ((ForceCastFlags & MFCF_ReplicateValue) || !bExactMatch)
+		{
+			// Replicate scalar if requested, or if we don't require an exact match (this can happen when force-casting to/from LWC
+			// The only way we *don't* replicate scalar is requesting an exact match without the ReplicateValue flag
+			// TODO - My guess is that case probably isn't relevant, and we should just always replicate scalar on cast, but trying to preserve behavior for now
+			CastFlags |= EMaterialCastFlags::ReplicateScalar;
+		}
+
+		FString FiniteCode = CastValue(GetParameterCode(Code), SourceType, DestType, CastFlags);
 		if (IsAnalyticDerivEnabled() && IsDerivativeValid(CodeDerivInfo.DerivativeStatus))
 		{
 			if (CodeDerivInfo.DerivativeStatus == EDerivativeStatus::Valid)
 			{
 				FString DerivString = *GetParameterCodeDeriv(Code, CompiledPDV_Analytic);
-				FString DDXCode = CastValue(DerivString + TEXT(".Ddx"), MakeNonLWCType(SourceType), MakeNonLWCType(DestType), Flags);
-				FString DDYCode = CastValue(DerivString + TEXT(".Ddy"), MakeNonLWCType(SourceType), MakeNonLWCType(DestType), Flags);
+				FString DDXCode = CastValue(DerivString + TEXT(".Ddx"), MakeNonLWCType(SourceType), MakeNonLWCType(DestType), CastFlags);
+				FString DDYCode = CastValue(DerivString + TEXT(".Ddy"), MakeNonLWCType(SourceType), MakeNonLWCType(DestType), CastFlags);
 				FString AnalyticCode = DerivativeAutogen.ConstructDeriv(FiniteCode, DDXCode, DDYCode, GetDerivType(DestType));
 				return AddCodeChunkInnerDeriv(*FiniteCode, *AnalyticCode, DestType, false, EDerivativeStatus::Valid);
 			}
