Fix broken links to documentation in header and source files.

diff --git a/Engine/Plugins/2D/Paper2D/Source/Paper2DEditor/Private/FlipbookEditor/FlipbookEditor.cpp b/Engine/Plugins/2D/Paper2D/Source/Paper2DEditor/Private/FlipbookEditor/FlipbookEditor.cpp
--- a/Engine/Plugins/2D/Paper2D/Source/Paper2DEditor/Private/FlipbookEditor/FlipbookEditor.cpp
+++ b/Engine/Plugins/2D/Paper2D/Source/Paper2DEditor/Private/FlipbookEditor/FlipbookEditor.cpp
@@ -422,7 +422,7 @@ FString FFlipbookEditor::GetWorldCentricTabPrefix() const
 
 FString FFlipbookEditor::GetDocumentationLink() const
 {
-	return TEXT("Engine/Paper2D/FlipbookEditor");
+	return TEXT("AnimatingObjects/Paper2D/Flipbooks");
 }
 
 FLinearColor FFlipbookEditor::GetWorldCentricTabColorScale() const
diff --git a/Engine/Plugins/2D/Paper2D/Source/Paper2DEditor/Private/SpriteEditor/SpriteEditor.cpp b/Engine/Plugins/2D/Paper2D/Source/Paper2DEditor/Private/SpriteEditor/SpriteEditor.cpp
--- a/Engine/Plugins/2D/Paper2D/Source/Paper2DEditor/Private/SpriteEditor/SpriteEditor.cpp
+++ b/Engine/Plugins/2D/Paper2D/Source/Paper2DEditor/Private/SpriteEditor/SpriteEditor.cpp
@@ -477,7 +477,7 @@ FString FSpriteEditor::GetWorldCentricTabPrefix() const
 
 FString FSpriteEditor::GetDocumentationLink() const
 {
-	return TEXT("Engine/Paper2D/SpriteEditor");
+	return TEXT("AnimatingObjects/Paper2D/Sprites");
 }
 
 void FSpriteEditor::OnToolkitHostingStarted(const TSharedRef<class IToolkit>& Toolkit)
diff --git a/Engine/Plugins/2D/Paper2D/Source/Paper2DEditor/Private/TileMapEditing/TileMapEditor.cpp b/Engine/Plugins/2D/Paper2D/Source/Paper2DEditor/Private/TileMapEditing/TileMapEditor.cpp
--- a/Engine/Plugins/2D/Paper2D/Source/Paper2DEditor/Private/TileMapEditing/TileMapEditor.cpp
+++ b/Engine/Plugins/2D/Paper2D/Source/Paper2DEditor/Private/TileMapEditing/TileMapEditor.cpp
@@ -399,7 +399,7 @@ FString FTileMapEditor::GetWorldCentricTabPrefix() const
 
 FString FTileMapEditor::GetDocumentationLink() const
 {
-	return TEXT("Engine/Paper2D/TileMapEditor");
+	return TEXT("AnimatingObjects/Paper2D/TileMaps");
 }
 
 void FTileMapEditor::OnToolkitHostingStarted(const TSharedRef< class IToolkit >& Toolkit)
diff --git a/Engine/Plugins/2D/Paper2D/Source/Paper2DEditor/Private/TileSetEditor.cpp b/Engine/Plugins/2D/Paper2D/Source/Paper2DEditor/Private/TileSetEditor.cpp
--- a/Engine/Plugins/2D/Paper2D/Source/Paper2DEditor/Private/TileSetEditor.cpp
+++ b/Engine/Plugins/2D/Paper2D/Source/Paper2DEditor/Private/TileSetEditor.cpp
@@ -440,8 +440,7 @@ FLinearColor FTileSetEditor::GetWorldCentricTabColorScale() const
 
 FString FTileSetEditor::GetDocumentationLink() const
 {
-	//@TODO: Need to make a page for this
-	return TEXT("Engine/Paper2D/TileSetEditor");
+	return TEXT("AnimatingObjects/Paper2D/TileMaps");
 }
 
 void FTileSetEditor::OnToolkitHostingStarted(const TSharedRef<class IToolkit>& Toolkit)
diff --git a/Engine/Plugins/Editor/BlueprintMaterialTextureNodes/Source/BlueprintMaterialTextureNodes/Public/BlueprintMaterialTextureNodesBPLibrary.h b/Engine/Plugins/Editor/BlueprintMaterialTextureNodes/Source/BlueprintMaterialTextureNodes/Public/BlueprintMaterialTextureNodesBPLibrary.h
--- a/Engine/Plugins/Editor/BlueprintMaterialTextureNodes/Source/BlueprintMaterialTextureNodes/Public/BlueprintMaterialTextureNodesBPLibrary.h
+++ b/Engine/Plugins/Editor/BlueprintMaterialTextureNodes/Source/BlueprintMaterialTextureNodes/Public/BlueprintMaterialTextureNodesBPLibrary.h
@@ -21,7 +21,7 @@
 *	Category -	the category your node will be under in the Blueprint drop-down menu.
 *
 *	For more info on custom blueprint nodes visit documentation:
-*	https://wiki.unrealengine.com/Custom_Blueprint_Node_Creation
+*	https://docs.unrealengine.com/ProgrammingAndScripting/Blueprints/TechnicalGuide/ExtendingBlueprints
 */
 UCLASS()
 class BLUEPRINTMATERIALTEXTURENODES_API UBlueprintMaterialTextureNodesBPLibrary : public UBlueprintFunctionLibrary
diff --git a/Engine/Plugins/Media/BinkMedia/Source/BinkMediaPlayerEditor/Private/BinkMediaPlayerEditorToolkit.h b/Engine/Plugins/Media/BinkMedia/Source/BinkMediaPlayerEditor/Private/BinkMediaPlayerEditorToolkit.h
--- a/Engine/Plugins/Media/BinkMedia/Source/BinkMediaPlayerEditor/Private/BinkMediaPlayerEditorToolkit.h
+++ b/Engine/Plugins/Media/BinkMedia/Source/BinkMediaPlayerEditor/Private/BinkMediaPlayerEditorToolkit.h
@@ -28,7 +28,7 @@ struct FBinkMediaPlayerEditorToolkit : public FAssetEditorToolkit, public FEdito
 
 	// FAssetEditorToolkit interface
 
-	virtual FString GetDocumentationLink() const override { return FString(TEXT("Engine/Content/Types/MediaAssets/Properties/Interface")); }
+	virtual FString GetDocumentationLink() const override { return FString(TEXT("WorkingWithMedia/IntegratingMedia/BinkVideo")); }
 	virtual void RegisterTabSpawners( const TSharedRef<class FTabManager>& TabManager ) override;
 	virtual void UnregisterTabSpawners( const TSharedRef<class FTabManager>& TabManager ) override;
 
diff --git a/Engine/Plugins/Media/MediaPlayerEditor/Source/MediaPlayerEditor/Private/Toolkits/MediaPlayerEditorToolkit.cpp b/Engine/Plugins/Media/MediaPlayerEditor/Source/MediaPlayerEditor/Private/Toolkits/MediaPlayerEditorToolkit.cpp
--- a/Engine/Plugins/Media/MediaPlayerEditor/Source/MediaPlayerEditor/Private/Toolkits/MediaPlayerEditorToolkit.cpp
+++ b/Engine/Plugins/Media/MediaPlayerEditor/Source/MediaPlayerEditor/Private/Toolkits/MediaPlayerEditorToolkit.cpp
@@ -144,7 +144,7 @@ void FMediaPlayerEditorToolkit::Initialize(UMediaPlayer* InMediaPlayer, const ET
 
 FString FMediaPlayerEditorToolkit::GetDocumentationLink() const
 {
-	return FString(TEXT("Engine/Content/Types/MediaAssets/Properties/Interface"));
+	return FString(TEXT("WorkingWithMedia/IntegratingMedia/MediaFramework"));
 }
 
 
diff --git a/Engine/Plugins/Media/MediaPlayerEditor/Source/MediaPlayerEditor/Private/Toolkits/MediaPlaylistEditorToolkit.cpp b/Engine/Plugins/Media/MediaPlayerEditor/Source/MediaPlayerEditor/Private/Toolkits/MediaPlaylistEditorToolkit.cpp
--- a/Engine/Plugins/Media/MediaPlayerEditor/Source/MediaPlayerEditor/Private/Toolkits/MediaPlaylistEditorToolkit.cpp
+++ b/Engine/Plugins/Media/MediaPlayerEditor/Source/MediaPlayerEditor/Private/Toolkits/MediaPlaylistEditorToolkit.cpp
@@ -97,7 +97,7 @@ void FMediaPlaylistEditorToolkit::Initialize(UMediaPlaylist* InMediaPlaylist, co
 
 FString FMediaPlaylistEditorToolkit::GetDocumentationLink() const
 {
-	return FString(TEXT("Engine/Content/Types/MediaAssets/Properties/Interface"));
+	return FString(TEXT("WorkingWithMedia/IntegratingMedia/MediaFramework"));
 }
 
 
diff --git a/Engine/Plugins/Runtime/Metasound/Source/MetasoundEditor/Private/MetasoundEditor.h b/Engine/Plugins/Runtime/Metasound/Source/MetasoundEditor/Private/MetasoundEditor.h
--- a/Engine/Plugins/Runtime/Metasound/Source/MetasoundEditor/Private/MetasoundEditor.h
+++ b/Engine/Plugins/Runtime/Metasound/Source/MetasoundEditor/Private/MetasoundEditor.h
@@ -95,7 +95,7 @@ namespace Metasound
 
 			virtual FString GetDocumentationLink() const override
 			{
-				return FString(TEXT("Engine/Audio/Metasounds/Editor"));
+				return FString(TEXT("working-with-audio/sound-sources/meta-sounds"));
 			}
 
 			/** FGCObject interface */
diff --git a/Engine/Plugins/Runtime/Oculus/OculusVR/Source/OculusEditor/Private/OculusToolWidget.cpp b/Engine/Plugins/Runtime/Oculus/OculusVR/Source/OculusEditor/Private/OculusToolWidget.cpp
--- a/Engine/Plugins/Runtime/Oculus/OculusVR/Source/OculusEditor/Private/OculusToolWidget.cpp
+++ b/Engine/Plugins/Runtime/Oculus/OculusVR/Source/OculusEditor/Private/OculusToolWidget.cpp
@@ -229,7 +229,7 @@ void SOculusToolWidget::RebuildLayout()
 	AddSimpleSetting(box, SimpleSettings.Find(FName("AndroidPackaging")));
 	AddSimpleSetting(box, SimpleSettings.Find(FName("AndroidQuestArch")));
 
-	box = NewCategory(scroller, LOCTEXT("PostProcessHeader", "<RichTextBlock.Bold>Post-Processing Settings:</>\nThe below settings all refer to your project's post-processing settings. Post-processing can be very expensive in VR, so we recommend disabling many expensive post-processing effects. You can fine-tune your post-processing settings with a Post Process Volume. <a href=\"https://docs.unrealengine.com/en-us/Platforms/VR/VRPerformance\" id=\"HyperlinkDecorator\">Read more.</>."));
+	box = NewCategory(scroller, LOCTEXT("PostProcessHeader", "<RichTextBlock.Bold>Post-Processing Settings:</>\nThe below settings all refer to your project's post-processing settings. Post-processing can be very expensive in VR, so we recommend disabling many expensive post-processing effects. You can fine-tune your post-processing settings with a Post Process Volume. <a href=\"https://docs.unrealengine.com/SharingAndReleasing/XRDevelopment/VR/VRPerformanceAndProfiling\" id=\"HyperlinkDecorator\">Read more.</>."));
 	AddSimpleSetting(box, SimpleSettings.Find(FName("LensFlare")));
 	AddSimpleSetting(box, SimpleSettings.Find(FName("AntiAliasing")));
 
diff --git a/Engine/Source/Editor/AnimationEditor/Private/AnimationEditor.h b/Engine/Source/Editor/AnimationEditor/Private/AnimationEditor.h
--- a/Engine/Source/Editor/AnimationEditor/Private/AnimationEditor.h
+++ b/Engine/Source/Editor/AnimationEditor/Private/AnimationEditor.h
@@ -83,7 +83,7 @@ public:
 	/** @return the documentation location for this editor */
 	virtual FString GetDocumentationLink() const override
 	{
-		return FString(TEXT("Engine/Animation/AnimationEditor"));
+		return FString(TEXT("AnimatingObjects/SkeletalMeshAnimation/Persona/Modes/Animation"));
 	}
 
 	/** FGCObject interface */
diff --git a/Engine/Source/Editor/Kismet/Private/BlueprintDetailsCustomization.h b/Engine/Source/Editor/Kismet/Private/BlueprintDetailsCustomization.h
--- a/Engine/Source/Editor/Kismet/Private/BlueprintDetailsCustomization.h
+++ b/Engine/Source/Editor/Kismet/Private/BlueprintDetailsCustomization.h
@@ -34,7 +34,7 @@ class SBlueprintNamespaceEntry;
 
 /**
  * Variable network replication options.
- * @see https://docs.unrealengine.com/latest/INT/Gameplay/Networking/Replication/
+ * @see https://docs.unrealengine.com/InteractiveExperiences/Networking/Blueprints
  */
 namespace EVariableReplication
 {
diff --git a/Engine/Source/Editor/Persona/Private/SAnimStreamableEditor.cpp b/Engine/Source/Editor/Persona/Private/SAnimStreamableEditor.cpp
--- a/Engine/Source/Editor/Persona/Private/SAnimStreamableEditor.cpp
+++ b/Engine/Source/Editor/Persona/Private/SAnimStreamableEditor.cpp
@@ -11,7 +11,7 @@
 
 TSharedRef<SWidget> SAnimStreamableEditor::CreateDocumentAnchor()
 {
-	return IDocumentation::Get()->CreateAnchor(TEXT("Engine/Animation/AnimationStreamable"));
+	return IDocumentation::Get()->CreateAnchor(TEXT("AnimatingObjects"));
 }
 
 void SAnimStreamableEditor::Construct(const FArguments& InArgs, const TSharedRef<class IPersonaPreviewScene>& InPreviewScene, const TSharedRef<class IEditableSkeleton>& InEditableSkeleton, const TSharedRef<FUICommandList>& InCommandList)
diff --git a/Engine/Source/Editor/SkeletalMeshEditor/Private/SkeletalMeshEditor.h b/Engine/Source/Editor/SkeletalMeshEditor/Private/SkeletalMeshEditor.h
--- a/Engine/Source/Editor/SkeletalMeshEditor/Private/SkeletalMeshEditor.h
+++ b/Engine/Source/Editor/SkeletalMeshEditor/Private/SkeletalMeshEditor.h
@@ -93,7 +93,7 @@ public:
 	/** @return the documentation location for this editor */
 	virtual FString GetDocumentationLink() const override
 	{
-		return FString(TEXT("Engine/Animation/SkeletalMeshEditor"));
+		return FString(TEXT("AnimatingObjects/SkeletalMeshAnimation/Persona/Modes/Mesh"));
 	}
 
 	/** FGCObject interface */
diff --git a/Engine/Source/Editor/SkeletonEditor/Private/SkeletonEditor.h b/Engine/Source/Editor/SkeletonEditor/Private/SkeletonEditor.h
--- a/Engine/Source/Editor/SkeletonEditor/Private/SkeletonEditor.h
+++ b/Engine/Source/Editor/SkeletonEditor/Private/SkeletonEditor.h
@@ -79,7 +79,7 @@ public:
 	/** @return the documentation location for this editor */
 	virtual FString GetDocumentationLink() const override
 	{
-		return FString(TEXT("Engine/Animation/SkeletonEditor"));
+		return FString(TEXT("AnimatingObjects/SkeletalMeshAnimation/Persona/Modes/Skeleton"));
 	}
 
 	/** Get the skeleton tree widget */
diff --git a/Engine/Source/Runtime/CoreUObject/Public/UObject/Object.h b/Engine/Source/Runtime/CoreUObject/Public/UObject/Object.h
--- a/Engine/Source/Runtime/CoreUObject/Public/UObject/Object.h
+++ b/Engine/Source/Runtime/CoreUObject/Public/UObject/Object.h
@@ -46,7 +46,7 @@ namespace ECastCheckedType
  * The base class of all UE objects. The type of an object is defined by its UClass.
  * This provides support functions for creating and using objects, and virtual functions that should be overridden in child classes.
  * 
- * @see https://docs.unrealengine.com/en-us/Programming/UnrealArchitecture/Objects
+ * @see https://docs.unrealengine.com/ProgrammingAndScripting/ProgrammingWithCPP/UnrealArchitecture/Objects
  */
 class COREUOBJECT_API UObject : public UObjectBaseUtility
 {
diff --git a/Engine/Source/Runtime/Engine/Classes/Components/AudioComponent.h b/Engine/Source/Runtime/Engine/Classes/Components/AudioComponent.h
--- a/Engine/Source/Runtime/Engine/Classes/Components/AudioComponent.h
+++ b/Engine/Source/Runtime/Engine/Classes/Components/AudioComponent.h
@@ -156,7 +156,7 @@ class ENGINE_API UInitialActiveSoundParams : public UObject
 /**
  * AudioComponent is used to play a Sound
  *
- * @see https://docs.unrealengine.com/latest/INT/Audio/Overview/index.html
+ * @see https://docs.unrealengine.com/WorkingWithAudio/Overview
  * @see USoundBase
  */
 UCLASS(ClassGroup=(Audio, Common), HideCategories=(Object, ActorComponent, Physics, Rendering, Mobility, LOD), ShowCategories=Trigger, meta=(BlueprintSpawnableComponent))
diff --git a/Engine/Source/Runtime/Engine/Classes/Components/PrimitiveComponent.h b/Engine/Source/Runtime/Engine/Classes/Components/PrimitiveComponent.h
--- a/Engine/Source/Runtime/Engine/Classes/Components/PrimitiveComponent.h
+++ b/Engine/Source/Runtime/Engine/Classes/Components/PrimitiveComponent.h
@@ -325,7 +325,7 @@ public:
 	 * If true, this component will generate overlap events when it is overlapping other components (eg Begin Overlap).
 	 * Both components (this and the other) must have this enabled for overlap events to occur.
 	 *
-	 * @see [Overlap Events](https://docs.unrealengine.com/latest/INT/Engine/Physics/Collision/index.html#overlapandgenerateoverlapevents)
+	 * @see [Overlap Events](https://docs.unrealengine.com/InteractiveExperiences/Physics/Collision/Overview#overlapandgenerateoverlapevents)
 	 * @see UpdateOverlaps(), BeginComponentOverlap(), EndComponentOverlap()
 	 */
 	UFUNCTION(BlueprintGetter)
@@ -1045,7 +1045,7 @@ public:
 	 * Begin tracking an overlap interaction with the component specified.
 	 * @param OtherComp - The component of the other actor that this component is now overlapping
 	 * @param bDoNotifies - True to dispatch appropriate begin/end overlap notifications when these events occur.
-	 * @see [Overlap Events](https://docs.unrealengine.com/latest/INT/Engine/Physics/Collision/index.html#overlapandgenerateoverlapevents)
+	 * @see [Overlap Events](https://docs.unrealengine.com/InteractiveExperiences/Physics/Collision/Overview#overlapandgenerateoverlapevents)
 	 */
 	void BeginComponentOverlap(const FOverlapInfo& OtherOverlap, bool bDoNotifies);
 	
@@ -1054,7 +1054,7 @@ public:
 	 * @param OtherComp The component of the other actor to stop overlapping
 	 * @param bDoNotifies True to dispatch appropriate begin/end overlap notifications when these events occur.
 	 * @param bSkipNotifySelf True to skip end overlap notifications to this component's.  Does not affect notifications to OtherComp's actor.
-	 * @see [Overlap Events](https://docs.unrealengine.com/latest/INT/Engine/Physics/Collision/index.html#overlapandgenerateoverlapevents)
+	 * @see [Overlap Events](https://docs.unrealengine.com/InteractiveExperiences/Physics/Collision/Overview#overlapandgenerateoverlapevents)
 	 */
 	void EndComponentOverlap(const FOverlapInfo& OtherOverlap, bool bDoNotifies=true, bool bSkipNotifySelf=false);
 
diff --git a/Engine/Source/Runtime/Engine/Classes/Engine/LODActor.h b/Engine/Source/Runtime/Engine/Classes/Engine/LODActor.h
--- a/Engine/Source/Runtime/Engine/Classes/Engine/LODActor.h
+++ b/Engine/Source/Runtime/Engine/Classes/Engine/LODActor.h
@@ -61,7 +61,7 @@ struct FHLODInstancingKey
  * LODActor is an instance of an autogenerated StaticMesh Actors by Hierarchical LOD System
  * This is essentially just StaticMeshActor that you can't move or edit, but it contains multiple actors reference
  *
- * @see https://docs.unrealengine.com/latest/INT/Engine/Actors/LODActor/
+ * @see https://docs.unrealengine.com/BuildingWorlds/HLOD
  * @see UStaticMesh
  */
 
diff --git a/Engine/Source/Runtime/Engine/Classes/GameFramework/Actor.h b/Engine/Source/Runtime/Engine/Classes/GameFramework/Actor.h
--- a/Engine/Source/Runtime/Engine/Classes/GameFramework/Actor.h
+++ b/Engine/Source/Runtime/Engine/Classes/GameFramework/Actor.h
@@ -157,8 +157,8 @@ public:
  * - AActor::BeginPlay: Called when the level starts ticking, only during actual gameplay.
  *                      This normally happens right after PostInitializeComponents but can be delayed for networked or child actors.
  *
- * @see https://docs.unrealengine.com/latest/INT/Programming/UnrealArchitecture/Actors/
- * @see https://docs.unrealengine.com/en-us/Programming/UnrealArchitecture/Actors/ActorLifecycle
+ * @see https://docs.unrealengine.com/Programming/UnrealArchitecture/Actors
+ * @see https://docs.unrealengine.com/Programming/UnrealArchitecture/Actors/ActorLifecycle
  * @see UActorComponent
  */
 UCLASS(BlueprintType, Blueprintable, config=Engine, meta=(ShortTooltip="An Actor is an object that can be placed or spawned in the world."))
@@ -184,7 +184,7 @@ public:
 	/**
 	 * Primary Actor tick function, which calls TickActor().
 	 * Tick functions can be configured to control whether ticking is enabled, at what time during a frame the update occurs, and to set up tick dependencies.
-	 * @see https://docs.unrealengine.com/latest/INT/API/Runtime/Engine/Engine/FTickFunction/
+	 * @see https://docs.unrealengine.com/API/Runtime/Engine/Engine/FTickFunction
 	 * @see AddTickPrerequisiteActor(), AddTickPrerequisiteComponent()
 	 */
 	UPROPERTY(EditDefaultsOnly, Category=Tick)
@@ -214,7 +214,7 @@ private:
 	 * If true, replicate movement/location related properties.
 	 * Actor must also be set to replicate.
 	 * @see SetReplicates()
-	 * @see https://docs.unrealengine.com/latest/INT/Gameplay/Networking/Replication/
+	 * @see https://docs.unrealengine.com/InteractiveExperiences/Networking/Actors
 	 */
 	UPROPERTY(ReplicatedUsing=OnRep_ReplicateMovement, Category=Replication, EditDefaultsOnly)
 	uint8 bReplicateMovement:1;    
@@ -528,7 +528,7 @@ public:
 	 * Properties flagged for replication will update on clients if they change on the server.
 	 * Internally changes the RemoteRole property and handles the cases where the actor needs to be added to the network actor list.
 	 * @param bInReplicates Whether this Actor replicates to network clients.
-	 * @see https://docs.unrealengine.com/latest/INT/Gameplay/Networking/Replication/
+	 * @see https://docs.unrealengine.com/InteractiveExperiences/Networking/Actors
 	 */
 	UFUNCTION(BlueprintCallable, BlueprintAuthorityOnly, Category=Networking)
 	void SetReplicates(bool bInReplicates);
diff --git a/Engine/Source/Runtime/Engine/Classes/PhysicsEngine/PhysicsAsset.h b/Engine/Source/Runtime/Engine/Classes/PhysicsEngine/PhysicsAsset.h
--- a/Engine/Source/Runtime/Engine/Classes/PhysicsEngine/PhysicsAsset.h
+++ b/Engine/Source/Runtime/Engine/Classes/PhysicsEngine/PhysicsAsset.h
@@ -90,7 +90,7 @@ enum class EPhysicsAssetSolverType: uint8
  * A SkeletalMesh has a single PhysicsAsset, which allows for easily turning ragdoll physics on or off for many SkeletalMeshComponents
  * The asset can be configured inside the Physics Asset Editor.
  *
- * @see https://docs.unrealengine.com/latest/INT/Engine/Physics/PhAT/Reference/index.html
+ * @see https://docs.unrealengine.com/InteractiveExperiences/Physics/PhysicsAssetEditor
  * @see USkeletalMesh
  */
 
diff --git a/Engine/Source/Runtime/Landscape/Classes/LandscapeProxy.h b/Engine/Source/Runtime/Landscape/Classes/LandscapeProxy.h
--- a/Engine/Source/Runtime/Landscape/Classes/LandscapeProxy.h
+++ b/Engine/Source/Runtime/Landscape/Classes/LandscapeProxy.h
@@ -648,7 +648,7 @@ public:
 	 * If true, Landscape will generate overlap events when other components are overlapping it (eg Begin Overlap).
 	 * Both the Landscape and the other component must have this flag enabled for overlap events to occur.
 	 *
-	 * @see [Overlap Events](https://docs.unrealengine.com/latest/INT/Engine/Physics/Collision/index.html#overlapandgenerateoverlapevents)
+	 * @see [Overlap Events](https://docs.unrealengine.com/InteractiveExperiences/Physics/Collision/Overview#overlapandgenerateoverlapevents)
 	 * @see UpdateOverlaps(), BeginComponentOverlap(), EndComponentOverlap()
 	 */
 	UPROPERTY(EditAnywhere, BlueprintReadOnly, Category = Collision)
diff --git a/Engine/Source/Runtime/Projects/Private/ProjectManager.cpp b/Engine/Source/Runtime/Projects/Private/ProjectManager.cpp
--- a/Engine/Source/Runtime/Projects/Private/ProjectManager.cpp
+++ b/Engine/Source/Runtime/Projects/Private/ProjectManager.cpp
@@ -39,7 +39,7 @@ bool FProjectManager::LoadProjectFile( const FString& InProjectFile )
 	}
 
 #if PLATFORM_IOS
-	FString UpdatedMessage = FString::Printf(TEXT("%s\n%s"), *FailureReason.ToString(), TEXT("For troubleshooting, please go to https://docs.unrealengine.com/latest/INT/Platforms/iOS/GettingStarted/index.html"));
+	FString UpdatedMessage = FString::Printf(TEXT("%s\n%s"), *FailureReason.ToString(), TEXT("For troubleshooting, please go to https://docs.unrealengine.com/SharingAndReleasing/Mobile/iOS"));
 	FailureReason = FText::FromString(UpdatedMessage);
 #endif
 	UE_LOG(LogProjectManager, Error, TEXT("%s"), *FailureReason.ToString());
