Refuse to save too large shadermap DDC as it won't be handled anyway.

diff --git a/Engine/Source/Runtime/Engine/Private/Materials/MaterialShader.cpp b/Engine/Source/Runtime/Engine/Private/Materials/MaterialShader.cpp
--- a/Engine/Source/Runtime/Engine/Private/Materials/MaterialShader.cpp
+++ b/Engine/Source/Runtime/Engine/Private/Materials/MaterialShader.cpp
@@ -1420,14 +1420,23 @@ void FMaterialShaderMap::SaveToDerivedDataCache()
 {
 	TRACE_CPUPROFILER_EVENT_SCOPE(FMaterialShaderMap::SaveToDerivedDataCache);
 	COOK_STAT(auto Timer = MaterialShaderCookStats::UsageStats.TimeSyncWork());
-	TArray<uint8> SaveData;
-	FMemoryWriter Ar(SaveData, true);
+	TArray64<uint8> SaveData;
+	FMemoryWriter64 Ar(SaveData, true);
 	Serialize(Ar);
 
+	// Refuse to save too large derived data, as it won't get handled by the cache anyway. It appends some metadata, so for safety we use a lower threshold.
+	const int64 kMaxSaveDataForDDC = MAX_int32 - (4 * 1024 * 1024);
+	if (SaveData.Num() >= kMaxSaveDataForDDC)
+	{
+		UE_LOG(LogMaterial, Warning, TEXT("Not saving %s shaders to DDC as its size (%lld MB) is larger than %lld MB, which may cause problems for the legacy DDC API."), 
+			GetMaterialPath(), SaveData.Num() / (1024 * 1024), kMaxSaveDataForDDC);
+		return;
+	}
+
 	TRACE_COUNTER_ADD(Shaders_FMaterialShaderMapDDCBytesSent, SaveData.Num());
 
 	FString DataKey = GetMaterialShaderMapKeyString(ShaderMapId, GetShaderPlatform());
-	UE_LOG(LogMaterial, Verbose, TEXT("Saved shaders for %s from DDC (key hash: %s)"), GetMaterialPath(), *FSHA1_HashString(DataKey));
+	UE_LOG(LogMaterial, Verbose, TEXT("Saved shaders for %s to DDC (key hash: %s)"), GetMaterialPath(), *FSHA1_HashString(DataKey));
 	GetDerivedDataCacheRef().Put(*DataKey, SaveData, FStringView(GetFriendlyName()));
 	COOK_STAT(Timer.AddMiss(SaveData.Num()));
 }
