Errata for CL 18727212 (print MB, not bytes).

diff --git a/Engine/Source/Runtime/Engine/Private/Materials/MaterialShader.cpp b/Engine/Source/Runtime/Engine/Private/Materials/MaterialShader.cpp
--- a/Engine/Source/Runtime/Engine/Private/Materials/MaterialShader.cpp
+++ b/Engine/Source/Runtime/Engine/Private/Materials/MaterialShader.cpp
@@ -1429,7 +1429,7 @@ void FMaterialShaderMap::SaveToDerivedDataCache()
 	if (SaveData.Num() >= kMaxSaveDataForDDC)
 	{
 		UE_LOG(LogMaterial, Warning, TEXT("Not saving %s shaders to DDC as its size (%lld MB) is larger than %lld MB, which may cause problems for the legacy DDC API."), 
-			GetMaterialPath(), SaveData.Num() / (1024 * 1024), kMaxSaveDataForDDC);
+			GetMaterialPath(), SaveData.Num() / (1024 * 1024), kMaxSaveDataForDDC / (1024 * 1024));
 		return;
 	}
 
