UE-102275 Oculus - Console Command window aligned off player view
-When the debug canvas is being drawn to an internal texture do not apply DPIScale from the viewport.

diff --git a/Engine/Source/Runtime/Engine/Private/Slate/DebugCanvas.cpp b/Engine/Source/Runtime/Engine/Private/Slate/DebugCanvas.cpp
--- a/Engine/Source/Runtime/Engine/Private/Slate/DebugCanvas.cpp
+++ b/Engine/Source/Runtime/Engine/Private/Slate/DebugCanvas.cpp
@@ -190,13 +190,18 @@ void FDebugCanvasDrawer::BeginRenderingCanvas( const FIntRect& CanvasRect )
 
 void FDebugCanvasDrawer::InitDebugCanvas(FViewportClient* ViewportClient, UWorld* InWorld)
 {
+	const bool bIsStereoscopic3D = GEngine && GEngine->IsStereoscopic3D();
+	IStereoLayers* const StereoLayers = (bIsStereoscopic3D && GEngine && GEngine->StereoRenderingDevice.IsValid()) ? GEngine->StereoRenderingDevice->GetStereoLayers() : nullptr;
+	const bool bUseInternalTexture = StereoLayers && bIsStereoscopic3D;
+
 	// If the canvas is not null there is more than one viewport draw call before slate draws.  This can happen on resizes. 
 	// We need to delete the old canvas
 		// This can also happen if we are debugging a HUD blueprint and in that case we need to continue using
 		// the same canvas
 	if (FSlateApplication::Get().IsNormalExecution())
 	{
-		GameThreadCanvas = MakeShared<FCanvas, ESPMode::ThreadSafe>(RenderTarget, nullptr, InWorld, InWorld ? InWorld->FeatureLevel.GetValue() : GMaxRHIFeatureLevel, FCanvas::CDM_DeferDrawing, ViewportClient->GetDPIScale());
+		const float DPIScale = bUseInternalTexture ? 1.0f : ViewportClient->GetDPIScale();
+		GameThreadCanvas = MakeShared<FCanvas, ESPMode::ThreadSafe>(RenderTarget, nullptr, InWorld, InWorld ? InWorld->FeatureLevel.GetValue() : GMaxRHIFeatureLevel, FCanvas::CDM_DeferDrawing, DPIScale);
 
 		// Do not allow the canvas to be flushed outside of our debug rendering path
 		GameThreadCanvas->SetAllowedModes(FCanvas::Allow_DeleteOnRender);
@@ -204,13 +209,9 @@ void FDebugCanvasDrawer::InitDebugCanvas(FViewportClient* ViewportClient, UWorld
 
 	if (GameThreadCanvas.IsValid())
 	{
-		const bool bIsStereoscopic3D = GEngine && GEngine->IsStereoscopic3D();
-		IStereoLayers* const StereoLayers = (bIsStereoscopic3D && GEngine && GEngine->StereoRenderingDevice.IsValid()) ? GEngine->StereoRenderingDevice->GetStereoLayers() : nullptr;
-		const bool bHMDAvailable = StereoLayers && bIsStereoscopic3D;
-
-		GameThreadCanvas->SetUseInternalTexture(bHMDAvailable);
+		GameThreadCanvas->SetUseInternalTexture(bUseInternalTexture);
 
-		if (bHMDAvailable && LayerTexture)
+		if (bUseInternalTexture && LayerTexture)
 		{
 			if (StereoLayers)
 			{
