Fixing compilation errors in the ALLOW_DITHERED_LOD_FOR_INSTANCED_STATIC_MESHES = false path.
This change has been sitting around for a while, I gave it to a licensee who shipped a game successfully with it, so it's properly tested.

#rnx

diff --git a/Engine/Source/Runtime/Engine/Private/InstancedStaticMesh.cpp b/Engine/Source/Runtime/Engine/Private/InstancedStaticMesh.cpp
--- a/Engine/Source/Runtime/Engine/Private/InstancedStaticMesh.cpp
+++ b/Engine/Source/Runtime/Engine/Private/InstancedStaticMesh.cpp
@@ -833,32 +833,28 @@ void FInstancedStaticMeshVertexFactory::InitRHI()
 #if !ALLOW_DITHERED_LOD_FOR_INSTANCED_STATIC_MESHES // position(and normal) only shaders cannot work with dithered LOD
 	// If the vertex buffer containing position is not the same vertex buffer containing the rest of the data,
 	// then initialize PositionStream and PositionDeclaration.
-	if(Data.PositionComponent.VertexBuffer != Data.TangentBasisComponents[0].VertexBuffer)
+	if (Data.PositionComponent.VertexBuffer != Data.TangentBasisComponents[0].VertexBuffer)
 	{
-		auto AddDeclaration = [&Data](EVertexInputStreamType InputStreamType, bool bInstanced, bool bAddNormal)
+		auto AddDeclaration = [&](EVertexInputStreamType InputStreamType)
 		{
 			FVertexDeclarationElementList StreamElements;
-			StreamElements.Add(AccessPositionStreamComponent(Data.PositionComponent, 0));
+			StreamElements.Add(AccessStreamComponent(Data.PositionComponent, 0, InputStreamType));
 
-			bAddNormal = bAddNormal && Data.TangentBasisComponents[1].VertexBuffer != NULL;
-			if (bAddNormal)
+			if (InputStreamType == EVertexInputStreamType::PositionAndNormalOnly && Data.TangentBasisComponents[1].VertexBuffer != NULL)
 			{
 				StreamElements.Add(AccessStreamComponent(Data.TangentBasisComponents[1], 2, InputStreamType));
 			}
 
-			if (bInstanced)
-			{
-				// toss in the instanced location stream
-				StreamElements.Add(AccessPositionStreamComponent(Data.InstanceOriginComponent, 8));
-				StreamElements.Add(AccessPositionStreamComponent(Data.InstanceTransformComponent[0], 9));
-				StreamElements.Add(AccessPositionStreamComponent(Data.InstanceTransformComponent[1], 10));
-				StreamElements.Add(AccessPositionStreamComponent(Data.InstanceTransformComponent[2], 11));
-			}
+			// Add the instanced location streams
+			StreamElements.Add(AccessStreamComponent(Data.InstanceOriginComponent, 8, InputStreamType));
+			StreamElements.Add(AccessStreamComponent(Data.InstanceTransformComponent[0], 9, InputStreamType));
+			StreamElements.Add(AccessStreamComponent(Data.InstanceTransformComponent[1], 10, InputStreamType));
+			StreamElements.Add(AccessStreamComponent(Data.InstanceTransformComponent[2], 11, InputStreamType));
 
 			InitDeclaration(StreamElements, InputStreamType);
 		};
-		AddDeclaration(EVertexInputStreamType::PositionOnly, bInstanced, false);
-		AddDeclaration(EVertexInputStreamType::PositionAndNormalOnly, bInstanced, true);
+		AddDeclaration(EVertexInputStreamType::PositionOnly);
+		AddDeclaration(EVertexInputStreamType::PositionAndNormalOnly);
 	}
 #endif
 
