Updating a comment

#rnx

diff --git a/Engine/Source/Runtime/Engine/Public/MaterialShared.h b/Engine/Source/Runtime/Engine/Public/MaterialShared.h
--- a/Engine/Source/Runtime/Engine/Public/MaterialShared.h
+++ b/Engine/Source/Runtime/Engine/Public/MaterialShared.h
@@ -87,7 +87,7 @@ template <class ElementType> class TLinkedList;
 
 #define HANDLE_CUSTOM_OUTPUTS_AS_MATERIAL_ATTRIBUTES 0
 
-// GPUCULL_TODO: This seems it does not even compile if turned off?
+// GPUCULL_TODO: This seems it does not even compile if turned off? (EDIT: compilation issues fixed in 18733736)
 #define ALLOW_DITHERED_LOD_FOR_INSTANCED_STATIC_MESHES (1)
 
 // disallow debug data in shipping or on non-desktop Test
