Build Mesh DDC key based on whether mesh is static or skeletal.

diff --git a/Engine/Source/Developer/MeshBuilder/Public/IMeshBuilderModule.h b/Engine/Source/Developer/MeshBuilder/Public/IMeshBuilderModule.h
--- a/Engine/Source/Developer/MeshBuilder/Public/IMeshBuilderModule.h
+++ b/Engine/Source/Developer/MeshBuilder/Public/IMeshBuilderModule.h
@@ -27,7 +27,7 @@ public:
 		return GetForPlatform(TargetPlatform);
 	}
 
-	virtual void AppendToDDCKey(FString& DDCKey) { }
+	virtual void AppendToDDCKey(FString& DDCKey, bool bSkeletal) { }
 
 	virtual bool BuildMesh(class FStaticMeshRenderData& OutRenderData, class UObject* Mesh, const class FStaticMeshLODGroup& LODGroup, bool bGenerateCoarseMeshStreamingLODs) = 0;
 
diff --git a/Engine/Source/Runtime/Engine/Private/SkeletalMeshRenderData.cpp b/Engine/Source/Runtime/Engine/Private/SkeletalMeshRenderData.cpp
--- a/Engine/Source/Runtime/Engine/Private/SkeletalMeshRenderData.cpp
+++ b/Engine/Source/Runtime/Engine/Private/SkeletalMeshRenderData.cpp
@@ -246,7 +246,7 @@ FString BuildSkeletalMeshDerivedDataKey(const ITargetPlatform* TargetPlatform, U
 		KeySuffix += TEXT("_MinMLOD");
 	}
 
-	IMeshBuilderModule::GetForPlatform(TargetPlatform).AppendToDDCKey(KeySuffix);
+	IMeshBuilderModule::GetForPlatform(TargetPlatform).AppendToDDCKey(KeySuffix, true);
 	const bool bUnlimitedBoneInfluences = FGPUBaseSkinVertexFactory::GetUnlimitedBoneInfluences();
 	KeySuffix += bUnlimitedBoneInfluences ? "1" : "0";
 
diff --git a/Engine/Source/Runtime/Engine/Private/StaticMesh.cpp b/Engine/Source/Runtime/Engine/Private/StaticMesh.cpp
--- a/Engine/Source/Runtime/Engine/Private/StaticMesh.cpp
+++ b/Engine/Source/Runtime/Engine/Private/StaticMesh.cpp
@@ -2594,7 +2594,7 @@ static FString BuildStaticMeshDerivedDataKeySuffix(const ITargetPlatform* Target
 		}
 	}
 
-	IMeshBuilderModule::GetForPlatform(TargetPlatform).AppendToDDCKey(KeySuffix);
+	IMeshBuilderModule::GetForPlatform(TargetPlatform).AppendToDDCKey(KeySuffix, false);
 
 	if (TargetPlatform->GetPlatformInfo().PlatformGroupName == TEXT("Desktop")
 		&& CVarStripMinLodDataDuringCooking.GetValueOnAnyThread() != 0
