Fix for studioanalytics ensure firing during long slow tasks dialogs

diff --git a/Engine/Source/Runtime/Engine/Private/Analytics/StudioAnalytics.cpp b/Engine/Source/Runtime/Engine/Private/Analytics/StudioAnalytics.cpp
--- a/Engine/Source/Runtime/Engine/Private/Analytics/StudioAnalytics.cpp
+++ b/Engine/Source/Runtime/Engine/Private/Analytics/StudioAnalytics.cpp
@@ -153,7 +153,7 @@ void FStudioAnalytics::FireEvent_Loading(const FString& LoadingName, double Seco
 	}
 
 	// Throw out anything over 10 hours - 
-	if (!ensureMsgf(SecondsSpentLoading < 36000, TEXT("The loading event shouldn't be over 10 hours, perhaps an uninitialized bit of memory?")))
+	if (SecondsSpentLoading > 36000)
 	{
 		return;
 	}
