Add method to override platform memory size bucket when previewing DPs.
Store any matched fragments in the DP instance when previewing.

diff --git a/Engine/Source/Runtime/Engine/Classes/DeviceProfiles/DeviceProfile.h b/Engine/Source/Runtime/Engine/Classes/DeviceProfiles/DeviceProfile.h
--- a/Engine/Source/Runtime/Engine/Classes/DeviceProfiles/DeviceProfile.h
+++ b/Engine/Source/Runtime/Engine/Classes/DeviceProfiles/DeviceProfile.h
@@ -167,11 +167,18 @@ public:
 	void AddExpandedCVars(const TMap<FString, FString>& CVarsToMerge);
 	void AddPreviewCVars(const TMap<FString, FString>& CVarsToMerge);
 	void ClearAllExpandedCVars();
+	/** Set the memory size bucket to be used when previewing this DP, changing this will reset the expanded cvars. */
+	void SetPreviewMemorySizeBucket(EPlatformMemorySizeBucket PreviewMemorySizeBucketIn);
+	EPlatformMemorySizeBucket GetPreviewMemorySizeBucket() const;
+
 private:
 	/** Resolved CVars, including expanded scalability cvars used to properly emulate one platform on another */
 	TMap<FString, FString> AllExpandedCVars;
 
 	/** The set of cvars that can be previewed (a subset of AllExpandedCVars) */
 	TMap<FString, FString> AllPreviewCVars;
+
+	/** The EPlatformMemorySizeBucket to use when processing the device profile for previewing. */
+	EPlatformMemorySizeBucket PreviewMemorySizeBucket = EPlatformMemorySizeBucket::Default;
 #endif
 };
diff --git a/Engine/Source/Runtime/Engine/Private/DeviceProfiles/DeviceProfile.cpp b/Engine/Source/Runtime/Engine/Private/DeviceProfiles/DeviceProfile.cpp
--- a/Engine/Source/Runtime/Engine/Private/DeviceProfiles/DeviceProfile.cpp
+++ b/Engine/Source/Runtime/Engine/Private/DeviceProfiles/DeviceProfile.cpp
@@ -420,6 +420,19 @@ void UDeviceProfile::ClearAllExpandedCVars()
 	AllPreviewCVars.Empty();
 }
 
+void UDeviceProfile::SetPreviewMemorySizeBucket(EPlatformMemorySizeBucket PreviewMemorySizeBucketIn)
+{ 
+	if (PreviewMemorySizeBucket != PreviewMemorySizeBucketIn)
+	{
+		PreviewMemorySizeBucket = PreviewMemorySizeBucketIn;
+		// If this changes then any cached cvars are likely to be invalid too.
+		ClearAllExpandedCVars();
+	}
+}
 
+EPlatformMemorySizeBucket UDeviceProfile::GetPreviewMemorySizeBucket() const
+{ 
+	return PreviewMemorySizeBucket;
+}
 
 #endif
\ No newline at end of file
diff --git a/Engine/Source/Runtime/Engine/Private/DeviceProfiles/DeviceProfileManager.cpp b/Engine/Source/Runtime/Engine/Private/DeviceProfiles/DeviceProfileManager.cpp
--- a/Engine/Source/Runtime/Engine/Private/DeviceProfiles/DeviceProfileManager.cpp
+++ b/Engine/Source/Runtime/Engine/Private/DeviceProfiles/DeviceProfileManager.cpp
@@ -145,6 +145,8 @@ static void ExpandScalabilityCVar(FConfigCacheIni* ConfigSystem, const FString&
 void UDeviceProfileManager::ProcessDeviceProfileIniSettings(const FString& DeviceProfileName, EDeviceProfileMode Mode)
 {
 	FConfigCacheIni* ConfigSystem = GConfig;
+	EPlatformMemorySizeBucket PreviewMemorySizeBucket = EPlatformMemorySizeBucket::Default;
+
 	if (Mode == EDeviceProfileMode::DPM_CacheValues)
 	{
 #if ALLOW_OTHER_PLATFORM_CONFIG
@@ -153,6 +155,7 @@ void UDeviceProfileManager::ProcessDeviceProfileIniSettings(const FString& Devic
 		check(Profile);
 		// use the DP's platform's configs, NOT the running platform
 		ConfigSystem = FConfigCacheIni::ForPlatform(*Profile->DeviceType);
+		PreviewMemorySizeBucket = Profile->GetPreviewMemorySizeBucket();
 #else
 		checkNoEntry();
 #endif
@@ -321,10 +324,10 @@ void UDeviceProfileManager::ProcessDeviceProfileIniSettings(const FString& Devic
 				FString ArrayName = TEXT("CVars");
 				if (Pass == 0)
 				{
-					// assume default when caching for another platform, since we don' thave a current device to emulate (maybe we want to be able to pass in override memory bucket?)
+					// PreviewMemorySizeBucket is used when previewing, since we don't have a current device to emulate.
 					if (Mode == EDeviceProfileMode::DPM_CacheValues)
 					{
-						ArrayName += TEXT("_Default");
+						ArrayName += BucketNames[(int32)PreviewMemorySizeBucket];
 					}
 					else
 					{
@@ -468,6 +471,8 @@ void UDeviceProfileManager::ProcessDeviceProfileIniSettings(const FString& Devic
 
 		// now remember all previewable cvars
 		Profile->AddPreviewCVars(PreviewableCVars);
+
+		Profile->SelectedFragments = FragmentsSelected;
 	}
 #endif
 }
@@ -1111,15 +1116,13 @@ bool UDeviceProfileManager::AreTextureGroupsTheSame(UDeviceProfile* Profile1, UD
 IDeviceProfileSelectorModule* UDeviceProfileManager::GetPreviewDeviceProfileSelectorModule(FConfigCacheIni* PreviewConfigSystemIn)
 {
 #if ALLOW_OTHER_PLATFORM_CONFIG && WITH_EDITOR
-	// If we're getting the selector for previewing the PIEPreviewDeviceProfileSelector device selector can be given a PreviewDeviceDesciption to return selector info for specific devices.
-	FString PreviewDeviceDesciption; 
-	if (PreviewConfigSystemIn->GetString(TEXT("DeviceProfileManager"), TEXT("PreviewDeviceDesciption"), PreviewDeviceDesciption, GEngineIni))
+	// If we're getting the selector for previewing the PreviewDeviceProfileSelectionModule, this can be separately configured to preview a target device.
+	FString PreviewDeviceProfileSelectionModuleName;
+	if (PreviewConfigSystemIn->GetString(TEXT("DeviceProfileManager"), TEXT("PreviewDeviceProfileSelectionModule"), PreviewDeviceProfileSelectionModuleName, GEngineIni))
 	{
 		// this should only be specified when previewing.
-		if (IPIEPreviewDeviceModule* DPSelectorModule = FModuleManager::LoadModulePtr<IPIEPreviewDeviceModule>(TEXT("PIEPreviewDeviceProfileSelector")))
+		if (IPIEPreviewDeviceModule* DPSelectorModule = FModuleManager::LoadModulePtr<IPIEPreviewDeviceModule>(*PreviewDeviceProfileSelectionModuleName))
 		{
-			// TODO: check success.
-			DPSelectorModule->SetPreviewDevice(PreviewDeviceDesciption);
 			return DPSelectorModule;
 		}
 	}
