Adding net.QueuedBunchTimeoutSeconds instead of using a hardcoded value
Adding additional logging when the queued timeout is hit

diff --git a/Engine/Source/Runtime/Engine/Classes/Engine/PackageMapClient.h b/Engine/Source/Runtime/Engine/Classes/Engine/PackageMapClient.h
--- a/Engine/Source/Runtime/Engine/Classes/Engine/PackageMapClient.h
+++ b/Engine/Source/Runtime/Engine/Classes/Engine/PackageMapClient.h
@@ -255,6 +255,7 @@ public:
 	void			SetAsyncLoadMode( const EAsyncLoadMode NewMode );
 	bool			ShouldAsyncLoad() const;
 	bool			CanClientLoadObject( const UObject* Object, const FNetworkGUID& NetGUID ) const;
+	FString			Describe(const FNetworkGUID& NetGUID) const;
 
 	void			AsyncPackageCallback(const FName& PackageName, UPackage * Package, EAsyncLoadingResult::Type Result);
 
diff --git a/Engine/Source/Runtime/Engine/Private/DataChannel.cpp b/Engine/Source/Runtime/Engine/Private/DataChannel.cpp
--- a/Engine/Source/Runtime/Engine/Private/DataChannel.cpp
+++ b/Engine/Source/Runtime/Engine/Private/DataChannel.cpp
@@ -79,10 +79,8 @@ FAutoConsoleVariableRef CVarNetSkipReplicatorForDestructionInfos(
 	GSkipReplicatorForDestructionInfos,
 	TEXT("If enabled, skip creation of object replicator in SetChannelActor when we know there is no content payload and we're going to immediately destroy the actor."));
 
-extern TAutoConsoleVariable<int32> CVarFilterGuidRemapping;
 extern NETCORE_API TAutoConsoleVariable<int32> CVarNetEnableDetailedScopeCounters;
 
-
 // Fairly large number, and probably a bad idea to even have a bunch this size, but want to be safe for now and not throw out legitimate data
 static int32 NetMaxConstructedPartialBunchSizeBytes = 1024 * 64;
 static FAutoConsoleVariableRef CVarNetMaxConstructedPartialBunchSizeBytes(
@@ -104,6 +102,20 @@ static TAutoConsoleVariable<bool> CVarEnableNetInitialSubObjects(
 	TEXT("Enables new SubObjects to set bNetInitial to true to make sure all replicated properties are replicated.")
 );
 
+namespace UE
+{
+	namespace Net
+	{
+		extern int32 FilterGuidRemapping;
+
+		static float QueuedBunchTimeoutSeconds = 30.0f;
+		static FAutoConsoleVariableRef CVarQueuedBunchTimeoutSeconds(
+			TEXT("net.QueuedBunchTimeoutSeconds"),
+			QueuedBunchTimeoutSeconds,
+			TEXT("Time in seconds to wait for queued bunches on a channel to flush before logging a warning."));
+	};
+};
+
 template<typename T>
 static const bool IsBunchTooLarge(UNetConnection* Connection, T* Bunch)
 {
@@ -1986,7 +1998,7 @@ void UControlChannel::ReceiveDestructionInfo(FInBunch& Bunch)
 					TheActor->PreDestroyFromReplication();
 					TheActor->Destroy(true);
 
-					if (CVarFilterGuidRemapping.GetValueOnAnyThread() > 0)
+					if (UE::Net::FilterGuidRemapping != 0)
 					{
 						// Remove this actor's NetGUID from the list of unmapped values, it will be added back if it replicates again
 						if (NetGUID.IsValid() && Connection != nullptr && Connection->Driver != nullptr && Connection->Driver->GuidCache.IsValid())
@@ -2161,7 +2173,7 @@ void UActorChannel::DestroyActorAndComponents()
 		Actor->Destroy(true);
 	}
 
-	if (CVarFilterGuidRemapping.GetValueOnAnyThread() > 0)
+	if (UE::Net::FilterGuidRemapping != 0)
 	{
 		// Remove this actor's NetGUID from the list of unmapped values, it will be added back if it replicates again
 		if (ActorNetGUID.IsValid() && Connection != nullptr && Connection->Driver != nullptr && Connection->Driver->GuidCache.IsValid())
@@ -2595,8 +2607,7 @@ bool UActorChannel::ProcessQueuedBunches()
 		}
 		else
 		{
-			const double QUEUED_BUNCH_TIMEOUT_IN_SECONDS = 30;
-			if ((FPlatformTime::Seconds() - QueuedBunchStartTime) > QUEUED_BUNCH_TIMEOUT_IN_SECONDS)
+			if ((FPlatformTime::Seconds() - QueuedBunchStartTime) > UE::Net::QueuedBunchTimeoutSeconds)
 			{
 				if (!bSuppressQueuedBunchWarningsDueToHitches && FPlatformProperties::RequiresCookedData())
 				{
@@ -2606,9 +2617,10 @@ bool UActorChannel::ProcessQueuedBunches()
 					{
 						for (const FNetworkGUID& Guid : PendingGuidResolves)
 						{
-							const bool bIsPending = Connection->Driver->GuidCache->IsGUIDPending(Guid);
-							const FString GuidPath = Connection->Driver->GuidCache->FullNetGUIDPath(Guid);
-							UE_LOG(LogNet, Log, TEXT("  PendingGuidResolve %s. Pending in cache: %d, path: %s"), *Guid.ToString(), !!bIsPending, *GuidPath);
+							UE_LOG(LogNet, Log, TEXT("  PendingGuidResolve [%s] %s"), *Guid.ToString(), *Connection->Driver->GuidCache->Describe(Guid));
+
+							const FNetworkGUID OuterGuid = Connection->Driver->GuidCache->GetOuterNetGUID(Guid);
+							UE_LOG(LogNet, Log, TEXT("      OuterGuid [%s] %s"), *OuterGuid.ToString(), *Connection->Driver->GuidCache->Describe(OuterGuid));
 						}
 					}
 				}
diff --git a/Engine/Source/Runtime/Engine/Private/NetDriver.cpp b/Engine/Source/Runtime/Engine/Private/NetDriver.cpp
--- a/Engine/Source/Runtime/Engine/Private/NetDriver.cpp
+++ b/Engine/Source/Runtime/Engine/Private/NetDriver.cpp
@@ -672,7 +672,16 @@ void UNetDriver::CancelAdaptiveReplication(FNetworkObjectInfo& InNetworkActor)
 
 static TAutoConsoleVariable<int32> CVarOptimizedRemapping( TEXT( "net.OptimizedRemapping" ), 1, TEXT( "Uses optimized path to remap unmapped network guids" ) );
 static TAutoConsoleVariable<int32> CVarMaxClientGuidRemaps( TEXT( "net.MaxClientGuidRemaps" ), 100, TEXT( "Max client resolves of unmapped network guids per tick" ) );
-TAutoConsoleVariable<int32> CVarFilterGuidRemapping( TEXT( "net.FilterGuidRemapping" ), 1, TEXT( "Remove destroyed and parent guids from unmapped list" ) );
+
+namespace UE
+{
+	namespace Net
+	{
+		int32 FilterGuidRemapping = 1;
+		static FAutoConsoleVariableRef CVarFilterGuidRemapping(TEXT("net.FilterGuidRemapping"), FilterGuidRemapping, TEXT("Remove destroyed and parent guids from unmapped list"));
+	};
+};
+
 bool CVar_NetDriver_ReportGameTickFlushTime = false;
 static FAutoConsoleVariableRef CVarNetDriverReportTickFlushTime(TEXT("net.ReportGameTickFlushTime"), CVar_NetDriver_ReportGameTickFlushTime, TEXT("Record and report to the perf tracking system the processing time of the GameNetDriver's TickFlush."), ECVF_Default);
 
@@ -1183,14 +1192,14 @@ void UNetDriver::TickFlush(float DeltaSeconds)
 		if (CVarOptimizedRemapping.GetValueOnAnyThread() && GuidCache.IsValid())
 		{
 			// Go over recently imported network guids, and see if there are any replicators that need to map them
-			TSet< FNetworkGUID >& ImportedNetGuids = GuidCache->ImportedNetGuids;
-			TMap< FNetworkGUID, TSet< FNetworkGUID > >& PendingOuterNetGuids = GuidCache->PendingOuterNetGuids;
+			TSet<FNetworkGUID>& ImportedNetGuidsRef = GuidCache->ImportedNetGuids;
+			TMap<FNetworkGUID, TSet<FNetworkGUID>>& PendingOuterNetGuidsRef = GuidCache->PendingOuterNetGuids;
 
-			SET_DWORD_STAT(STAT_ImportedNetGuids, ImportedNetGuids.Num());
-			SET_DWORD_STAT(STAT_PendingOuterNetGuids, PendingOuterNetGuids.Num());
+			SET_DWORD_STAT(STAT_ImportedNetGuids, ImportedNetGuidsRef.Num());
+			SET_DWORD_STAT(STAT_PendingOuterNetGuids, PendingOuterNetGuidsRef.Num());
 			SET_DWORD_STAT(STAT_UnmappedReplicators, UnmappedReplicators.Num());
 
-			TSet< FObjectReplicator* > ForceUpdateReplicators;
+			TSet<FObjectReplicator*> ForceUpdateReplicators;
 
 			for (FObjectReplicator* Replicator : UnmappedReplicators)
 			{
@@ -1201,16 +1210,15 @@ void UNetDriver::TickFlush(float DeltaSeconds)
 				}
 			}
 
-			if (ImportedNetGuids.Num() || ForceUpdateReplicators.Num())
+			if (ImportedNetGuidsRef.Num() || ForceUpdateReplicators.Num())
 			{
 				int32 NumRemapTests = 0;
-				int32 MaxRemaps = IsServer() ? 0 : CVarMaxClientGuidRemaps.GetValueOnAnyThread();
-				bool bFilterGuidRemapping = (CVarFilterGuidRemapping.GetValueOnAnyThread() > 0);
+				const int32 MaxRemaps = IsServer() ? 0 : CVarMaxClientGuidRemaps.GetValueOnAnyThread();
 
-				TArray< FNetworkGUID > UnmappedGuids;
-				TArray< FNetworkGUID > NewlyMappedGuids;
+				TArray<FNetworkGUID> UnmappedGuids;
+				TArray<FNetworkGUID> NewlyMappedGuids;
 
-				for (auto It = ImportedNetGuids.CreateIterator(); It; ++It)
+				for (auto It = ImportedNetGuidsRef.CreateIterator(); It; ++It)
 				{
 					const FNetworkGUID NetworkGuid = *It;
 					bool bMappedOrBroken = false;
@@ -1230,15 +1238,15 @@ void UNetDriver::TickFlush(float DeltaSeconds)
 
 					if (!bMappedOrBroken)
 					{
-						FNetworkGUID OuterGUID = GuidCache->GetOuterNetGUID(NetworkGuid);
+						const FNetworkGUID OuterGUID = GuidCache->GetOuterNetGUID(NetworkGuid);
 
 						// we're missing the outer, stop checking until we map it
-						if (bFilterGuidRemapping && OuterGUID.IsValid() && !OuterGUID.IsDefault() && !GuidCache->IsGUIDLoaded(OuterGUID) && !GuidCache->IsGUIDPending(OuterGUID))
+						if ((UE::Net::FilterGuidRemapping != 0) && OuterGUID.IsValid() && !OuterGUID.IsDefault() && !GuidCache->IsGUIDLoaded(OuterGUID) && !GuidCache->IsGUIDPending(OuterGUID))
 						{
-							UE_LOG(LogNetPackageMap, Log, TEXT("Missing outer (%s) for unmapped guid (%s), marking pending"), *GuidCache->FullNetGUIDPath(OuterGUID), *GuidCache->FullNetGUIDPath(NetworkGuid));
+							UE_LOG(LogNetPackageMap, Log, TEXT("Missing outer (%s) for unmapped guid (%s), marking pending"), *GuidCache->Describe(OuterGUID), *GuidCache->Describe(NetworkGuid));
 
-							TSet<FNetworkGUID>& PendingGuids = PendingOuterNetGuids.FindOrAdd(OuterGUID);
-							PendingGuids.Add(NetworkGuid);
+							TSet<FNetworkGUID>& PendingGuidsRef = PendingOuterNetGuidsRef.FindOrAdd(OuterGUID);
+							PendingGuidsRef.Add(NetworkGuid);
 						}
 						else
 						{
@@ -1260,38 +1268,35 @@ void UNetDriver::TickFlush(float DeltaSeconds)
 				// attempt to resolve dependent guids next tick (outer is now mapped)
 				for (const FNetworkGUID& NetGuid : NewlyMappedGuids)
 				{
-					TSet< FNetworkGUID >* DependentGuids = PendingOuterNetGuids.Find(NetGuid);
-					if (DependentGuids)
+					if (TSet<FNetworkGUID>* DependentGuids = PendingOuterNetGuidsRef.Find(NetGuid))
 					{
 						UE_LOG(LogNetPackageMap, Log, TEXT("Newly mapped outer (%s) removing from pending"), *GuidCache->FullNetGUIDPath(NetGuid));
 
-						ImportedNetGuids.Append(*DependentGuids);
-						PendingOuterNetGuids.Remove(NetGuid);
+						ImportedNetGuidsRef.Append(*DependentGuids);
+						PendingOuterNetGuidsRef.Remove(NetGuid);
 					}
 				}
 
 				// any tested guids that could not yet be mapped are added back to the list
 				if (UnmappedGuids.Num())
 				{
-					ImportedNetGuids.CompactStable();
-					ImportedNetGuids.Append(UnmappedGuids);
+					ImportedNetGuidsRef.CompactStable();
+					ImportedNetGuidsRef.Append(UnmappedGuids);
 				}
 
 				if (NewlyMappedGuids.Num() || ForceUpdateReplicators.Num())
 				{
-					TSet< FObjectReplicator* > AllReplicators = ForceUpdateReplicators;
+					TSet<FObjectReplicator*> ReplicatorsToUpdate = MoveTemp(ForceUpdateReplicators);
 
 					for (const FNetworkGUID& NetGuid : NewlyMappedGuids)
 					{
-						TSet< FObjectReplicator* >* Replicators = GuidToReplicatorMap.Find(NetGuid);
-
-						if (Replicators)
+						if (TSet<FObjectReplicator*>* Replicators = GuidToReplicatorMap.Find(NetGuid))
 						{
-							AllReplicators.Append(*Replicators);
+							ReplicatorsToUpdate.Append(*Replicators);
 						}
 					}
 
-					for (FObjectReplicator* Replicator : AllReplicators)
+					for (FObjectReplicator* Replicator : ReplicatorsToUpdate)
 					{
 						if (UnmappedReplicators.Contains(Replicator))
 						{
diff --git a/Engine/Source/Runtime/Engine/Private/PackageMapClient.cpp b/Engine/Source/Runtime/Engine/Private/PackageMapClient.cpp
--- a/Engine/Source/Runtime/Engine/Private/PackageMapClient.cpp
+++ b/Engine/Source/Runtime/Engine/Private/PackageMapClient.cpp
@@ -54,7 +54,14 @@ CSV_DEFINE_CATEGORY(PackageMap, true);
 static const int INTERNAL_LOAD_OBJECT_RECURSION_LIMIT = 16;
 
 extern FAutoConsoleVariableRef CVarEnableMultiplayerWorldOriginRebasing;
-extern TAutoConsoleVariable<int32> CVarFilterGuidRemapping;
+
+namespace UE
+{
+	namespace Net
+	{
+		extern int32 FilterGuidRemapping;
+	};
+};
 
 static TAutoConsoleVariable<int32> CVarAllowAsyncLoading(
 	TEXT("net.AllowAsyncLoading"),
@@ -409,8 +416,7 @@ bool UPackageMapClient::SerializeNewActor(FArchive& Ar, class UActorChannel *Cha
 		return false;
 	}
 
-	bool bFilterGuidRemapping = (CVarFilterGuidRemapping.GetValueOnAnyThread() > 0);
-	if (!bFilterGuidRemapping)
+	if (UE::Net::FilterGuidRemapping == 0)
 	{
 		if ( GuidCache.IsValid() )
 		{
@@ -447,7 +453,7 @@ bool UPackageMapClient::SerializeNewActor(FArchive& Ar, class UActorChannel *Cha
 		return false;		// This doesn't mean an error. This just simply means we didn't spawn an actor.
 	}
 
-	if (bFilterGuidRemapping)
+	if (UE::Net::FilterGuidRemapping != 0)
 	{
 		// Do not mark guid as imported until we know we aren't deleting it
 		if ( GuidCache.IsValid() )
@@ -679,7 +685,7 @@ bool UPackageMapClient::SerializeNewActor(FArchive& Ar, class UActorChannel *Cha
 			UE_LOG( LogNetPackageMap, Log, TEXT( "SerializeNewActor: Failed to find static actor: FullNetGuidPath: %s, Channel: %d" ), *GuidCache->FullNetGUIDPath( NetGUID ), Channel->ChIndex );
 		}
 
-		if (bFilterGuidRemapping)
+		if (UE::Net::FilterGuidRemapping != 0)
 		{
 			// Do not attempt to resolve this missing actor
 			if ( GuidCache.IsValid() )
@@ -3233,7 +3239,7 @@ UObject* FNetGUIDCache::GetObjectFromNetGUID( const FNetworkGUID& NetGUID, const
 			return NULL;
 		}
 
-		// If outer is broken, we will never load, set outselves to broken as well and bail
+		// If outer is broken, we will never load, set ourselves to broken as well and bail
 		if ( OuterCacheObject->bIsBroken )
 		{
 			UE_LOG( LogNetPackageMap, Error, TEXT( "GetObjectFromNetGUID: Outer is broken. FullNetGUIDPath: %s" ), *FullNetGUIDPath( NetGUID ) );
@@ -3566,6 +3572,22 @@ FString FNetGUIDCache::FullNetGUIDPath( const FNetworkGUID& NetGUID ) const
 	return FullPath;
 }
 
+FString	FNetGUIDCache::Describe(const FNetworkGUID& NetGUID) const
+{
+	FString Desc = FString::Printf(TEXT("NetworkGUID [%s]"), *NetGUID.ToString());
+
+	if (FNetGuidCacheObject const* const CacheObjectPtr = GetCacheObject(NetGUID))
+	{
+		Desc += FString::Printf(TEXT(" NoLoad [%d] Pending [%d] Broken [%d] Outer [%s] FullPath [%s] Object [%s]"), !!CacheObjectPtr->bNoLoad, !!CacheObjectPtr->bIsPending, !!CacheObjectPtr->bIsBroken, *CacheObjectPtr->OuterGUID.ToString(), *FullNetGUIDPath(NetGUID), *GetFullNameSafe(CacheObjectPtr->Object.Get()));
+	}
+	else
+	{
+		Desc += TEXT(" Unregistered");
+	}
+
+	return Desc;
+}
+
 void FNetGUIDCache::GenerateFullNetGUIDPath_r( const FNetworkGUID& NetGUID, FString& FullPath ) const
 {
 	if ( !NetGUID.IsValid() )
