Fix build fail from partial merge.

diff --git a/Engine/Source/Runtime/Engine/Classes/DeviceProfiles/DeviceProfile.h b/Engine/Source/Runtime/Engine/Classes/DeviceProfiles/DeviceProfile.h
--- a/Engine/Source/Runtime/Engine/Classes/DeviceProfiles/DeviceProfile.h
+++ b/Engine/Source/Runtime/Engine/Classes/DeviceProfiles/DeviceProfile.h
@@ -100,6 +100,9 @@ private:
 	FOnCVarsUpdated CVarsUpdatedDelegate;
 
 public:
+	// The selected result after running the MatchingRules process.
+	TArray<FSelectedFragmentProperties> SelectedFragments;
+
 	/* ValidateProfile()
 	* Validate the Profile after changes by loading it's config (.ini)
 	*/
