Adding scoped CPU markers for the most expensive parts of `GlobalBeginCompileShader` and `FShaderCompilingManager::SubmitJobs`.

- FShaderType::AddReferencedUniformBufferIncludes
- GetInputHash

diff --git a/Engine/Source/Runtime/Engine/Private/ShaderCompiler/ShaderCompiler.cpp b/Engine/Source/Runtime/Engine/Private/ShaderCompiler/ShaderCompiler.cpp
--- a/Engine/Source/Runtime/Engine/Private/ShaderCompiler/ShaderCompiler.cpp
+++ b/Engine/Source/Runtime/Engine/Private/ShaderCompiler/ShaderCompiler.cpp
@@ -524,6 +524,7 @@ void FShaderCompileJobCollection::SubmitJobs(const TArray<FShaderCommonCompileJo
 			// Just precompute the InputHash for each job in multiple-thread.
 			if (ShaderCompiler::IsJobCacheEnabled())
 			{
+				TRACE_CPUPROFILER_EVENT_SCOPE(GetInputHash);
 				ParallelFor(InJobs.Num(), [&InJobs](int32 Index) { InJobs[Index]->GetInputHash(); });
 			}
 
@@ -3760,7 +3761,7 @@ void FShaderCompilingManager::SubmitJobs(TArray<FShaderCommonCompileJobPtr>& New
 	{
 		return;
 	}
-
+	
 	TRACE_CPUPROFILER_EVENT_SCOPE(FShaderCompilingManager::SubmitJobs);
 	check(!FPlatformProperties::RequiresCookedData());
 
diff --git a/Engine/Source/Runtime/RenderCore/Private/ShaderParameters.cpp b/Engine/Source/Runtime/RenderCore/Private/ShaderParameters.cpp
--- a/Engine/Source/Runtime/RenderCore/Private/ShaderParameters.cpp
+++ b/Engine/Source/Runtime/RenderCore/Private/ShaderParameters.cpp
@@ -347,6 +347,8 @@ void FShaderType::FlushShaderFileCache(const TMap<FString, TArray<const TCHAR*>
 
 void FShaderType::AddReferencedUniformBufferIncludes(FShaderCompilerEnvironment& OutEnvironment, FString& OutSourceFilePrefix, EShaderPlatform Platform) const
 {
+	TRACE_CPUPROFILER_EVENT_SCOPE(FShaderType::AddReferencedUniformBufferIncludes);
+
 	// Cache uniform buffer struct declarations referenced by this shader type's files
 	if (CachedUniformBufferPlatform != Platform)
 	{
