Fix build issue.

diff --git a/Engine/Source/Runtime/Engine/Classes/DeviceProfiles/DeviceProfile.h b/Engine/Source/Runtime/Engine/Classes/DeviceProfiles/DeviceProfile.h
--- a/Engine/Source/Runtime/Engine/Classes/DeviceProfiles/DeviceProfile.h
+++ b/Engine/Source/Runtime/Engine/Classes/DeviceProfiles/DeviceProfile.h
@@ -82,6 +82,12 @@ public:
 		return CVarsUpdatedDelegate;
 	}
 
+	/** 
+	 * Accessor to a fragment by tag.
+	 * @param FragmentTag the name of the tag to find
+	 * @return FSelectedFragmentProperties ptr or null if FragmentTag not found.
+	 */	
+	const FSelectedFragmentProperties* GetFragmentByTag(FName& FragmentTag) const;
 public:
 	/** 
 	 * Access to the device profiles Texture LOD Settings
diff --git a/Engine/Source/Runtime/Engine/Private/DeviceProfiles/DeviceProfile.cpp b/Engine/Source/Runtime/Engine/Private/DeviceProfiles/DeviceProfile.cpp
--- a/Engine/Source/Runtime/Engine/Private/DeviceProfiles/DeviceProfile.cpp
+++ b/Engine/Source/Runtime/Engine/Private/DeviceProfiles/DeviceProfile.cpp
@@ -31,6 +31,17 @@ UDeviceProfile::UDeviceProfile(const FObjectInitializer& ObjectInitializer)
 //	LoadConfig(GetClass(), *DeviceProfileFileName, UE::LCPF_ReadParentSections);
 }
 
+const FSelectedFragmentProperties* UDeviceProfile::GetFragmentByTag(FName& FragmentTag) const
+{
+	for (const FSelectedFragmentProperties& SelectedFragment : SelectedFragments)
+	{
+		if (SelectedFragment.Tag == FragmentTag)
+		{
+			return &SelectedFragment;
+		}
+	}
+	return nullptr;
+}
 
 void UDeviceProfile::GatherParentCVarInformationRecursively(OUT TMap<FString, FString>& CVarInformation) const
 {
