FSkeletalMeshObjectGPUSkin would leak references to an RHIBuffer by holding an SRV and not clear it during ReleaseRHI (when using Apex Cloth).

diff --git a/Engine/Source/Runtime/Engine/Private/GPUSkinVertexFactory.cpp b/Engine/Source/Runtime/Engine/Private/GPUSkinVertexFactory.cpp
--- a/Engine/Source/Runtime/Engine/Private/GPUSkinVertexFactory.cpp
+++ b/Engine/Source/Runtime/Engine/Private/GPUSkinVertexFactory.cpp
@@ -654,6 +654,12 @@ void TGPUSkinAPEXClothVertexFactory<BoneInfluenceType>::ReleaseDynamicRHI()
 {
 	Super::ReleaseDynamicRHI();
 	ClothShaderData.ReleaseClothSimulData();
+
+	// Release the RHIResource reference held in FGPUSkinAPEXClothDataType
+	if (ClothDataPtr)
+	{
+		ClothDataPtr->ClothBuffer.SafeRelease();
+	}
 }
 
 /*-----------------------------------------------------------------------------
