Fix Volume not updating player collision model when brush shape is modified.
Fix Volume not updating player collision model when undoing/redoing brush shape modifications.

#rnx

diff --git a/Engine/Source/Runtime/Engine/Private/Volume.cpp b/Engine/Source/Runtime/Engine/Private/Volume.cpp
--- a/Engine/Source/Runtime/Engine/Private/Volume.cpp
+++ b/Engine/Source/Runtime/Engine/Private/Volume.cpp
@@ -47,7 +47,9 @@ void AVolume::PostEditChangeProperty(FPropertyChangedEvent& PropertyChangedEvent
 	static FName NAME_BrushBuilder(TEXT("BrushBuilder"));
 
 	// The brush builder that created this volume has changed. Notify listeners
-	if( PropertyChangedEvent.ChangeType != EPropertyChangeType::Interactive && PropertyChangedEvent.MemberProperty && PropertyChangedEvent.MemberProperty->GetFName() == NAME_BrushBuilder )
+	// Also notify on null property change events submitted during undo/redo
+	if( PropertyChangedEvent.ChangeType != EPropertyChangeType::Interactive &&
+		((GIsTransacting && !PropertyChangedEvent.MemberProperty) || (PropertyChangedEvent.MemberProperty && PropertyChangedEvent.MemberProperty->GetFName() == NAME_BrushBuilder)) )
 	{
 		OnVolumeShapeChanged.Broadcast(*this);
 	}
