Initialize FVectors in SerializeNewActor in the loading path

diff --git a/Engine/Source/Runtime/Engine/Private/PackageMapClient.cpp b/Engine/Source/Runtime/Engine/Private/PackageMapClient.cpp
--- a/Engine/Source/Runtime/Engine/Private/PackageMapClient.cpp
+++ b/Engine/Source/Runtime/Engine/Private/PackageMapClient.cpp
@@ -470,11 +470,11 @@ bool UPackageMapClient::SerializeNewActor(FArchive& Ar, class UActorChannel *Cha
 	{
 		UObject* Archetype = nullptr;
 		UObject* ActorLevel = nullptr;
-		FVector Location;
-		FVector Scale;
-		FVector Velocity;
-		FRotator Rotation;
-		bool SerSuccess;
+		FVector Location = FVector::ZeroVector;
+		FVector Scale = FVector::OneVector;
+		FVector Velocity = FVector::ZeroVector;
+		FRotator Rotation = FRotator::ZeroRotator;
+		bool SerSuccess = false;
 
 		if (Ar.IsSaving())
 		{
@@ -499,14 +499,10 @@ bool UPackageMapClient::SerializeNewActor(FArchive& Ar, class UActorChannel *Cha
 			if (RootComponent)
 			{
 				Location = FRepMovement::RebaseOntoZeroOrigin(Actor->GetActorLocation(), Actor);
-			} 
-			else
-			{
-				Location = FVector::ZeroVector;
+				Rotation = Actor->GetActorRotation();
+				Scale = Actor->GetActorScale();
+				Velocity = Actor->GetVelocity();
 			}
-			Rotation = RootComponent ? Actor->GetActorRotation() : FRotator::ZeroRotator;
-			Scale = RootComponent ? Actor->GetActorScale() : FVector::OneVector;
-			Velocity = RootComponent ? Actor->GetVelocity() : FVector::ZeroVector;
 		}
 
 		FNetworkGUID ArchetypeNetGUID;
