Added a FixedTick to ActorComponent and Actor to have a thread safe way to access physics data as part of a sim callback

diff --git a/Engine/Source/Runtime/Engine/Classes/Components/ActorComponent.h b/Engine/Source/Runtime/Engine/Classes/Components/ActorComponent.h
--- a/Engine/Source/Runtime/Engine/Classes/Components/ActorComponent.h
+++ b/Engine/Source/Runtime/Engine/Classes/Components/ActorComponent.h
@@ -291,6 +291,10 @@ private:
 	uint8 bMarkedForPreEndOfFrameSync : 1;
 
 public:
+
+	/** Whether to use use the fixed physics tick with this component. Requires async fixed tick to be enabled (see project settings) */
+	uint8 bFixedTickEnabled : 1;
+
 	/** Describes how a component instance will be created */
 	UPROPERTY()
 	EComponentCreationMethod CreationMethod;
@@ -675,6 +679,16 @@ public:
 	 * @param ThisTickFunction - Internal tick function struct that caused this to run
 	 */
 	virtual void TickComponent(float DeltaTime, enum ELevelTick TickType, FActorComponentTickFunction *ThisTickFunction);
+
+	/**
+	 * Function called for every physics step on this ActorComponent. Override this function to implement custom logic to be executed every physics step.
+	 * Only executes if the component is registered, and also bFixedTick must be set to true.
+	 * Requires async fixed tick to be enabled (see physics project settings)
+	 *	
+	 * @param DeltaTime - The fixed tick delta time
+	 * @param SimTime - This is the total sim time since the sim began, it is an integer multiple of DeltaTime
+	 */
+	virtual void FixedTickComponent(float DeltaTime, float SimTime) { ReceiveFixedTick(DeltaTime, SimTime); }
 	
 	/** 
 	 * Set up a tick function for a component in the standard way. 
@@ -981,6 +995,10 @@ public:
 	/** Event called every frame if tick is enabled */
 	UFUNCTION(BlueprintImplementableEvent, meta=(DisplayName = "Tick"))
 	void ReceiveTick(float DeltaSeconds);
+
+	/** Event called every physics tick if EnabledFixedTick is true and async fixed tick is enabled (see physics project settings)*/
+	UFUNCTION(BlueprintImplementableEvent, meta = (DisplayName = "Fixed Tick"))
+	void ReceiveFixedTick(float DeltaSeconds, float SimSeconds);
 	
 	/** 
 	 *  Called by owner actor on position shifting
diff --git a/Engine/Source/Runtime/Engine/Classes/Components/PrimitiveComponent.h b/Engine/Source/Runtime/Engine/Classes/Components/PrimitiveComponent.h
--- a/Engine/Source/Runtime/Engine/Classes/Components/PrimitiveComponent.h
+++ b/Engine/Source/Runtime/Engine/Classes/Components/PrimitiveComponent.h
@@ -1951,6 +1951,18 @@ public:
 	*/
 	virtual FBodyInstance* GetBodyInstance(FName BoneName = NAME_None, bool bGetWelded = true, int32 Index = INDEX_NONE) const;
 
+	/**
+	 * Returns BodyInstanceFixedTickHandle of the component. For use in the FixedTick event
+	*
+	* @param BoneName				Used to get body associated with specific bone. NAME_None automatically gets the root most body
+	* @param bGetWelded				If the component has been welded to another component and bGetWelded is true we return the single welded BodyInstance that is used in the simulation
+	* @param Index					Index used in Components with multiple body instances
+	*
+	* @return		Returns the BodyInstanceFixedTickHandle based on various states (does component have multiple bodies? Is the body welded to another body?)
+	*/
+	UFUNCTION(BlueprintPure, Category = "Physics")
+	FBodyInstanceFixedTickHandle GetBodyInstanceFixedTickHandle(FName BoneName = NAME_None, bool bGetWelded = true, int32 Index = -1) const;
+
 	/** 
 	 * Returns The square of the distance to closest Body Instance surface. 
 	 *
diff --git a/Engine/Source/Runtime/Engine/Classes/GameFramework/Actor.h b/Engine/Source/Runtime/Engine/Classes/GameFramework/Actor.h
--- a/Engine/Source/Runtime/Engine/Classes/GameFramework/Actor.h
+++ b/Engine/Source/Runtime/Engine/Classes/GameFramework/Actor.h
@@ -453,6 +453,10 @@ private:
 
 	/** Set if an Actor tries to be destroyed while it is beginning play so that once BeginPlay ends we can issue the destroy call. */
 	uint8 bActorWantsDestroyDuringBeginPlay : 1;
+	
+	/** Whether to use use the fixed physics tick with this actor. Requires async fixed tick to be enabled (see project settings) */
+	UPROPERTY(EditAnywhere, Category=Physics)
+	uint8 bFixedTickEnabled : 1;
 
 	/** Enum defining if BeginPlay has started or finished */
 	enum class EActorBeginPlayState : uint8
@@ -1760,6 +1764,10 @@ public:
 	UFUNCTION(BlueprintImplementableEvent, meta=(DisplayName = "Tick"))
 	void ReceiveTick(float DeltaSeconds);
 
+	/** Event called every physics tick if EnabledFixedTick is true and async fixed tick is enabled (see physics project settings)*/
+	UFUNCTION(BlueprintImplementableEvent, meta = (DisplayName = "Fixed Tick"))
+	void ReceiveFixedTick(float DeltaSeconds, float SimSeconds);
+
 	/** 
 	 *	Event when this actor overlaps another actor, for example a player walking into a trigger.
 	 *	For events when objects have a blocking collision, for example a player hitting a wall, see 'Hit' events.
@@ -2486,6 +2494,16 @@ public:
 	 */
 	virtual void TickActor( float DeltaTime, enum ELevelTick TickType, FActorTickFunction& ThisTickFunction );
 
+	/**
+	 * Function called for every physics step on this Actor. Override this function to implement custom logic to be executed every physics step.
+	 * bFixedTick must be set to true.
+	 * Requires async fixed tick to be enabled (see physics project settings)
+	 *	
+	 * @param DeltaTime - The fixed tick delta time
+	 * @param SimTime - This is the total sim time since the sim began, it is an integer multiple of DeltaTime
+	 */
+	virtual void FixedTickActor(float DeltaTime, float SimTime) { ReceiveFixedTick(DeltaTime, SimTime); }
+
 	/**
 	 * Called when an actor is done spawning into the world (from UWorld::SpawnActor), both in the editor and during gameplay
 	 * For actors with a root component, the location and rotation will have already been set.
diff --git a/Engine/Source/Runtime/Engine/Classes/PhysicsEngine/BodyInstance.h b/Engine/Source/Runtime/Engine/Classes/PhysicsEngine/BodyInstance.h
--- a/Engine/Source/Runtime/Engine/Classes/PhysicsEngine/BodyInstance.h
+++ b/Engine/Source/Runtime/Engine/Classes/PhysicsEngine/BodyInstance.h
@@ -248,6 +248,23 @@ enum class BodyInstanceSceneState : uint8
 	Removed
 };
 
+namespace Chaos
+{
+	class FRigidBodyHandle_Internal;
+}
+
+USTRUCT(BlueprintType)
+struct ENGINE_API FBodyInstanceFixedTickHandle
+{
+	GENERATED_BODY()
+	FPhysicsActorHandle Proxy = nullptr;
+
+	Chaos::FRigidBodyHandle_Internal* operator->();
+
+	bool IsValid() const;
+	operator bool() const { return IsValid(); }
+};
+
 /** Container for a physics representation of an object */
 USTRUCT(BlueprintType)
 struct ENGINE_API FBodyInstance : public FBodyInstanceCore
@@ -545,6 +562,8 @@ public:
 	// Internal physics representation of our body instance
 	FPhysicsActorHandle ActorHandle;
 
+	FBodyInstanceFixedTickHandle GetBodyInstancePhysicsThreadHandle() const { return FBodyInstanceFixedTickHandle{ ActorHandle }; }
+
 #if USE_BODYINSTANCE_DEBUG_NAMES
 	TSharedPtr<TArray<ANSICHAR>> CharDebugName;
 #endif
diff --git a/Engine/Source/Runtime/Engine/Private/Actor.cpp b/Engine/Source/Runtime/Engine/Private/Actor.cpp
--- a/Engine/Source/Runtime/Engine/Private/Actor.cpp
+++ b/Engine/Source/Runtime/Engine/Private/Actor.cpp
@@ -108,6 +108,7 @@ void AActor::InitializeDefaults()
 	PrimaryActorTick.bCanEverTick = false;
 	PrimaryActorTick.bStartWithTickEnabled = true;
 	PrimaryActorTick.SetTickFunctionEnable(false); 
+	bFixedTickEnabled = false;
 
 	CustomTimeDilation = 1.0f;
 
@@ -1142,6 +1143,23 @@ void AActor::RegisterAllActorTickFunctions(bool bRegister, bool bDoComponents)
 				}
 			}
 		}
+
+#if WITH_CHAOS
+		if (bFixedTickEnabled)
+		{
+			if (FPhysScene_Chaos* Scene = static_cast<FPhysScene_Chaos*>(GetWorld()->GetPhysicsScene()))
+			{
+				if (bRegister)
+				{
+					Scene->RegisterFixedTickActor(this);
+				}
+				else
+				{
+					Scene->UnregisterFixedTickActor(this);
+				}
+			}
+		}
+#endif
 	}
 }
 
diff --git a/Engine/Source/Runtime/Engine/Private/Components/ActorComponent.cpp b/Engine/Source/Runtime/Engine/Private/Components/ActorComponent.cpp
--- a/Engine/Source/Runtime/Engine/Private/Components/ActorComponent.cpp
+++ b/Engine/Source/Runtime/Engine/Private/Components/ActorComponent.cpp
@@ -327,6 +327,7 @@ UActorComponent::UActorComponent(const FObjectInitializer& ObjectInitializer /*=
 	bNavigationRelevant = false;
 
 	bMarkedForPreEndOfFrameSync = false;
+	bFixedTickEnabled = false;
 }
 
 void UActorComponent::PostInitProperties()
@@ -1190,7 +1191,24 @@ void UActorComponent::RegisterAllComponentTickFunctions(bool bRegister)
 			checkf(GTestRegisterComponentTickFunctions == this, TEXT("Failed to route component RegisterTickFunctions (%s)"), *GetFullName());
 			GTestRegisterComponentTickFunctions = NULL;
 		}
+
+#if WITH_CHAOS
+		if(bFixedTickEnabled)
+		{
+			if(FPhysScene_Chaos* Scene = static_cast<FPhysScene_Chaos*>(WorldPrivate->GetPhysicsScene()))
+			{
+				if(bRegister)
+				{
+					Scene->RegisterFixedTickComponent(this);
+				}
+				else
+				{
+					Scene->UnregisterFixedTickComponent(this);
+				}
+			}
+		}
 	}
+#endif
 }
 
 void UActorComponent::TickComponent(float DeltaTime, enum ELevelTick TickType, FActorComponentTickFunction *ThisTickFunction)
diff --git a/Engine/Source/Runtime/Engine/Private/PhysicsEngine/BodyInstance.cpp b/Engine/Source/Runtime/Engine/Private/PhysicsEngine/BodyInstance.cpp
--- a/Engine/Source/Runtime/Engine/Private/PhysicsEngine/BodyInstance.cpp
+++ b/Engine/Source/Runtime/Engine/Private/PhysicsEngine/BodyInstance.cpp
@@ -4512,6 +4512,24 @@ FRecalculatedMassProperties& FBodyInstance::OnRecalculatedMassProperties()
 }
 
 
+bool FBodyInstanceFixedTickHandle::IsValid() const
+{
+#if WITH_CHAOS
+	return Proxy && Proxy->GetPhysicsThreadAPI() != nullptr;
+#else
+	return false;
+#endif
+}
+
+Chaos::FRigidBodyHandle_Internal* FBodyInstanceFixedTickHandle::operator->()
+{
+#if WITH_CHAOS
+	return Proxy->GetPhysicsThreadAPI();
+#else
+	return nullptr;
+#endif
+}
+
 ////////////////////////////////////////////////////////////////////////////
 // FBodyInstanceEditorHelpers
 
diff --git a/Engine/Source/Runtime/Engine/Private/PhysicsEngine/Experimental/PhysScene_Chaos.cpp b/Engine/Source/Runtime/Engine/Private/PhysicsEngine/Experimental/PhysScene_Chaos.cpp
--- a/Engine/Source/Runtime/Engine/Private/PhysicsEngine/Experimental/PhysScene_Chaos.cpp
+++ b/Engine/Source/Runtime/Engine/Private/PhysicsEngine/Experimental/PhysScene_Chaos.cpp
@@ -71,6 +71,35 @@ DECLARE_CYCLE_STAT(TEXT("Update Kinematics On Deferred SkelMeshes"), STAT_Update
 #include "Editor.h"
 #endif
 
+class FFixedTickCallback : public Chaos::TSimCallbackObject<>
+{
+public:
+	FFixedTickCallback()
+		: TSimCallbackObject<>(/*InRunOnFrozenGameThread=*/ true)
+	{
+
+	}
+
+	TSet<UActorComponent*> FixedTickComponents;
+	TSet<AActor*> FixedTickActors;
+	virtual void OnPreSimulate_Internal() override
+	{
+		using namespace Chaos;
+		const FReal DeltaTime = GetDeltaTime_Internal();
+		const FReal SimTime = GetSimTime_Internal();
+		//TODO: handle case where callbacks modify FixedTickComponents or FixedTickActors
+		for (UActorComponent* Component : FixedTickComponents)
+		{
+			Component->FixedTickComponent(DeltaTime, SimTime);
+		}
+
+		for(AActor* Actor : FixedTickActors)
+		{
+			Actor->FixedTickActor(DeltaTime, SimTime);
+		}
+	}
+};
+
 DEFINE_LOG_CATEGORY_STATIC(LogFPhysScene_ChaosSolver, Log, All);
 
 void DumpHierarchyStats(const TArray<FString>& Args)
@@ -419,6 +448,11 @@ FPhysScene_Chaos::~FPhysScene_Chaos()
 {
 #if WITH_CHAOS
 
+	if (FixedTickCallback)
+	{
+		SceneSolver->UnregisterAndFreeSimCallbackObject_External(FixedTickCallback);
+	}
+
 	// Must ensure deferred components do not hold onto scene pointer.
 	ProcessDeferredCreatePhysicsState();
 
@@ -1813,6 +1847,41 @@ void FPhysScene_Chaos::ResimNFrames(const int32 NumFramesRequested)
 #endif
 }
 
+void FPhysScene_Chaos::EnableFixedTickCallback()
+{
+	if (FixedTickCallback == nullptr)
+	{
+		FixedTickCallback = SceneSolver->CreateAndRegisterSimCallbackObject_External<FFixedTickCallback>();
+	}
+}
+
+void FPhysScene_Chaos::RegisterFixedTickComponent(UActorComponent* Component)
+{
+	EnableFixedTickCallback();
+	FixedTickCallback->FixedTickComponents.Add(Component);
+}
+
+void FPhysScene_Chaos::UnregisterFixedTickComponent(UActorComponent* Component)
+{
+	if (FixedTickCallback)
+	{
+		FixedTickCallback->FixedTickComponents.Remove(Component);
+	}
+}
+
+void FPhysScene_Chaos::RegisterFixedTickActor(AActor* Actor)
+{
+	EnableFixedTickCallback();
+	FixedTickCallback->FixedTickActors.Add(Actor);
+}
+
+void FPhysScene_Chaos::UnregisterFixedTickActor(AActor* Actor)
+{
+	if (FixedTickCallback)
+	{
+		FixedTickCallback->FixedTickActors.Remove(Actor);
+	}
+}
 
 TSharedPtr<IPhysicsReplicationFactory> FPhysScene_Chaos::PhysicsReplicationFactory;
 
diff --git a/Engine/Source/Runtime/Engine/Private/PrimitiveComponentPhysics.cpp b/Engine/Source/Runtime/Engine/Private/PrimitiveComponentPhysics.cpp
--- a/Engine/Source/Runtime/Engine/Private/PrimitiveComponentPhysics.cpp
+++ b/Engine/Source/Runtime/Engine/Private/PrimitiveComponentPhysics.cpp
@@ -817,6 +817,19 @@ FBodyInstance* UPrimitiveComponent::GetBodyInstance(FName BoneName, bool bGetWel
 	return const_cast<FBodyInstance*>((bGetWelded && BodyInstance.WeldParent) ? BodyInstance.WeldParent : &BodyInstance);
 }
 
+FBodyInstanceFixedTickHandle UPrimitiveComponent::GetBodyInstanceFixedTickHandle(FName BoneName, bool bGetWelded, int32 Index) const
+{
+	if(FBodyInstance* BI = GetBodyInstance(BoneName, bGetWelded, Index))
+	{
+		return BI->GetBodyInstancePhysicsThreadHandle();
+	}
+	else
+	{
+		return FBodyInstanceFixedTickHandle();
+	}
+	
+}
+
 bool UPrimitiveComponent::GetSquaredDistanceToCollision(const FVector& Point, float& OutSquaredDistance, FVector& OutClosestPointOnCollision) const
 {
 	OutClosestPointOnCollision = Point;
diff --git a/Engine/Source/Runtime/Engine/Public/Physics/Experimental/PhysScene_Chaos.h b/Engine/Source/Runtime/Engine/Public/Physics/Experimental/PhysScene_Chaos.h
--- a/Engine/Source/Runtime/Engine/Public/Physics/Experimental/PhysScene_Chaos.h
+++ b/Engine/Source/Runtime/Engine/Public/Physics/Experimental/PhysScene_Chaos.h
@@ -121,6 +121,12 @@ public:
 
 	void UnRegisterForCollisionEvents(UPrimitiveComponent* Component);
 
+	void RegisterFixedTickComponent(UActorComponent* Component);
+	void UnregisterFixedTickComponent(UActorComponent* Component);
+
+	void RegisterFixedTickActor(AActor* Actor);
+	void UnregisterFixedTickActor(AActor* Actor);
+
 	/**
 	 * Called during creation of the physics state for gamethread objects to pass off an object to the physics thread
 	 */
@@ -282,6 +288,8 @@ private:
 	virtual void OnSyncBodies(Chaos::FPhysicsSolverBase* Solver) override;
 #endif
 
+	void EnableFixedTickCallback();
+
 #if 0
 	void SetKinematicTransform(FPhysicsActorHandle& InActorReference,const Chaos::TRigidTransform<float,3>& NewTransform)
 	{
@@ -349,6 +357,7 @@ private:
 #if WITH_CHAOS
 	Chaos::FReal LastEventDispatchTime;
 #endif 
+	class FFixedTickCallback* FixedTickCallback = nullptr;
 
 #if WITH_EDITOR
 	// Counter used to check a match with the single step status.
diff --git a/Engine/Source/Runtime/Engine/Public/Physics/Experimental/PhysicsThreadLibrary.h b/Engine/Source/Runtime/Engine/Public/Physics/Experimental/PhysicsThreadLibrary.h
new file mode 100644
--- /dev/null
+++ b/Engine/Source/Runtime/Engine/Public/Physics/Experimental/PhysicsThreadLibrary.h
@@ -0,0 +1,47 @@
+// Copyright Epic Games, Inc. All Rights Reserved.
+
+#pragma once
+
+#include "CoreMinimal.h"
+#include "UObject/ObjectMacros.h"
+#include "UObject/Object.h"
+#include "Templates/SubclassOf.h"
+#include "Engine/EngineTypes.h"
+#include "UObject/UnrealType.h"
+#include "UObject/ScriptMacros.h"
+#include "UObject/Interface.h"
+#include "Kismet/BlueprintFunctionLibrary.h"
+#include "Engine/CollisionProfile.h"
+#include "PhysicsProxy/SingleParticlePhysicsProxy.h"
+#include "PhysicsEngine/BodyInstance.h"
+#include "PhysicsThreadLibrary.generated.h"
+
+UCLASS(meta = (ScriptName = "PhysicsThreadLibrary"))
+class ENGINE_API UPhysicsThreadLibrary : public UBlueprintFunctionLibrary
+{
+	GENERATED_BODY()
+
+	/**
+	 *	Add a force to a single rigid body.
+	 *  This is like a 'thruster'. Good for adding a burst over some (non zero) time. Should be called every frame for the duration of the force.
+	 *
+	 *	@param	Force		 Force vector to apply. Magnitude indicates strength of force.
+	 *  @param  bAccelChange If true, Force is taken as a change in acceleration instead of a physical force (i.e. mass will have no effect).
+	 */
+	UFUNCTION(BlueprintCallable, Category = "Utilities")
+	static void AddForce(FBodyInstanceFixedTickHandle Handle, FVector Force, bool bAccelChange = false)
+	{
+		if (Handle)
+		{
+			if(bAccelChange)
+			{
+				Handle->AddForce(Force * Handle->M());
+			}
+			else
+			{
+				Handle->AddForce(Force);
+			}
+		}
+	}
+
+};
\ No newline at end of file
diff --git a/Engine/Source/Runtime/Experimental/Chaos/Private/Chaos/Framework/PhysicsSolverBase.cpp b/Engine/Source/Runtime/Experimental/Chaos/Private/Chaos/Framework/PhysicsSolverBase.cpp
--- a/Engine/Source/Runtime/Experimental/Chaos/Private/Chaos/Framework/PhysicsSolverBase.cpp
+++ b/Engine/Source/Runtime/Experimental/Chaos/Private/Chaos/Framework/PhysicsSolverBase.cpp
@@ -107,7 +107,7 @@ namespace Chaos
 		//We are on GT, but we know PhysicsThread is waiting so we're actually going to operate on PT data
 #if PHYSICS_THREAD_CONTEXT
 		FPhysicsThreadContextScope Scope(/*IsPhysicsThreadContext=*/true);
-		//TODO: add something to make sure we don't access GT solver data by mistake
+		FFrozenGameThreadContextScope FrozenScope;	//Make sure we fire ensures if any physics GT data is used
 #endif
 
 		Solver.SetOnFrozenGameThread(true);
@@ -423,7 +423,7 @@ namespace Chaos
 				PendingTasks = TGraphTask<FPhysicsSolverProcessPushDataTask>::CreateTask(&Prereqs).ConstructAndDispatchWhenReady(*this, PushData);
 				Prereqs.Add(PendingTasks);
 
-				if(bSolverHasFrozenGameThreadCallbacks)
+				if(bSolverHasFrozenGameThreadCallbacks && IsUsingFixedDt())
 				{
 					PendingTasks = TGraphTask<FPhysicsSolverFrozenGTPreSimCallbacks>::CreateTask(&Prereqs).ConstructAndDispatchWhenReady(*this);
 					Prereqs.Add(PendingTasks);
@@ -447,4 +447,15 @@ namespace Chaos
 
 		return BlockingTasks;
 	}
+
+
+	void FAllSolverTasks::AdvanceSolver()
+	{
+		ProcessPushData.ProcessPushData();
+		if (Solver.IsUsingFixedDt())
+		{
+			GTPreSimCallbacks.GTPreSimCallbacks();
+		}
+		AdvanceTask.AdvanceSolver();
+	}
 }
diff --git a/Engine/Source/Runtime/Experimental/Chaos/Public/Chaos/Framework/PhysicsSolverBase.h b/Engine/Source/Runtime/Experimental/Chaos/Public/Chaos/Framework/PhysicsSolverBase.h
--- a/Engine/Source/Runtime/Experimental/Chaos/Public/Chaos/Framework/PhysicsSolverBase.h
+++ b/Engine/Source/Runtime/Experimental/Chaos/Public/Chaos/Framework/PhysicsSolverBase.h
@@ -136,12 +136,13 @@ namespace Chaos
 		FPushPhysicsData* PushData;
 	};
 
-	struct FAllSolverTasks
+	struct CHAOS_API FAllSolverTasks
 	{
-		FAllSolverTasks(FPhysicsSolverBase& Solver, FPushPhysicsData* PushData)
-			: ProcessPushData(Solver, PushData)
-			, GTPreSimCallbacks(Solver)
-			, AdvanceTask(Solver, PushData)
+		FAllSolverTasks(FPhysicsSolverBase& InSolver, FPushPhysicsData* PushData)
+			: ProcessPushData(InSolver, PushData)
+			, GTPreSimCallbacks(InSolver)
+			, AdvanceTask(InSolver, PushData)
+			, Solver(InSolver)
 		{
 		}
 
@@ -149,12 +150,9 @@ namespace Chaos
 		FPhysicsSolverFrozenGTPreSimCallbacks GTPreSimCallbacks;
 		FPhysicsSolverAdvanceTask AdvanceTask;
 
-		void AdvanceSolver()
-		{
-			ProcessPushData.ProcessPushData();
-			GTPreSimCallbacks.GTPreSimCallbacks();
-			AdvanceTask.AdvanceSolver();
-		}
+		void AdvanceSolver();
+
+		FPhysicsSolverBase& Solver;
 	};
 
 	class FPersistentPhysicsTask;
diff --git a/Engine/Source/Runtime/Experimental/Chaos/Public/Framework/Threading.h b/Engine/Source/Runtime/Experimental/Chaos/Public/Framework/Threading.h
--- a/Engine/Source/Runtime/Experimental/Chaos/Public/Framework/Threading.h
+++ b/Engine/Source/Runtime/Experimental/Chaos/Public/Framework/Threading.h
@@ -148,7 +148,7 @@ public:
 
 	bool IsInGameThreadContext() const
 	{
-		return IsInGameThread() || GameThreadContext > 0;
+		return (IsInGameThread() || GameThreadContext > 0) && !bFrozenGameThread;
 	}
 
 	void IncPhysicsSimContext()
@@ -173,9 +173,22 @@ public:
 		--GameThreadContext;
 	}
 
+	void FreezeGameThreadContext()
+	{
+		ensure(!bFrozenGameThread);
+		bFrozenGameThread = true;
+	}
+
+	void UnFreezeGameThreadContext()
+	{
+		ensure(bFrozenGameThread);
+		bFrozenGameThread = false;
+	}
+
 private:
 	int32 PhysicsSimContext = 0;
 	int32 GameThreadContext = 0;
+	bool bFrozenGameThread = false;
 };
 
 struct FPhysicsThreadContextScope
@@ -222,6 +235,19 @@ struct FGameThreadContextScope
 	bool bParentIsGameThreadContext;
 };
 
+struct FFrozenGameThreadContextScope
+{
+	FFrozenGameThreadContextScope()
+	{
+		FPhysicsThreadContext::Get().FreezeGameThreadContext();
+	}
+
+	~FFrozenGameThreadContextScope()
+	{
+		FPhysicsThreadContext::Get().UnFreezeGameThreadContext();
+	}
+};
+
 
 FORCEINLINE bool IsInPhysicsThreadContext()
 {
