Removing unncessary check on audio thread.

diff --git a/Engine/Source/Runtime/Engine/Private/AudioDevice.cpp b/Engine/Source/Runtime/Engine/Private/AudioDevice.cpp
--- a/Engine/Source/Runtime/Engine/Private/AudioDevice.cpp
+++ b/Engine/Source/Runtime/Engine/Private/AudioDevice.cpp
@@ -6650,8 +6650,6 @@ void FAudioDevice::UnregisterPluginListener(const TAudioPluginListenerPtr Plugin
 
 bool FAudioDevice::IsAudioDeviceMuted() const
 {
-	check(IsInAudioThread());
-
 	FAudioDeviceManager* DeviceManager = GEngine->GetAudioDeviceManager();
 	if(DeviceManager)
 	{
