Add low-level callback system to AudioDevice that allows other systems to monitor audio rendering progress. This will allow those systems to do things like prepare themselves for future callbacks for audio.

diff --git a/Engine/Source/Runtime/AudioMixer/Private/AudioMixerDevice.cpp b/Engine/Source/Runtime/AudioMixer/Private/AudioMixerDevice.cpp
--- a/Engine/Source/Runtime/AudioMixer/Private/AudioMixerDevice.cpp
+++ b/Engine/Source/Runtime/AudioMixer/Private/AudioMixerDevice.cpp
@@ -732,6 +732,11 @@ namespace Audio
 		// Update the audio render thread time at the head of the render
 		AudioThreadTimingData.AudioRenderThreadTime = FPlatformTime::Seconds() - AudioThreadTimingData.StartTime;
 
+		// notify interested parties
+		FAudioDeviceRenderInfo RenderInfo;
+		RenderInfo.NumFrames = SourceManager->GetNumOutputFrames();
+		NotifyAudioDevicePreRender(RenderInfo);
+
 		// Pump the command queue to the audio render thread
 		PumpCommandQueue();
 
@@ -801,6 +806,9 @@ namespace Audio
 		// Update the audio clock
 		AudioClock += AudioClockDelta;
 
+		// notify interested parties
+		NotifyAudioDevicePostRender(RenderInfo);
+
 		return true;
 	}
 
diff --git a/Engine/Source/Runtime/Engine/Private/AudioDevice.cpp b/Engine/Source/Runtime/Engine/Private/AudioDevice.cpp
--- a/Engine/Source/Runtime/Engine/Private/AudioDevice.cpp
+++ b/Engine/Source/Runtime/Engine/Private/AudioDevice.cpp
@@ -4663,6 +4663,16 @@ void FAudioDevice::UpdateAudioEngineSubsystems()
 	}
 }
 
+void FAudioDevice::NotifyAudioDevicePreRender(const FAudioDeviceRenderInfo& InInfo) const
+{
+	OnAudioDevicePreRender.Broadcast(InInfo);
+}
+
+void FAudioDevice::NotifyAudioDevicePostRender(const FAudioDeviceRenderInfo& InInfo) const
+{
+	OnAudioDevicePostRender.Broadcast(InInfo);
+}
+
 void FAudioDevice::SendUpdateResultsToGameThread(const int32 FirstActiveIndex)
 {
 	DECLARE_CYCLE_STAT(TEXT("FGameThreadAudioTask.AudioSendResults"), STAT_AudioSendResults, STATGROUP_TaskGraphTasks);
diff --git a/Engine/Source/Runtime/Engine/Public/AudioDevice.h b/Engine/Source/Runtime/Engine/Public/AudioDevice.h
--- a/Engine/Source/Runtime/Engine/Public/AudioDevice.h
+++ b/Engine/Source/Runtime/Engine/Public/AudioDevice.h
@@ -412,6 +412,13 @@ public:
 	virtual void OnNewSubmixBuffer(const USoundSubmix* OwningSubmix, float* AudioData, int32 NumSamples, int32 NumChannels, const int32 SampleRate, double AudioClock) = 0;
 };
 
+struct FAudioDeviceRenderInfo
+{
+	int32 NumFrames = 0;
+};
+
+DECLARE_MULTICAST_DELEGATE_OneParam(FOnAudioDevicePreRender, FAudioDeviceRenderInfo);
+DECLARE_MULTICAST_DELEGATE_OneParam(FOnAudioDevicePostRender, FAudioDeviceRenderInfo);
 
 class ENGINE_API FAudioDevice : public FExec
 {
@@ -940,6 +947,9 @@ public:
 		return nullptr;
 	}
 
+	FOnAudioDevicePreRender OnAudioDevicePreRender;
+	FOnAudioDevicePostRender OnAudioDevicePostRender;
+
 	virtual void InitSoundEffectPresets() {}
 
 	/**
@@ -1493,7 +1503,11 @@ protected:
 	 */
 	virtual void OnListenerUpdated(const TArray<FListener>& InListeners) {};
 
-	private:
+	void NotifyAudioDevicePreRender(const FAudioDeviceRenderInfo& InInfo) const;
+
+	void NotifyAudioDevicePostRender(const FAudioDeviceRenderInfo& InInfo) const;
+	
+private:
 
 	/**
 	 * Adds an active sound to the audio device. Can be a new active sound or one provided by the re-triggering
