Fixed crash in GPULightmass due the usage of a bad scene pointer

diff --git a/Engine/Source/Runtime/Engine/Private/InstancedStaticMesh.cpp b/Engine/Source/Runtime/Engine/Private/InstancedStaticMesh.cpp
--- a/Engine/Source/Runtime/Engine/Private/InstancedStaticMesh.cpp
+++ b/Engine/Source/Runtime/Engine/Private/InstancedStaticMesh.cpp
@@ -4563,7 +4563,7 @@ void FInstancedStaticMeshVertexFactoryShaderParameters::GetElementShaderBindings
 
 	ShaderBindings.Add(InstanceOffset, InstanceOffsetValue);
 	
-	if (!UseGPUScene(Scene->GetShaderPlatform()))
+	if (!UseGPUScene(GMaxRHIShaderPlatform, GMaxRHIFeatureLevel))
 	{
 		ShaderBindings.Add(Shader->GetUniformBufferParameter<FInstancedStaticMeshVertexFactoryUniformShaderParameters>(), InstancedVertexFactory->GetUniformBuffer());
 		if (InstancedVertexFactory->SupportsManualVertexFetch(FeatureLevel))
