Fix GetReversedHit
It wasn't properly swapping BoneName and MyBoneName

fyi fred.kimberley

diff --git a/Engine/Source/Runtime/Engine/Classes/Engine/EngineTypes.h b/Engine/Source/Runtime/Engine/Classes/Engine/EngineTypes.h
--- a/Engine/Source/Runtime/Engine/Classes/Engine/EngineTypes.h
+++ b/Engine/Source/Runtime/Engine/Classes/Engine/EngineTypes.h
@@ -2463,7 +2463,7 @@ struct ENGINE_API FHitResult
 
 		FName TempBoneName = Result.BoneName;
 		Result.BoneName = Result.MyBoneName;
-		Result.MyBoneName = Result.BoneName;
+		Result.MyBoneName = TempBoneName;
 		return Result;
 	}
 
