Option to disable manual send clamp

diff --git a/Engine/Source/Runtime/AudioMixer/Private/AudioMixerSource.cpp b/Engine/Source/Runtime/AudioMixer/Private/AudioMixerSource.cpp
--- a/Engine/Source/Runtime/AudioMixer/Private/AudioMixerSource.cpp
+++ b/Engine/Source/Runtime/AudioMixer/Private/AudioMixerSource.cpp
@@ -1459,7 +1459,14 @@ namespace Audio
 				}
 				else if (SendInfo.SendLevelControlMethod == ESendLevelControlMethod::Manual)
 				{
-					SendLevel = FMath::Clamp(SendInfo.SendLevel, 0.0f, 1.0f);
+					if (SendInfo.DisableManualSendClamp)
+					{
+						SendLevel = SendInfo.SendLevel;
+					}
+					else
+					{
+						SendLevel = FMath::Clamp(SendInfo.SendLevel, 0.0f, 1.0f);
+					}
 				}
 				else
 				{
diff --git a/Engine/Source/Runtime/Engine/Classes/Sound/SoundSubmixSend.h b/Engine/Source/Runtime/Engine/Classes/Sound/SoundSubmixSend.h
--- a/Engine/Source/Runtime/Engine/Classes/Sound/SoundSubmixSend.h
+++ b/Engine/Source/Runtime/Engine/Classes/Sound/SoundSubmixSend.h
@@ -116,6 +116,10 @@ struct ENGINE_API FSoundSubmixSendInfo
 	UPROPERTY(EditAnywhere, BlueprintReadWrite, Category = SubmixSend)
 	float SendLevel;
 
+	// Whether to disable the 0-1 clamp for manual SendLevel control
+	UPROPERTY(EditAnywhere, BlueprintReadWrite, Category = SubmixSend)
+	bool DisableManualSendClamp;
+
 	// The amount to send to master when sound is located at a distance equal to value specified in the min send distance.
 	UPROPERTY(EditAnywhere, BlueprintReadWrite, Category = SubmixSend)
 	float MinSendLevel;
diff --git a/Engine/Source/Runtime/Engine/Private/SoundSubmixSend.cpp b/Engine/Source/Runtime/Engine/Private/SoundSubmixSend.cpp
--- a/Engine/Source/Runtime/Engine/Private/SoundSubmixSend.cpp
+++ b/Engine/Source/Runtime/Engine/Private/SoundSubmixSend.cpp
@@ -8,6 +8,7 @@ FSoundSubmixSendInfo::FSoundSubmixSendInfo()
 	, SendStage(ESubmixSendStage::PostDistanceAttenuation)
 	, SoundSubmix(nullptr)
 	, SendLevel(0.0f)
+	, DisableManualSendClamp(false)
 	, MinSendLevel(0.0f)
 	, MaxSendLevel(1.0f)
 	, MinSendDistance(100.0f)
