Don't set owner on created material instance decal when in game.
Persistent owner should only be set for editing workflow where we don't want material to be transient.
Fixes bug where material created in construction script can be unloaded with it's owner, even when still referenced by primitives in the world.

diff --git a/Engine/Source/Runtime/Engine/Private/KismetMaterialLibrary.cpp b/Engine/Source/Runtime/Engine/Private/KismetMaterialLibrary.cpp
--- a/Engine/Source/Runtime/Engine/Private/KismetMaterialLibrary.cpp
+++ b/Engine/Source/Runtime/Engine/Private/KismetMaterialLibrary.cpp
@@ -126,10 +126,10 @@ class UMaterialInstanceDynamic* UKismetMaterialLibrary::CreateDynamicMaterialIns
 	if (Parent)
 	{
 
-		// MIDs need to be created within a persistent object if in the construction script (or editor utility) or else they will not be saved.
-		// If this MID is created at runtime then put it in the transient package
+		// In editor MIDs need to be created within a persistent object or else they will not be saved.
+		// If this MID is created at runtime or specifically marked as transient then put it in the transient package.
 		UWorld* World = GEngine->GetWorldFromContextObject(WorldContextObject, EGetWorldErrorMode::ReturnNull);
-		UObject* MIDOuter = (!EnumHasAnyFlags(CreationFlags, EMIDCreationFlags::Transient) &&  World && (World->bIsRunningConstructionScript  || !World->IsGameWorld()) ? WorldContextObject : nullptr);
+		UObject* MIDOuter = !EnumHasAnyFlags(CreationFlags, EMIDCreationFlags::Transient) && World && !World->IsGameWorld() ? WorldContextObject : nullptr;
 		NewMID = UMaterialInstanceDynamic::Create(Parent, MIDOuter, OptionalName);
 		if (MIDOuter == nullptr)
 		{
