Change default object trace filter state to unfiltered. Allows animation insights to work without the Source data filtering plugin enabled.

diff --git a/Engine/Source/Runtime/Engine/Private/TraceFilter.cpp b/Engine/Source/Runtime/Engine/Private/TraceFilter.cpp
--- a/Engine/Source/Runtime/Engine/Private/TraceFilter.cpp
+++ b/Engine/Source/Runtime/Engine/Private/TraceFilter.cpp
@@ -19,7 +19,7 @@
 struct ENGINE_API FTraceFilterObjectAnnotation
 {
 	FTraceFilterObjectAnnotation()
-		: bIsTraceable(false)
+		: bIsTraceable(true)
 	{}
 
 	bool bIsTraceable;
@@ -27,7 +27,7 @@ struct ENGINE_API FTraceFilterObjectAnnotation
 	/** Determine if this annotation is default - required for annotations */
 	FORCEINLINE bool IsDefault() const
 	{
-		return !bIsTraceable;
+		return bIsTraceable;
 	}
 };
 
@@ -292,18 +292,14 @@ void ENGINE_API FTraceFilter::SetObjectIsTraceable</*bForceThreadSafe = */ true>
 {
 	ensure(InObject);
 		
+	FTraceFilterObjectAnnotation Annotation;
+	Annotation.bIsTraceable = bIsTraceable;
+	GObjectFilterAnnotations.AddAnnotation(InObject, Annotation);
+
 	if (bIsTraceable)
 	{
-		FTraceFilterObjectAnnotation Annotation;
-		Annotation.bIsTraceable = true;
-		GObjectFilterAnnotations.AddAnnotation(InObject, Annotation);
-
 		TRACE_OBJECT(InObject);
 	}
-	else
-	{
-		GObjectFilterAnnotations.RemoveAnnotation(InObject);
-	}
 }
 
 template<>
@@ -313,15 +309,15 @@ void ENGINE_API FTraceFilter::SetObjectIsTraceable</*bForceThreadSafe = */ false
 
 	check(GObjectFilterAnnotations.IsLocked());
 	TMap<const UObjectBase*, FTraceFilterObjectAnnotation>& AnnotationMap = GObjectFilterAnnotations.GetAnnotationMap();
-	if (bIsTraceable)
+	if (!bIsTraceable)
 	{
-		AnnotationMap.FindOrAdd(InObject).bIsTraceable = true;
-		TRACE_OBJECT(InObject);
+		AnnotationMap.FindOrAdd(InObject).bIsTraceable = false;
 	}
 	else
 	{
 		AnnotationMap.Remove(InObject);
-	}
+		TRACE_OBJECT(InObject);
+	}	
 }
 
 template<>
