Fix for sounds cutting off if they started playing in the binarual radius, and the listener leaves that radius

diff --git a/Engine/Source/Runtime/Engine/Private/ActiveSound.cpp b/Engine/Source/Runtime/Engine/Private/ActiveSound.cpp
--- a/Engine/Source/Runtime/Engine/Private/ActiveSound.cpp
+++ b/Engine/Source/Runtime/Engine/Private/ActiveSound.cpp
@@ -99,6 +99,8 @@ FActiveSound::FActiveSound()
 	, bEnableBusSendRoutingOverride(false)
 	, bEnableMainSubmixOutputOverride(false)
 	, bEnableSubmixSendRoutingOverride(false)
+	, bIsFirstAttenuationUpdate(true)
+	, bStartedWithinNonBinauralRadius(false)
 	, UserIndex(0)
 	, FadeOut(EFadeOut::None)
 	, bIsOccluded(false)
@@ -1803,8 +1805,13 @@ void FActiveSound::UpdateAttenuation(float DeltaTime, FSoundParseParameters& Par
 	ParseParams.bUseSpatialization |= Settings->bSpatialize;
 	ParseParams.bEnableSourceDataOverride |= Settings->bEnableSourceDataOverride;
 
-	// Check the binaural radius to determine if we're going to HRTF spatialize
-	if (ListenerData.ListenerToSoundDistance < Settings->BinauralRadius)
+	// Check the binaural radius to determine if we're going to HRTF spatialize, cache the result
+	if (bIsFirstAttenuationUpdate)
+	{
+		bStartedWithinNonBinauralRadius = ListenerData.ListenerToSoundDistance < Settings->BinauralRadius;
+	}
+
+	if (bStartedWithinNonBinauralRadius)
 	{
 		ParseParams.SpatializationMethod = ESoundSpatializationAlgorithm::SPATIALIZATION_Default;
 	}
@@ -1836,4 +1843,6 @@ void FActiveSound::UpdateAttenuation(float DeltaTime, FSoundParseParameters& Par
 	{
 		FocusData.PriorityHighest = FocusDataToApply.PriorityHighest;
 	}
+
+	bIsFirstAttenuationUpdate = false;
 }
diff --git a/Engine/Source/Runtime/Engine/Public/ActiveSound.h b/Engine/Source/Runtime/Engine/Public/ActiveSound.h
--- a/Engine/Source/Runtime/Engine/Public/ActiveSound.h
+++ b/Engine/Source/Runtime/Engine/Public/ActiveSound.h
@@ -567,6 +567,9 @@ public:
 	uint8 bEnableMainSubmixOutputOverride : 1;
 	uint8 bEnableSubmixSendRoutingOverride : 1;
 
+	uint8 bIsFirstAttenuationUpdate : 1;
+	uint8 bStartedWithinNonBinauralRadius : 1;
+
 	uint8 UserIndex;
 
 	/** Type of fade out currently being applied */
