#ue5 - FPlane conversion is explicit.

diff --git a/Engine/Source/Runtime/Core/Public/Math/Plane.h b/Engine/Source/Runtime/Core/Public/Math/Plane.h
--- a/Engine/Source/Runtime/Core/Public/Math/Plane.h
+++ b/Engine/Source/Runtime/Core/Public/Math/Plane.h
@@ -335,10 +335,9 @@ public:
 		return true;
 	}
 
-	// Conversion to other type. TODO: explicit!
+	// Conversion to other type.
 	template<typename FArg, TEMPLATE_REQUIRES(!TIsSame<T, FArg>::Value)>
-	TPlane(const TPlane<FArg>& From) : TPlane<T>((T)From.X, (T)From.Y, (T)From.Z, (T)From.W) {}
-
+	explicit TPlane(const TPlane<FArg>& From) : TPlane<T>((T)From.X, (T)From.Y, (T)From.Z, (T)From.W) {}
 };
 
 
diff --git a/Engine/Source/Runtime/Engine/Private/ModelCollision.cpp b/Engine/Source/Runtime/Engine/Private/ModelCollision.cpp
--- a/Engine/Source/Runtime/Engine/Private/ModelCollision.cpp
+++ b/Engine/Source/Runtime/Engine/Private/ModelCollision.cpp
@@ -30,7 +30,7 @@ void UModel::GetSurfacePlanes(
 
 		for (int32 SurfaceIndex = 0; SurfaceIndex < Surfs.Num(); SurfaceIndex++)
 		{
-			OutPlanes.Add(Surfs[SurfaceIndex].Plane.TransformByUsingAdjointT(FMatrix44f(Matrix), float(DetM), FMatrix44f(MatrixTA)));	// LWC_TODO: Precision loss
+			OutPlanes.Add(FPlane(Surfs[SurfaceIndex].Plane).TransformByUsingAdjointT(Matrix, DetM, MatrixTA));
 		}
 	}
 }
diff --git a/Engine/Source/Runtime/Engine/Private/PhysicsEngine/PhysUtils.cpp b/Engine/Source/Runtime/Engine/Private/PhysicsEngine/PhysUtils.cpp
--- a/Engine/Source/Runtime/Engine/Private/PhysicsEngine/PhysUtils.cpp
+++ b/Engine/Source/Runtime/Engine/Private/PhysicsEngine/PhysUtils.cpp
@@ -84,7 +84,7 @@ static bool ModelToHullsWorker(FKAggregateGeom* outGeom,
 	// BACK
 	if (node.iBack != INDEX_NONE) // If there is a child, recurse into it.
 	{
-		planes.Add(node.Plane);
+		planes.Add(FPlane(node.Plane));
 		if (!ModelToHullsWorker(outGeom, inModel, node.iBack, node.ChildOutside(0, bOutside), planes))
 		{
 			return false;
@@ -93,7 +93,7 @@ static bool ModelToHullsWorker(FKAggregateGeom* outGeom,
 	}
 	else if (!node.ChildOutside(0, bOutside)) // If its a leaf, and solid (inside)
 	{
-		planes.Add(node.Plane);
+		planes.Add(FPlane(node.Plane));
 		if (!AddConvexPrim(outGeom, planes, inModel))
 		{
 			return false;
@@ -104,7 +104,7 @@ static bool ModelToHullsWorker(FKAggregateGeom* outGeom,
 	// FRONT
 	if (node.iFront != INDEX_NONE)
 	{
-		planes.Add(node.Plane.Flip());
+		planes.Add(FPlane(node.Plane.Flip()));
 		if (!ModelToHullsWorker(outGeom, inModel, node.iFront, node.ChildOutside(1, bOutside), planes))
 		{
 			return false;
@@ -113,7 +113,7 @@ static bool ModelToHullsWorker(FKAggregateGeom* outGeom,
 	}
 	else if (!node.ChildOutside(1, bOutside))
 	{
-		planes.Add(node.Plane.Flip());
+		planes.Add(FPlane(node.Plane.Flip()));
 		if (!AddConvexPrim(outGeom, planes, inModel))
 		{
 			return false;
