New HLSL translator supports flow control in preshaders

diff --git a/Engine/Source/Runtime/Engine/Private/HLSLTree/HLSLTree.cpp b/Engine/Source/Runtime/Engine/Private/HLSLTree/HLSLTree.cpp
--- a/Engine/Source/Runtime/Engine/Private/HLSLTree/HLSLTree.cpp
+++ b/Engine/Source/Runtime/Engine/Private/HLSLTree/HLSLTree.cpp
@@ -7,6 +7,9 @@
 #include "Shader/ShaderTypes.h"
 #include "Shader/Preshader.h"
 
+// TODO - M_ForLoop doesn't work yet
+// FPreparedType::GetEvaluation takes scope, checks loop scope automatically
+
 namespace UE
 {
 namespace HLSLTree
@@ -32,8 +35,9 @@ public:
 	FExpressionLocalPHI(const FExpressionLocalPHI* Source, EDerivativeCoordinate Coord);
 
 	virtual void ComputeAnalyticDerivatives(FTree& Tree, FExpressionDerivatives& OutResult) const override;
-	virtual bool PrepareValue(FEmitContext& Context, const FRequestedType& RequestedType, FPrepareValueResult& OutResult) const override;
+	virtual bool PrepareValue(FEmitContext& Context, FEmitScope& Scope, const FRequestedType& RequestedType, FPrepareValueResult& OutResult) const override;
 	virtual void EmitValueShader(FEmitContext& Context, FEmitScope& Scope, const FRequestedType& RequestedType, FEmitValueShaderResult& OutResult) const override;
+	virtual void EmitValuePreshader(FEmitContext& Context, FEmitScope& Scope, const FRequestedType& RequestedType, Shader::FPreshaderData& OutPreshader) const override;
 
 	TArray<EDerivativeCoordinate, TInlineAllocator<8>> DerivativeChain;
 	FName LocalName;
@@ -51,8 +55,9 @@ class FExpressionFunctionCall final : public FExpression
 public:
 	FExpressionFunctionCall(FFunction* InFunction, int32 InOutputIndex) : Function(InFunction), OutputIndex(InOutputIndex) {}
 
-	virtual bool PrepareValue(FEmitContext& Context, const FRequestedType& RequestedType, FPrepareValueResult& OutResult) const override;
+	virtual bool PrepareValue(FEmitContext& Context, FEmitScope& Scope, const FRequestedType& RequestedType, FPrepareValueResult& OutResult) const override;
 	virtual void EmitValueShader(FEmitContext& Context, FEmitScope& Scope, const FRequestedType& RequestedType, FEmitValueShaderResult& OutResult) const override;
+	virtual void EmitValuePreshader(FEmitContext& Context, FEmitScope& Scope, const FRequestedType& RequestedType, Shader::FPreshaderData& OutPreshader) const override;
 
 	FFunction* Function;
 	int32 OutputIndex;
@@ -94,7 +99,7 @@ FBinaryOpDescription GetBinaryOpDesription(EBinaryOp Op)
 	case EBinaryOp::Sub: return FBinaryOpDescription(TEXT("Subtract"), TEXT("-"), Shader::EPreshaderOpcode::Sub); break;
 	case EBinaryOp::Mul: return FBinaryOpDescription(TEXT("Multiply"), TEXT("*"), Shader::EPreshaderOpcode::Mul); break;
 	case EBinaryOp::Div: return FBinaryOpDescription(TEXT("Divide"), TEXT("/"), Shader::EPreshaderOpcode::Div); break;
-	case EBinaryOp::Less: return FBinaryOpDescription(TEXT("Less"), TEXT("<"), Shader::EPreshaderOpcode::Nop); break;
+	case EBinaryOp::Less: return FBinaryOpDescription(TEXT("Less"), TEXT("<"), Shader::EPreshaderOpcode::Less); break;
 	default: checkNoEntry(); return FBinaryOpDescription();
 	}
 }
@@ -110,18 +115,64 @@ EExpressionEvaluation CombineEvaluations(EExpressionEvaluation Lhs, EExpressionE
 	{
 		return Lhs;
 	}
-	else if (Lhs == EExpressionEvaluation::Constant && Rhs == EExpressionEvaluation::Constant)
+	else if (Lhs == EExpressionEvaluation::Unknown)
 	{
-		// 2 constants make a constant
-		return EExpressionEvaluation::Constant;
+		return Rhs;
+	}
+	else if (Rhs == EExpressionEvaluation::Unknown)
+	{
+		return Lhs;
 	}
 	else if (Lhs == EExpressionEvaluation::Shader || Rhs == EExpressionEvaluation::Shader)
 	{
 		// If either requires shader, shader is required
 		return EExpressionEvaluation::Shader;
 	}
-	// Any combination of constants/preshader can make a preshader
-	return EExpressionEvaluation::Preshader;
+	else if (Lhs == EExpressionEvaluation::PreshaderLoop || Rhs == EExpressionEvaluation::PreshaderLoop)
+	{
+		// Otherwise if either requires preshader, preshader is required
+		return EExpressionEvaluation::PreshaderLoop;
+	}
+	else if (Lhs == EExpressionEvaluation::Preshader || Rhs == EExpressionEvaluation::Preshader)
+	{
+		// Otherwise if either requires preshader, preshader is required
+		return EExpressionEvaluation::Preshader;
+	}
+	else if (Lhs == EExpressionEvaluation::ConstantLoop || Rhs == EExpressionEvaluation::ConstantLoop)
+	{
+		return EExpressionEvaluation::ConstantLoop;
+	}
+
+	// Otherwise must be constant
+	check(Lhs == EExpressionEvaluation::Constant);
+	check(Rhs == EExpressionEvaluation::Constant);
+	return EExpressionEvaluation::Constant;
+}
+
+EExpressionEvaluation MakeLoopEvaluation(EExpressionEvaluation Evaluation)
+{
+	if (Evaluation == EExpressionEvaluation::Preshader)
+	{
+		return EExpressionEvaluation::PreshaderLoop;
+	}
+	else if (Evaluation == EExpressionEvaluation::Constant)
+	{
+		return EExpressionEvaluation::ConstantLoop;
+	}
+	return Evaluation;
+}
+
+EExpressionEvaluation MakeNonLoopEvaluation(EExpressionEvaluation Evaluation)
+{
+	if (Evaluation == EExpressionEvaluation::PreshaderLoop)
+	{
+		return EExpressionEvaluation::Preshader;
+	}
+	else if (Evaluation == EExpressionEvaluation::ConstantLoop)
+	{
+		return EExpressionEvaluation::Constant;
+	}
+	return Evaluation;
 }
 
 FScope* FScope::FindSharedParent(FScope* Lhs, FScope* Rhs)
@@ -167,7 +218,7 @@ void FExpressionLocalPHI::ComputeAnalyticDerivatives(FTree& Tree, FExpressionDer
 	OutResult.ExpressionDdy = Tree.NewExpression<FExpressionLocalPHI>(this, EDerivativeCoordinate::Ddy);
 }
 
-bool FExpressionLocalPHI::PrepareValue(FEmitContext& Context, const FRequestedType& RequestedType, FPrepareValueResult& OutResult) const
+bool FExpressionLocalPHI::PrepareValue(FEmitContext& Context, FEmitScope& Scope, const FRequestedType& RequestedType, FPrepareValueResult& OutResult) const
 {
 	check(NumValues <= MaxNumPreviousScopes);
 	FExpression* ForwardExpression = Values[0];
@@ -189,15 +240,17 @@ bool FExpressionLocalPHI::PrepareValue(FEmitContext& Context, const FRequestedTy
 	if (bForwardExpressionValid)
 	{
 		check(ForwardExpression);
-		return OutResult.SetForwardValue(Context, RequestedType, ForwardExpression);
+		return OutResult.SetForwardValue(Context, Scope, RequestedType, ForwardExpression);
 	}
 
 	// 2) PHI has different values in previous scopes, but possible some previous scopes may become dead due to constant folding
 	// In this case, we check to see if the value is the same in all live scopes, and forward if possible
+	FEmitScope* EmitScopes[MaxNumPreviousScopes] = { nullptr };
 	for (int32 i = 0; i < NumValues; ++i)
 	{
 		// Ignore values in dead scopes
-		if (Context.PrepareScope(Scopes[i]))
+		EmitScopes[i] = Context.PrepareScope(Scopes[i]);
+		if (EmitScopes[i])
 		{
 			FExpression* ScopeExpression = Values[i];
 			if (!ForwardExpression)
@@ -215,7 +268,7 @@ bool FExpressionLocalPHI::PrepareValue(FEmitContext& Context, const FRequestedTy
 	if (bForwardExpressionValid)
 	{
 		check(ForwardExpression);
-		return OutResult.SetForwardValue(Context, RequestedType, ForwardExpression);
+		return OutResult.SetForwardValue(Context, Scope, RequestedType, ForwardExpression);
 	}
 
 	FPreparedType TypePerValue[MaxNumPreviousScopes];
@@ -226,9 +279,13 @@ bool FExpressionLocalPHI::PrepareValue(FEmitContext& Context, const FRequestedTy
 	{
 		for (int32 i = 0; i < NumValues; ++i)
 		{
-			if (TypePerValue[i].IsVoid() && Context.PrepareScope(Scopes[i]))
+			if (!EmitScopes[i] || EmitScopes[i]->Evaluation == EExpressionEvaluation::Unknown)
 			{
-				const FPreparedType& ValueType = Context.PrepareExpression(Values[i], RequestedType);
+				EmitScopes[i] = Context.PrepareScope(Scopes[i]);
+			}
+			if (TypePerValue[i].IsVoid() && EmitScopes[i])
+			{
+				const FPreparedType& ValueType = Context.PrepareExpression(Values[i], *EmitScopes[i], RequestedType);
 				if (!ValueType.IsVoid())
 				{
 					TypePerValue[i] = ValueType;
@@ -240,6 +297,7 @@ bool FExpressionLocalPHI::PrepareValue(FEmitContext& Context, const FRequestedTy
 							ValueType.GetType().GetName());
 					}
 					CurrentType = MergedType;
+					CurrentType.MergeEvaluation(EmitScopes[i]->Evaluation);
 					check(NumValidTypes < NumValues);
 					NumValidTypes++;
 				}
@@ -256,8 +314,8 @@ bool FExpressionLocalPHI::PrepareValue(FEmitContext& Context, const FRequestedTy
 	}
 
 	// Assuming we have at least one value with a valid type, we use that to initialize our type
-	CurrentType.SetEvaluation(EExpressionEvaluation::Shader); // TODO - No support for preshader flow control
-	if (!OutResult.SetType(Context, RequestedType, CurrentType))
+	const FPreparedType InitialType(CurrentType);
+	if (!OutResult.SetType(Context, RequestedType, InitialType))
 	{
 		return false;
 	}
@@ -273,6 +331,26 @@ bool FExpressionLocalPHI::PrepareValue(FEmitContext& Context, const FRequestedTy
 		{
 			return Context.Errors->AddError(TEXT("Failed to compute all types for LocalPHI"));
 		}
+
+		if (CurrentType != InitialType)
+		{
+			// Update type again based on computing remaining values
+			if (!OutResult.SetType(Context, RequestedType, CurrentType))
+			{
+				return false;
+			}
+
+			// Since we changed our type, need to update any dependant values again
+			for (int32 i = 0; i < NumValues; ++i)
+			{
+				const FPreparedType& ValueType = Context.PrepareExpression(Values[i], *EmitScopes[i], RequestedType);
+				// Don't expect types to change *again*
+				if (ValueType.IsVoid() || MergePreparedTypes(CurrentType, ValueType) != CurrentType)
+				{
+					return Context.Errors->AddError(TEXT("Mismatched types for local variable %s and %s"));
+				}
+			}
+		}
 	}
 
 	return true;
@@ -351,14 +429,63 @@ void FExpressionLocalPHI::EmitValueShader(FEmitContext& Context, FEmitScope& Sco
 	OutResult.Code = EmitExpression;
 }
 
-bool FExpressionFunctionCall::PrepareValue(FEmitContext& Context, const FRequestedType& RequestedType, FPrepareValueResult& OutResult) const
+void FExpressionLocalPHI::EmitValuePreshader(FEmitContext& Context, FEmitScope& Scope, const FRequestedType& RequestedType, Shader::FPreshaderData& OutPreshader) const
+{
+	int32 ValueStackPosition = INDEX_NONE;
+	for (int32 Index = Context.PreshaderLocalPHIScopes.Num() - 1; Index >= 0; --Index)
+	{
+		const FPreshaderLocalPHIScope* LocalPHIScope = Context.PreshaderLocalPHIScopes[Index];
+		if (LocalPHIScope->ExpressionLocalPHI == this)
+		{
+			ValueStackPosition = LocalPHIScope->ValueStackPosition;
+			break;
+		}
+	}
+
+	if (ValueStackPosition == INDEX_NONE)
+	{
+		// Assign the initial value
+		const FRequestedType LocalType = GetRequestedType();
+		Context.PreshaderStackPosition++;
+		OutPreshader.WriteOpcode(Shader::EPreshaderOpcode::ConstantZero).Write(LocalType.GetType());
+
+		ValueStackPosition = Context.PreshaderStackPosition;
+		const FPreshaderLocalPHIScope LocalPHIScope(this, ValueStackPosition);
+		Context.PreshaderLocalPHIScopes.Add(&LocalPHIScope);
+
+		FEmitScope* EmitRootScope = nullptr;// &Scope;
+		FEmitPreshaderScope PreshaderScopes[MaxNumPreviousScopes];
+		for (int32 i = 0; i < NumValues; ++i)
+		{
+			FEmitPreshaderScope& PreshaderScope = PreshaderScopes[i];
+			PreshaderScope.Scope = Context.AcquireEmitScope(Scopes[i]);
+			PreshaderScope.Value = Values[i];
+			EmitRootScope = FEmitScope::FindSharedParent(PreshaderScopes[i].Scope, EmitRootScope);
+		}
+
+		Context.EmitPreshaderScope(*EmitRootScope, RequestedType, MakeArrayView(PreshaderScopes, NumValues), OutPreshader);
+		verify(Context.PreshaderLocalPHIScopes.Pop(false) == &LocalPHIScope);
+		check(Context.PreshaderStackPosition == ValueStackPosition);
+	}
+	else
+	{
+		const int32 PreshaderStackOffset = Context.PreshaderStackPosition - ValueStackPosition;
+		check(PreshaderStackOffset >= 0 && PreshaderStackOffset <= 0xffff);
+
+		Context.PreshaderStackPosition++;
+		OutPreshader.WriteOpcode(Shader::EPreshaderOpcode::PushValue).Write((uint16)PreshaderStackOffset);
+	}
+}
+
+bool FExpressionFunctionCall::PrepareValue(FEmitContext& Context, FEmitScope& Scope, const FRequestedType& RequestedType, FPrepareValueResult& OutResult) const
 {
-	if (!Context.PrepareScopeWithParent(Function->RootScope, Function->CalledScope))
+	FEmitScope* EmitFunctionScope = Context.PrepareScopeWithParent(Function->RootScope, Function->CalledScope);
+	if (!EmitFunctionScope)
 	{
 		return false;
 	}
 
-	FPreparedType OutputType = Context.PrepareExpression(Function->OutputExpressions[OutputIndex], RequestedType);
+	FPreparedType OutputType = Context.PrepareExpression(Function->OutputExpressions[OutputIndex], Scope, RequestedType);
 	return OutResult.SetType(Context, RequestedType, OutputType);
 }
 
@@ -378,9 +505,13 @@ void FExpressionFunctionCall::EmitValueShader(FEmitContext& Context, FEmitScope&
 	OutResult.Code = Context.EmitInlineExpressionWithDependency(Scope, Dependency, EmitFunctionOutput->Type, TEXT("%"), EmitFunctionOutput);
 }
 
+void FExpressionFunctionCall::EmitValuePreshader(FEmitContext& Context, FEmitScope& Scope, const FRequestedType& RequestedType, Shader::FPreshaderData& OutPreshader) const
+{
+	Function->OutputExpressions[OutputIndex]->GetValuePreshader(Context, Scope, RequestedType, OutPreshader);
+}
+
 void FExpression::Reset()
 {
-	CurrentRequestedType.Reset();
 	PrepareValueResult = FPrepareValueResult();
 }
 
@@ -481,6 +612,32 @@ FRequestedType FRequestedType::GetField(const Shader::FStructField* Field) const
 	return Result;
 }
 
+EExpressionEvaluation FPreparedComponent::GetEvaluation(const FEmitScope& Scope) const
+{
+	EExpressionEvaluation Result = Evaluation;
+	if (IsLoopEvaluation(Result))
+	{
+		// We only want to return a 'Loop' evaluation if we're within the loop's scope
+		if (!Scope.HasParent(LoopScope))
+		{
+			switch (Result)
+			{
+			case EExpressionEvaluation::ConstantLoop: Result = EExpressionEvaluation::Constant; break;
+			case EExpressionEvaluation::PreshaderLoop: Result = EExpressionEvaluation::Preshader; break;
+			default: checkNoEntry(); break;
+			}
+		}
+	}
+	return Result;
+}
+
+FPreparedComponent CombineComponents(const FPreparedComponent& Lhs, const FPreparedComponent& Rhs)
+{
+	const EExpressionEvaluation Evaluation = CombineEvaluations(Lhs.Evaluation, Rhs.Evaluation);
+	FEmitScope* LoopScope = IsLoopEvaluation(Evaluation) ? FEmitScope::FindSharedParent(Lhs.LoopScope, Rhs.LoopScope) : nullptr;
+	return FPreparedComponent(Evaluation, LoopScope);
+}
+
 FPreparedType::FPreparedType(const Shader::FType& InType)
 {
 	if (InType.IsStruct())
@@ -501,7 +658,7 @@ int32 FPreparedType::GetNumComponents() const
 	}
 	else if (ValueComponentType != Shader::EValueComponentType::Void)
 	{
-		const int32 MaxComponentIndex = PreparedComponents.FindLastByPredicate([](EExpressionEvaluation InEvaluation) { return InEvaluation != EExpressionEvaluation::None; });
+		const int32 MaxComponentIndex = PreparedComponents.FindLastByPredicate([](const FPreparedComponent& InComponent) { return InComponent.Evaluation != EExpressionEvaluation::None; });
 		if (MaxComponentIndex != INDEX_NONE)
 		{
 			return MaxComponentIndex + 1;
@@ -549,8 +706,8 @@ FRequestedType FPreparedType::GetRequestedType() const
 		}
 		for (int32 Index = 0; Index < NumComponents; ++Index)
 		{
-			const EExpressionEvaluation ComponentEvaluation = GetComponentEvaluation(Index);
-			if (ComponentEvaluation != EExpressionEvaluation::None)
+			const FPreparedComponent Component = GetComponent(Index);
+			if (!Component.IsNone())
 			{
 				Result.SetComponentRequest(Index);
 			}
@@ -559,84 +716,124 @@ FRequestedType FPreparedType::GetRequestedType() const
 	return Result;
 }
 
-EExpressionEvaluation FPreparedType::GetEvaluation() const
+EExpressionEvaluation FPreparedType::GetEvaluation(const FEmitScope& Scope) const
 {
 	EExpressionEvaluation Result = EExpressionEvaluation::None;
 	for (int32 Index = 0; Index < PreparedComponents.Num(); ++Index)
 	{
-		Result = CombineEvaluations(Result, PreparedComponents[Index]);
+		Result = CombineEvaluations(Result, PreparedComponents[Index].GetEvaluation(Scope));
 	}
 	return Result;
 }
 
-EExpressionEvaluation FPreparedType::GetEvaluation(const FRequestedType& RequestedType) const
+EExpressionEvaluation FPreparedType::GetEvaluation(const FEmitScope& Scope, const FRequestedType& RequestedType) const
 {
 	EExpressionEvaluation Result = EExpressionEvaluation::None;
 	for (int32 Index = 0; Index < PreparedComponents.Num(); ++Index)
 	{
 		if (RequestedType.IsComponentRequested(Index))
 		{
-			Result = CombineEvaluations(Result, PreparedComponents[Index]);
+			Result = CombineEvaluations(Result, PreparedComponents[Index].GetEvaluation(Scope));
 		}
 	}
 	return Result;
 }
 
-EExpressionEvaluation FPreparedType::GetFieldEvaluation(int32 ComponentIndex, int32 NumComponents) const
+EExpressionEvaluation FPreparedType::GetFieldEvaluation(const FEmitScope& Scope, int32 ComponentIndex, int32 NumComponents) const
 {
 	EExpressionEvaluation Result = EExpressionEvaluation::None;
 	for (int32 Index = 0; Index < NumComponents; ++Index)
 	{
-		Result = CombineEvaluations(Result, GetComponentEvaluation(ComponentIndex + Index));
+		const FPreparedComponent Component = GetComponent(Index);
+		Result = CombineEvaluations(Result, Component.GetEvaluation(Scope));
 	}
 	return Result;
 }
 
-EExpressionEvaluation FPreparedType::GetComponentEvaluation(int32 Index) const
+FPreparedComponent FPreparedType::GetComponent(int32 Index) const
 {
-	return PreparedComponents.IsValidIndex(Index) ? PreparedComponents[Index] : EExpressionEvaluation::None;
+	return PreparedComponents.IsValidIndex(Index) ? PreparedComponents[Index] : FPreparedComponent();
 }
 
-void FPreparedType::SetComponentEvaluation(int32 Index, EExpressionEvaluation Evaluation)
+void FPreparedType::EnsureNumComponents(int32 NumComponents)
 {
-	if (Evaluation != EExpressionEvaluation::None && Index >= PreparedComponents.Num())
+	if (NumComponents > PreparedComponents.Num())
 	{
-		PreparedComponents.AddZeroed(Index + 1 - PreparedComponents.Num());
+		static_assert((uint8)EExpressionEvaluation::None == 0u, "Assume zero initializes to None");
+		PreparedComponents.AddZeroed(NumComponents - PreparedComponents.Num());
+	}
+}
+
+void FPreparedType::SetComponent(int32 Index, const FPreparedComponent& InComponent)
+{
+	if (InComponent.Evaluation != EExpressionEvaluation::None)
+	{
+		EnsureNumComponents(Index + 1);
 	}
 	if (PreparedComponents.IsValidIndex(Index))
 	{
-		PreparedComponents[Index] = Evaluation;
+		PreparedComponents[Index] = InComponent;
 	}
 }
 
-void FPreparedType::MergeComponentEvaluation(int32 Index, EExpressionEvaluation Evaluation)
+void FPreparedType::MergeComponent(int32 Index, const FPreparedComponent& InComponent)
 {
-	if (Evaluation != EExpressionEvaluation::None && Index >= PreparedComponents.Num())
+	if (InComponent.Evaluation != EExpressionEvaluation::None)
 	{
-		PreparedComponents.AddZeroed(Index + 1 - PreparedComponents.Num());
+		EnsureNumComponents(Index + 1);
 	}
 	if (PreparedComponents.IsValidIndex(Index))
 	{
-		PreparedComponents[Index] = CombineEvaluations(PreparedComponents[Index], Evaluation);
+		PreparedComponents[Index] = CombineComponents(PreparedComponents[Index], InComponent);
 	}
 }
 
 void FPreparedType::SetEvaluation(EExpressionEvaluation Evaluation)
 {
+	check(!IsLoopEvaluation(Evaluation));
 	for (int32 Index = 0; Index < PreparedComponents.Num(); ++Index)
 	{
-		if (PreparedComponents[Index] != EExpressionEvaluation::None)
+		if (!PreparedComponents[Index].IsNone())
 		{
 			PreparedComponents[Index] = Evaluation;
 		}
 	}
 }
 
+void FPreparedType::MergeEvaluation(EExpressionEvaluation Evaluation)
+{
+	check(!IsLoopEvaluation(Evaluation));
+	for (int32 Index = 0; Index < PreparedComponents.Num(); ++Index)
+	{
+		if (!PreparedComponents[Index].IsNone())
+		{
+			FPreparedComponent& Component = PreparedComponents[Index];
+			Component = CombineComponents(Component, Evaluation);
+		}
+	}
+}
+
+void FPreparedType::SetLoopEvaluation(FEmitScope& Scope, const FRequestedType& RequestedType)
+{
+	for (int32 Index = 0; Index < PreparedComponents.Num(); ++Index)
+	{
+		if (RequestedType.IsComponentRequested(Index))
+		{
+			FPreparedComponent& Component = PreparedComponents[Index];
+			Component.Evaluation = MakeLoopEvaluation(Component.Evaluation);
+			if (IsLoopEvaluation(Component.Evaluation))
+			{
+				Component.LoopScope = FEmitScope::FindSharedParent(&Scope, Component.LoopScope);
+			}
+		}
+	}
+}
+
 void FPreparedType::SetField(const Shader::FStructField* Field, const FPreparedType& FieldType)
 {
 	for (int32 Index = 0; Index < Field->GetNumComponents(); ++Index)
 	{
-		SetComponentEvaluation(Field->ComponentIndex + Index, FieldType.GetComponentEvaluation(Index));
+		SetComponent(Field->ComponentIndex + Index, FieldType.GetComponent(Index));
 	}
 }
 
@@ -645,7 +842,7 @@ FPreparedType FPreparedType::GetFieldType(const Shader::FStructField* Field) con
 	FPreparedType Result(Field->Type);
 	for (int32 Index = 0; Index < Field->GetNumComponents(); ++Index)
 	{
-		Result.SetComponentEvaluation(Index, GetComponentEvaluation(Field->ComponentIndex + Index));
+		Result.SetComponent(Index, GetComponent(Field->ComponentIndex + Index));
 	}
 	return Result;
 }
@@ -682,9 +879,9 @@ FPreparedType MergePreparedTypes(const FPreparedType& Lhs, const FPreparedType&
 
 	for (int32 Index = 0; Index < NumComponents; ++Index)
 	{
-		const EExpressionEvaluation LhsEvaluation = Lhs.GetComponentEvaluation(Index);
-		const EExpressionEvaluation RhsEvaluation = Rhs.GetComponentEvaluation(Index);
-		Result.SetComponentEvaluation(Index, CombineEvaluations(LhsEvaluation, RhsEvaluation));
+		const FPreparedComponent LhsComponent = Lhs.GetComponent(Index);
+		const FPreparedComponent RhsComponent = Rhs.GetComponent(Index);
+		Result.SetComponent(Index, CombineComponents(LhsComponent, RhsComponent));
 	}
 
 	return Result;
@@ -742,7 +939,7 @@ bool FPrepareValueResult::SetType(FEmitContext& Context, const FRequestedType& R
 			{
 				if (RequestedType.IsComponentRequested(Index))
 				{
-					PreparedType.MergeComponentEvaluation(Index, Evaluation);
+					PreparedType.MergeComponent(Index, Evaluation);
 				}
 			}
 		}
@@ -760,7 +957,7 @@ bool FPrepareValueResult::SetType(FEmitContext& Context, const FRequestedType& R
 		{
 			if (RequestedType.IsComponentRequested(Index))
 			{
-				PreparedType.MergeComponentEvaluation(Index, Type.GetComponentEvaluation(Index));
+				PreparedType.MergeComponent(Index, Type.GetComponent(Index));
 			}
 		}
 		return true;
@@ -768,14 +965,19 @@ bool FPrepareValueResult::SetType(FEmitContext& Context, const FRequestedType& R
 	return false;
 }
 
-bool FPrepareValueResult::SetForwardValue(FEmitContext& Context, const FRequestedType& RequestedType, FExpression* InForwardValue)
+bool FPrepareValueResult::SetForwardValue(FEmitContext& Context, FEmitScope& Scope, const FRequestedType& RequestedType, FExpression* InForwardValue)
 {
 	if (InForwardValue != ForwardValue)
 	{
-		PreparedType = Context.PrepareExpression(InForwardValue, RequestedType);
+		PreparedType = Context.PrepareExpression(InForwardValue, Scope, RequestedType);
 		ForwardValue = InForwardValue;
 	}
-	return (bool)InForwardValue;
+	return !PreparedType.IsVoid();
+}
+
+void FStatement::EmitPreshader(FEmitContext& Context, FEmitScope& Scope, const FRequestedType& RequestedType, TArrayView<const FEmitPreshaderScope> Scopes, Shader::FPreshaderData& OutPreshader) const
+{
+	check(false);
 }
 
 void FExpression::ComputeAnalyticDerivatives(FTree& Tree, FExpressionDerivatives& OutResult) const
@@ -788,7 +990,7 @@ void FExpression::EmitValueShader(FEmitContext& Context, FEmitScope& Scope, cons
 	check(false);
 }
 
-void FExpression::EmitValuePreshader(FEmitContext& Context, const FRequestedType& RequestedType, Shader::FPreshaderData& OutPreshader) const
+void FExpression::EmitValuePreshader(FEmitContext& Context, FEmitScope& Scope, const FRequestedType& RequestedType, Shader::FPreshaderData& OutPreshader) const
 {
 	check(false);
 }
@@ -801,7 +1003,7 @@ FEmitShaderExpression* FExpression::GetValueShader(FEmitContext& Context, FEmitS
 		return PrepareValueResult.ForwardValue->GetValueShader(Context, Scope, RequestedType, ResultType);
 	}
 
-	const EExpressionEvaluation Evaluation = PrepareValueResult.PreparedType.GetEvaluation(RequestedType);
+	const EExpressionEvaluation Evaluation = PrepareValueResult.PreparedType.GetEvaluation(Scope, RequestedType);
 	check(Evaluation != EExpressionEvaluation::None);
 
 	FEmitShaderExpression* Value = nullptr;
@@ -811,7 +1013,7 @@ FEmitShaderExpression* FExpression::GetValueShader(FEmitContext& Context, FEmitS
 	}
 	else
 	{
-		check(Evaluation == EExpressionEvaluation::Shader);
+		check(Evaluation != EExpressionEvaluation::None && Evaluation != EExpressionEvaluation::Unknown);
 		FEmitValueShaderResult Result;
 		EmitValueShader(Context, Scope, RequestedType, Result);
 		Value = Result.Code;
@@ -831,50 +1033,52 @@ FEmitShaderExpression* FExpression::GetValueShader(FEmitContext& Context, FEmitS
 	return GetValueShader(Context, Scope, GetRequestedType());
 }
 
-void FExpression::GetValuePreshader(FEmitContext& Context, const FRequestedType& RequestedType, Shader::FPreshaderData& OutPreshader)
+void FExpression::GetValuePreshader(FEmitContext& Context, FEmitScope& Scope, const FRequestedType& RequestedType, Shader::FPreshaderData& OutPreshader)
 {
 	FOwnerScope OwnerScope(*Context.Errors, GetOwner());
 	if (PrepareValueResult.ForwardValue)
 	{
-		return PrepareValueResult.ForwardValue->GetValuePreshader(Context, RequestedType, OutPreshader);
+		return PrepareValueResult.ForwardValue->GetValuePreshader(Context, Scope, RequestedType, OutPreshader);
 	}
 
-	check(!bReentryFlag);
-	const EExpressionEvaluation Evaluation = PrepareValueResult.PreparedType.GetEvaluation(RequestedType);
-	if (Evaluation == EExpressionEvaluation::Preshader)
-	{
-		bReentryFlag = true;
-		EmitValuePreshader(Context, RequestedType, OutPreshader);
-		bReentryFlag = false;
-	}
-	else if (Evaluation == EExpressionEvaluation::Constant)
+	const int32 PrevStackPosition = Context.PreshaderStackPosition;
+	const EExpressionEvaluation Evaluation = PrepareValueResult.PreparedType.GetEvaluation(Scope, RequestedType);
+
+	if (Evaluation == EExpressionEvaluation::Constant)
 	{
-		const Shader::FValue ConstantValue = GetValueConstant(Context, RequestedType);
+		const Shader::FValue ConstantValue = GetValueConstant(Context, Scope, RequestedType);
+		Context.PreshaderStackPosition++;
 		OutPreshader.WriteOpcode(Shader::EPreshaderOpcode::Constant).Write(ConstantValue);
 	}
 	else
 	{
-		// Value is not used, write a dummy value
-		check(Evaluation == EExpressionEvaluation::None);
-		OutPreshader.WriteOpcode(Shader::EPreshaderOpcode::ConstantZero).Write(Shader::EValueType::Float1);
+		check(Evaluation != EExpressionEvaluation::None && Evaluation != EExpressionEvaluation::Unknown && Evaluation != EExpressionEvaluation::Shader);
+		EmitValuePreshader(Context, Scope, RequestedType, OutPreshader);
 	}
+	check(Context.PreshaderStackPosition == PrevStackPosition + 1);
 }
 
-Shader::FValue FExpression::GetValueConstant(FEmitContext& Context, const FRequestedType& RequestedType)
+Shader::FValue FExpression::GetValueConstant(FEmitContext& Context, FEmitScope& Scope, const FRequestedType& RequestedType)
 {
 	FOwnerScope OwnerScope(*Context.Errors, GetOwner());
 	if (PrepareValueResult.ForwardValue)
 	{
-		return PrepareValueResult.ForwardValue->GetValueConstant(Context, RequestedType);
+		return PrepareValueResult.ForwardValue->GetValueConstant(Context, Scope, RequestedType);
 	}
 
 	check(!bReentryFlag);
-	check(PrepareValueResult.PreparedType.GetEvaluation(RequestedType) == EExpressionEvaluation::Constant);
+
+	const EExpressionEvaluation Evaluation = PrepareValueResult.PreparedType.GetEvaluation(Scope, RequestedType);
+	check(Evaluation == EExpressionEvaluation::Constant || Evaluation == EExpressionEvaluation::ConstantLoop);
 	
 	Shader::FPreshaderData ConstantPreshader;
-	bReentryFlag = true;
-	EmitValuePreshader(Context, RequestedType, ConstantPreshader);
-	bReentryFlag = false;
+	{
+		FExpressionReentryScope ReentryScope(this);
+		const int32 PrevPreshaderStackPosition = Context.PreshaderStackPosition;
+		EmitValuePreshader(Context, Scope, RequestedType, ConstantPreshader);
+		check(Context.PreshaderStackPosition == PrevPreshaderStackPosition + 1);
+		Context.PreshaderStackPosition--;
+	}
 
 	// Evaluate the constant preshader and store its value
 	Shader::FPreshaderStack Stack;
diff --git a/Engine/Source/Runtime/Engine/Private/HLSLTree/HLSLTreeCommon.cpp b/Engine/Source/Runtime/Engine/Private/HLSLTree/HLSLTreeCommon.cpp
--- a/Engine/Source/Runtime/Engine/Private/HLSLTree/HLSLTreeCommon.cpp
+++ b/Engine/Source/Runtime/Engine/Private/HLSLTree/HLSLTreeCommon.cpp
@@ -10,6 +10,12 @@ namespace UE
 namespace HLSLTree
 {
 
+struct FPreshaderLoopScope
+{
+	const FStatement* BreakStatement = nullptr;
+	Shader::FPreshaderLabel BreakLabel;
+};
+
 FExternalInputDescription GetExternalInputDescription(EExternalInput Input)
 {
 	switch (Input)
@@ -141,13 +147,14 @@ void FExpressionConstant::ComputeAnalyticDerivatives(FTree& Tree, FExpressionDer
 	}
 }
 
-bool FExpressionConstant::PrepareValue(FEmitContext& Context, const FRequestedType& RequestedType, FPrepareValueResult& OutResult) const
+bool FExpressionConstant::PrepareValue(FEmitContext& Context, FEmitScope& Scope, const FRequestedType& RequestedType, FPrepareValueResult& OutResult) const
 {
 	return OutResult.SetType(Context, RequestedType, EExpressionEvaluation::Constant, Value.Type);
 }
 
-void FExpressionConstant::EmitValuePreshader(FEmitContext& Context, const FRequestedType& RequestedType, Shader::FPreshaderData& OutPreshader) const
+void FExpressionConstant::EmitValuePreshader(FEmitContext& Context, FEmitScope& Scope, const FRequestedType& RequestedType, Shader::FPreshaderData& OutPreshader) const
 {
+	Context.PreshaderStackPosition++;
 	OutPreshader.WriteOpcode(Shader::EPreshaderOpcode::Constant).Write(Value);
 }
 
@@ -163,14 +170,15 @@ void FExpressionMaterialParameter::ComputeAnalyticDerivatives(FTree& Tree, FExpr
 	}
 }
 
-bool FExpressionMaterialParameter::PrepareValue(FEmitContext& Context, const FRequestedType& RequestedType, FPrepareValueResult& OutResult) const
+bool FExpressionMaterialParameter::PrepareValue(FEmitContext& Context, FEmitScope& Scope, const FRequestedType& RequestedType, FPrepareValueResult& OutResult) const
 {
 	const EExpressionEvaluation Evaluation = IsStaticMaterialParameter(ParameterType) ? EExpressionEvaluation::Constant : EExpressionEvaluation::Preshader;
 	return OutResult.SetType(Context, RequestedType, Evaluation, GetShaderValueType(ParameterType));
 }
 
-void FExpressionMaterialParameter::EmitValuePreshader(FEmitContext& Context, const FRequestedType& RequestedType, Shader::FPreshaderData& OutPreshader) const
+void FExpressionMaterialParameter::EmitValuePreshader(FEmitContext& Context, FEmitScope& Scope, const FRequestedType& RequestedType, Shader::FPreshaderData& OutPreshader) const
 {
+	Context.PreshaderStackPosition++;
 	if (ParameterType == EMaterialParameterType::StaticSwitch)
 	{
 		const FMaterialParameterInfo ParameterInfo(ParameterName);
@@ -215,7 +223,7 @@ void FExpressionExternalInput::ComputeAnalyticDerivatives(FTree& Tree, FExpressi
 	}
 }
 
-bool FExpressionExternalInput::PrepareValue(FEmitContext& Context, const FRequestedType& RequestedType, FPrepareValueResult& OutResult) const
+bool FExpressionExternalInput::PrepareValue(FEmitContext& Context, FEmitScope& Scope, const FRequestedType& RequestedType, FPrepareValueResult& OutResult) const
 {
 	const FExternalInputDescription InputDesc = GetExternalInputDescription(InputType);
 	return OutResult.SetType(Context, RequestedType, EExpressionEvaluation::Shader, InputDesc.Type);
@@ -272,11 +280,11 @@ void FExpressionExternalInput::EmitValueShader(FEmitContext& Context, FEmitScope
 	}
 }
 
-bool FExpressionTextureSample::PrepareValue(FEmitContext& Context, const FRequestedType& RequestedType, FPrepareValueResult& OutResult) const
+bool FExpressionTextureSample::PrepareValue(FEmitContext& Context, FEmitScope& Scope, const FRequestedType& RequestedType, FPrepareValueResult& OutResult) const
 {
-	Context.PrepareExpression(TexCoordExpression, Shader::EValueType::Float2);
-	Context.PrepareExpression(TexCoordDerivatives.ExpressionDdx, Shader::EValueType::Float2);
-	Context.PrepareExpression(TexCoordDerivatives.ExpressionDdy, Shader::EValueType::Float2);
+	Context.PrepareExpression(TexCoordExpression, Scope, Shader::EValueType::Float2);
+	Context.PrepareExpression(TexCoordDerivatives.ExpressionDdx, Scope, Shader::EValueType::Float2);
+	Context.PrepareExpression(TexCoordDerivatives.ExpressionDdy, Scope, Shader::EValueType::Float2);
 
 	return OutResult.SetType(Context, RequestedType, EExpressionEvaluation::Shader, Shader::EValueType::Float4);
 }
@@ -435,12 +443,12 @@ void FExpressionGetStructField::ComputeAnalyticDerivatives(FTree& Tree, FExpress
 	}
 }
 
-bool FExpressionGetStructField::PrepareValue(FEmitContext& Context, const FRequestedType& RequestedType, FPrepareValueResult& OutResult) const
+bool FExpressionGetStructField::PrepareValue(FEmitContext& Context, FEmitScope& Scope, const FRequestedType& RequestedType, FPrepareValueResult& OutResult) const
 {
 	FRequestedType RequestedStructType(StructType, false);
 	RequestedStructType.SetField(Field, RequestedType);
 
-	const FPreparedType& StructPreparedType = Context.PrepareExpression(StructExpression, RequestedStructType);
+	const FPreparedType& StructPreparedType = Context.PrepareExpression(StructExpression, Scope, RequestedStructType);
 	if (!StructPreparedType.IsVoid() && StructPreparedType.StructType != StructType)
 	{
 		return Context.Errors->AddErrorf(TEXT("Expected type %s"), StructType->Name);
@@ -461,12 +469,12 @@ void FExpressionGetStructField::EmitValueShader(FEmitContext& Context, FEmitScop
 		Field->Name);
 }
 
-void FExpressionGetStructField::EmitValuePreshader(FEmitContext& Context, const FRequestedType& RequestedType, Shader::FPreshaderData& OutPreshader) const
+void FExpressionGetStructField::EmitValuePreshader(FEmitContext& Context, FEmitScope& Scope, const FRequestedType& RequestedType, Shader::FPreshaderData& OutPreshader) const
 {
 	FRequestedType RequestedStructType(StructType, false);
 	RequestedStructType.SetField(Field, RequestedType);
 
-	StructExpression->GetValuePreshader(Context, RequestedStructType, OutPreshader);
+	StructExpression->GetValuePreshader(Context, Scope, RequestedStructType, OutPreshader);
 	OutPreshader.WriteOpcode(Shader::EPreshaderOpcode::GetField).Write(Field->Type).Write(Field->ComponentIndex);
 }
 
@@ -487,21 +495,21 @@ void FExpressionSetStructField::ComputeAnalyticDerivatives(FTree& Tree, FExpress
 	}
 }
 
-bool FExpressionSetStructField::PrepareValue(FEmitContext& Context, const FRequestedType& RequestedType, FPrepareValueResult& OutResult) const
+bool FExpressionSetStructField::PrepareValue(FEmitContext& Context, FEmitScope& Scope, const FRequestedType& RequestedType, FPrepareValueResult& OutResult) const
 {
 	FRequestedType RequestedStructType(RequestedType);
 	RequestedStructType.ClearFieldRequested(Field);
 
-	const FPreparedType& StructPreparedType = Context.PrepareExpression(StructExpression, RequestedStructType);
+	const FPreparedType& StructPreparedType = Context.PrepareExpression(StructExpression, Scope, RequestedStructType);
 	if (!StructPreparedType.IsVoid() && StructPreparedType.StructType != StructType)
 	{
 		return Context.Errors->AddErrorf(TEXT("Expected type %s"), StructType->Name);
 	}
 
-	const FPreparedType& FieldPreparedType = Context.PrepareExpression(FieldExpression, RequestedType.GetField(Field));
+	const FPreparedType& FieldPreparedType = Context.PrepareExpression(FieldExpression, Scope, RequestedType.GetField(Field));
 	if (FieldPreparedType.IsVoid())
 	{
-		return OutResult.SetForwardValue(Context, RequestedType, StructExpression);
+		return OutResult.SetForwardValue(Context, Scope, RequestedType, StructExpression);
 	}
 	else
 	{
@@ -519,10 +527,10 @@ void FExpressionSetStructField::EmitValueShader(FEmitContext& Context, FEmitScop
 {
 	FRequestedType RequestedStructType(RequestedType);
 	RequestedStructType.ClearFieldRequested(Field);
-	const EExpressionEvaluation StructEvaluation = StructExpression->GetEvaluation(RequestedStructType);
+	const EExpressionEvaluation StructEvaluation = StructExpression->GetEvaluation(Scope, RequestedStructType);
 
 	const FRequestedType RequestedFieldType = RequestedType.GetField(Field);
-	const EExpressionEvaluation FieldEvaluation = FieldExpression->GetEvaluation(RequestedFieldType);
+	const EExpressionEvaluation FieldEvaluation = FieldExpression->GetEvaluation(Scope, RequestedFieldType);
 	FEmitShaderExpression* FieldValue = FieldExpression->GetValueShader(Context, Scope, RequestedFieldType, Field->Type);
 
 	if (StructEvaluation == EExpressionEvaluation::None)
@@ -548,39 +556,44 @@ void FExpressionSetStructField::EmitValueShader(FEmitContext& Context, FEmitScop
 	}
 }
 
-void FExpressionSetStructField::EmitValuePreshader(FEmitContext& Context, const FRequestedType& RequestedType, Shader::FPreshaderData& OutPreshader) const
+void FExpressionSetStructField::EmitValuePreshader(FEmitContext& Context, FEmitScope& Scope, const FRequestedType& RequestedType, Shader::FPreshaderData& OutPreshader) const
 {
 	FRequestedType RequestedStructType(RequestedType);
 	RequestedStructType.ClearFieldRequested(Field);
-	const EExpressionEvaluation StructEvaluation = StructExpression->GetEvaluation(RequestedStructType);
+	const EExpressionEvaluation StructEvaluation = StructExpression->GetEvaluation(Scope, RequestedStructType);
 
 	const FRequestedType RequestedFieldType = RequestedType.GetField(Field);
-	const EExpressionEvaluation FieldEvaluation = FieldExpression->GetEvaluation(RequestedFieldType);
+	const EExpressionEvaluation FieldEvaluation = FieldExpression->GetEvaluation(Scope, RequestedFieldType);
 
 	if (StructEvaluation != EExpressionEvaluation::None)
 	{
-		StructExpression->GetValuePreshader(Context, RequestedStructType, OutPreshader);
+		StructExpression->GetValuePreshader(Context, Scope, RequestedStructType, OutPreshader);
 	}
 	else
 	{
+		Context.PreshaderStackPosition++;
 		OutPreshader.WriteOpcode(Shader::EPreshaderOpcode::ConstantZero).Write(Shader::FType(StructType));
 	}
 
-	FieldExpression->GetValuePreshader(Context, RequestedFieldType, OutPreshader);
+	FieldExpression->GetValuePreshader(Context, Scope, RequestedFieldType, OutPreshader);
+
+	check(Context.PreshaderStackPosition > 0);
+	Context.PreshaderStackPosition--;
+
 	OutPreshader.WriteOpcode(Shader::EPreshaderOpcode::SetField).Write(Field->ComponentIndex).Write(Field->GetNumComponents());
 }
 
-bool FExpressionSelect::PrepareValue(FEmitContext& Context, const FRequestedType& RequestedType, FPrepareValueResult& OutResult) const
+bool FExpressionSelect::PrepareValue(FEmitContext& Context, FEmitScope& Scope, const FRequestedType& RequestedType, FPrepareValueResult& OutResult) const
 {
-	const FPreparedType& ConditionType = Context.PrepareExpression(ConditionExpression, 1);
-	if (ConditionType.GetEvaluation(Shader::EValueType::Bool1) == EExpressionEvaluation::Constant)
+	const FPreparedType& ConditionType = Context.PrepareExpression(ConditionExpression, Scope, Shader::EValueType::Bool1);
+	if (ConditionType.GetEvaluation(Scope, Shader::EValueType::Bool1) == EExpressionEvaluation::Constant)
 	{
-		const bool bCondition = ConditionExpression->GetValueConstant(Context, Shader::EValueType::Bool1).AsBoolScalar();
-		return OutResult.SetForwardValue(Context, RequestedType, bCondition ? TrueExpression : FalseExpression);
+		const bool bCondition = ConditionExpression->GetValueConstant(Context, Scope, Shader::EValueType::Bool1).AsBoolScalar();
+		return OutResult.SetForwardValue(Context, Scope, RequestedType, bCondition ? TrueExpression : FalseExpression);
 	}
 
-	const FPreparedType& LhsType = Context.PrepareExpression(FalseExpression, RequestedType);
-	const FPreparedType& RhsType = Context.PrepareExpression(TrueExpression, RequestedType);
+	const FPreparedType& LhsType = Context.PrepareExpression(FalseExpression, Scope, RequestedType);
+	const FPreparedType& RhsType = Context.PrepareExpression(TrueExpression, Scope, RequestedType);
 	
 	if (LhsType.ValueComponentType != RhsType.ValueComponentType ||
 		LhsType.StructType != RhsType.StructType)
@@ -602,7 +615,7 @@ void FExpressionSelect::EmitValueShader(FEmitContext& Context, FEmitScope& Scope
 		FalseValue);
 }
 
-void FExpressionSelect::EmitValuePreshader(FEmitContext& Context, const FRequestedType& RequestedType, Shader::FPreshaderData& OutPreshader) const
+void FExpressionSelect::EmitValuePreshader(FEmitContext& Context, FEmitScope& Scope, const FRequestedType& RequestedType, Shader::FPreshaderData& OutPreshader) const
 {
 	check(false); // TODO
 }
@@ -638,9 +651,9 @@ void FExpressionUnaryOp::ComputeAnalyticDerivatives(FTree& Tree, FExpressionDeri
 	}
 }
 
-bool FExpressionUnaryOp::PrepareValue(FEmitContext& Context, const FRequestedType& RequestedType, FPrepareValueResult& OutResult) const
+bool FExpressionUnaryOp::PrepareValue(FEmitContext& Context, FEmitScope& Scope, const FRequestedType& RequestedType, FPrepareValueResult& OutResult) const
 {
-	const FPreparedType& InputType = Context.PrepareExpression(Input, RequestedType);
+	const FPreparedType& InputType = Context.PrepareExpression(Input, Scope, RequestedType);
 	if (InputType.IsVoid())
 	{
 		return false;
@@ -697,12 +710,12 @@ void FExpressionUnaryOp::EmitValueShader(FEmitContext& Context, FEmitScope& Scop
 	}
 }
 
-void FExpressionUnaryOp::EmitValuePreshader(FEmitContext& Context, const FRequestedType& RequestedType, Shader::FPreshaderData& OutPreshader) const
+void FExpressionUnaryOp::EmitValuePreshader(FEmitContext& Context, FEmitScope& Scope, const FRequestedType& RequestedType, Shader::FPreshaderData& OutPreshader) const
 {
 	const FUnaryOpDescription OpDesc = GetUnaryOpDesription(Op);
 	check(OpDesc.PreshaderOpcode != Shader::EPreshaderOpcode::Nop);
 
-	Input->GetValuePreshader(Context, RequestedType, OutPreshader);
+	Input->GetValuePreshader(Context, Scope, RequestedType, OutPreshader);
 	OutPreshader.WriteOpcode(OpDesc.PreshaderOpcode);
 }
 
@@ -763,10 +776,10 @@ void FExpressionBinaryOp::ComputeAnalyticDerivatives(FTree& Tree, FExpressionDer
 	}
 }
 
-bool FExpressionBinaryOp::PrepareValue(FEmitContext& Context, const FRequestedType& RequestedType, FPrepareValueResult& OutResult) const
+bool FExpressionBinaryOp::PrepareValue(FEmitContext& Context, FEmitScope& Scope, const FRequestedType& RequestedType, FPrepareValueResult& OutResult) const
 {
-	const FPreparedType& LhsType = Context.PrepareExpression(Lhs, RequestedType);
-	const FPreparedType& RhsType = Context.PrepareExpression(Rhs, RequestedType);
+	const FPreparedType& LhsType = Context.PrepareExpression(Lhs, Scope, RequestedType);
+	const FPreparedType& RhsType = Context.PrepareExpression(Rhs, Scope, RequestedType);
 	if (LhsType.IsVoid() || RhsType.IsVoid())
 	{
 		return false;
@@ -875,14 +888,18 @@ void FExpressionBinaryOp::EmitValueShader(FEmitContext& Context, FEmitScope& Sco
 	}
 }
 
-void FExpressionBinaryOp::EmitValuePreshader(FEmitContext& Context, const FRequestedType& RequestedType, Shader::FPreshaderData& OutPreshader) const
+void FExpressionBinaryOp::EmitValuePreshader(FEmitContext& Context, FEmitScope& Scope, const FRequestedType& RequestedType, Shader::FPreshaderData& OutPreshader) const
 {
 	const FRequestedType InputType = Private::GetBinaryOpInputType(RequestedType, Op, Lhs->GetRequestedType(), Rhs->GetRequestedType());
 	const FBinaryOpDescription OpDesc = GetBinaryOpDesription(Op);
 	check(OpDesc.PreshaderOpcode != Shader::EPreshaderOpcode::Nop);
 
-	Lhs->GetValuePreshader(Context, InputType, OutPreshader);
-	Rhs->GetValuePreshader(Context, InputType, OutPreshader);
+	Lhs->GetValuePreshader(Context, Scope, InputType, OutPreshader);
+	Rhs->GetValuePreshader(Context, Scope, InputType, OutPreshader);
+
+	check(Context.PreshaderStackPosition > 0);
+	Context.PreshaderStackPosition--;
+
 	OutPreshader.WriteOpcode(OpDesc.PreshaderOpcode);
 }
 
@@ -896,16 +913,16 @@ void FExpressionSwizzle::ComputeAnalyticDerivatives(FTree& Tree, FExpressionDeri
 	}
 }
 
-bool FExpressionSwizzle::PrepareValue(FEmitContext& Context, const FRequestedType& RequestedType, FPrepareValueResult& OutResult) const
+bool FExpressionSwizzle::PrepareValue(FEmitContext& Context, FEmitScope& Scope, const FRequestedType& RequestedType, FPrepareValueResult& OutResult) const
 {
 	const FRequestedType RequestedInputType = Parameters.GetRequestedInputType(RequestedType);
-	const FPreparedType& InputType = Context.PrepareExpression(Input, RequestedInputType);
+	const FPreparedType& InputType = Context.PrepareExpression(Input, Scope, RequestedInputType);
 
 	FPreparedType ResultType(InputType.ValueComponentType);
 	for (int32 ComponentIndex = 0; ComponentIndex < Parameters.NumComponents; ++ComponentIndex)
 	{
 		const int32 SwizzledComponentIndex = Parameters.ComponentIndex[ComponentIndex];
-		ResultType.SetComponentEvaluation(ComponentIndex, InputType.GetComponentEvaluation(SwizzledComponentIndex));
+		ResultType.SetComponent(ComponentIndex, InputType.GetComponent(SwizzledComponentIndex));
 	}
 
 	return OutResult.SetType(Context, RequestedType, ResultType);
@@ -942,11 +959,11 @@ void FExpressionSwizzle::EmitValueShader(FEmitContext& Context, FEmitScope& Scop
 	}
 }
 
-void FExpressionSwizzle::EmitValuePreshader(FEmitContext& Context, const FRequestedType& RequestedType, Shader::FPreshaderData& OutPreshader) const
+void FExpressionSwizzle::EmitValuePreshader(FEmitContext& Context, FEmitScope& Scope, const FRequestedType& RequestedType, Shader::FPreshaderData& OutPreshader) const
 {
 	const FRequestedType RequestedInputType = Parameters.GetRequestedInputType(RequestedType);
 
-	Input->GetValuePreshader(Context, RequestedInputType, OutPreshader);
+	Input->GetValuePreshader(Context, Scope, RequestedInputType, OutPreshader);
 	OutPreshader.WriteOpcode(Shader::EPreshaderOpcode::ComponentSwizzle)
 		.Write((uint8)Parameters.NumComponents)
 		.Write((uint8)Parameters.ComponentIndex[0])
@@ -966,9 +983,9 @@ void FExpressionAppend::ComputeAnalyticDerivatives(FTree& Tree, FExpressionDeriv
 	}
 }
 
-bool FExpressionAppend::PrepareValue(FEmitContext& Context, const FRequestedType& RequestedType, FPrepareValueResult& OutResult) const
+bool FExpressionAppend::PrepareValue(FEmitContext& Context, FEmitScope& Scope, const FRequestedType& RequestedType, FPrepareValueResult& OutResult) const
 {
-	const FPreparedType& LhsType = Context.PrepareExpression(Lhs, RequestedType);
+	const FPreparedType& LhsType = Context.PrepareExpression(Lhs, Scope, RequestedType);
 	const int32 NumRequestedComponents = RequestedType.GetNumComponents();
 	const int32 NumLhsComponents = FMath::Min(LhsType.GetNumComponents(), NumRequestedComponents);
 
@@ -978,7 +995,7 @@ bool FExpressionAppend::PrepareValue(FEmitContext& Context, const FRequestedType
 		RhsRequestedType.SetComponentRequest(Index - NumLhsComponents, RequestedType.IsComponentRequested(Index));
 	}
 
-	const FPreparedType& RhsType = Context.PrepareExpression(Rhs, RhsRequestedType);
+	const FPreparedType& RhsType = Context.PrepareExpression(Rhs, Scope, RhsRequestedType);
 	const int32 NumRhsComponents = FMath::Min(RhsType.GetNumComponents(), NumRequestedComponents - NumLhsComponents);
 
 	if (LhsType.ValueComponentType != RhsType.ValueComponentType)
@@ -989,11 +1006,11 @@ bool FExpressionAppend::PrepareValue(FEmitContext& Context, const FRequestedType
 	FPreparedType ResultType(LhsType.ValueComponentType);
 	for (int32 Index = 0; Index < NumLhsComponents; ++Index)
 	{
-		ResultType.SetComponentEvaluation(Index, LhsType.GetComponentEvaluation(Index));
+		ResultType.SetComponent(Index, LhsType.GetComponent(Index));
 	}
 	for (int32 Index = 0; Index < NumRhsComponents; ++Index)
 	{
-		ResultType.SetComponentEvaluation(NumLhsComponents + Index, LhsType.GetComponentEvaluation(Index));
+		ResultType.SetComponent(NumLhsComponents + Index, LhsType.GetComponent(Index));
 	}
 
 	return OutResult.SetType(Context, RequestedType, ResultType);
@@ -1059,18 +1076,22 @@ void FExpressionAppend::EmitValueShader(FEmitContext& Context, FEmitScope& Scope
 	}
 }
 
-void FExpressionAppend::EmitValuePreshader(FEmitContext& Context, const FRequestedType& RequestedType, Shader::FPreshaderData& OutPreshader) const
+void FExpressionAppend::EmitValuePreshader(FEmitContext& Context, FEmitScope& Scope, const FRequestedType& RequestedType, Shader::FPreshaderData& OutPreshader) const
 {
 	const Private::FAppendTypes Types = Private::GetAppendTypes(RequestedType, Lhs->GetRequestedType(), Rhs->GetRequestedType());
-	Lhs->GetValuePreshader(Context, Types.LhsType, OutPreshader);
+	Lhs->GetValuePreshader(Context, Scope, Types.LhsType, OutPreshader);
 	if (Types.RhsType.IsVoid())
 	{
-		Rhs->GetValuePreshader(Context, Types.RhsType, OutPreshader);
+		Rhs->GetValuePreshader(Context, Scope, Types.RhsType, OutPreshader);
+
+		check(Context.PreshaderStackPosition > 0);
+		Context.PreshaderStackPosition--;
+
 		OutPreshader.WriteOpcode(Shader::EPreshaderOpcode::AppendVector);
 	}
 }
 
-bool FExpressionReflectionVector::PrepareValue(FEmitContext& Context, const FRequestedType& RequestedType, FPrepareValueResult& OutResult) const
+bool FExpressionReflectionVector::PrepareValue(FEmitContext& Context, FEmitScope& Scope, const FRequestedType& RequestedType, FPrepareValueResult& OutResult) const
 {
 	Context.bReadMaterialNormal = true;
 	return OutResult.SetType(Context, RequestedType, EExpressionEvaluation::Shader, Shader::EValueType::Float3);
@@ -1081,7 +1102,7 @@ void FExpressionReflectionVector::EmitValueShader(FEmitContext& Context, FEmitSc
 	OutResult.Code = Context.EmitInlineExpression(Scope, Shader::EValueType::Float3, TEXT("Parameters.ReflectionVector"));
 }
 
-bool FStatementBreak::Prepare(FEmitContext& Context) const
+bool FStatementBreak::Prepare(FEmitContext& Context, FEmitScope& Scope) const
 {
 	return true;
 }
@@ -1091,7 +1112,15 @@ void FStatementBreak::EmitShader(FEmitContext& Context, FEmitScope& Scope) const
 	Context.EmitStatement(Scope, TEXT("break;"));
 }
 
-bool FStatementReturn::Prepare(FEmitContext& Context) const
+void FStatementBreak::EmitPreshader(FEmitContext& Context, FEmitScope& Scope, const FRequestedType& RequestedType, TArrayView<const FEmitPreshaderScope> Scopes, Shader::FPreshaderData& OutPreshader) const
+{
+	FPreshaderLoopScope* LoopScope = Context.PreshaderLoopScopes.Last();
+	check(!LoopScope->BreakStatement);
+	LoopScope->BreakStatement = this;
+	LoopScope->BreakLabel = OutPreshader.WriteJump(Shader::EPreshaderOpcode::Jump);
+}
+
+bool FStatementReturn::Prepare(FEmitContext& Context, FEmitScope& Scope) const
 {
 	return true;
 }
@@ -1101,34 +1130,34 @@ void FStatementReturn::EmitShader(FEmitContext& Context, FEmitScope& Scope) cons
 	Context.EmitStatement(Scope, TEXT("return %;"), Expression->GetValueShader(Context, Scope));
 }
 
-bool FStatementIf::Prepare(FEmitContext& Context) const
+bool FStatementIf::Prepare(FEmitContext& Context, FEmitScope& Scope) const
 {
-	const FPreparedType& ConditionType = Context.PrepareExpression(ConditionExpression, 1);
+	const FPreparedType& ConditionType = Context.PrepareExpression(ConditionExpression, Scope, Shader::EValueType::Bool1);
 	if (ConditionType.IsVoid())
 	{
 		return false;
 	}
 
-	const EExpressionEvaluation ConditionEvaluation = ConditionType.GetEvaluation(Shader::EValueType::Bool1);
+	const EExpressionEvaluation ConditionEvaluation = ConditionType.GetEvaluation(Scope, Shader::EValueType::Bool1);
 	check(ConditionEvaluation != EExpressionEvaluation::None);
 	if (ConditionEvaluation == EExpressionEvaluation::Constant)
 	{
-		const bool bCondition = ConditionExpression->GetValueConstant(Context, Shader::EValueType::Bool1).AsBoolScalar();
+		const bool bCondition = ConditionExpression->GetValueConstant(Context, Scope, Shader::EValueType::Bool1).AsBoolScalar();
 		if (bCondition)
 		{
-			Context.PrepareScope(ThenScope);
-			Context.MarkScopeDead(ElseScope);
+			Context.MarkScopeEvaluation(Scope, ThenScope, EExpressionEvaluation::Constant);
+			Context.MarkScopeDead(Scope, ElseScope);
 		}
 		else
 		{
-			Context.MarkScopeDead(ThenScope);
-			Context.PrepareScope(ElseScope);
+			Context.MarkScopeDead(Scope, ThenScope);
+			Context.MarkScopeEvaluation(Scope, ElseScope, EExpressionEvaluation::Constant);
 		}
 	}
 	else
 	{
-		Context.PrepareScope(ThenScope);
-		Context.PrepareScope(ElseScope);
+		Context.MarkScopeEvaluation(Scope, ThenScope, ConditionEvaluation);
+		Context.MarkScopeEvaluation(Scope, ElseScope, ConditionEvaluation);
 	}
 
 	return true;
@@ -1137,10 +1166,10 @@ bool FStatementIf::Prepare(FEmitContext& Context) const
 void FStatementIf::EmitShader(FEmitContext& Context, FEmitScope& Scope) const
 {
 	FEmitShaderNode* Dependency = nullptr;
-	const EExpressionEvaluation ConditionEvaluation = ConditionExpression->GetEvaluation(Shader::EValueType::Bool1);
+	const EExpressionEvaluation ConditionEvaluation = ConditionExpression->GetEvaluation(Scope, Shader::EValueType::Bool1);
 	if (ConditionEvaluation == EExpressionEvaluation::Constant)
 	{
-		const bool bCondition = ConditionExpression->GetValueConstant(Context, Shader::EValueType::Bool1).AsBoolScalar();
+		const bool bCondition = ConditionExpression->GetValueConstant(Context, Scope, Shader::EValueType::Bool1).AsBoolScalar();
 		if (bCondition)
 		{
 			Dependency = Context.EmitNextScope(Scope, ThenScope);
@@ -1159,9 +1188,33 @@ void FStatementIf::EmitShader(FEmitContext& Context, FEmitScope& Scope) const
 	Context.EmitNextScopeWithDependency(Scope, Dependency, NextScope);
 }
 
-bool FStatementLoop::Prepare(FEmitContext& Context) const
+void FStatementIf::EmitPreshader(FEmitContext& Context, FEmitScope& Scope, const FRequestedType& RequestedType, TArrayView<const FEmitPreshaderScope> Scopes, Shader::FPreshaderData& OutPreshader) const
 {
-	Context.PrepareScope(LoopScope);
+	ConditionExpression->GetValuePreshader(Context, Scope, Shader::EValueType::Bool1, OutPreshader);
+
+	check(Context.PreshaderStackPosition > 0);
+	Context.PreshaderStackPosition--;
+	const Shader::FPreshaderLabel Label0 = OutPreshader.WriteJump(Shader::EPreshaderOpcode::JumpIfFalse);
+
+	Context.EmitPreshaderScope(ThenScope, RequestedType, Scopes, OutPreshader);
+
+	const Shader::FPreshaderLabel Label1 = OutPreshader.WriteJump(Shader::EPreshaderOpcode::Jump);
+	OutPreshader.SetLabel(Label0);
+	
+	Context.EmitPreshaderScope(ElseScope, RequestedType, Scopes, OutPreshader);
+
+	OutPreshader.SetLabel(Label1);
+}
+
+bool FStatementLoop::Prepare(FEmitContext& Context, FEmitScope& Scope) const
+{
+	FEmitScope* BreakScope = Context.PrepareScope(&BreakStatement->GetParentScope());
+	if (!BreakScope)
+	{
+		return false;
+	}
+
+	Context.MarkScopeEvaluation(Scope, LoopScope, BreakScope->Evaluation);
 	return true;
 }
 
@@ -1171,5 +1224,19 @@ void FStatementLoop::EmitShader(FEmitContext& Context, FEmitScope& Scope) const
 	Context.EmitNextScopeWithDependency(Scope, Dependency, NextScope);
 }
 
+void FStatementLoop::EmitPreshader(FEmitContext& Context, FEmitScope& Scope, const FRequestedType& RequestedType, TArrayView<const FEmitPreshaderScope> Scopes, Shader::FPreshaderData& OutPreshader) const
+{
+	FPreshaderLoopScope PreshaderLoopScope;
+	Context.PreshaderLoopScopes.Add(&PreshaderLoopScope);
+
+	const Shader::FPreshaderLabel Label = OutPreshader.GetLabel();
+	Context.EmitPreshaderScope(LoopScope, RequestedType, Scopes, OutPreshader);
+	OutPreshader.WriteJump(Shader::EPreshaderOpcode::Jump, Label);
+
+	verify(Context.PreshaderLoopScopes.Pop() == &PreshaderLoopScope);
+	check(PreshaderLoopScope.BreakStatement == BreakStatement);
+	OutPreshader.SetLabel(PreshaderLoopScope.BreakLabel);
+}
+
 } // namespace HLSLTree
 } // namespace UE
diff --git a/Engine/Source/Runtime/Engine/Private/HLSLTree/HLSLTreeEmit.cpp b/Engine/Source/Runtime/Engine/Private/HLSLTree/HLSLTreeEmit.cpp
--- a/Engine/Source/Runtime/Engine/Private/HLSLTree/HLSLTreeEmit.cpp
+++ b/Engine/Source/Runtime/Engine/Private/HLSLTree/HLSLTreeEmit.cpp
@@ -134,6 +134,11 @@ FEmitScope* FEmitScope::FindSharedParent(FEmitScope* Lhs, FEmitScope* Rhs)
 {
 	FEmitScope* Scope0 = Lhs;
 	FEmitScope* Scope1 = Rhs;
+	if (!Scope0)
+	{
+		return Scope1;
+	}
+
 	if (Scope1)
 	{
 		while (Scope0 != Scope1)
@@ -153,6 +158,42 @@ FEmitScope* FEmitScope::FindSharedParent(FEmitScope* Lhs, FEmitScope* Rhs)
 	return Scope0;
 }
 
+bool FEmitScope::HasParent(const FEmitScope* InParent) const
+{
+	if (InParent)
+	{
+		const FEmitScope* CheckParent = this;
+		while (CheckParent)
+		{
+			if (CheckParent == InParent)
+			{
+				return true;
+			}
+			CheckParent = CheckParent->ParentScope;
+		}
+	}
+	return false;
+}
+
+bool FEmitScope::IsLoop() const
+{
+	return OwnerStatement && OwnerStatement->IsLoop();
+}
+
+FEmitScope* FEmitScope::FindLoop()
+{
+	FEmitScope* CheckParent = this;
+	while (CheckParent)
+	{
+		if (CheckParent->IsLoop())
+		{
+			return CheckParent;
+		}
+		CheckParent = CheckParent->ParentScope;
+	}
+	return nullptr;
+}
+
 FEmitContext::FEmitContext(FMemStackBase& InAllocator, FErrorHandlerInterface& InErrors, const Shader::FStructTypeRegistry& InTypeRegistry)
 	: Allocator(&InAllocator)
 	, Errors(&InErrors)
@@ -164,123 +205,229 @@ FEmitContext::~FEmitContext()
 {
 }
 
-const FPreparedType& FEmitContext::PrepareExpression(FExpression* InExpression, const FRequestedType& RequestedType)
+FPreparedType FEmitContext::PrepareExpression(FExpression* InExpression, FEmitScope& Scope, const FRequestedType& RequestedType)
 {
-	static FPreparedType VoidType;
-	if (!InExpression)
+	if (!InExpression || RequestedType.IsVoid())
 	{
-		return VoidType;
+		return FPreparedType();
 	}
 
 	FOwnerScope OwnerScope(*Errors, InExpression->GetOwner());
+	FPrepareValueResult& Result = InExpression->PrepareValueResult;
 	if (InExpression->bReentryFlag)
 	{
 		// Valid for this to be called reentrantly
 		// Code should ensure that the type is set before the reentrant call, otherwise type will not be valid here
 		// LocalPHI nodes rely on this to break loops
-		return InExpression->PrepareValueResult.GetPreparedType();
+		FEmitScope* LoopScope = Scope.FindLoop();
+		check(LoopScope);
+		Result.PreparedType.SetLoopEvaluation(*LoopScope, RequestedType);
+		return Result.GetPreparedType();
 	}
 
-	bool bNeedToUpdateType = false;
-	if (InExpression->CurrentRequestedType.RequestedComponents.Num() == 0)
+	bool bResult = false;
 	{
-		InExpression->CurrentRequestedType = RequestedType;
-		bNeedToUpdateType = !RequestedType.IsVoid();
+		FExpressionReentryScope ReentryScope(InExpression);
+		bResult = InExpression->PrepareValue(*this, Scope, RequestedType, Result);
 	}
-	else if (InExpression->CurrentRequestedType.GetStructType() != RequestedType.GetStructType())
+
+	if (!bResult)
 	{
-		Errors->AddError(TEXT("Type mismatch"));
-		return VoidType;
+		Result.SetTypeVoid();
 	}
 	else
 	{
-		const int32 NumComponents = RequestedType.GetNumComponents();
-		InExpression->CurrentRequestedType.RequestedComponents.PadToNum(NumComponents, false);
-		for (int32 Index = 0; Index < NumComponents; ++Index)
-		{
-			const bool bPrevRequest = InExpression->CurrentRequestedType.IsComponentRequested(Index);
-			const bool bRequest = RequestedType.IsComponentRequested(Index);
-			if (!bPrevRequest && bRequest)
-			{
-				InExpression->CurrentRequestedType.SetComponentRequest(Index);
-				bNeedToUpdateType = true;
-			}
-		}
+		check(!Result.GetPreparedType().IsVoid());
 	}
 
-	if (bNeedToUpdateType)
+	return Result.GetPreparedType();
+}
+
+FEmitScope* FEmitContext::InternalPrepareScope(FScope* InScope, FScope* InParentScope)
+{
+	if (!InScope)
 	{
-		check(!InExpression->CurrentRequestedType.IsVoid());
+		return nullptr;
+	}
 
-		bool bResult = false;
+	TArray<FEmitScope*, TInlineAllocator<16>> EmitScopeChain;
+	{
+		TArray<FScope*, TInlineAllocator<16>> ScopeChain;
+		ScopeChain.Add(InScope);
 		{
-			FExpressionReentryScope ReentryScope(InExpression);
-			bResult = InExpression->PrepareValue(*this, InExpression->CurrentRequestedType, InExpression->PrepareValueResult);
+			FScope* CurrentParentScope = InParentScope;
+			while (CurrentParentScope)
+			{
+				ScopeChain.Add(CurrentParentScope);
+				CurrentParentScope = CurrentParentScope->ParentScope;
+			}
 		}
 
-		if (!bResult)
+		EmitScopeChain.Reserve(ScopeChain.Num());
+		FEmitScope* EmitParentScope = nullptr;
+		bool bFoundUninitializedScope = false;
+		for (int32 Index = ScopeChain.Num() - 1; Index >= 0; --Index)
 		{
-			// If we failed to assign a valid type, reset the requested type as well
-			// This ensures we'll try to compute a type again the next time we're called
-			InExpression->CurrentRequestedType.Reset();
-			InExpression->PrepareValueResult.SetTypeVoid();
+			FScope* CurrentScope = ScopeChain[Index];
+			FEmitScope* EmitScope = AcquireEmitScopeWithParent(CurrentScope, EmitParentScope);
+			if (EmitScope->State == EEmitScopeState::Initializing)
+			{
+				// If we hit an initializing scope, reset the chain
+				// We only want the chain to include scopes *below* the last initializing scope
+				EmitScopeChain.Reset();
+				// Don't expect any uninitialized scopes *above* the initializing scope (scopes are initialized top down)
+				check(!bFoundUninitializedScope);
+			}
+			else
+			{
+				if (EmitScope->State == EEmitScopeState::Uninitialized)
+				{
+					bFoundUninitializedScope = true;
+					EmitScope->State = EEmitScopeState::Initializing;
+				}
+				EmitScopeChain.Add(EmitScope);
+			}
+			EmitParentScope = EmitScope;
 		}
-		else
+
+		if (EmitScopeChain.Num() == 0)
 		{
-			check(!InExpression->PrepareValueResult.GetPreparedType().IsVoid());
+			// This can happen if the current scope is initializing
+			check(EmitParentScope);
+			check(EmitParentScope->State == EEmitScopeState::Initializing);
+			return EmitParentScope;
 		}
 	}
 
-	return InExpression->PrepareValueResult.GetPreparedType();
-}
-
-FEmitScope* FEmitContext::InternalPrepareScope(FScope* Scope, FScope* ParentScope, bool bMarkDead)
-{
-	FEmitScope* EmitScope = nullptr;
-	if (Scope)
+	bool bMarkDead = false;
+	for (int32 Index = 0; Index < EmitScopeChain.Num(); ++Index)
 	{
-		FEmitScope* EmitParentScope = InternalPrepareScope(ParentScope, ParentScope ? ParentScope->ParentScope : nullptr, false);
-		EmitScope = AcquireEmitScopeWithParent(Scope, EmitParentScope);
-
-		if (bMarkDead || (EmitParentScope && EmitParentScope->State == EEmitScopeState::Dead))
+		FEmitScope* EmitScope = EmitScopeChain[Index];
+		if (bMarkDead)
 		{
+			EmitScope->Evaluation = EExpressionEvaluation::None;
 			EmitScope->State = EEmitScopeState::Dead;
 		}
-		else if (EmitScope->State == EEmitScopeState::Uninitialized)
+		else if (EmitScope->State == EEmitScopeState::Initializing)
 		{
-			EmitScope->State = EEmitScopeState::Initializing;
-			bool bPrepareResult = true;
-			if (Scope->OwnerStatement)
+			if (EmitScope->OwnerStatement)
 			{
-				bPrepareResult = Scope->OwnerStatement->Prepare(*this);
+				FEmitScope* EmitParentScope = EmitScope->ParentScope;
+				check(EmitParentScope);
+				if (EmitScope->OwnerStatement->Prepare(*this, *EmitParentScope))
+				{
+					if (EmitScope->State == EEmitScopeState::Initializing)
+					{
+						check(EmitScope->Evaluation != EExpressionEvaluation::None);
+						EmitScope->State = EEmitScopeState::Live;
+					}
+					else
+					{
+						check(EmitScope->State == EEmitScopeState::Dead);
+					}
+				}
+				else
+				{
+					EmitScope->State = EEmitScopeState::Uninitialized;
+					EmitScope->Evaluation = EExpressionEvaluation::Unknown;
+				}
 			}
-			if (EmitScope->State == EEmitScopeState::Initializing)
+			else
 			{
-				// If Prepare() returns false, we're still uninitialized
-				EmitScope->State = bPrepareResult ? EEmitScopeState::Live : EEmitScopeState::Uninitialized;
+				EmitScope->State = EEmitScopeState::Live;
 			}
 		}
+
+		if (EmitScope->State == EEmitScopeState::Dead)
+		{
+			EmitScope->Evaluation = EExpressionEvaluation::None;
+			bMarkDead = true;
+		}
 	}
 
-	return EmitScope;
+	EExpressionEvaluation CurrentEvaluation = EExpressionEvaluation::Constant;
+	for (int32 Index = EmitScopeChain.Num() - 1; Index >= 0; --Index)
+	{
+		FEmitScope* EmitScope = EmitScopeChain[Index];
+		if (EmitScope->State == EEmitScopeState::Dead)
+		{
+			check(EmitScope->Evaluation == EExpressionEvaluation::None);
+			CurrentEvaluation = EExpressionEvaluation::Constant;
+		}
+		else if (EmitScope->State == EEmitScopeState::Uninitialized || CurrentEvaluation == EExpressionEvaluation::Unknown)
+		{
+			CurrentEvaluation = EExpressionEvaluation::Unknown;
+			EmitScope->Evaluation = EExpressionEvaluation::Unknown;
+		}
+		else
+		{
+			CurrentEvaluation = CombineEvaluations(CurrentEvaluation, EmitScope->Evaluation);
+			EmitScope->Evaluation = CurrentEvaluation;
+		}
+	}
+
+	return EmitScopeChain.Last();
 }
 
-bool FEmitContext::PrepareScope(FScope* Scope)
+FEmitScope* FEmitContext::PrepareScope(FScope* Scope)
 {
-	FEmitScope* EmitScope = InternalPrepareScope(Scope, Scope ? Scope->ParentScope : nullptr, false);
-	return EmitScope && EmitScope->State != EEmitScopeState::Dead;
+	FEmitScope* EmitScope = InternalPrepareScope(Scope, Scope ? Scope->ParentScope : nullptr);
+	return EmitScope && EmitScope->State != EEmitScopeState::Dead ? EmitScope : nullptr;
 }
 
-bool FEmitContext::PrepareScopeWithParent(FScope* Scope, FScope* ParentScope)
+FEmitScope* FEmitContext::PrepareScopeWithParent(FScope* Scope, FScope* ParentScope)
 {
-	FEmitScope* EmitScope = InternalPrepareScope(Scope, ParentScope, false);
-	return EmitScope && EmitScope->State != EEmitScopeState::Dead;
+	FEmitScope* EmitScope = InternalPrepareScope(Scope, ParentScope);
+	return EmitScope && EmitScope->State != EEmitScopeState::Dead ? EmitScope : nullptr;
 }
 
-bool FEmitContext::MarkScopeDead(FScope* Scope)
+void FEmitContext::MarkScopeEvaluation(FEmitScope& EmitParentScope, FScope* Scope, EExpressionEvaluation Evaluation)
 {
-	InternalPrepareScope(Scope, Scope ? Scope->ParentScope : nullptr, true);
-	return false;
+	FEmitScope* EmitScope = AcquireEmitScopeWithParent(Scope, &EmitParentScope);
+	if (EmitScope && EmitScope->State != EEmitScopeState::Dead)
+	{
+		// Don't set loop evaluations on scopes
+		EmitScope->Evaluation = CombineEvaluations(EmitScope->Evaluation, MakeNonLoopEvaluation(Evaluation));
+	}
+}
+
+void FEmitContext::MarkScopeDead(FEmitScope& EmitParentScope, FScope* Scope)
+{
+	FEmitScope* EmitScope = AcquireEmitScopeWithParent(Scope, &EmitParentScope);
+	if (EmitScope)
+	{
+		EmitScope->State = EEmitScopeState::Dead;
+		EmitScope->Evaluation = EExpressionEvaluation::None;
+	}
+}
+
+void FEmitContext::EmitPreshaderScope(const FScope* Scope, const FRequestedType& RequestedType, TArrayView<const FEmitPreshaderScope> PreshaderScopes, Shader::FPreshaderData& OutPreshader)
+{
+	FEmitScope* EmitScope = FindEmitScope(Scope);
+	if (EmitScope)
+	{
+		EmitPreshaderScope(*EmitScope, RequestedType, PreshaderScopes, OutPreshader);
+	}
+}
+
+void FEmitContext::EmitPreshaderScope(FEmitScope& EmitScope, const FRequestedType& RequestedType, TArrayView<const FEmitPreshaderScope> PreshaderScopes, Shader::FPreshaderData& OutPreshader)
+{
+	for (const FEmitPreshaderScope& PreshaderScope : PreshaderScopes)
+	{
+		if (PreshaderScope.Scope == &EmitScope)
+		{
+			PreshaderScope.Value->GetValuePreshader(*this, EmitScope, RequestedType, OutPreshader);
+			check(PreshaderStackPosition > 0);
+			PreshaderStackPosition--;
+			OutPreshader.WriteOpcode(Shader::EPreshaderOpcode::Assign); // assign the new value
+			break;
+		}
+	}
+
+	if (EmitScope.ContainedStatement)
+	{
+		EmitScope.ContainedStatement->EmitPreshader(*this, EmitScope, RequestedType, PreshaderScopes, OutPreshader);
+	}
 }
 
 FEmitScope* FEmitContext::AcquireEmitScopeWithParent(const FScope* Scope, FEmitScope* EmitParentScope)
@@ -293,8 +440,12 @@ FEmitScope* FEmitContext::AcquireEmitScopeWithParent(const FScope* Scope, FEmitS
 		if (!EmitScope)
 		{
 			EmitScope = new(*Allocator) FEmitScope();
+			EmitScope->OwnerStatement = Scope->OwnerStatement;
+			EmitScope->ContainedStatement = Scope->ContainedStatement;
 			EmitScope->ParentScope = EmitParentScope;
 			EmitScope->NestedLevel = EmitParentScope ? EmitParentScope->NestedLevel + 1 : 0;
+			EmitScope->Evaluation = EExpressionEvaluation::Unknown;
+			EmitScope->State = EEmitScopeState::Uninitialized;
 			EmitScopeMap.Add(Scope, EmitScope);
 		}
 		check(!EmitParentScope || EmitScope->ParentScope == EmitParentScope);
@@ -308,9 +459,15 @@ FEmitScope* FEmitContext::AcquireEmitScope(const FScope* Scope)
 	return AcquireEmitScopeWithParent(Scope, EmitParentScope);
 }
 
+FEmitScope* FEmitContext::FindEmitScope(const FScope* Scope)
+{
+	FEmitScope* const* PrevEmitScope = EmitScopeMap.Find(Scope);
+	return PrevEmitScope ? *PrevEmitScope : nullptr;
+}
+
 FEmitScope* FEmitContext::InternalEmitScope(const FScope* Scope)
 {
-	FEmitScope* EmitScope = AcquireEmitScope(Scope);
+	FEmitScope* EmitScope = FindEmitScope(Scope);
 	if (EmitScope && EmitScope->State != EEmitScopeState::Dead)
 	{
 		if (Scope->ContainedStatement)
@@ -505,7 +662,7 @@ void WriteMaterialUniformAccess(Shader::EValueComponentType ComponentType, uint3
 FEmitShaderExpression* FEmitContext::EmitPreshaderOrConstant(FEmitScope& Scope, const FRequestedType& RequestedType, FExpression* Expression)
 {
 	Shader::FPreshaderData LocalPreshader;
-	Expression->EmitValuePreshader(*this, RequestedType, LocalPreshader);
+	Expression->EmitValuePreshader(*this, Scope, RequestedType, LocalPreshader);
 
 	const Shader::FType Type = RequestedType.GetType();
 
@@ -544,7 +701,7 @@ FEmitShaderExpression* FEmitContext::EmitPreshaderOrConstant(FEmitScope& Scope,
 		const Shader::EValueType FieldType = Type.GetFlatFieldType(FieldIndex);
 		const Shader::FValueTypeDescription TypeDesc = Shader::GetValueTypeDescription(FieldType);
 		const int32 NumFieldComponents = TypeDesc.NumComponents;
-		const EExpressionEvaluation FieldEvaluation = Expression->GetPreparedType().GetFieldEvaluation(ComponentIndex, NumFieldComponents);
+		const EExpressionEvaluation FieldEvaluation = Expression->GetPreparedType().GetFieldEvaluation(Scope, ComponentIndex, NumFieldComponents);
 
 		if (FieldEvaluation == EExpressionEvaluation::Preshader)
 		{
@@ -558,7 +715,7 @@ FEmitShaderExpression* FEmitContext::EmitPreshaderOrConstant(FEmitScope& Scope,
 				PreshaderHeader->FieldIndex = UniformExpressionSet.UniformPreshaderFields.Num();
 				PreshaderHeader->NumFields = 0u;
 				PreshaderHeader->OpcodeOffset = UniformExpressionSet.UniformPreshaderData.Num();
-				Expression->EmitValuePreshader(*this, RequestedType, UniformExpressionSet.UniformPreshaderData);
+				Expression->EmitValuePreshader(*this, Scope, RequestedType, UniformExpressionSet.UniformPreshaderData);
 				PreshaderHeader->OpcodeSize = UniformExpressionSet.UniformPreshaderData.Num() - PreshaderHeader->OpcodeOffset;
 			}
 
@@ -630,7 +787,9 @@ FEmitShaderExpression* FEmitContext::EmitPreshaderOrConstant(FEmitScope& Scope,
 		else
 		{
 			// We allow FieldEvaluation to be 'None', since in that case we still need to fill in a value for the HLSL initializer
-			check(FieldEvaluation == EExpressionEvaluation::Constant || FieldEvaluation == EExpressionEvaluation::None);
+			check(FieldEvaluation == EExpressionEvaluation::Constant ||
+				FieldEvaluation == EExpressionEvaluation::ConstantLoop ||
+				FieldEvaluation == EExpressionEvaluation::None);
 
 			// The type generated by the preshader might not match the expected type
 			// In the future, with new HLSLTree, preshader could potentially include explicit cast opcodes, and avoid implicit conversions
diff --git a/Engine/Source/Runtime/Engine/Private/Materials/MaterialExpressionHLSL.cpp b/Engine/Source/Runtime/Engine/Private/Materials/MaterialExpressionHLSL.cpp
--- a/Engine/Source/Runtime/Engine/Private/Materials/MaterialExpressionHLSL.cpp
+++ b/Engine/Source/Runtime/Engine/Private/Materials/MaterialExpressionHLSL.cpp
@@ -446,7 +446,7 @@ bool UMaterialExpressionWhileLoop::GenerateHLSLStatements(FMaterialHLSLGenerator
 	LoopStatement->LoopScope->AddPreviousScope(*IfStatement->ThenScope);
 	LoopStatement->NextScope->AddPreviousScope(*IfStatement->ElseScope);
 
-	Generator.GetTree().NewStatement<FStatementBreak>(*IfStatement->ElseScope);
+	LoopStatement->BreakStatement = Generator.GetTree().NewStatement<FStatementBreak>(*IfStatement->ElseScope);
 
 	IfStatement->ConditionExpression = Condition.AcquireHLSLExpression(Generator, *LoopStatement->LoopScope);
 	LoopBody.GenerateHLSLStatements(Generator, *IfStatement->ThenScope);
@@ -519,7 +519,7 @@ bool UMaterialExpressionForLoop::GenerateHLSLStatements(FMaterialHLSLGenerator&
 	LoopStatement->LoopScope->AddPreviousScope(*IfStatement->ThenScope);
 	LoopStatement->NextScope->AddPreviousScope(*IfStatement->ElseScope);
 
-	Generator.GetTree().NewStatement<FStatementBreak>(*IfStatement->ElseScope);
+	LoopStatement->BreakStatement = Generator.GetTree().NewStatement<FStatementBreak>(*IfStatement->ElseScope);
 
 	FExpression* LocalExpression = Generator.GetTree().AcquireLocal(*LoopStatement->LoopScope, ExpressionData->LocalName);
 
diff --git a/Engine/Source/Runtime/Engine/Private/Materials/MaterialHLSLEmitter.cpp b/Engine/Source/Runtime/Engine/Private/Materials/MaterialHLSLEmitter.cpp
--- a/Engine/Source/Runtime/Engine/Private/Materials/MaterialHLSLEmitter.cpp
+++ b/Engine/Source/Runtime/Engine/Private/Materials/MaterialHLSLEmitter.cpp
@@ -732,10 +732,12 @@ bool MaterialEmitHLSL(const FMaterialCompileTargetParameters& InCompilerTarget,
 	const FStructField* NormalField = Generator.GetMaterialAttributesType()->FindFieldByName(*FMaterialAttributeDefinitionMap::GetAttributeName(MP_Normal));
 	check(NormalField);
 
+	FEmitScope* EmitResultScope = EmitContext.PrepareScope(&Generator.GetResultStatement()->GetParentScope());
+
 	// Prepare all fields *except* normal
 	FRequestedType RequestedMaterialAttributesType(Generator.GetMaterialAttributesType());
 	RequestedMaterialAttributesType.ClearFieldRequested(NormalField);
-	if (EmitContext.PrepareExpression(Generator.GetResultExpression(), RequestedMaterialAttributesType).IsVoid())
+	if (EmitContext.PrepareExpression(Generator.GetResultExpression(), *EmitResultScope, RequestedMaterialAttributesType).IsVoid())
 	{
 		return false;
 	}
@@ -748,7 +750,7 @@ bool MaterialEmitHLSL(const FMaterialCompileTargetParameters& InCompilerTarget,
 	{
 		// No access to material normal, can execute everything in phase0
 		RequestedMaterialAttributesType.SetFieldRequested(NormalField);
-		if (EmitContext.PrepareExpression(Generator.GetResultExpression(), RequestedMaterialAttributesType).IsVoid())
+		if (EmitContext.PrepareExpression(Generator.GetResultExpression(), *EmitResultScope, RequestedMaterialAttributesType).IsVoid())
 		{
 			return false;
 		}
@@ -763,10 +765,12 @@ bool MaterialEmitHLSL(const FMaterialCompileTargetParameters& InCompilerTarget,
 		// Reset state
 		HLSLTree->ResetNodes();
 
+		EmitResultScope = EmitContext.PrepareScope(&Generator.GetResultStatement()->GetParentScope());
+
 		// Prepare code for just the normal
 		FRequestedType RequestedMaterialNormal(Generator.GetMaterialAttributesType(), false);
 		RequestedMaterialNormal.SetFieldRequested(NormalField);
-		if (EmitContext.PrepareExpression(Generator.GetResultExpression(), RequestedMaterialNormal).IsVoid())
+		if (EmitContext.PrepareExpression(Generator.GetResultExpression(), *EmitResultScope, RequestedMaterialNormal).IsVoid())
 		{
 			return false;
 		}
diff --git a/Engine/Source/Runtime/Engine/Private/Materials/MaterialUniformExpressions.cpp b/Engine/Source/Runtime/Engine/Private/Materials/MaterialUniformExpressions.cpp
--- a/Engine/Source/Runtime/Engine/Private/Materials/MaterialUniformExpressions.cpp
+++ b/Engine/Source/Runtime/Engine/Private/Materials/MaterialUniformExpressions.cpp
@@ -751,6 +751,15 @@ void FUniformExpressionSet::FillUniformBuffer(const FMaterialRenderContext& Mate
 						*DestAddress++ = FloatValue[i];
 					}
 				}
+				else if (UniformTypeDesc.ComponentType == EValueComponentType::Int)
+				{
+					const FIntValue IntValue = FieldValue.AsInt();
+					int32* DestAddress = (int32*)PreshaderBuffer + PreshaderField.BufferOffset;
+					for (int32 i = 0; i < NumFieldComponents; ++i)
+					{
+						*DestAddress++ = IntValue[i];
+					}
+				}
 				else if (UniformTypeDesc.ComponentType == EValueComponentType::Bool)
 				{
 					const FBoolValue BoolValue = FieldValue.AsBool();
diff --git a/Engine/Source/Runtime/Engine/Private/Shader/Preshader.cpp b/Engine/Source/Runtime/Engine/Private/Shader/Preshader.cpp
--- a/Engine/Source/Runtime/Engine/Private/Shader/Preshader.cpp
+++ b/Engine/Source/Runtime/Engine/Private/Shader/Preshader.cpp
@@ -85,10 +85,10 @@ FPreshaderValue FPreshaderStack::PopValue()
 	return Value;
 }
 
-FPreshaderValue FPreshaderStack::PeekValue()
+FPreshaderValue FPreshaderStack::PeekValue(int32 Offset)
 {
 	FPreshaderValue Value;
-	Value.Type = Values.Last();
+	Value.Type = Values.Last(Offset);
 
 	const int32 NumComponents = Value.Type.GetNumComponents();
 	const int32 ComponentIndex = Components.Num() - NumComponents;
@@ -178,6 +178,39 @@ void FPreshaderData::WriteValue(const FValue& Value)
 	}
 }
 
+FPreshaderLabel FPreshaderData::WriteJump(EPreshaderOpcode Op)
+{
+	WriteOpcode(Op);
+	const int32 Offset = Data.Num();
+	Write((uint32)0xffffffff); // Write a placeholder for the jump offset
+	return FPreshaderLabel(Offset);
+}
+
+void FPreshaderData::WriteJump(EPreshaderOpcode Op, FPreshaderLabel Label)
+{
+	WriteOpcode(Op);
+	const int32 Offset = Data.Num();
+	const int32 JumpOffset = Label.Offset - Offset - 4; // Compute the offset to jump
+	Write(JumpOffset);
+}
+
+void FPreshaderData::SetLabel(FPreshaderLabel InLabel)
+{
+	const int32 TargetOffset = Data.Num();
+	const int32 BaseOffset = InLabel.Offset;
+	const int32 JumpOffset = TargetOffset - BaseOffset - 4; // Compute the offset to jump
+	check(JumpOffset >= 0);
+
+	uint8* Dst = &Data[BaseOffset];
+	check(Dst[0] == 0xff && Dst[1] == 0xff && Dst[2] == 0xff && Dst[3] == 0xff);
+	FMemory::Memcpy(Dst, &JumpOffset, 4); // Patch the offset into the jump opcode
+}
+
+FPreshaderLabel FPreshaderData::GetLabel()
+{
+	return FPreshaderLabel(Data.Num());
+}
+
 FPreshaderDataContext::FPreshaderDataContext(const FPreshaderData& InData)
 	: Ptr(InData.Data.GetData())
 	, EndPtr(Ptr + InData.Data.Num())
@@ -312,6 +345,27 @@ static void EvaluateGetField(FPreshaderStack& Stack, FPreshaderDataContext& REST
 	Stack.PushValue(FieldType, FieldComponents);
 }
 
+static void EvaluatePushValue(FPreshaderStack& Stack, FPreshaderDataContext& RESTRICT Data)
+{
+	const int32 StackOffset = ReadPreshaderValue<uint16>(Data);
+	const FPreshaderValue Value = Stack.PeekValue(StackOffset);
+	// Make a local copy of the component array, as it will be invalidated when pushing the copy
+	const TArray<FValueComponent, TInlineAllocator<64>> LocalComponent(Value.Component);
+	Stack.PushValue(Value.Type, LocalComponent);
+}
+
+static void EvaluateAssign(FPreshaderStack& Stack)
+{
+	const FPreshaderValue Value = Stack.PopValue();
+	// Make a local copy of the component array, as it will be invalidated when pushing the copy
+	const TArray<FValueComponent, TInlineAllocator<64>> LocalComponent(Value.Component);
+
+	// Remove the old value
+	Stack.PopValue();
+	// Replace with the new value
+	Stack.PushValue(Value.Type, LocalComponent);
+}
+
 static void EvaluateParameter(FPreshaderStack& Stack, const FUniformExpressionSet* UniformExpressionSet, uint32 ParameterIndex, const FMaterialRenderContext& Context)
 {
 	if (!UniformExpressionSet)
@@ -503,6 +557,25 @@ static void EvaluateRuntimeVirtualTextureUniform(const FMaterialRenderContext& C
 	}
 }
 
+static void EvaluateJump(FPreshaderDataContext& RESTRICT Data)
+{
+	const int32 JumpOffset = ReadPreshaderValue<int32>(Data);
+	check(Data.Ptr + JumpOffset <= Data.EndPtr);
+	Data.Ptr += JumpOffset;
+}
+
+static void EvaluateJumpIfFalse(FPreshaderStack& Stack, FPreshaderDataContext& RESTRICT Data)
+{
+	const int32 JumpOffset = ReadPreshaderValue<int32>(Data);
+	check(Data.Ptr + JumpOffset <= Data.EndPtr);
+
+	const FValue ConditionValue = Stack.PopValue().AsShaderValue();
+	if (!ConditionValue.AsBoolScalar())
+	{
+		Data.Ptr += JumpOffset;
+	}
+}
+
 FPreshaderValue EvaluatePreshader(const FUniformExpressionSet* UniformExpressionSet, const FMaterialRenderContext& Context, FPreshaderStack& Stack, FPreshaderDataContext& RESTRICT Data)
 {
 	uint8 const* const DataEnd = Data.EndPtr;
@@ -520,10 +593,13 @@ FPreshaderValue EvaluatePreshader(const FUniformExpressionSet* UniformExpression
 		case EPreshaderOpcode::Parameter:
 			EvaluateParameter(Stack, UniformExpressionSet, ReadPreshaderValue<uint16>(Data), Context);
 			break;
+		case EPreshaderOpcode::PushValue: EvaluatePushValue(Stack, Data); break;
+		case EPreshaderOpcode::Assign: EvaluateAssign(Stack); break;
 		case EPreshaderOpcode::Add: EvaluateBinaryOp(Stack, Add); break;
 		case EPreshaderOpcode::Sub: EvaluateBinaryOp(Stack, Sub); break;
 		case EPreshaderOpcode::Mul: EvaluateBinaryOp(Stack, Mul); break;
 		case EPreshaderOpcode::Div: EvaluateBinaryOp(Stack, Div); break;
+		case EPreshaderOpcode::Less: EvaluateBinaryOp(Stack, Less); break;
 		case EPreshaderOpcode::Fmod: EvaluateBinaryOp(Stack, Fmod); break;
 		case EPreshaderOpcode::Min: EvaluateBinaryOp(Stack, Min); break;
 		case EPreshaderOpcode::Max: EvaluateBinaryOp(Stack, Max); break;
@@ -560,6 +636,8 @@ FPreshaderValue EvaluatePreshader(const FUniformExpressionSet* UniformExpression
 		case EPreshaderOpcode::ExternalTextureCoordinateScaleRotation: EvaluateExternalTextureCoordinateScaleRotation(Context, Stack, Data); break;
 		case EPreshaderOpcode::ExternalTextureCoordinateOffset: EvaluateExternalTextureCoordinateOffset(Context, Stack, Data); break;
 		case EPreshaderOpcode::RuntimeVirtualTextureUniform: EvaluateRuntimeVirtualTextureUniform(Context, Stack, Data); break;
+		case EPreshaderOpcode::Jump: EvaluateJump(Data); break;
+		case EPreshaderOpcode::JumpIfFalse: EvaluateJumpIfFalse(Stack, Data); break;
 		default:
 			UE_LOG(LogMaterial, Fatal, TEXT("Unknown preshader opcode %d"), (uint8)Opcode);
 			break;
diff --git a/Engine/Source/Runtime/Engine/Private/Shader/ShaderValue.cpp b/Engine/Source/Runtime/Engine/Private/Shader/ShaderValue.cpp
--- a/Engine/Source/Runtime/Engine/Private/Shader/ShaderValue.cpp
+++ b/Engine/Source/Runtime/Engine/Private/Shader/ShaderValue.cpp
@@ -945,6 +945,7 @@ struct FOpMin : public FOpBase { template<typename T> T operator()(T Lhs, T Rhs)
 struct FOpMax : public FOpBase { template<typename T> T operator()(T Lhs, T Rhs) const { return FMath::Max(Lhs, Rhs); } };
 struct FOpFmod : public FOpBaseNoInt { template<typename T> T operator()(T Lhs, T Rhs) const { return FMath::Fmod(Lhs, Rhs); } };
 struct FOpAtan2 : public FOpBaseNoInt { template<typename T> T operator()(T Lhs, T Rhs) const { return FMath::Atan2(Lhs, Rhs); } };
+struct FOpLess : public FOpBase { template<typename T> bool operator()(T Lhs, T Rhs) const { return Lhs < Rhs; } };
 
 template<typename Operation>
 inline FValue UnaryOp(const Operation& Op, const FValue& Value)
@@ -1053,6 +1054,56 @@ inline FValue BinaryOp(const Operation& Op, const FValue& Lhs, const FValue& Rhs
 	return Result;
 }
 
+template<typename Operation>
+inline FValue CompareOp(const Operation& Op, const FValue& Lhs, const FValue& Rhs)
+{
+	if (Lhs.Type.IsStruct() || Rhs.Type.IsStruct())
+	{
+		return FValue();
+	}
+	const FValueTypeDescription LhsDesc = GetValueTypeDescription(Lhs.Type);
+	const FValueTypeDescription RhsDesc = GetValueTypeDescription(Rhs.Type);
+	const int8 NumComponents = GetNumComponentsResult(LhsDesc.NumComponents, RhsDesc.NumComponents);
+
+	FValue Result;
+	Result.Type = MakeValueType(EValueComponentType::Bool, NumComponents);
+	if constexpr (Operation::SupportsDouble)
+	{
+		if (LhsDesc.ComponentType == EValueComponentType::Double || RhsDesc.ComponentType == EValueComponentType::Double)
+		{
+			const FDoubleValue LhsCast = Lhs.AsDouble();
+			const FDoubleValue RhsCast = Rhs.AsDouble();
+			for (int32 i = 0; i < NumComponents; ++i)
+			{
+				Result.Component.Add(Op(LhsCast.Component[i], RhsCast.Component[i]));
+			}
+			return Result;
+		}
+	}
+
+	if constexpr (Operation::SupportsInt)
+	{
+		if (LhsDesc.ComponentType != EValueComponentType::Float && RhsDesc.ComponentType != EValueComponentType::Float)
+		{
+			const FIntValue LhsCast = Lhs.AsInt();
+			const FIntValue RhsCast = Rhs.AsInt();
+			for (int32 i = 0; i < NumComponents; ++i)
+			{
+				Result.Component.Add(Op(LhsCast.Component[i], RhsCast.Component[i]));
+			}
+			return Result;
+		}
+	}
+
+	const FFloatValue LhsCast = Lhs.AsFloat();
+	const FFloatValue RhsCast = Rhs.AsFloat();
+	for (int32 i = 0; i < NumComponents; ++i)
+	{
+		Result.Component.Add(Op(LhsCast.Component[i], RhsCast.Component[i]));
+	}
+	return Result;
+}
+
 } // namespace Private
 
 bool operator==(const FValue& Lhs, const FValue& Rhs)
@@ -1223,6 +1274,11 @@ FValue Div(const FValue& Lhs, const FValue& Rhs)
 	return Private::BinaryOp(Private::FOpDiv(), Lhs, Rhs);
 }
 
+FValue Less(const FValue& Lhs, const FValue& Rhs)
+{
+	return Private::CompareOp(Private::FOpLess(), Lhs, Rhs);
+}
+
 FValue Min(const FValue& Lhs, const FValue& Rhs)
 {
 	return Private::BinaryOp(Private::FOpMin(), Lhs, Rhs);
diff --git a/Engine/Source/Runtime/Engine/Public/HLSLTree/HLSLTree.h b/Engine/Source/Runtime/Engine/Public/HLSLTree/HLSLTree.h
--- a/Engine/Source/Runtime/Engine/Public/HLSLTree/HLSLTree.h
+++ b/Engine/Source/Runtime/Engine/Public/HLSLTree/HLSLTree.h
@@ -34,6 +34,8 @@ namespace HLSLTree
 
 class FNode;
 class FScope;
+class FStatement;
+class FExpression;
 class FFunction;
 class FExpressionLocalPHI;
 class FRequestedType;
@@ -183,6 +185,15 @@ private:
 	friend class FTree;
 };
 
+struct FEmitPreshaderScope
+{
+	FEmitPreshaderScope() = default;
+	FEmitPreshaderScope(FEmitScope* InScope, FExpression* InValue) : Scope(InScope), Value(InValue) {}
+
+	FEmitScope* Scope = nullptr;
+	FExpression* Value = nullptr;
+};
+
 /**
  * Represents an HLSL statement.  This is a piece of code that doesn't evaluate to any value, but instead should be executed sequentially, and likely has side-effects.
  * Examples include assigning a value, or various control flow structures (if, for, while, etc)
@@ -193,15 +204,19 @@ class FStatement : public FNode
 public:
 	FScope& GetParentScope() const { return *ParentScope; }
 
+	virtual bool IsLoop() const { return false; }
+
 protected:
-	virtual bool Prepare(FEmitContext& Context) const = 0;
+	virtual bool Prepare(FEmitContext& Context, FEmitScope& Scope) const = 0;
 	virtual void EmitShader(FEmitContext& Context, FEmitScope& Scope) const = 0;
+	virtual void EmitPreshader(FEmitContext& Context, FEmitScope& Scope, const FRequestedType& RequestedType, TArrayView<const FEmitPreshaderScope> Scopes, Shader::FPreshaderData& OutPreshader) const;
 
 private:
 	FScope* ParentScope = nullptr;
 
 	friend class FTree;
 	friend class FEmitContext;
+	friend class FExpressionLocalPHI;
 };
 
 /**
@@ -256,9 +271,47 @@ public:
 	/** 1 bit per component, a value of 'true' means the specified component is requsted */
 	TBitArray<> RequestedComponents;
 };
+inline bool operator==(const FRequestedType& Lhs, const FRequestedType& Rhs)
+{
+	return Lhs.ValueComponentType == Rhs.ValueComponentType &&
+		Lhs.StructType == Rhs.StructType &&
+		Lhs.RequestedComponents == Rhs.RequestedComponents;
+}
+inline bool operator!=(const FRequestedType& Lhs, const FRequestedType& Rhs)
+{
+	return !operator==(Lhs, Rhs);
+}
 
 FRequestedType MakeRequestedType(Shader::EValueComponentType ComponentType, const FRequestedType& RequestedComponents);
 
+struct FPreparedComponent
+{
+	FPreparedComponent() = default;
+	FPreparedComponent(EExpressionEvaluation InEvaluation, FEmitScope* InLoopScope = nullptr)
+		: LoopScope(InLoopScope), Evaluation(InEvaluation)
+	{
+		check(!IsLoopEvaluation(InEvaluation) || InLoopScope);
+	}
+
+	inline bool IsNone() const { return Evaluation == EExpressionEvaluation::None; }
+
+	EExpressionEvaluation GetEvaluation(const FEmitScope& Scope) const;
+
+	FEmitScope* LoopScope = nullptr;
+	EExpressionEvaluation Evaluation = EExpressionEvaluation::None;
+};
+inline bool operator==(const FPreparedComponent& Lhs, const FPreparedComponent& Rhs)
+{
+	return Lhs.Evaluation == Rhs.Evaluation &&
+		Lhs.LoopScope == Rhs.LoopScope;
+}
+inline bool operator!=(const FPreparedComponent& Lhs, const FPreparedComponent& Rhs)
+{
+	return !operator==(Lhs, Rhs);
+}
+
+FPreparedComponent CombineComponents(const FPreparedComponent& Lhs, const FPreparedComponent& Rhs);
+
 /**
  * Like FRequestedType, but tracks an EExpressionEvaluation per component, rather than a simple requested flag
  */
@@ -271,6 +324,8 @@ public:
 	FPreparedType(const Shader::FType& InType);
 
 	void SetEvaluation(EExpressionEvaluation Evaluation);
+	void MergeEvaluation(EExpressionEvaluation Evaluation);
+	void SetLoopEvaluation(FEmitScope& Scope, const FRequestedType& RequestedType);
 
 	void SetField(const Shader::FStructField* Field, const FPreparedType& FieldType);
 	FPreparedType GetFieldType(const Shader::FStructField* Field) const;
@@ -283,21 +338,33 @@ public:
 	bool IsInitialized() const { return StructType != nullptr || ValueComponentType != Shader::EValueComponentType::Void; }
 	bool IsVoid() const;
 
-	EExpressionEvaluation GetEvaluation() const;
-	EExpressionEvaluation GetEvaluation(const FRequestedType& RequestedType) const;
-	EExpressionEvaluation GetFieldEvaluation(int32 ComponentIndex, int32 NumComponents) const;
-	EExpressionEvaluation GetComponentEvaluation(int32 Index) const;
+	EExpressionEvaluation GetEvaluation(const FEmitScope& Scope) const;
+	EExpressionEvaluation GetEvaluation(const FEmitScope& Scope, const FRequestedType& RequestedType) const;
+	EExpressionEvaluation GetFieldEvaluation(const FEmitScope& Scope, int32 ComponentIndex, int32 NumComponents) const;
+	FPreparedComponent GetComponent(int32 Index) const;
 
-	void SetComponentEvaluation(int32 Index, EExpressionEvaluation Evaluation);
-	void MergeComponentEvaluation(int32 Index, EExpressionEvaluation Evaluation);
+	void SetComponent(int32 Index, const FPreparedComponent& InComponent);
+	void MergeComponent(int32 Index, const FPreparedComponent& InComponent);
 
 	/** Unlike FRequestedType, one of these should be set */
 	const Shader::FStructType* StructType = nullptr;
 	Shader::EValueComponentType ValueComponentType = Shader::EValueComponentType::Void;
 
+	void EnsureNumComponents(int32 NumComponents);
+
 	/** Evaluation type for each component, may be 'None' for components that are unused */
-	TArray<EExpressionEvaluation, TInlineAllocator<16>> PreparedComponents;
+	TArray<FPreparedComponent, TInlineAllocator<16>> PreparedComponents;
 };
+inline bool operator==(const FPreparedType& Lhs, const FPreparedType& Rhs)
+{
+	return Lhs.ValueComponentType == Rhs.ValueComponentType &&
+		Lhs.StructType == Rhs.StructType &&
+		Lhs.PreparedComponents == Rhs.PreparedComponents;
+}
+inline bool operator!=(const FPreparedType& Lhs, const FPreparedType& Rhs)
+{
+	return !operator==(Lhs, Rhs);
+}
 
 FPreparedType MergePreparedTypes(const FPreparedType& Lhs, const FPreparedType& Rhs);
 
@@ -311,7 +378,7 @@ public:
 	bool SetType(FEmitContext& Context, const FRequestedType& RequestedType, EExpressionEvaluation Evaluation, const Shader::FType& Type);
 	bool SetType(FEmitContext& Context, const FRequestedType& RequestedType, const FPreparedType& Type);
 
-	bool SetForwardValue(FEmitContext& Context, const FRequestedType& RequestedType, FExpression* InValue);
+	bool SetForwardValue(FEmitContext& Context, FEmitScope& Scope, const FRequestedType& RequestedType, FExpression* InValue);
 
 private:
 	bool TryMergePreparedType(FEmitContext& Context, const Shader::FStructType* StructType, Shader::EValueComponentType ComponentType);
@@ -320,6 +387,7 @@ private:
 	FPreparedType PreparedType;
 
 	friend class FExpression;
+	friend class FEmitContext;
 };
 
 struct FEmitValueShaderResult
@@ -355,26 +423,25 @@ public:
 	const FPreparedType& GetPreparedType() const { return PrepareValueResult.PreparedType; }
 	FRequestedType GetRequestedType() const { return PrepareValueResult.PreparedType.GetRequestedType(); }
 	Shader::FType GetType() const { return PrepareValueResult.PreparedType.GetType(); }
-	EExpressionEvaluation GetEvaluation(const FRequestedType& RequestedType) const { return PrepareValueResult.PreparedType.GetEvaluation(RequestedType); }
+	EExpressionEvaluation GetEvaluation(const FEmitScope& Scope, const FRequestedType& RequestedType) const { return PrepareValueResult.PreparedType.GetEvaluation(Scope, RequestedType); }
 
 	virtual void Reset() override;
 
 	FEmitShaderExpression* GetValueShader(FEmitContext& Context, FEmitScope& Scope, const FRequestedType& RequestedType, const Shader::FType& ResultType);
-	void GetValuePreshader(FEmitContext& Context, const FRequestedType& RequestedType, Shader::FPreshaderData& OutPreshader);
-	Shader::FValue GetValueConstant(FEmitContext& Context, const FRequestedType& RequestedType);
+	void GetValuePreshader(FEmitContext& Context, FEmitScope& Scope, const FRequestedType& RequestedType, Shader::FPreshaderData& OutPreshader);
+	Shader::FValue GetValueConstant(FEmitContext& Context, FEmitScope& Scope, const FRequestedType& RequestedType);
 
 	FEmitShaderExpression* GetValueShader(FEmitContext& Context, FEmitScope& Scope, const FRequestedType& RequestedType);
 	FEmitShaderExpression* GetValueShader(FEmitContext& Context, FEmitScope& Scope);
 
 protected:
 	virtual void ComputeAnalyticDerivatives(FTree& Tree, FExpressionDerivatives& OutResult) const;
-	virtual bool PrepareValue(FEmitContext& Context, const FRequestedType& RequestedType, FPrepareValueResult& OutResult) const = 0;
+	virtual bool PrepareValue(FEmitContext& Context, FEmitScope& Scope, const FRequestedType& RequestedType, FPrepareValueResult& OutResult) const = 0;
 	virtual void EmitValueShader(FEmitContext& Context, FEmitScope& Scope, const FRequestedType& RequestedType, FEmitValueShaderResult& OutResult) const;
-	virtual void EmitValuePreshader(FEmitContext& Context, const FRequestedType& RequestedType, Shader::FPreshaderData& OutPreshader) const;
+	virtual void EmitValuePreshader(FEmitContext& Context, FEmitScope& Scope, const FRequestedType& RequestedType, Shader::FPreshaderData& OutPreshader) const;
 
 private:
 	FExpressionDerivatives Derivatives;
-	FRequestedType CurrentRequestedType;
 	FPrepareValueResult PrepareValueResult;
 	bool bReentryFlag = false;
 	bool bComputedDerivatives = false;
@@ -456,8 +523,8 @@ private:
 	friend class FExpression;
 	friend class FEmitContext;
 
-	FStatement* OwnerStatement = nullptr;
 	FScope* ParentScope = nullptr;
+	FStatement* OwnerStatement = nullptr;
 	FStatement* ContainedStatement = nullptr;
 	FScope* PreviousScope[MaxNumPreviousScopes];
 	TMap<FName, FExpression*> LocalMap;
diff --git a/Engine/Source/Runtime/Engine/Public/HLSLTree/HLSLTreeCommon.h b/Engine/Source/Runtime/Engine/Public/HLSLTree/HLSLTreeCommon.h
--- a/Engine/Source/Runtime/Engine/Public/HLSLTree/HLSLTreeCommon.h
+++ b/Engine/Source/Runtime/Engine/Public/HLSLTree/HLSLTreeCommon.h
@@ -21,8 +21,8 @@ public:
 	Shader::FValue Value;
 
 	virtual void ComputeAnalyticDerivatives(FTree& Tree, FExpressionDerivatives& OutResult) const override;
-	virtual bool PrepareValue(FEmitContext& Context, const FRequestedType& RequestedType, FPrepareValueResult& OutResult) const override;
-	virtual void EmitValuePreshader(FEmitContext& Context, const FRequestedType& RequestedType, Shader::FPreshaderData& OutPreshader) const override;
+	virtual bool PrepareValue(FEmitContext& Context, FEmitScope& Scope, const FRequestedType& RequestedType, FPrepareValueResult& OutResult) const override;
+	virtual void EmitValuePreshader(FEmitContext& Context, FEmitScope& Scope, const FRequestedType& RequestedType, Shader::FPreshaderData& OutPreshader) const override;
 };
 
 class FExpressionMaterialParameter : public FExpression
@@ -37,8 +37,8 @@ public:
 	EMaterialParameterType ParameterType;
 
 	virtual void ComputeAnalyticDerivatives(FTree& Tree, FExpressionDerivatives& OutResult) const override;
-	virtual bool PrepareValue(FEmitContext& Context, const FRequestedType& RequestedType, FPrepareValueResult& OutResult) const override;
-	virtual void EmitValuePreshader(FEmitContext& Context, const FRequestedType& RequestedType, Shader::FPreshaderData& OutPreshader) const override;
+	virtual bool PrepareValue(FEmitContext& Context, FEmitScope& Scope, const FRequestedType& RequestedType, FPrepareValueResult& OutResult) const override;
+	virtual void EmitValuePreshader(FEmitContext& Context, FEmitScope& Scope, const FRequestedType& RequestedType, Shader::FPreshaderData& OutPreshader) const override;
 };
 
 enum class EExternalInput : uint8
@@ -122,7 +122,7 @@ public:
 	EExternalInput InputType;
 
 	virtual void ComputeAnalyticDerivatives(FTree& Tree, FExpressionDerivatives& OutResult) const override;
-	virtual bool PrepareValue(FEmitContext& Context, const FRequestedType& RequestedType, FPrepareValueResult& OutResult) const override;
+	virtual bool PrepareValue(FEmitContext& Context, FEmitScope& Scope, const FRequestedType& RequestedType, FPrepareValueResult& OutResult) const override;
 	virtual void EmitValueShader(FEmitContext& Context, FEmitScope& Scope, const FRequestedType& RequestedType, FEmitValueShaderResult& OutResult) const override;
 };
 
@@ -143,7 +143,7 @@ public:
 	ESamplerSourceMode SamplerSource;
 	ETextureMipValueMode MipValueMode;
 
-	virtual bool PrepareValue(FEmitContext& Context, const FRequestedType& RequestedType, FPrepareValueResult& OutResult) const override;
+	virtual bool PrepareValue(FEmitContext& Context, FEmitScope& Scope, const FRequestedType& RequestedType, FPrepareValueResult& OutResult) const override;
 	virtual void EmitValueShader(FEmitContext& Context, FEmitScope& Scope, const FRequestedType& RequestedType, FEmitValueShaderResult& OutResult) const override;
 };
 
@@ -162,9 +162,9 @@ public:
 	FExpression* StructExpression;
 
 	virtual void ComputeAnalyticDerivatives(FTree& Tree, FExpressionDerivatives& OutResult) const override;
-	virtual bool PrepareValue(FEmitContext& Context, const FRequestedType& RequestedType, FPrepareValueResult& OutResult) const override;
+	virtual bool PrepareValue(FEmitContext& Context, FEmitScope& Scope, const FRequestedType& RequestedType, FPrepareValueResult& OutResult) const override;
 	virtual void EmitValueShader(FEmitContext& Context, FEmitScope& Scope, const FRequestedType& RequestedType, FEmitValueShaderResult& OutResult) const override;
-	virtual void EmitValuePreshader(FEmitContext& Context, const FRequestedType& RequestedType, Shader::FPreshaderData& OutPreshader) const override;
+	virtual void EmitValuePreshader(FEmitContext& Context, FEmitScope& Scope, const FRequestedType& RequestedType, Shader::FPreshaderData& OutPreshader) const override;
 };
 
 class FExpressionSetStructField : public FExpression
@@ -186,9 +186,9 @@ public:
 	FExpression* FieldExpression;
 
 	virtual void ComputeAnalyticDerivatives(FTree& Tree, FExpressionDerivatives& OutResult) const override;
-	virtual bool PrepareValue(FEmitContext& Context, const FRequestedType& RequestedType, FPrepareValueResult& OutResult) const override;
+	virtual bool PrepareValue(FEmitContext& Context, FEmitScope& Scope, const FRequestedType& RequestedType, FPrepareValueResult& OutResult) const override;
 	virtual void EmitValueShader(FEmitContext& Context, FEmitScope& Scope, const FRequestedType& RequestedType, FEmitValueShaderResult& OutResult) const override;
-	virtual void EmitValuePreshader(FEmitContext& Context, const FRequestedType& RequestedType, Shader::FPreshaderData& OutPreshader) const override;
+	virtual void EmitValuePreshader(FEmitContext& Context, FEmitScope& Scope, const FRequestedType& RequestedType, Shader::FPreshaderData& OutPreshader) const override;
 };
 
 class FExpressionSelect : public FExpression
@@ -204,9 +204,9 @@ public:
 	FExpression* TrueExpression;
 	FExpression* FalseExpression;
 
-	virtual bool PrepareValue(FEmitContext& Context, const FRequestedType& RequestedType, FPrepareValueResult& OutResult) const override;
+	virtual bool PrepareValue(FEmitContext& Context, FEmitScope& Scope, const FRequestedType& RequestedType, FPrepareValueResult& OutResult) const override;
 	virtual void EmitValueShader(FEmitContext& Context, FEmitScope& Scope, const FRequestedType& RequestedType, FEmitValueShaderResult& OutResult) const override;
-	virtual void EmitValuePreshader(FEmitContext& Context, const FRequestedType& RequestedType, Shader::FPreshaderData& OutPreshader) const override;
+	virtual void EmitValuePreshader(FEmitContext& Context, FEmitScope& Scope, const FRequestedType& RequestedType, Shader::FPreshaderData& OutPreshader) const override;
 };
 
 class FExpressionUnaryOp : public FExpression
@@ -221,9 +221,9 @@ public:
 	FExpression* Input;
 
 	virtual void ComputeAnalyticDerivatives(FTree& Tree, FExpressionDerivatives& OutResult) const override;
-	virtual bool PrepareValue(FEmitContext& Context, const FRequestedType& RequestedType, FPrepareValueResult& OutResult) const override;
+	virtual bool PrepareValue(FEmitContext& Context, FEmitScope& Scope, const FRequestedType& RequestedType, FPrepareValueResult& OutResult) const override;
 	virtual void EmitValueShader(FEmitContext& Context, FEmitScope& Scope, const FRequestedType& RequestedType, FEmitValueShaderResult& OutResult) const override;
-	virtual void EmitValuePreshader(FEmitContext& Context, const FRequestedType& RequestedType, Shader::FPreshaderData& OutPreshader) const override;
+	virtual void EmitValuePreshader(FEmitContext& Context, FEmitScope& Scope, const FRequestedType& RequestedType, Shader::FPreshaderData& OutPreshader) const override;
 };
 
 class FExpressionBinaryOp : public FExpression
@@ -240,9 +240,9 @@ public:
 	FExpression* Rhs;
 
 	virtual void ComputeAnalyticDerivatives(FTree& Tree, FExpressionDerivatives& OutResult) const override;
-	virtual bool PrepareValue(FEmitContext& Context, const FRequestedType& RequestedType, FPrepareValueResult& OutResult) const override;
+	virtual bool PrepareValue(FEmitContext& Context, FEmitScope& Scope, const FRequestedType& RequestedType, FPrepareValueResult& OutResult) const override;
 	virtual void EmitValueShader(FEmitContext& Context, FEmitScope& Scope, const FRequestedType& RequestedType, FEmitValueShaderResult& OutResult) const override;
-	virtual void EmitValuePreshader(FEmitContext& Context, const FRequestedType& RequestedType, Shader::FPreshaderData& OutPreshader) const override;
+	virtual void EmitValuePreshader(FEmitContext& Context, FEmitScope& Scope, const FRequestedType& RequestedType, Shader::FPreshaderData& OutPreshader) const override;
 };
 
 struct FSwizzleParameters
@@ -270,9 +270,9 @@ public:
 	FExpression* Input;
 
 	virtual void ComputeAnalyticDerivatives(FTree& Tree, FExpressionDerivatives& OutResult) const override;
-	virtual bool PrepareValue(FEmitContext& Context, const FRequestedType& RequestedType, FPrepareValueResult& OutResult) const override;
+	virtual bool PrepareValue(FEmitContext& Context, FEmitScope& Scope, const FRequestedType& RequestedType, FPrepareValueResult& OutResult) const override;
 	virtual void EmitValueShader(FEmitContext& Context, FEmitScope& Scope, const FRequestedType& RequestedType, FEmitValueShaderResult& OutResult) const override;
-	virtual void EmitValuePreshader(FEmitContext& Context, const FRequestedType& RequestedType, Shader::FPreshaderData& OutPreshader) const override;
+	virtual void EmitValuePreshader(FEmitContext& Context, FEmitScope& Scope, const FRequestedType& RequestedType, Shader::FPreshaderData& OutPreshader) const override;
 };
 
 class FExpressionAppend : public FExpression
@@ -287,15 +287,15 @@ public:
 	FExpression* Rhs;
 
 	virtual void ComputeAnalyticDerivatives(FTree& Tree, FExpressionDerivatives& OutResult) const override;
-	virtual bool PrepareValue(FEmitContext& Context, const FRequestedType& RequestedType, FPrepareValueResult& OutResult) const override;
+	virtual bool PrepareValue(FEmitContext& Context, FEmitScope& Scope, const FRequestedType& RequestedType, FPrepareValueResult& OutResult) const override;
 	virtual void EmitValueShader(FEmitContext& Context, FEmitScope& Scope, const FRequestedType& RequestedType, FEmitValueShaderResult& OutResult) const override;
-	virtual void EmitValuePreshader(FEmitContext& Context, const FRequestedType& RequestedType, Shader::FPreshaderData& OutPreshader) const override;
+	virtual void EmitValuePreshader(FEmitContext& Context, FEmitScope& Scope, const FRequestedType& RequestedType, Shader::FPreshaderData& OutPreshader) const override;
 };
 
 class FExpressionReflectionVector : public FExpression
 {
 public:
-	virtual bool PrepareValue(FEmitContext& Context, const FRequestedType& RequestedType, FPrepareValueResult& OutResult) const override;
+	virtual bool PrepareValue(FEmitContext& Context, FEmitScope& Scope, const FRequestedType& RequestedType, FPrepareValueResult& OutResult) const override;
 	virtual void EmitValueShader(FEmitContext& Context, FEmitScope& Scope, const FRequestedType& RequestedType, FEmitValueShaderResult& OutResult) const override;
 };
 
@@ -304,15 +304,16 @@ class FStatementReturn : public FStatement
 public:
 	FExpression* Expression;
 
-	virtual bool Prepare(FEmitContext& Context) const override;
+	virtual bool Prepare(FEmitContext& Context, FEmitScope& Scope) const override;
 	virtual void EmitShader(FEmitContext& Context, FEmitScope& Scope) const override;
 };
 
 class FStatementBreak : public FStatement
 {
 public:
-	virtual bool Prepare(FEmitContext& Context) const override;
+	virtual bool Prepare(FEmitContext& Context, FEmitScope& Scope) const override;
 	virtual void EmitShader(FEmitContext& Context, FEmitScope& Scope) const override;
+	virtual void EmitPreshader(FEmitContext& Context, FEmitScope& Scope, const FRequestedType& RequestedType, TArrayView<const FEmitPreshaderScope> Scopes, Shader::FPreshaderData& OutPreshader) const override;
 };
 
 class FStatementIf : public FStatement
@@ -323,18 +324,22 @@ public:
 	FScope* ElseScope;
 	FScope* NextScope;
 
-	virtual bool Prepare(FEmitContext& Context) const override;
+	virtual bool Prepare(FEmitContext& Context, FEmitScope& Scope) const override;
 	virtual void EmitShader(FEmitContext& Context, FEmitScope& Scope) const override;
+	virtual void EmitPreshader(FEmitContext& Context, FEmitScope& Scope, const FRequestedType& RequestedType, TArrayView<const FEmitPreshaderScope> Scopes, Shader::FPreshaderData& OutPreshader) const override;
 };
 
 class FStatementLoop : public FStatement
 {
 public:
+	FStatement* BreakStatement;
 	FScope* LoopScope;
 	FScope* NextScope;
 
-	virtual bool Prepare(FEmitContext& Context) const override;
+	virtual bool IsLoop() const override { return true; }
+	virtual bool Prepare(FEmitContext& Context, FEmitScope& Scope) const override;
 	virtual void EmitShader(FEmitContext& Context, FEmitScope& Scope) const override;
+	virtual void EmitPreshader(FEmitContext& Context, FEmitScope& Scope, const FRequestedType& RequestedType, TArrayView<const FEmitPreshaderScope> Scopes, Shader::FPreshaderData& OutPreshader) const override;
 };
 
 }
diff --git a/Engine/Source/Runtime/Engine/Public/HLSLTree/HLSLTreeEmit.h b/Engine/Source/Runtime/Engine/Public/HLSLTree/HLSLTreeEmit.h
--- a/Engine/Source/Runtime/Engine/Public/HLSLTree/HLSLTreeEmit.h
+++ b/Engine/Source/Runtime/Engine/Public/HLSLTree/HLSLTreeEmit.h
@@ -16,6 +16,12 @@ struct FStaticParameterSet;
 
 namespace UE
 {
+
+namespace Shader
+{
+class FPreshaderData;
+} // namespace Shader
+
 namespace HLSLTree
 {
 
@@ -30,6 +36,8 @@ class FEmitScope;
 class FEmitShaderExpression;
 class FEmitShaderStatement;
 
+struct FEmitPreshaderScope;
+
 struct FEmitShaderScopeEntry
 {
 	FEmitShaderScopeEntry() = default;
@@ -170,13 +178,27 @@ class FEmitScope
 {
 public:
 	static FEmitScope* FindSharedParent(FEmitScope* Lhs, FEmitScope* Rhs);
+	bool HasParent(const FEmitScope* InParent) const;
+	bool IsLoop() const;
+	FEmitScope* FindLoop();
 
 	void EmitShaderCode(FEmitShaderScopeStack& Stack);
 
 	FEmitScope* ParentScope = nullptr;
-	FEmitShaderNode* FirstNode;
+	FStatement* OwnerStatement = nullptr;
+	FStatement* ContainedStatement = nullptr;
+	FEmitShaderNode* FirstNode = nullptr;
 	int32 NestedLevel = 0;
 	EEmitScopeState State = EEmitScopeState::Uninitialized;
+	EExpressionEvaluation Evaluation = EExpressionEvaluation::None;
+};
+
+struct FPreshaderLocalPHIScope
+{
+	FPreshaderLocalPHIScope(const FExpression* InExpression, int32 InValueStackPosition) : ExpressionLocalPHI(InExpression), ValueStackPosition(InValueStackPosition) {}
+
+	const FExpression* ExpressionLocalPHI;
+	int32 ValueStackPosition;
 };
 
 /** Tracks shared state while emitting HLSL code */
@@ -186,15 +208,20 @@ public:
 	explicit FEmitContext(FMemStackBase& InAllocator, FErrorHandlerInterface& InErrors, const Shader::FStructTypeRegistry& InTypeRegistry);
 	~FEmitContext();
 
-	const FPreparedType& PrepareExpression(FExpression* InExpression, const FRequestedType& RequestedType);
+	FPreparedType PrepareExpression(FExpression* InExpression, FEmitScope& Scope, const FRequestedType& RequestedType);
+
+	FEmitScope* InternalPrepareScope(FScope* Scope, FScope* ParentScope);
+	FEmitScope* PrepareScope(FScope* Scope);
+	FEmitScope* PrepareScopeWithParent(FScope* Scope, FScope* ParentScope);
+	void MarkScopeEvaluation(FEmitScope& EmitParentScope, FScope* Scope, EExpressionEvaluation Evaluation);
+	void MarkScopeDead(FEmitScope& EmitParentScope, FScope* Scope);
 
-	FEmitScope* InternalPrepareScope(FScope* Scope, FScope* ParentScope, bool bMarkDead);
-	bool PrepareScope(FScope* Scope);
-	bool PrepareScopeWithParent(FScope* Scope, FScope* ParentScope);
-	bool MarkScopeDead(FScope* Scope);
+	void EmitPreshaderScope(const FScope* Scope, const FRequestedType& RequestedType, TArrayView<const FEmitPreshaderScope> PreshaderScopes, Shader::FPreshaderData& OutPreshader);
+	void EmitPreshaderScope(FEmitScope& EmitScope, const FRequestedType& RequestedType, TArrayView<const FEmitPreshaderScope> PreshaderScopes, Shader::FPreshaderData& OutPreshader);
 
 	FEmitScope* AcquireEmitScopeWithParent(const FScope* Scope, FEmitScope* EmitParentScope);
 	FEmitScope* AcquireEmitScope(const FScope* Scope);
+	FEmitScope* FindEmitScope(const FScope* Scope);
 	FEmitScope* InternalEmitScope(const FScope* Scope);
 
 	template<typename T>
@@ -346,10 +373,14 @@ public:
 
 	TArray<FEmitShaderNode*> EmitNodes;
 	TMap<const FScope*, FEmitScope*> EmitScopeMap;
+	TMap<const FExpression*, FEmitScope*> PrepareLocalPHIMap;
 	TMap<const FExpression*, FEmitShaderExpression*> EmitLocalPHIMap;
 	TMap<FXxHash64, FEmitShaderExpression*> EmitExpressionMap;
 	TMap<FXxHash64, FEmitShaderExpression*> EmitPreshaderMap;
 	TMap<const FFunction*, FEmitShaderNode*> EmitFunctionMap;
+	TArray<struct FPreshaderLoopScope*> PreshaderLoopScopes;
+	TArray<const FPreshaderLocalPHIScope*> PreshaderLocalPHIScopes;
+	int32 PreshaderStackPosition = 0;
 
 	// TODO - remove preshader material dependency
 	const FMaterial* Material = nullptr;
diff --git a/Engine/Source/Runtime/Engine/Public/HLSLTree/HLSLTreeTypes.h b/Engine/Source/Runtime/Engine/Public/HLSLTree/HLSLTreeTypes.h
--- a/Engine/Source/Runtime/Engine/Public/HLSLTree/HLSLTreeTypes.h
+++ b/Engine/Source/Runtime/Engine/Public/HLSLTree/HLSLTreeTypes.h
@@ -33,17 +33,31 @@ enum class EExpressionEvaluation : uint8
 	/** Invalid/uninitialized */
 	None,
 
+	/** Valid, but not yet known */
+	Unknown,
+
 	/** The expression outputs HLSL code (via FExpressionEmitResult::Writer) */
 	Shader,
 
+	PreshaderLoop,
+
 	/** The expression outputs preshader code evaluated at runtime (via FExpressionEmitResult::Preshader) */
 	Preshader,
 
+	ConstantLoop,
+
 	/** The expression outputs constant preshader code evaluated at compile time (via FExpressionEmitResult::Preshader) */
 	Constant,
 };
 
 EExpressionEvaluation CombineEvaluations(EExpressionEvaluation Lhs, EExpressionEvaluation Rhs);
+EExpressionEvaluation MakeLoopEvaluation(EExpressionEvaluation Evaluation);
+EExpressionEvaluation MakeNonLoopEvaluation(EExpressionEvaluation Evaluation);
+
+inline bool IsLoopEvaluation(EExpressionEvaluation Evaluation)
+{
+	return Evaluation == EExpressionEvaluation::PreshaderLoop || Evaluation == EExpressionEvaluation::ConstantLoop;
+}
 
 enum class EUnaryOp : uint8
 {
diff --git a/Engine/Source/Runtime/Engine/Public/Shader/Preshader.h b/Engine/Source/Runtime/Engine/Public/Shader/Preshader.h
--- a/Engine/Source/Runtime/Engine/Public/Shader/Preshader.h
+++ b/Engine/Source/Runtime/Engine/Public/Shader/Preshader.h
@@ -64,6 +64,11 @@ enum class EPreshaderOpcode : uint8
 	GetField,
 	SetField,
 	Neg,
+	Jump,
+	JumpIfFalse,
+	PushValue,
+	Less,
+	Assign,
 };
 
 struct FPreshaderStructType
@@ -74,6 +79,13 @@ struct FPreshaderStructType
 	LAYOUT_FIELD(int32, NumComponents);
 };
 
+struct FPreshaderLabel
+{
+	explicit FPreshaderLabel(int32 InOffset = INDEX_NONE) : Offset(InOffset) {}
+
+	int32 Offset;
+};
+
 class FPreshaderData
 {
 	DECLARE_TYPE_LAYOUT(FPreshaderData, NonVirtual);
@@ -100,6 +112,12 @@ public:
 	void WriteType(const FType& Type);
 	void WriteValue(const FValue& Value);
 
+	FPreshaderLabel WriteJump(EPreshaderOpcode Op);
+	void WriteJump(EPreshaderOpcode Op, FPreshaderLabel Label);
+
+	FPreshaderLabel GetLabel();
+	void SetLabel(FPreshaderLabel InLabel);
+
 	template<typename T>
 	FPreshaderData& Write(const T& Value) { WriteData(&Value, sizeof(T)); return *this; }
 
diff --git a/Engine/Source/Runtime/Engine/Public/Shader/PreshaderTypes.h b/Engine/Source/Runtime/Engine/Public/Shader/PreshaderTypes.h
--- a/Engine/Source/Runtime/Engine/Public/Shader/PreshaderTypes.h
+++ b/Engine/Source/Runtime/Engine/Public/Shader/PreshaderTypes.h
@@ -56,7 +56,7 @@ public:
 
 	/** Returned values are invalidated when anything is pushed onto the stack */
 	FPreshaderValue PopValue();
-	FPreshaderValue PeekValue();
+	FPreshaderValue PeekValue(int32 Offset = 0);
 
 	void Reset()
 	{
diff --git a/Engine/Source/Runtime/Engine/Public/Shader/ShaderTypes.h b/Engine/Source/Runtime/Engine/Public/Shader/ShaderTypes.h
--- a/Engine/Source/Runtime/Engine/Public/Shader/ShaderTypes.h
+++ b/Engine/Source/Runtime/Engine/Public/Shader/ShaderTypes.h
@@ -425,6 +425,7 @@ ENGINE_API FValue Add(const FValue& Lhs, const FValue& Rhs);
 ENGINE_API FValue Sub(const FValue& Lhs, const FValue& Rhs);
 ENGINE_API FValue Mul(const FValue& Lhs, const FValue& Rhs);
 ENGINE_API FValue Div(const FValue& Lhs, const FValue& Rhs);
+ENGINE_API FValue Less(const FValue& Lhs, const FValue& Rhs);
 ENGINE_API FValue Min(const FValue& Lhs, const FValue& Rhs);
 ENGINE_API FValue Max(const FValue& Lhs, const FValue& Rhs);
 ENGINE_API FValue Clamp(const FValue& Value, const FValue& Low, const FValue& High);
