PR #8820: Add an API to retrieve the changelist a replay was recorded on while playing back the replay. (Contributed by [REDACTED_PERSON])

diff --git a/Engine/Source/Runtime/Engine/Classes/Engine/DemoNetDriver.h b/Engine/Source/Runtime/Engine/Classes/Engine/DemoNetDriver.h
--- a/Engine/Source/Runtime/Engine/Classes/Engine/DemoNetDriver.h
+++ b/Engine/Source/Runtime/Engine/Classes/Engine/DemoNetDriver.h
@@ -607,6 +607,7 @@ public:
 		return ReplayHelper.ActiveReplayName;
 	}
 
+	uint32 GetPlaybackDemoChangelist() const { return ReplayHelper.PlaybackDemoHeader.EngineVersion.GetChangelist(); }
 	uint32 GetPlaybackDemoVersion() const { return ReplayHelper.PlaybackDemoHeader.Version; }
 
 	FString GetDemoPath() const;
