PR #8811: [5.0] Add missing EShaderCompressionFormat::Type Enum Entries (Contributed by [REDACTED_PERSON])

diff --git a/Engine/Source/Runtime/Engine/Classes/Engine/RendererSettings.h b/Engine/Source/Runtime/Engine/Classes/Engine/RendererSettings.h
--- a/Engine/Source/Runtime/Engine/Classes/Engine/RendererSettings.h
+++ b/Engine/Source/Runtime/Engine/Classes/Engine/RendererSettings.h
@@ -202,6 +202,8 @@ namespace EShaderCompressionFormat
 	{
 		None = 0 UMETA(DisplayName = "Do not compress", ToolTip = "Fastest, but disk and memory footprint will be large"),
 		LZ4 = 1 UMETA(DisplayName = "LZ4", ToolTip = "Compressing using LZ4"),
+		Oodle = 2 UMETA(DisplayName = "Oodle", ToolTip = "Compressing using Oodle (default)"),
+		Zlib = 3 UMETA(DisplayName = "ZLib", ToolTip = "Compressing using Zlib")
 	};
 }
 
