ULevel::GetLevelInfoFromAssetRegistry: remove conditions to call ScanFilesSynchronous

diff --git a/Engine/Source/Runtime/Engine/Private/Level.cpp b/Engine/Source/Runtime/Engine/Private/Level.cpp
--- a/Engine/Source/Runtime/Engine/Private/Level.cpp
+++ b/Engine/Source/Runtime/Engine/Private/Level.cpp
@@ -2037,11 +2037,10 @@ namespace LevelAssetRegistryHelper
 	static bool GetLevelInfoFromAssetRegistry(FName LevelPackage, TFunctionRef<bool(const FAssetData & Asset)> Func)
 	{
 		IAssetRegistry& AssetRegistry = FModuleManager::LoadModuleChecked<FAssetRegistryModule>(TEXT("AssetRegistry")).Get();
-		if (AssetRegistry.IsLoadingAssets() || IsRunningGame())
-		{
-			AssetRegistry.ScanFilesSynchronous({ LevelPackage.ToString()});
-		}
 
+		// Always ask for scan in case no initial scan was done (Commandlets, IsRunningGame) or if AssetRegistry isn't done scanning
+		AssetRegistry.ScanFilesSynchronous({ LevelPackage.ToString()});
+		
 		TArray<FAssetData> LevelPackageAssets;
 		AssetRegistry.GetAssetsByPackageName(LevelPackage, LevelPackageAssets, true);
 
