FLevelStreamingGCHelper: Add EnableForCommandlet() method

diff --git a/Engine/Source/Runtime/Engine/Classes/Engine/World.h b/Engine/Source/Runtime/Engine/Classes/Engine/World.h
--- a/Engine/Source/Runtime/Engine/Classes/Engine/World.h
+++ b/Engine/Source/Runtime/Engine/Classes/Engine/World.h
@@ -456,12 +456,19 @@ struct ENGINE_API FLevelStreamingGCHelper
 	 * @return	The number of levels pending a purge by the garbage collector
 	 */
 	static int32 GetNumLevelsPendingPurge();
+
+	/**
+	 * Allows FLevelStreamingGCHelper to be used in a commandlet.
+	 */
+	static void EnableForCommandlet();
 	
 private:
 	/** Static array of levels that should be unloaded */
 	static TArray<TWeakObjectPtr<ULevel> > LevelsPendingUnload;
 	/** Static array of level packages that have been marked by PrepareStreamedOutLevelsForGC */
 	static TArray<FName> LevelPackageNames;
+	/** Static bool allows FLevelStreamingGCHelper to be used in a commandlet */
+	static bool bEnabledForCommandlet;
 };
 
 /** Saved editor viewport state information */
diff --git a/Engine/Source/Runtime/Engine/Private/World.cpp b/Engine/Source/Runtime/Engine/Private/World.cpp
--- a/Engine/Source/Runtime/Engine/Private/World.cpp
+++ b/Engine/Source/Runtime/Engine/Private/World.cpp
@@ -3175,6 +3175,7 @@ FGameTime FGameTime::GetTimeSinceAppStart()
 FLevelStreamingGCHelper::FOnGCStreamedOutLevelsEvent FLevelStreamingGCHelper::OnGCStreamedOutLevels;
 TArray<TWeakObjectPtr<ULevel> > FLevelStreamingGCHelper::LevelsPendingUnload;
 TArray<FName> FLevelStreamingGCHelper::LevelPackageNames;
+bool FLevelStreamingGCHelper::bEnabledForCommandlet = false;
 
 void FLevelStreamingGCHelper::AddGarbageCollectorCallback()
 {
@@ -3188,9 +3189,15 @@ void FLevelStreamingGCHelper::AddGarbageCollectorCallback()
 	}
 }
 
+void FLevelStreamingGCHelper::EnableForCommandlet()
+{
+	check(IsRunningCommandlet());
+	bEnabledForCommandlet = true;
+}
+
 void FLevelStreamingGCHelper::RequestUnload( ULevel* InLevel )
 {
-	if (!IsRunningCommandlet())
+	if (!IsRunningCommandlet() || bEnabledForCommandlet)
 	{
 		check( InLevel );
 		check( InLevel->bIsVisible == false );
@@ -3216,7 +3223,7 @@ void FLevelStreamingGCHelper::PrepareStreamedOutLevelsForGC()
 	{
 		ULevel*	Level = LevelsPendingUnload[LevelIndex].Get();
 
-		if( Level && (!GIsEditor || Level->GetOutermost()->HasAnyPackageFlags(PKG_PlayInEditor) ))
+		if( Level && ((!GIsEditor || bEnabledForCommandlet) || Level->GetOutermost()->HasAnyPackageFlags(PKG_PlayInEditor)))
 		{
 			UPackage* LevelPackage = Level->GetOutermost();
 			UE_LOG(LogStreaming, Log, TEXT("PrepareStreamedOutLevelsForGC called on '%s'"), *LevelPackage->GetName() );
