WorldPartition: HLOD Embedded Level Instance Fix
- HLOD Generation now uses UWorldPartitionLevelStreamingDynamic::LoadInEditor to load sub actors of a AWorldPartitionHLOD

#tests AncientGame hlod gen + cook + pie

diff --git a/Engine/Plugins/Editor/WorldPartitionHLODUtilities/Source/Private/WorldPartitionHLODUtilities.cpp b/Engine/Plugins/Editor/WorldPartitionHLODUtilities/Source/Private/WorldPartitionHLODUtilities.cpp
--- a/Engine/Plugins/Editor/WorldPartitionHLODUtilities/Source/Private/WorldPartitionHLODUtilities.cpp
+++ b/Engine/Plugins/Editor/WorldPartitionHLODUtilities/Source/Private/WorldPartitionHLODUtilities.cpp
@@ -7,9 +7,11 @@
 #include "WorldPartition/WorldPartition.h"
 #include "WorldPartition/WorldPartitionHandle.h"
 #include "WorldPartition/HLOD/HLODActor.h"
+#include "WorldPartition/HLOD/HLODSubActor.h"
 #include "WorldPartition/HLOD/HLODActorDesc.h"
 #include "WorldPartition/HLOD/HLODBuilder.h"
 #include "WorldPartition/HLOD/HLODLayer.h"
+#include "WorldPartition/WorldPartitionLevelStreamingDynamic.h"
 
 #include "ISMPartition/ISMComponentDescriptor.h"
 
@@ -22,52 +24,35 @@
 #include "Materials/Material.h"
 #include "AssetCompilingManager.h"
 
-#include "LevelInstance/LevelInstanceActor.h"
-#include "LevelInstance/LevelInstanceSubsystem.h"
-
 #include "HLODBuilderInstancing.h"
 #include "HLODBuilderMeshMerge.h"
 #include "HLODBuilderMeshSimplify.h"
 #include "HLODBuilderMeshApproximate.h"
 
-static bool LoadSubActors(AWorldPartitionHLOD* InHLODActor, TArray<FWorldPartitionReference>& OutActors)
-{
-	OutActors.Reserve(InHLODActor->GetSubActors().Num());
+#include "Algo/Transform.h"
 
+static UWorldPartitionLevelStreamingDynamic* CreateLevelStreamingFromHLODActor(AWorldPartitionHLOD* InHLODActor, bool& bOutDirty)
+{
+	bOutDirty = false;
 	UWorld* World = InHLODActor->GetWorld();
 	UWorldPartition* WorldPartition = World->GetWorldPartition();
 	check(WorldPartition);
+	
+	const FName LevelStreamingName = FName(*FString::Printf(TEXT("HLODLevelStreaming_%s"), *InHLODActor->GetName()));
+	TArray<FWorldPartitionRuntimeCellObjectMapping> Mappings;
+	Mappings.Reserve(InHLODActor->GetSubActors().Num());
+	Algo::Transform(InHLODActor->GetSubActors(), Mappings, [](const FHLODSubActor& SubActor) { return FWorldPartitionRuntimeCellObjectMapping(SubActor.ActorPackage, SubActor.ActorPath, SubActor.ContainerID, SubActor.ContainerTransform, SubActor.ContainerPackage); });
 
-	bool bIsDirty = false;
+	UWorldPartitionLevelStreamingDynamic* LevelStreaming = UWorldPartitionLevelStreamingDynamic::LoadInEditor(World, LevelStreamingName, Mappings);
+	check(LevelStreaming);
 
-	// Gather (and potentially load) actors
-	for (const FGuid& SubActorGuid : InHLODActor->GetSubActors())
+	if (!LevelStreaming->GetLoadSucceeded())
 	{
-		FWorldPartitionReference ActorRef(WorldPartition, SubActorGuid);
-
-		if (ActorRef.IsValid())
-		{
-			AActor* LoadedActor = ActorRef.Get()->GetActor();
-
-			// Load level instances
-			if (ALevelInstance* LevelInstance = Cast<ALevelInstance>(LoadedActor))
-			{
-				// Wait for level instance loading
-				if (LevelInstance->SupportsLoading())
-				{
-					LevelInstance->GetLevelInstanceSubsystem()->BlockLoadLevelInstance(LevelInstance);
-				}
-			}
-
-			OutActors.Add(MoveTemp(ActorRef));
-		}
-		else
-		{
-			bIsDirty = true;
-		}
+		bOutDirty = true;
+		UE_LOG(LogHLODBuilder, Warning, TEXT("HLOD actor \"%s\" needs to be rebuilt as it didn't succeed in loading all actors."), *InHLODActor->GetActorLabel());
 	}
 
-	return !bIsDirty;
+	return LevelStreaming;
 }
 
 static uint32 GetCRC(const UHLODLayer* InHLODLayer)
@@ -88,12 +73,12 @@ static uint32 GetCRC(const UHLODLayer* InHLODLayer)
 	return CRC;
 }
 
-static uint32 ComputeHLODHash(AWorldPartitionHLOD* InHLODActor, const TArray<FWorldPartitionReference>& InActors)
+static uint32 ComputeHLODHash(AWorldPartitionHLOD* InHLODActor, const TArray<AActor*>& InActors)
 {
 	FArchiveCrc32 Ar;
 
 	// Base key, changing this will force a rebuild of all HLODs
-	FString HLODBaseKey = "4CE9431F9C6842F996786C980641B63A";
+	FString HLODBaseKey = "5052091956924DB3BD9ACE00B71944AC";
 	Ar << HLODBaseKey;
 
 	// HLOD Layer
@@ -163,7 +148,7 @@ TArray<AWorldPartitionHLOD*> FWorldPartitionHLODUtilities::CreateHLODActors(FHLO
 {
 	struct FSubActorsInfo
 	{
-		TArray<FGuid>			SubActors;
+		TArray<FHLODSubActor>	SubActors;
 		bool					bIsSpatiallyLoaded;
 	};
 	TMap<UHLODLayer*, FSubActorsInfo> SubActorsInfos;
@@ -178,7 +163,7 @@ TArray<AWorldPartitionHLOD*> FWorldPartitionHLODUtilities::CreateHLODActors(FHLO
 			{
 				FSubActorsInfo& SubActorsInfo = SubActorsInfos.FindOrAdd(HLODLayer);
 
-				SubActorsInfo.SubActors.Add(ActorInstance.Actor);
+				SubActorsInfo.SubActors.Emplace(ActorDescView.GetGuid(), ActorDescView.GetActorPackage(), ActorDescView.GetActorPath(), ActorInstance.ContainerInstance->ID, ActorInstance.ContainerInstance->Container->GetContainerPackage(), ActorInstance.ContainerInstance->Transform);
 				if (ActorDescView.GetIsSpatiallyLoaded())
 				{
 					SubActorsInfo.bIsSpatiallyLoaded = true;
@@ -245,8 +230,8 @@ TArray<AWorldPartitionHLOD*> FWorldPartitionHLODUtilities::CreateHLODActors(FHLO
 			bool bSubActorsChanged = HLODActor->GetSubActors().Num() != SubActorsInfo.SubActors.Num();
 			if (!bSubActorsChanged)
 			{
-				TArray<FGuid> A = HLODActor->GetSubActors();
-				TArray<FGuid> B = SubActorsInfo.SubActors;
+				TArray<FHLODSubActor> A = HLODActor->GetSubActors();
+				TArray<FHLODSubActor> B = SubActorsInfo.SubActors;
 				A.Sort();
 				B.Sort();
 				bSubActorsChanged = A != B;
@@ -367,19 +352,19 @@ UHLODBuilderSettings* FWorldPartitionHLODUtilities::CreateHLODBuilderSettings(UH
 
 uint32 FWorldPartitionHLODUtilities::BuildHLOD(AWorldPartitionHLOD* InHLODActor)
 {
-	TArray<FWorldPartitionReference> SubActors;
-	
-	bool bIsDirty = !LoadSubActors(InHLODActor, SubActors);
-	if (bIsDirty)
+	bool bIsDirty = false;
+	UWorldPartitionLevelStreamingDynamic* LevelStreaming = CreateLevelStreamingFromHLODActor(InHLODActor, bIsDirty);
+	ON_SCOPE_EXIT
 	{
-		UE_LOG(LogHLODBuilder, Warning, TEXT("HLOD actor \"%s\" needs to be rebuilt as it references actors that have been deleted."), *InHLODActor->GetActorLabel());
-	}
+		UWorldPartitionLevelStreamingDynamic::UnloadFromEditor(LevelStreaming);
+	};
 
 	uint32 OldHLODHash = bIsDirty ? 0 : InHLODActor->GetHLODHash();
-	uint32 NewHLODHash = ComputeHLODHash(InHLODActor, SubActors);
+	uint32 NewHLODHash = ComputeHLODHash(InHLODActor, LevelStreaming->GetLoadedLevel()->Actors);
 
 	if (OldHLODHash == NewHLODHash)
 	{
+		UE_LOG(LogHLODBuilder, Verbose, TEXT("HLOD actor \"%s\" doesn't need to be rebuilt."), *InHLODActor->GetActorLabel());
 		return OldHLODHash;
 	}
 
@@ -398,7 +383,7 @@ uint32 FWorldPartitionHLODUtilities::BuildHLOD(AWorldPartitionHLOD* InHLODActor)
 				FAssetCompilingManager::Get().FinishAllCompilation();
 			}
 
-			HLODBuilder->Build(InHLODActor, HLODLayer, SubActors);
+			HLODBuilder->Build(InHLODActor, HLODLayer, LevelStreaming->GetLoadedLevel()->Actors);
 			HLODBuilder->RemoveFromRoot();
 		}
 			
diff --git a/Engine/Source/Editor/UnrealEd/Private/WorldPartition/WorldPartitionHLODsBuilder.cpp b/Engine/Source/Editor/UnrealEd/Private/WorldPartition/WorldPartitionHLODsBuilder.cpp
--- a/Engine/Source/Editor/UnrealEd/Private/WorldPartition/WorldPartitionHLODsBuilder.cpp
+++ b/Engine/Source/Editor/UnrealEd/Private/WorldPartition/WorldPartitionHLODsBuilder.cpp
@@ -287,6 +287,9 @@ bool UWorldPartitionHLODsBuilder::RunInternal(UWorld* World, const FCellInfo& In
 	WorldPartition = World->GetWorldPartition();
 	check(WorldPartition);
 
+	// Allows HLOD Streaming levels to be GCed properly
+	FLevelStreamingGCHelper::EnableForCommandlet();
+
 	SourceControlHelper = new FSourceControlHelper(PackageHelper);
 
 	bool bRet = true;
@@ -767,14 +770,14 @@ bool UWorldPartitionHLODsBuilder::GetHLODActorsToBuild(TArray<FGuid>& HLODActors
 	}
 	else
 	{
-		TArray<TArray<FGuid>> HLODWorkloads = GetHLODWorldloads(1);
+		TArray<TArray<FGuid>> HLODWorkloads = GetHLODWorkloads(1);
 		HLODActorsToBuild = MoveTemp(HLODWorkloads[0]);
 	}
 
 	return bRet;
 }
 
-TArray<TArray<FGuid>> UWorldPartitionHLODsBuilder::GetHLODWorldloads(int32 NumWorkloads) const
+TArray<TArray<FGuid>> UWorldPartitionHLODsBuilder::GetHLODWorkloads(int32 NumWorkloads) const
 {
 	// Build a mapping of 1 HLOD[Level] -> N HLOD[Level - 1]
 	TMap<FGuid, TArray<FGuid>>	HLODParenting;
@@ -782,12 +785,15 @@ TArray<TArray<FGuid>> UWorldPartitionHLODsBuilder::GetHLODWorldloads(int32 NumWo
 	{
 		TArray<FGuid>& ChildHLODs = HLODParenting.Add(HLODIterator->GetGuid());
 
-		for (const FGuid& SubActorGuid : HLODIterator->GetSubActors())
+		for (const FHLODSubActorDesc& SubActor : HLODIterator->GetSubActors())
 		{
-			FWorldPartitionActorDesc* SubActorDesc = WorldPartition->GetActorDesc(SubActorGuid);
-			if (SubActorDesc && SubActorDesc->GetActorClass()->IsChildOf<AWorldPartitionHLOD>())
+			if (SubActor.ContainerID.IsMainContainer())
 			{
-				ChildHLODs.Add(SubActorGuid);
+				FWorldPartitionActorDesc* SubActorDesc = WorldPartition->GetActorDesc(SubActor.ActorGuid);
+				if (SubActorDesc && SubActorDesc->GetActorClass()->IsChildOf<AWorldPartitionHLOD>())
+				{
+					ChildHLODs.Add(SubActor.ActorGuid);
+				}
 			}
 		}
 	}
@@ -874,19 +880,22 @@ bool UWorldPartitionHLODsBuilder::ValidateWorkload(const TArray<FGuid>&Workload)
 
 		const FHLODActorDesc* HLODActorDesc = static_cast<const FHLODActorDesc*>(ActorDesc);
 
-		for (const FGuid& SubActorGuid : HLODActorDesc->GetSubActors())
+		for (const FHLODSubActorDesc& SubActor : HLODActorDesc->GetSubActors())
 		{
-			FWorldPartitionActorDesc* SubActorDesc = WorldPartition->GetActorDesc(SubActorGuid);
+			if (SubActor.ContainerID.IsMainContainer())
+			{
+				FWorldPartitionActorDesc* SubActorDesc = WorldPartition->GetActorDesc(SubActor.ActorGuid);
 
-			// Invalid sub actor guid found, this is unexpected when running distributed builds as the build step is always preceeded by the setup step.
-			check(SubActorDesc || !bDistributedBuild);
+				// Invalid sub actor guid found, this is unexpected when running distributed builds as the build step is always preceeded by the setup step.
+				check(SubActorDesc || !bDistributedBuild);
 
-			if (SubActorDesc && SubActorDesc->GetActorClass()->IsChildOf<AWorldPartitionHLOD>())
-			{
-				if(!ProcessedHLOD.Contains(SubActorGuid))
+				if (SubActorDesc && SubActorDesc->GetActorClass()->IsChildOf<AWorldPartitionHLOD>())
 				{
-					UE_LOG(LogWorldPartitionHLODsBuilder, Error, TEXT("Child HLOD actor missing or out of order in HLOD workload, exiting..."));
-					return false;
+					if (!ProcessedHLOD.Contains(SubActor.ActorGuid))
+					{
+						UE_LOG(LogWorldPartitionHLODsBuilder, Error, TEXT("Child HLOD actor missing or out of order in HLOD workload, exiting..."));
+						return false;
+					}
 				}
 			}
 		}
@@ -899,7 +908,7 @@ bool UWorldPartitionHLODsBuilder::ValidateWorkload(const TArray<FGuid>&Workload)
 
 bool UWorldPartitionHLODsBuilder::GenerateBuildManifest(TMap<FString, int32>& FilesToBuilderMap) const
 {
-	TArray<TArray<FGuid>> BuildersWorkload = GetHLODWorldloads(BuilderCount);
+	TArray<TArray<FGuid>> BuildersWorkload = GetHLODWorkloads(BuilderCount);
 
 	FConfigFile ConfigFile;
 
diff --git a/Engine/Source/Editor/UnrealEd/Public/WorldPartition/WorldPartitionHLODsBuilder.h b/Engine/Source/Editor/UnrealEd/Public/WorldPartition/WorldPartitionHLODsBuilder.h
--- a/Engine/Source/Editor/UnrealEd/Public/WorldPartition/WorldPartitionHLODsBuilder.h
+++ b/Engine/Source/Editor/UnrealEd/Public/WorldPartition/WorldPartitionHLODsBuilder.h
@@ -95,7 +95,7 @@ protected:
 	bool GenerateBuildManifest(TMap<FString, int32>& FilesToBuilderMap) const;
 	bool GetHLODActorsToBuild(TArray<FGuid>& HLODActorsToBuild) const;
 
-	TArray<TArray<FGuid>> GetHLODWorldloads(int32 NumWorkloads) const;
+	TArray<TArray<FGuid>> GetHLODWorkloads(int32 NumWorkloads) const;
 	bool ValidateWorkload(const TArray<FGuid>& Workload) const;
 
 	bool CopyFilesToWorkingDir(const FString& TargetDir, const FHLODModifiedFiles& ModifiedFiles, TArray<FString>& BuildProducts);
diff --git a/Engine/Source/Runtime/Core/Public/UObject/UE5ReleaseStreamObjectVersion.h b/Engine/Source/Runtime/Core/Public/UObject/UE5ReleaseStreamObjectVersion.h
--- a/Engine/Source/Runtime/Core/Public/UObject/UE5ReleaseStreamObjectVersion.h
+++ b/Engine/Source/Runtime/Core/Public/UObject/UE5ReleaseStreamObjectVersion.h
@@ -90,6 +90,9 @@ struct CORE_API FUE5ReleaseStreamObjectVersion
 		// Convolution bloom now take into account FPostProcessSettings::BloomIntensity for scatter dispersion.
 		ConvolutionBloomIntensity,
 
+		// Serialize FHLODSubActors instead of FGuids in WorldPartition HLODActorDesc
+		WorldPartitionHLODActorDescSerializeHLODSubActors,
+
 		// -----<new versions can be added above this line>-------------------------------------------------
 		VersionPlusOne,
 		LatestVersion = VersionPlusOne - 1
diff --git a/Engine/Source/Runtime/Engine/Private/WorldPartition/HLOD/HLODActor.cpp b/Engine/Source/Runtime/Engine/Private/WorldPartition/HLOD/HLODActor.cpp
--- a/Engine/Source/Runtime/Engine/Private/WorldPartition/HLOD/HLODActor.cpp
+++ b/Engine/Source/Runtime/Engine/Private/WorldPartition/HLOD/HLODActor.cpp
@@ -18,6 +18,16 @@
 #include "Engine/TextureStreamingTypes.h"
 #endif
 
+#if WITH_EDITORONLY_DATA
+FArchive& operator<<(FArchive& Ar, FHLODSubActorDesc& SubActor)
+{
+	Ar << SubActor.ActorGuid;
+	Ar << SubActor.ContainerID.ID;
+
+	return Ar;
+}
+#endif
+
 AWorldPartitionHLOD::AWorldPartitionHLOD(const FObjectInitializer& ObjectInitializer)
 	: Super(ObjectInitializer)
 	, bRequireWarmup(false)
@@ -130,14 +140,14 @@ void AWorldPartitionHLOD::SetHLODPrimitives(const TArray<UPrimitiveComponent*>&
 	}
 }
 
-void AWorldPartitionHLOD::SetSubActors(const TArray<FGuid>& InSubActors)
+void AWorldPartitionHLOD::SetSubActors(const TArray<FHLODSubActor>& InHLODSubActors)
 {
-	SubActors = InSubActors;
+	HLODSubActors = InHLODSubActors;
 }
 
-const TArray<FGuid>& AWorldPartitionHLOD::GetSubActors() const
+const TArray<FHLODSubActor>& AWorldPartitionHLOD::GetSubActors() const
 {
-	return SubActors;
+	return HLODSubActors;
 }
 
 void AWorldPartitionHLOD::SetSubActorsHLODLayer(const UHLODLayer* InSubActorsHLODLayer)
diff --git a/Engine/Source/Runtime/Engine/Private/WorldPartition/HLOD/HLODActorDesc.cpp b/Engine/Source/Runtime/Engine/Private/WorldPartition/HLOD/HLODActorDesc.cpp
--- a/Engine/Source/Runtime/Engine/Private/WorldPartition/HLOD/HLODActorDesc.cpp
+++ b/Engine/Source/Runtime/Engine/Private/WorldPartition/HLOD/HLODActorDesc.cpp
@@ -7,6 +7,7 @@
 #include "Algo/RemoveIf.h"
 #include "Hash/CityHashHelpers.h"
 #include "UObject/UE5MainStreamObjectVersion.h"
+#include "UObject/UE5ReleaseStreamObjectVersion.h"
 #include "WorldPartition/HLOD/HLODActor.h"
 #include "WorldPartition/HLOD/HLODLayer.h"
 
@@ -16,7 +17,8 @@ void FHLODActorDesc::Init(const AActor* InActor)
 
 	const AWorldPartitionHLOD* HLODActor = CastChecked<AWorldPartitionHLOD>(InActor);
 
-	SubActors = HLODActor->GetSubActors();
+	HLODSubActors.Reserve(HLODActor->GetSubActors().Num());
+	Algo::Transform(HLODActor->GetSubActors(), HLODSubActors, [](const FHLODSubActor& SubActor) { return FHLODSubActorDesc(SubActor.ActorGuid, SubActor.ContainerID); });
 	
 	CellHash = 0;
 	if (const UHLODLayer* SubActorsHLODLayer = HLODActor->GetSubActorsHLODLayer())
@@ -35,11 +37,20 @@ void FHLODActorDesc::Init(const AActor* InActor)
 void FHLODActorDesc::Serialize(FArchive& Ar)
 {
 	Ar.UsingCustomVersion(FUE5MainStreamObjectVersion::GUID);
+	Ar.UsingCustomVersion(FUE5ReleaseStreamObjectVersion::GUID);
 
 	FWorldPartitionActorDesc::Serialize(Ar);
 
-	Ar << SubActors;
-
+	if (Ar.CustomVer(FUE5ReleaseStreamObjectVersion::GUID) < FUE5ReleaseStreamObjectVersion::WorldPartitionHLODActorDescSerializeHLODSubActors)
+	{
+		TArray<FGuid> SubActors;
+		Ar << SubActors;
+	}
+	else
+	{
+		Ar << HLODSubActors;
+	}
+	
 	if (Ar.CustomVer(FUE5MainStreamObjectVersion::GUID) < FUE5MainStreamObjectVersion::WorldPartitionHLODActorDescSerializeHLODLayer)
 	{
 		FString HLODLayer_Deprecated;
@@ -58,10 +69,10 @@ bool FHLODActorDesc::Equals(const FWorldPartitionActorDesc* Other) const
 	{
 		const FHLODActorDesc* HLODActorDesc = (FHLODActorDesc*)Other;
 
-		if (CellHash == HLODActorDesc->CellHash && SubActors.Num() == HLODActorDesc->SubActors.Num())
+		if (CellHash == HLODActorDesc->CellHash && HLODSubActors.Num() == HLODActorDesc->HLODSubActors.Num())
 		{
-			TArray<FGuid> SortedSubActors(SubActors);
-			TArray<FGuid> SortedSubActorsOther(HLODActorDesc->SubActors);
+			TArray<FHLODSubActorDesc> SortedSubActors(HLODSubActors);
+			TArray<FHLODSubActorDesc> SortedSubActorsOther(HLODActorDesc->HLODSubActors);
 			SortedSubActors.Sort();
 			SortedSubActorsOther.Sort();
 			return SortedSubActors == SortedSubActorsOther;
diff --git a/Engine/Source/Runtime/Engine/Private/WorldPartition/HLOD/HLODBuilder.cpp b/Engine/Source/Runtime/Engine/Private/WorldPartition/HLOD/HLODBuilder.cpp
--- a/Engine/Source/Runtime/Engine/Private/WorldPartition/HLOD/HLODBuilder.cpp
+++ b/Engine/Source/Runtime/Engine/Private/WorldPartition/HLOD/HLODBuilder.cpp
@@ -46,7 +46,7 @@ TArray<UPrimitiveComponent*> UHLODBuilder::CreateComponents(AWorldPartitionHLOD*
 	return TArray<UPrimitiveComponent*>();
 }
 
-void UHLODBuilder::Build(AWorldPartitionHLOD* InHLODActor, const UHLODLayer* InHLODLayer, const TArray<FWorldPartitionReference>& InSubActors)
+void UHLODBuilder::Build(AWorldPartitionHLOD* InHLODActor, const UHLODLayer* InHLODLayer, const TArray<AActor*>& InSubActors)
 {
 	TArray<UPrimitiveComponent*> SubComponents = GatherPrimitiveComponents(InSubActors);
 	if (SubComponents.IsEmpty())
@@ -104,7 +104,7 @@ void UHLODBuilder::Build(AWorldPartitionHLOD* InHLODActor, const UHLODLayer* InH
 	}
 }
 
-TArray<UPrimitiveComponent*> UHLODBuilder::GatherPrimitiveComponents(const TArray<FWorldPartitionReference>& InActors)
+TArray<UPrimitiveComponent*> UHLODBuilder::GatherPrimitiveComponents(const TArray<AActor*>& InActors)
 {
 	TArray<UPrimitiveComponent*> PrimitiveComponents;
 
@@ -139,9 +139,8 @@ TArray<UPrimitiveComponent*> UHLODBuilder::GatherPrimitiveComponents(const TArra
 
 	TSet<AActor*> UnderlyingActors;
 
-	for (const FWorldPartitionReference& ActorRef : InActors)
+	for (AActor* Actor : InActors)
 	{
-		AActor* Actor = ActorRef->GetActor();
 		if (!Actor->IsHLODRelevant())
 		{
 			continue;
@@ -149,19 +148,6 @@ TArray<UPrimitiveComponent*> UHLODBuilder::GatherPrimitiveComponents(const TArra
 
 		// Gather primitives from the Actor
 		GatherPrimitivesFromActor(Actor);
-
-		// Retrieve all underlying actors (ex: all sub actors of a LevelInstance)
-		UnderlyingActors.Reset();
-		Actor->EditorGetUnderlyingActors(UnderlyingActors);
-
-		// Gather primitives from underlying actors
-		for (const AActor* UnderlyingActor : UnderlyingActors)
-		{
-			if (UnderlyingActor->IsHLODRelevant())
-			{
-				GatherPrimitivesFromActor(UnderlyingActor, Actor);
-			}
-		}
 	}
 
 	return PrimitiveComponents;
diff --git a/Engine/Source/Runtime/Engine/Private/WorldPartition/RuntimeSpatialHash/RuntimeSpatialHashHLOD.cpp b/Engine/Source/Runtime/Engine/Private/WorldPartition/RuntimeSpatialHash/RuntimeSpatialHashHLOD.cpp
--- a/Engine/Source/Runtime/Engine/Private/WorldPartition/RuntimeSpatialHash/RuntimeSpatialHashHLOD.cpp
+++ b/Engine/Source/Runtime/Engine/Private/WorldPartition/RuntimeSpatialHash/RuntimeSpatialHashHLOD.cpp
@@ -113,8 +113,9 @@ static void GameTick(UWorld* InWorld)
 	}
 }
 
-static TArray<FGuid> GenerateHLODsForGrid(UWorldPartition* WorldPartition, const FActorContainerInstance& MainContainerInstance, const FSpatialHashRuntimeGrid& RuntimeGrid, uint32 HLODLevel, FHLODCreationContext& Context, ISourceControlHelper* SourceControlHelper, bool bCreateActorsOnly, const TArray<const FActorClusterInstance*>& ClusterInstances)
+static TArray<FGuid> GenerateHLODsForGrid(UWorldPartition* WorldPartition, const FActorClusterContext& ActorClusterContext, const FSpatialHashRuntimeGrid& RuntimeGrid, uint32 HLODLevel, FHLODCreationContext& Context, ISourceControlHelper* SourceControlHelper, bool bCreateActorsOnly, const TArray<const FActorClusterInstance*>& ClusterInstances)
 {
+	const FActorContainerInstance& MainContainerInstance = *ActorClusterContext.GetClusterInstance(WorldPartition);
 	const FBox WorldBounds = MainContainerInstance.Bounds;
 
 	const FSquare2DGridHelper PartitionedActors = GetPartitionedActors(WorldPartition, WorldBounds, RuntimeGrid, ClusterInstances);
@@ -225,9 +226,10 @@ static TArray<FGuid> GenerateHLODsForGrid(UWorldPartition* WorldPartition, const
 						check(WorldPartition->GetActorDesc(CellHLODActor->GetActorGuid()));
 
 						UE_LOG(LogWorldPartitionRuntimeSpatialHashHLOD, Verbose, TEXT("Created HLOD actor %s - %s, for cell %s, represented actors:"), *CellHLODActor->GetName(), *CellHLODActor->GetActorGuid().ToString(), *CellName.ToString());
-						for (const auto& SubActorGuid : CellHLODActor->GetSubActors())
+						for (const auto& SubActor : CellHLODActor->GetSubActors())
 						{
-							UE_LOG(LogWorldPartitionRuntimeSpatialHashHLOD, Verbose, TEXT("\t\t%s - %s"), *WorldPartition->GetActorDesc(SubActorGuid)->GetActorLabel().ToString(), *SubActorGuid.ToString());
+							const FActorContainerInstance& ContainerInstance = *ActorClusterContext.GetClusterInstance(SubActor.ContainerID);
+							UE_LOG(LogWorldPartitionRuntimeSpatialHashHLOD, Verbose, TEXT("\t\t%s - %s - %s"), *SubActor.ContainerID.ToString(), *ContainerInstance.GetActorDescView(SubActor.ActorGuid).GetActorLabel().ToString(), *SubActor.ActorGuid.ToString());
 						}
 					}
 
@@ -467,10 +469,10 @@ bool UWorldPartitionRuntimeSpatialHash::GenerateHLOD(ISourceControlHelper* Sourc
 	// Keep track of all valid HLOD actors, along with which runtime grid they live in
 	TMap<FName, TArray<FGuid>> GridsHLODActors;
 
-	auto GenerateHLODs = [&GridsHLODActors, &MainContainerInstance, WorldPartition, &Context, SourceControlHelper, bCreateActorsOnly](const FSpatialHashRuntimeGrid& RuntimeGrid, uint32 HLODLevel, const TArray<const FActorClusterInstance*>& ActorClusterInstances)
+	auto GenerateHLODs = [&GridsHLODActors, &MainContainerInstance, &ActorClusterContext, WorldPartition, &Context, SourceControlHelper, bCreateActorsOnly](const FSpatialHashRuntimeGrid& RuntimeGrid, uint32 HLODLevel, const TArray<const FActorClusterInstance*>& ActorClusterInstances)
 	{
 		// Generate HLODs for this grid
-		TArray<FGuid> HLODActors = GenerateHLODsForGrid(WorldPartition, MainContainerInstance, RuntimeGrid, HLODLevel, Context, SourceControlHelper, bCreateActorsOnly, ActorClusterInstances);
+		TArray<FGuid> HLODActors = GenerateHLODsForGrid(WorldPartition, ActorClusterContext, RuntimeGrid, HLODLevel, Context, SourceControlHelper, bCreateActorsOnly, ActorClusterInstances);
 
 		for (const FGuid& HLODActorGuid : HLODActors)
 		{
diff --git a/Engine/Source/Runtime/Engine/Private/WorldPartition/WorldPartitionActorCluster.cpp b/Engine/Source/Runtime/Engine/Private/WorldPartition/WorldPartitionActorCluster.cpp
--- a/Engine/Source/Runtime/Engine/Private/WorldPartition/WorldPartitionActorCluster.cpp
+++ b/Engine/Source/Runtime/Engine/Private/WorldPartition/WorldPartitionActorCluster.cpp
@@ -157,6 +157,19 @@ void CreateActorCluster(const FWorldPartitionActorDescView& ActorDescView, TMap<
 	}
 }
 
+const FActorContainerInstance* FActorClusterContext::GetClusterInstance(const FActorContainerID& InContainerID) const
+{
+	for (const FActorContainerInstance& ContainerInstance : ContainerInstances)
+	{
+		if (ContainerInstance.ID == InContainerID)
+		{
+			return &ContainerInstance;
+		}
+	}
+
+	return nullptr;
+}
+
 FActorContainerInstance* FActorClusterContext::GetClusterInstance(const UActorDescContainer* InContainer)
 {
 	for (FActorContainerInstance& ContainerInstance : ContainerInstances)
diff --git a/Engine/Source/Runtime/Engine/Private/WorldPartition/WorldPartitionLevelHelper.cpp b/Engine/Source/Runtime/Engine/Private/WorldPartition/WorldPartitionLevelHelper.cpp
--- a/Engine/Source/Runtime/Engine/Private/WorldPartition/WorldPartitionLevelHelper.cpp
+++ b/Engine/Source/Runtime/Engine/Private/WorldPartition/WorldPartitionLevelHelper.cpp
@@ -217,7 +217,7 @@ ULevel* FWorldPartitionLevelHelper::CreateEmptyLevelForRuntimeCell(const UWorld*
 	return NewLevel;
 }
 
-bool FWorldPartitionLevelHelper::LoadActors(ULevel* InDestLevel, TArrayView<FWorldPartitionRuntimeCellObjectMapping> InActorPackages, FWorldPartitionPackageCache& InPackageCache, TFunction<void(bool)> InCompletionCallback, bool bLoadForPlay, FLinkerInstancingContext* InOutInstancingContext /*=nullptr*/)
+bool FWorldPartitionLevelHelper::LoadActors(ULevel* InDestLevel, TArrayView<FWorldPartitionRuntimeCellObjectMapping> InActorPackages, FWorldPartitionPackageCache& InPackageCache, TFunction<void(bool)> InCompletionCallback, bool bInLoadAsync, FLinkerInstancingContext* InOutInstancingContext /*=nullptr*/)
 {
 	UPackage* DestPackage = InDestLevel ? InDestLevel->GetPackage() : nullptr;
 	FString ShortLevelPackageName = DestPackage? FPackageName::GetShortName(DestPackage->GetFName()) : FString();
@@ -237,13 +237,13 @@ bool FWorldPartitionLevelHelper::LoadActors(ULevel* InDestLevel, TArrayView<FWor
 
 	// Levels to Load with actors to duplicate
 	TMap<FActorContainerID, TArray<FWorldPartitionRuntimeCellObjectMapping*>> PackagesToDuplicate;
-	check(!bLoadForPlay || InOutInstancingContext);
+	check(!InDestLevel || InOutInstancingContext);
 
 	for (FWorldPartitionRuntimeCellObjectMapping& PackageObjectMapping : InActorPackages)
 	{
 		if (PackageObjectMapping.ContainerID.IsMainContainer())
 		{
-			if (bLoadForPlay)
+			if (InOutInstancingContext)
 			{
 				FString ObjectPath = PackageObjectMapping.Package.ToString();
 				FString ActorPackageName = FPackageName::ObjectPathToPackageName(ObjectPath);
@@ -269,7 +269,7 @@ bool FWorldPartitionLevelHelper::LoadActors(ULevel* InDestLevel, TArrayView<FWor
 		const FString ActorPackageName = FPackageName::ObjectPathToPackageName(PackageObjectMapping.Package.ToString());
 		FName ActorName = *FPaths::GetExtension(PackageObjectMapping.Path.ToString());
 
-		FLoadPackageAsyncDelegate CompletionCallback = FLoadPackageAsyncDelegate::CreateLambda([LoadProgress, ActorName, bLoadForPlay, InDestLevel, InCompletionCallback](const FName& LoadedPackageName, UPackage* LoadedPackage, EAsyncLoadingResult::Type Result)
+		FLoadPackageAsyncDelegate CompletionCallback = FLoadPackageAsyncDelegate::CreateLambda([LoadProgress, ActorName, InDestLevel, InCompletionCallback](const FName& LoadedPackageName, UPackage* LoadedPackage, EAsyncLoadingResult::Type Result)
 		{
 			check(LoadProgress->NumPendingLoadRequests);
 			LoadProgress->NumPendingLoadRequests--;
@@ -278,7 +278,7 @@ bool FWorldPartitionLevelHelper::LoadActors(ULevel* InDestLevel, TArrayView<FWor
 
 			if (Actor)
 			{
-				if (bLoadForPlay)
+				if (InDestLevel)
 				{
 					check(Actor->IsPackageExternal());
 					InDestLevel->Actors.Add(Actor);
@@ -300,7 +300,7 @@ bool FWorldPartitionLevelHelper::LoadActors(ULevel* InDestLevel, TArrayView<FWor
 			}
 		});
 
-		if (bLoadForPlay)
+		if (bInLoadAsync)
 		{
 			FPackagePath PackagePath = FPackagePath::FromPackageNameChecked(*ActorPackageName);
 			check(DestPackage);
@@ -309,7 +309,13 @@ bool FWorldPartitionLevelHelper::LoadActors(ULevel* InDestLevel, TArrayView<FWor
 		}
 		else
 		{
-			UPackage* Package = LoadPackage(nullptr, *ActorPackageName, LOAD_None, nullptr, InOutInstancingContext);
+			UPackage* InstancingPackage = nullptr;
+			if(InOutInstancingContext)
+			{ 
+				InstancingPackage = CreatePackage(*ActorPackageInstanceNames[ChildIndex]);
+			}
+
+			UPackage* Package = LoadPackage(InstancingPackage, *ActorPackageName, LOAD_None, nullptr, InOutInstancingContext);
 			CompletionCallback.Execute(*ActorPackageName, Package, Package ? EAsyncLoadingResult::Succeeded : EAsyncLoadingResult::Failed);
 		}
 	}
@@ -392,6 +398,7 @@ bool FWorldPartitionLevelHelper::LoadActors(ULevel* InDestLevel, TArrayView<FWor
 					else
 					{
 						Mapping->LoadedPath = NAME_None;
+						LoadProgress->NumFailedLoadedRequests++;
 						UE_LOG(LogEngine, Warning, TEXT("Failed to find actor %s in package %s"), *ActorName, *LoadedPackageName.ToString());
 					}
 				}
@@ -412,7 +419,7 @@ bool FWorldPartitionLevelHelper::LoadActors(ULevel* InDestLevel, TArrayView<FWor
 		FName PackageDuplicateSeedName(*FString::Printf(TEXT("/Temp%s_DUP"), *PackageToLoadFrom));
 		
 		// Use package cache here because level instances can load the same levels across cells we want to avoid Loading the same package multiple times...
-		InPackageCache.LoadPackage(PackageDuplicateSeedName, *PackageToLoadFrom, CompletionCallback, /*bLoadAsync*/bLoadForPlay, /*bInWorldPackage*/true);
+		InPackageCache.LoadPackage(PackageDuplicateSeedName, *PackageToLoadFrom, CompletionCallback, /*bLoadAsync*/bInLoadAsync, /*bInWorldPackage*/true);
 	}
 
 	return (LoadProgress->NumPendingLoadRequests == 0);
diff --git a/Engine/Source/Runtime/Engine/Private/WorldPartition/WorldPartitionLevelStreamingDynamic.cpp b/Engine/Source/Runtime/Engine/Private/WorldPartition/WorldPartitionLevelStreamingDynamic.cpp
--- a/Engine/Source/Runtime/Engine/Private/WorldPartition/WorldPartitionLevelStreamingDynamic.cpp
+++ b/Engine/Source/Runtime/Engine/Private/WorldPartition/WorldPartitionLevelStreamingDynamic.cpp
@@ -33,13 +33,57 @@ UWorldPartitionLevelStreamingDynamic::UWorldPartitionLevelStreamingDynamic(const
 #if WITH_EDITOR
 	, RuntimeLevel(nullptr)
 	, bLoadRequestInProgress(false)
+	, bLoadSucceeded(false)
 #endif
 	, bShouldBeAlwaysLoaded(false)
 {
+#if WITH_EDITOR
+	SetShouldBeVisibleInEditor(false);
+#endif
 }
 
 #if WITH_EDITOR
 
+UWorldPartitionLevelStreamingDynamic* UWorldPartitionLevelStreamingDynamic::LoadInEditor(UWorld* World, FName LevelStreamingName, const TArray<FWorldPartitionRuntimeCellObjectMapping>& InPackages)
+{
+	check(World->WorldType == EWorldType::Editor);
+	UWorldPartitionLevelStreamingDynamic* LevelStreaming = NewObject<UWorldPartitionLevelStreamingDynamic>(World, LevelStreamingName, RF_Transient);
+	
+	FString PackageName = FString::Printf(TEXT("/Memory/%s"), *LevelStreamingName.ToString());
+	TSoftObjectPtr<UWorld> WorldAsset(FSoftObjectPath(FString::Printf(TEXT("%s.%s"), *PackageName, *World->GetName())));
+	LevelStreaming->SetWorldAsset(WorldAsset);
+	
+	LevelStreaming->LevelTransform = FTransform::Identity;
+	LevelStreaming->Initialize(World, InPackages);
+	LevelStreaming->SetShouldBeVisibleInEditor(true);
+	World->AddStreamingLevel(LevelStreaming);
+	World->FlushLevelStreaming();
+	
+	return LevelStreaming;
+}
+
+void UWorldPartitionLevelStreamingDynamic::UnloadFromEditor(UWorldPartitionLevelStreamingDynamic* InLevelStreaming)
+{
+	UWorld* World = InLevelStreaming->GetWorld();
+	check(World->WorldType == EWorldType::Editor);
+
+	ULevel* Level = InLevelStreaming->GetLoadedLevel();
+	InLevelStreaming->SetShouldBeVisibleInEditor(false);
+	InLevelStreaming->SetIsRequestingUnloadAndRemoval(true);
+	World->RemoveLevel(Level);
+	World->FlushLevelStreaming();
+	
+}
+
+void UWorldPartitionLevelStreamingDynamic::Initialize(UWorld* OuterWorld, const TArray<FWorldPartitionRuntimeCellObjectMapping>& InPackages)
+{
+	ChildPackages = InPackages;
+
+	OriginalLevelPackageName = OuterWorld->GetPackage()->GetLoadedPath().GetPackageFName();
+	PackageNameToLoad = GetWorldAssetPackageFName();
+	OuterWorldPartition = OuterWorld->GetWorldPartition();
+}
+
 /**
  * Initializes a UWorldPartitionLevelStreamingDynamic.
  */
@@ -54,14 +98,11 @@ void UWorldPartitionLevelStreamingDynamic::Initialize(const UWorldPartitionRunti
 
 	bShouldBeAlwaysLoaded = InCell.IsAlwaysLoaded();
 	StreamingPriority = 0;
-	ChildPackages = InCell.GetPackages();
 	UnsavedActorsContainer = InCell.UnsavedActorsContainer;
 	ActorFolders = InCell.GetActorFolders().Array();
 
 	UWorld* OuterWorld = InCell.GetOuterUWorldPartition()->GetTypedOuter<UWorld>();
-	OriginalLevelPackageName = OuterWorld->GetPackage()->GetLoadedPath().GetPackageFName();
-	PackageNameToLoad = GetWorldAssetPackageFName();
-	OuterWorldPartition = OuterWorld->GetWorldPartition();	
+	Initialize(OuterWorld, InCell.GetPackages());
 }
 
 /**
@@ -83,22 +124,22 @@ void UWorldPartitionLevelStreamingDynamic::CreateRuntimeLevel()
 {
 	check(PendingUnloadLevel == nullptr);
 	check(RuntimeLevel == nullptr);
-	const UWorld* PlayWorld = GetWorld();
-	check(PlayWorld && PlayWorld->IsGameWorld());
+	const UWorld* World = GetWorld();
+	check(World && (World->IsGameWorld() || GetShouldBeVisibleInEditor()));
 
 	// Create streaming cell Level package (make sure this package isn't discoverable until it is fully loaded)
 	const FString WorldAssetPath = GetWorldAsset().ToString();
 	const FString WorldAssetPackageName = FPackageName::ObjectPathToPackageName(WorldAssetPath) + TEXT("_AsyncLoad.") + FPackageName::ObjectPathToObjectName(WorldAssetPath);
-	RuntimeLevel = FWorldPartitionLevelHelper::CreateEmptyLevelForRuntimeCell(PlayWorld, WorldAssetPackageName);
+	RuntimeLevel = FWorldPartitionLevelHelper::CreateEmptyLevelForRuntimeCell(World, WorldAssetPackageName);
 	check(RuntimeLevel);
 
 	// Propagate ActorFolder flag to the runtime level and prepare its ActorFolders list
-	if (PlayWorld->PersistentLevel->IsUsingActorFolders() && ActorFolders.Num())
+	if (World->PersistentLevel->IsUsingActorFolders() && ActorFolders.Num())
 	{
 		FLevelActorFoldersHelper::SetUseActorFolders(RuntimeLevel, true);
 		for (const FGuid& ActorFolderGuid : ActorFolders)
 		{
-			if (UActorFolder* ActorFolder = PlayWorld->PersistentLevel->GetActorFolder(ActorFolderGuid))
+			if (UActorFolder* ActorFolder = World->PersistentLevel->GetActorFolder(ActorFolderGuid))
 			{
 				FLevelActorFoldersHelper::AddActorFolder(RuntimeLevel, ActorFolder, /*bInShouldDirtyLevel*/ false, /*bInShouldBroadcast*/ false);
 			}
@@ -179,7 +220,7 @@ bool UWorldPartitionLevelStreamingDynamic::RequestLevel(UWorld* InPersistentWorl
 		check(CellWorld->PersistentLevel != LoadedLevel);
 
 		// Level already exists but may have the wrong type due to being inactive before, so copy data over
-		check(InPersistentWorld->IsGameWorld());
+		check(InPersistentWorld->IsGameWorld() || GetShouldBeVisibleInEditor());
 		CellWorld->WorldType = InPersistentWorld->WorldType;
 		CellWorld->PersistentLevel->OwningWorld = InPersistentWorld;
 
@@ -204,7 +245,9 @@ bool UWorldPartitionLevelStreamingDynamic::RequestLevel(UWorld* InPersistentWorl
 			check(CellLevelPackage);
 			check(UWorld::FindWorldInPackage(CellLevelPackage));
 			check(RuntimeLevel->OwningWorld);
-			check(RuntimeLevel->OwningWorld->WorldType == EWorldType::PIE || ((IsRunningGame() || IsRunningDedicatedServer()) && RuntimeLevel->OwningWorld->WorldType == EWorldType::Game));
+			check(RuntimeLevel->OwningWorld->WorldType == EWorldType::PIE || 
+				((IsRunningGame() || IsRunningDedicatedServer()) && RuntimeLevel->OwningWorld->WorldType == EWorldType::Game) || 
+				(RuntimeLevel->OwningWorld->WorldType == EWorldType::Editor && GetShouldBeVisibleInEditor()));
 
 			if (IssueLoadRequests())
 			{
@@ -235,11 +278,11 @@ bool UWorldPartitionLevelStreamingDynamic::RequestLevel(UWorld* InPersistentWorl
  */
 bool UWorldPartitionLevelStreamingDynamic::IssueLoadRequests()
 {
-	check(ShouldBeLoaded());
+	check(ShouldBeLoaded() || GetShouldBeVisibleInEditor());
 	check(!HasLoadedLevel());
 	check(RuntimeLevel);
 	check(!bLoadRequestInProgress);
-
+	bLoadSucceeded = false;
 	bLoadRequestInProgress = true;
 	FLinkerInstancingContext InstancingContext;
 	UPackage* RuntimePackage = RuntimeLevel->GetPackage();
@@ -295,29 +338,27 @@ bool UWorldPartitionLevelStreamingDynamic::IssueLoadRequests()
 		ActorContainerDup->MarkAsGarbage();
 	}
 
-	auto FinalizeLoading = [this]()
+	auto FinalizeLoading = [this](bool bSucceeded)
 	{
 		check(bLoadRequestInProgress);
 		bLoadRequestInProgress = false;
+		bLoadSucceeded = bSucceeded;
+		if (!bSucceeded)
+		{
+			UE_LOG(LogLevelStreaming, Warning, TEXT("UWorldPartitionLevelStreamingDynamic::IssueLoadRequests failed %s"), *GetWorldAssetPackageName());
+		}
+
 		FinalizeRuntimeLevel();
 	};
 
 	// Load saved actors
 	if (ChildPackagesToLoad.Num())
 	{
-		FWorldPartitionLevelHelper::LoadActors(RuntimeLevel, ChildPackagesToLoad, PackageCache, [this, FinalizeLoading](bool bSucceeded)
-		{
-			if (!bSucceeded)
-			{
-				UE_LOG(LogLevelStreaming, Warning, TEXT("UWorldPartitionLevelStreamingDynamic::IssueLoadRequests failed %s"), *GetWorldAssetPackageName());
-			}
-
-			FinalizeLoading();
-		},/*bLoadForPlay=*/true, &InstancingContext);
+		FWorldPartitionLevelHelper::LoadActors(RuntimeLevel, ChildPackagesToLoad, PackageCache, FinalizeLoading, GetWorld()->IsGameWorld(), &InstancingContext);
 	}
 	else
 	{
-		FinalizeLoading();
+		FinalizeLoading(true);
 	}
 
 	return bLoadRequestInProgress;
@@ -343,7 +384,7 @@ void UWorldPartitionLevelStreamingDynamic::FinalizeRuntimeLevel()
 	const FString WorldAssetPackageName = FPackageName::ObjectPathToPackageName(WorldAssetPath);
 	RuntimePackage->Rename(*WorldAssetPackageName, nullptr, REN_ForceNoResetLoaders | REN_DontCreateRedirectors | REN_NonTransactional | REN_DoNotDirty);
 
-	if (!StreamingCell->GetIsHLOD())
+	if (StreamingCell.IsValid() && !StreamingCell->GetIsHLOD())
 	{
 		TRACE_CPUPROFILER_EVENT_SCOPE(FixupIDs);
 
@@ -499,4 +540,4 @@ UWorld* UWorldPartitionLevelStreamingDynamic::GetOuterWorld() const
 	return OuterWorldPartition->GetTypedOuter<UWorld>();
 }
 
-#undef LOCTEXT_NAMESPACE
\ No newline at end of file
+#undef LOCTEXT_NAMESPACE
diff --git a/Engine/Source/Runtime/Engine/Private/WorldPartition/WorldPartitionRuntimeLevelStreamingCell.cpp b/Engine/Source/Runtime/Engine/Private/WorldPartition/WorldPartitionRuntimeLevelStreamingCell.cpp
--- a/Engine/Source/Runtime/Engine/Private/WorldPartition/WorldPartitionRuntimeLevelStreamingCell.cpp
+++ b/Engine/Source/Runtime/Engine/Private/WorldPartition/WorldPartitionRuntimeLevelStreamingCell.cpp
@@ -144,7 +144,8 @@ UWorldPartitionLevelStreamingDynamic* UWorldPartitionRuntimeLevelStreamingCell::
 void UWorldPartitionRuntimeLevelStreamingCell::LoadActorsForCook()
 {
 	FWorldPartitionPackageCache PackageCache;
-	verify(FWorldPartitionLevelHelper::LoadActors(nullptr, Packages, PackageCache, [](bool){}, /*bLoadForPlay=*/false));
+	const bool bLoadAsync = false;
+	verify(FWorldPartitionLevelHelper::LoadActors(nullptr, Packages, PackageCache, [](bool){}, bLoadAsync));
 }
 
 void UWorldPartitionRuntimeLevelStreamingCell::MoveAlwaysLoadedContentToPersistentLevel()
diff --git a/Engine/Source/Runtime/Engine/Public/LevelInstance/LevelInstanceActor.h b/Engine/Source/Runtime/Engine/Public/LevelInstance/LevelInstanceActor.h
--- a/Engine/Source/Runtime/Engine/Public/LevelInstance/LevelInstanceActor.h
+++ b/Engine/Source/Runtime/Engine/Public/LevelInstance/LevelInstanceActor.h
@@ -84,7 +84,6 @@ public:
 	virtual bool IsLockLocation() const override;
 	virtual ELevelInstanceRuntimeBehavior GetDesiredRuntimeBehavior() const { return DesiredRuntimeBehavior; }
 	virtual ELevelInstanceRuntimeBehavior GetDefaultRuntimeBehavior() const { return ELevelInstanceRuntimeBehavior::Partitioned; }
-	virtual bool IsHLODRelevant() const override { return true; }
 
 	bool CanEdit(FText* OutReason = nullptr) const;
 	bool CanCommit(FText* OutReason = nullptr) const;
diff --git a/Engine/Source/Runtime/Engine/Public/WorldPartition/HLOD/HLODActor.h b/Engine/Source/Runtime/Engine/Public/WorldPartition/HLOD/HLODActor.h
--- a/Engine/Source/Runtime/Engine/Public/WorldPartition/HLOD/HLODActor.h
+++ b/Engine/Source/Runtime/Engine/Public/WorldPartition/HLOD/HLODActor.h
@@ -6,6 +6,7 @@
 #include "Containers/Set.h"
 #include "WorldPartition/WorldPartitionRuntimeCell.h"
 #include "WorldPartition/WorldPartitionActorDesc.h"
+#include "WorldPartition/HLOD/HLODSubActor.h"
 #include "HLODActor.generated.h"
 
 class UHLODLayer;
@@ -29,8 +30,8 @@ public:
 #if WITH_EDITOR
 	void SetHLODPrimitives(const TArray<UPrimitiveComponent*>& InHLODPrimitives);
 
-	void SetSubActors(const TArray<FGuid>& InSubActors);
-	const TArray<FGuid>& GetSubActors() const;
+	void SetSubActors(const TArray<FHLODSubActor>& InSubActorMappings);
+	const TArray<FHLODSubActor>& GetSubActors() const;
 
 	void SetSubActorsHLODLayer(const UHLODLayer* InSubActorsHLODLayer);
 	const UHLODLayer* GetSubActorsHLODLayer() const { return SubActorsHLODLayer; }
@@ -91,7 +92,7 @@ protected:
 private:
 #if WITH_EDITORONLY_DATA
 	UPROPERTY()
-	TArray<FGuid> SubActors;
+	TArray<FHLODSubActor> HLODSubActors;
 
 	UPROPERTY()
 	TObjectPtr<const UHLODLayer> SubActorsHLODLayer;
diff --git a/Engine/Source/Runtime/Engine/Public/WorldPartition/HLOD/HLODActorDesc.h b/Engine/Source/Runtime/Engine/Public/WorldPartition/HLOD/HLODActorDesc.h
--- a/Engine/Source/Runtime/Engine/Public/WorldPartition/HLOD/HLODActorDesc.h
+++ b/Engine/Source/Runtime/Engine/Public/WorldPartition/HLOD/HLODActorDesc.h
@@ -8,6 +8,7 @@
 
 #include "WorldPartition/WorldPartitionActorDesc.h"
 #include "WorldPartition/DataLayer/DataLayersID.h"
+#include "WorldPartition/HLOD/HLODSubActor.h"
 
 class UHLODLayer;
 
@@ -20,7 +21,7 @@ class ENGINE_API FHLODActorDesc : public FWorldPartitionActorDesc
 	friend class FHLODActorDescFactory;
 
 public:
-	inline const TArray<FGuid>& GetSubActors() const { return SubActors; }
+	inline const TArray<FHLODSubActorDesc>& GetSubActors() const { return HLODSubActors; }
 	inline uint64 GetCellHash() const { return CellHash; }
 
 	static uint64 ComputeCellHash(const FString HLODLayerName, uint64 GridIndexX, uint64 GridIndexY, uint64 GridIndexZ, FDataLayersID DataLayersID);
@@ -32,7 +33,7 @@ protected:
 	virtual bool Equals(const FWorldPartitionActorDesc* Other) const override;
 	virtual void Serialize(FArchive& Ar) override;
 
-	TArray<FGuid> SubActors;
+	TArray<FHLODSubActorDesc> HLODSubActors;
 
 	uint64 CellHash = 0;
 #endif
diff --git a/Engine/Source/Runtime/Engine/Public/WorldPartition/HLOD/HLODBuilder.h b/Engine/Source/Runtime/Engine/Public/WorldPartition/HLOD/HLODBuilder.h
--- a/Engine/Source/Runtime/Engine/Public/WorldPartition/HLOD/HLODBuilder.h
+++ b/Engine/Source/Runtime/Engine/Public/WorldPartition/HLOD/HLODBuilder.h
@@ -3,9 +3,9 @@
 #pragma once
 
 #include "CoreMinimal.h"
-#include "WorldPartition/WorldPartitionHandle.h"
 #include "HLODBuilder.generated.h"
 
+class AActor;
 class AWorldPartitionHLOD;
 class UPrimitiveComponent;
 class UHLODLayer;
@@ -42,8 +42,8 @@ public:
 
 	virtual TArray<UPrimitiveComponent*> CreateComponents(AWorldPartitionHLOD* InHLODActor, const UHLODLayer* InHLODLayer, const TArray<UPrimitiveComponent*>& InSubComponents) const;
 
-	void Build(AWorldPartitionHLOD* InHLODActor, const UHLODLayer* InHLODLayer, const TArray<FWorldPartitionReference>& InSubActors);
+	void Build(AWorldPartitionHLOD* InHLODActor, const UHLODLayer* InHLODLayer, const TArray<AActor*>& InSubActors);
 
-	static TArray<UPrimitiveComponent*> GatherPrimitiveComponents(const TArray<FWorldPartitionReference>& InActors);
+	static TArray<UPrimitiveComponent*> GatherPrimitiveComponents(const TArray<AActor*>& InActors);
 #endif // WITH_EDITOR
 };
diff --git a/Engine/Source/Runtime/Engine/Public/WorldPartition/HLOD/HLODSubActor.h b/Engine/Source/Runtime/Engine/Public/WorldPartition/HLOD/HLODSubActor.h
new file mode 100644
--- /dev/null
+++ b/Engine/Source/Runtime/Engine/Public/WorldPartition/HLOD/HLODSubActor.h
@@ -0,0 +1,118 @@
+// Copyright Epic Games, Inc. All Rights Reserved.
+#pragma once
+
+#include "Misc/Guid.h"
+#include "WorldPartition/WorldPartitionActorCluster.h"
+#include "HLODSubActor.generated.h"
+
+USTRUCT()
+struct FHLODSubActorDesc
+{
+	GENERATED_USTRUCT_BODY()
+
+	FHLODSubActorDesc()
+	{}
+
+	FHLODSubActorDesc(FGuid InActorGuid, const FActorContainerID& InContainerID)
+#if WITH_EDITORONLY_DATA
+		: ActorGuid(InActorGuid)
+		, ContainerID(InContainerID)
+#endif
+	{}
+
+#if WITH_EDITORONLY_DATA
+	friend bool operator==(const FHLODSubActorDesc& Lhs, const FHLODSubActorDesc& Rhs)
+	{
+		return Lhs.ActorGuid == Rhs.ActorGuid && Lhs.ContainerID == Rhs.ContainerID;
+	}
+
+	friend bool operator!=(const FHLODSubActorDesc& Lhs, const FHLODSubActorDesc& Rhs)
+	{
+		return !(Lhs == Rhs);
+	}
+
+	friend bool operator<(const FHLODSubActorDesc& Lhs, const FHLODSubActorDesc& Rhs)
+	{
+		if (Lhs.ActorGuid == Rhs.ActorGuid)
+		{
+			return Lhs.ContainerID.ID < Rhs.ContainerID.ID;
+		}
+
+		return Lhs.ActorGuid < Rhs.ActorGuid;
+	}
+	
+	ENGINE_API friend FArchive& operator<<(FArchive& Ar, FHLODSubActorDesc& SubActor);
+
+	UPROPERTY()
+	FGuid ActorGuid;
+
+	UPROPERTY()
+	FActorContainerID ContainerID;
+#endif
+};
+
+USTRUCT()
+struct FHLODSubActor
+{
+	GENERATED_USTRUCT_BODY()
+
+	FHLODSubActor()
+#if WITH_EDITORONLY_DATA
+		: ActorPackage(NAME_None)
+		, ActorPath(NAME_None)
+		, ContainerPackage(NAME_None)
+		, ContainerTransform(FTransform::Identity)
+#endif
+	{}
+
+	FHLODSubActor(FGuid InActorGuid, FName InActorPackage, FName InActorPath, const FActorContainerID& InContainerID, FName InContainerPackage, const FTransform& InContainerTransform)
+#if WITH_EDITORONLY_DATA
+		: ActorGuid(InActorGuid)
+		, ActorPackage(InActorPackage)
+		, ActorPath(InActorPath)
+		, ContainerID(InContainerID)
+		, ContainerPackage(InContainerPackage)
+		, ContainerTransform(InContainerTransform)
+#endif
+	{}
+
+#if WITH_EDITORONLY_DATA
+	friend bool operator==(const FHLODSubActor& Lhs, const FHLODSubActor& Rhs)
+	{
+		return Lhs.ActorGuid == Rhs.ActorGuid && Lhs.ContainerID == Rhs.ContainerID;
+	}
+
+	friend bool operator!=(const FHLODSubActor& Lhs, const FHLODSubActor& Rhs)
+	{
+		return !(Lhs == Rhs);
+	}
+
+	friend bool operator<(const FHLODSubActor& Lhs, const FHLODSubActor& Rhs)
+	{
+		if (Lhs.ActorGuid == Rhs.ActorGuid)
+		{
+			return Lhs.ContainerID.ID < Rhs.ContainerID.ID;
+		}
+
+		return Lhs.ActorGuid < Rhs.ActorGuid;
+	}
+		
+	UPROPERTY()
+	FGuid ActorGuid;
+	
+	UPROPERTY()
+	FName ActorPackage;
+
+	UPROPERTY()
+	FName ActorPath;
+
+	UPROPERTY()
+	FActorContainerID ContainerID;
+
+	UPROPERTY()
+	FName ContainerPackage;
+
+	UPROPERTY()
+	FTransform ContainerTransform;
+#endif
+};
diff --git a/Engine/Source/Runtime/Engine/Public/WorldPartition/WorldPartitionActorCluster.h b/Engine/Source/Runtime/Engine/Public/WorldPartition/WorldPartitionActorCluster.h
--- a/Engine/Source/Runtime/Engine/Public/WorldPartition/WorldPartitionActorCluster.h
+++ b/Engine/Source/Runtime/Engine/Public/WorldPartition/WorldPartitionActorCluster.h
@@ -159,6 +159,7 @@ public:
 	 */
 	const TArray<FActorClusterInstance>& GetClusterInstances() const { return ClusterInstances; }
 	
+	const FActorContainerInstance* GetClusterInstance(const FActorContainerID& InContainerID) const;
 	FActorContainerInstance* GetClusterInstance(const UActorDescContainer* InContainer);
 	const FActorContainerInstance* GetClusterInstance(const UActorDescContainer* InContainer) const;
 
diff --git a/Engine/Source/Runtime/Engine/Public/WorldPartition/WorldPartitionLevelHelper.h b/Engine/Source/Runtime/Engine/Public/WorldPartition/WorldPartitionLevelHelper.h
--- a/Engine/Source/Runtime/Engine/Public/WorldPartition/WorldPartitionLevelHelper.h
+++ b/Engine/Source/Runtime/Engine/Public/WorldPartition/WorldPartitionLevelHelper.h
@@ -25,7 +25,7 @@ public:
 	static void MoveExternalActorsToLevel(const TArray<FWorldPartitionRuntimeCellObjectMapping>& InChildPackages, ULevel* InLevel);
 	static void RemapLevelSoftObjectPaths(ULevel* InLevel, UWorldPartition* InWorldPartition);
 	
-	static bool LoadActors(ULevel* InDestLevel, TArrayView<FWorldPartitionRuntimeCellObjectMapping> InActorPackages, FWorldPartitionPackageCache& InPackageCache, TFunction<void(bool)> InCompletionCallback, bool bLoadForPlay, FLinkerInstancingContext* InOutInstancingContext = nullptr);
+	static bool LoadActors(ULevel* InDestLevel, TArrayView<FWorldPartitionRuntimeCellObjectMapping> InActorPackages, FWorldPartitionPackageCache& InPackageCache, TFunction<void(bool)> InCompletionCallback, bool bInLoadAsync, FLinkerInstancingContext* InOutInstancingContext = nullptr);
 	
 	static FString AddActorContainerIDToActorPath(const FActorContainerID& InContainerID, const FString& InActorPath);
 	static FString AddActorContainerIDToSubPathString(const FActorContainerID& InContainerID, const FString& InSubPathString);
diff --git a/Engine/Source/Runtime/Engine/Public/WorldPartition/WorldPartitionLevelStreamingDynamic.h b/Engine/Source/Runtime/Engine/Public/WorldPartition/WorldPartitionLevelStreamingDynamic.h
--- a/Engine/Source/Runtime/Engine/Public/WorldPartition/WorldPartitionLevelStreamingDynamic.h
+++ b/Engine/Source/Runtime/Engine/Public/WorldPartition/WorldPartitionLevelStreamingDynamic.h
@@ -37,6 +37,9 @@ class ENGINE_API UWorldPartitionLevelStreamingDynamic : public ULevelStreamingDy
 	virtual bool ShouldRequireFullVisibilityToRender() const override { return true; }
 
 #if WITH_EDITOR
+	static UWorldPartitionLevelStreamingDynamic* LoadInEditor(UWorld* World, FName LevelStreamingName, const TArray<FWorldPartitionRuntimeCellObjectMapping>& InPackages);
+	static void UnloadFromEditor(UWorldPartitionLevelStreamingDynamic* InLevelStreaming);
+
 	void Initialize(const UWorldPartitionRuntimeLevelStreamingCell& InCell);
 
 	// Override ULevelStreaming
@@ -44,12 +47,16 @@ class ENGINE_API UWorldPartitionLevelStreamingDynamic : public ULevelStreamingDy
 	virtual void BeginDestroy() override;
 	virtual TOptional<FFolder::FRootObject> GetFolderRootObject() const override;
 
+	bool GetLoadSucceeded() const { return bLoadSucceeded; }
+
 private:
 	void CreateRuntimeLevel();
 	bool IssueLoadRequests();
 	void FinalizeRuntimeLevel();
 	void OnCleanupLevel();
 
+	void Initialize(UWorld* OuterWorld, const TArray<FWorldPartitionRuntimeCellObjectMapping>& InPackages);
+
 	FName OriginalLevelPackageName;
 	TArray<FWorldPartitionRuntimeCellObjectMapping> ChildPackages;
 	TArray<FWorldPartitionRuntimeCellObjectMapping> ChildPackagesToLoad;
@@ -62,6 +69,7 @@ private:
 
 	FDelegateHandle OnCleanupLevelDelegateHandle;
 	bool bLoadRequestInProgress;
+	bool bLoadSucceeded;
 	FWorldPartitionPackageCache PackageCache;
 #endif
 
