Change subobject outer ensure to a log, fall back to default behaviour (actor as outer) if outer is unavailable

diff --git a/Engine/Source/Runtime/Engine/Private/DataChannel.cpp b/Engine/Source/Runtime/Engine/Private/DataChannel.cpp
--- a/Engine/Source/Runtime/Engine/Private/DataChannel.cpp
+++ b/Engine/Source/Runtime/Engine/Private/DataChannel.cpp
@@ -3906,9 +3906,13 @@ UObject* UActorChannel::ReadContentBlockHeader(FInBunch& Bunch, bool& bObjectDel
 				return nullptr;
 			}
 
-			if (!ensureMsgf(ObjOuter != nullptr, TEXT("UActorChannel::ReadContentBlockHeader: Unable to serialize subobject's outer. Ensure that subobjects are replicated top-down, so outers are received first. Class: %s, Actor: %s"), *GetNameSafe(SubObjClass), *GetNameSafe(Actor)))
+			if (ObjOuter == nullptr && SubObj == nullptr)
 			{
-				return nullptr;
+				// When replicating subobjects, the engine will by default replicate lower subobjects before their outers, using the owning actor as the outer on the client.
+				// To have a chain of subobjects maintain their correct outers on the client, the subobjects should be replicated top-down, so the outers are received before any subobjects they own.
+				UE_LOG(LogNetTraffic, Log, TEXT("UActorChannel::ReadContentBlockHeader: Unable to serialize subobject's outer, using owning actor as outer instead. Class: %s, Actor: %s"), *GetNameSafe(SubObjClass), *GetNameSafe(Actor));
+				
+				ObjOuter = Actor;
 			}
 		}
 	}
