Unpause rendering world when a single frame is advanced or toggling between PIE & SIE. It fixes motion blur artifacts when single stepping in PIE.

diff --git a/Engine/Source/Editor/UnrealEd/Private/EditorEngine.cpp b/Engine/Source/Editor/UnrealEd/Private/EditorEngine.cpp
--- a/Engine/Source/Editor/UnrealEd/Private/EditorEngine.cpp
+++ b/Engine/Source/Editor/UnrealEd/Private/EditorEngine.cpp
@@ -1649,6 +1649,8 @@ void UEditorEngine::Tick( float DeltaSeconds, bool bIdleMode )
 		}
 	}
 
+	bool bToggledBetweenPIEandSIE = bIsToggleBetweenPIEandSIEQueued;
+
 	// Kick off a Play Session request if one was queued up during the last frame.
 	if (PlaySessionRequest.IsSet())
 	{
@@ -1763,6 +1765,9 @@ void UEditorEngine::Tick( float DeltaSeconds, bool bIdleMode )
 
 					// tick the level
 					PieContext.World()->Tick( LEVELTICK_All, TickDeltaSeconds );
+#if WITH_EDITOR
+					PieContext.World()->bToggledBetweenPIEandSIEThisFrame = bToggledBetweenPIEandSIE;
+#endif
 					bAWorldTicked = true;
 					TickType = LEVELTICK_All;
 
diff --git a/Engine/Source/Runtime/Engine/Classes/Engine/World.h b/Engine/Source/Runtime/Engine/Classes/Engine/World.h
--- a/Engine/Source/Runtime/Engine/Classes/Engine/World.h
+++ b/Engine/Source/Runtime/Engine/Classes/Engine/World.h
@@ -1315,6 +1315,12 @@ public:
 #if WITH_EDITOR
 	/** When set, will tell us to pause simulation after one tick.  If a breakpoint is encountered before tick is complete we will stop there instead. */
 	uint8 bDebugFrameStepExecution:1;
+
+	/** Indicates that a single frame advance happened this frame. */
+	uint8 bDebugFrameStepExecutedThisFrame : 1;
+
+	/** Indicates that toggling between Play-in-Editor and Simulate-in-Editor happened this frame. */
+	uint8 bToggledBetweenPIEandSIEThisFrame : 1;
 #endif
 
 	/** Keeps track whether actors moved via PostEditMove and therefore constraint syncup should be performed. */
diff --git a/Engine/Source/Runtime/Engine/Private/GPUSkinVertexFactory.cpp b/Engine/Source/Runtime/Engine/Private/GPUSkinVertexFactory.cpp
--- a/Engine/Source/Runtime/Engine/Private/GPUSkinVertexFactory.cpp
+++ b/Engine/Source/Runtime/Engine/Private/GPUSkinVertexFactory.cpp
@@ -717,7 +717,7 @@ public:
 				// So check if bone data is updated this frame, if not then the previous frame data is stale and not suitable for motion blur.
 				bool bBoneDataUpdatedThisFrame = (View->Family->FrameNumber == ShaderData.UpdatedFrameNumber);
 				// If world is paused, use current frame bone matrices, so velocity is canceled and skeletal mesh isn't blurred from motion.
-				bool bPrevious = !View->Family->bWorldIsPaused_IncludingSimulatingInEditor && bBoneDataUpdatedThisFrame;
+				bool bPrevious = !View->Family->bWorldIsPaused && bBoneDataUpdatedThisFrame;
 				FRHIShaderResourceView* PreviousData = ShaderData.GetBoneBufferForReading(bPrevious).VertexBufferSRV;
 				ShaderBindings.Add(PreviousBoneMatrices, PreviousData);
 			}
@@ -809,7 +809,7 @@ public:
 		// So check if bone data is updated this frame, if not then the previous frame data is stale and not suitable for motion blur.
 		bool bBoneDataUpdatedThisFrame = (View->Family->FrameNumber == PassthroughVertexFactory->GetUpdatedFrameNumber());
 		// If world is paused, use current frame bone matrices, so velocity is canceled and skeletal mesh isn't blurred from motion.
-		bool bVerticesInMotion = !View->Family->bWorldIsPaused_IncludingSimulatingInEditor && bBoneDataUpdatedThisFrame;
+		bool bVerticesInMotion = !View->Family->bWorldIsPaused && bBoneDataUpdatedThisFrame;
 
 		if (bUsesSkinCache)
 		{
@@ -1104,7 +1104,7 @@ public:
 		{
 			const FMorphVertexBuffer* MorphVertexBuffer = nullptr;
 			const auto* GPUSkinVertexFactory = (const FGPUBaseSkinVertexFactory*)VertexFactory;
-			MorphVertexBuffer = GPUSkinVertexFactory->GetMorphVertexBuffer(!View->Family->bWorldIsPaused_IncludingSimulatingInEditor, View->Family->FrameNumber);
+			MorphVertexBuffer = GPUSkinVertexFactory->GetMorphVertexBuffer(!View->Family->bWorldIsPaused, View->Family->FrameNumber);
 			if (MorphVertexBuffer)
 			{
 				ShaderBindings.Add(PreviousMorphBufferParameter, MorphVertexBuffer->GetSRV());
diff --git a/Engine/Source/Runtime/Engine/Private/LevelTick.cpp b/Engine/Source/Runtime/Engine/Private/LevelTick.cpp
--- a/Engine/Source/Runtime/Engine/Private/LevelTick.cpp
+++ b/Engine/Source/Runtime/Engine/Private/LevelTick.cpp
@@ -1720,6 +1720,7 @@ void UWorld::Tick( ELevelTick TickType, float DeltaSeconds )
 
 #if WITH_EDITOR
 	// Finish up.
+	bDebugFrameStepExecutedThisFrame = bDebugFrameStepExecution;
 	if(bDebugFrameStepExecution)
 	{
 		bDebugPauseExecution = true;
diff --git a/Engine/Source/Runtime/Engine/Private/SceneView.cpp b/Engine/Source/Runtime/Engine/Private/SceneView.cpp
--- a/Engine/Source/Runtime/Engine/Private/SceneView.cpp
+++ b/Engine/Source/Runtime/Engine/Private/SceneView.cpp
@@ -2729,7 +2729,6 @@ FSceneViewFamily::FSceneViewFamily(const ConstructionValues& CVS)
 	SceneCaptureSource(SCS_FinalColorLDR),
 	SceneCaptureCompositeMode(SCCM_Overwrite),
 	bWorldIsPaused(false),
-	bWorldIsPaused_IncludingSimulatingInEditor(false),
 	bIsHDR(false),
 	bRequireMultiView(false),
 	GammaCorrection(CVS.GammaCorrection),
@@ -2779,8 +2778,12 @@ FSceneViewFamily::FSceneViewFamily(const ConstructionValues& CVS)
 				EngineShowFlags.LOD = 1;
 			}
 
+#if WITH_EDITOR
+			// If a single frame step or toggling between Play-in-Editor and Simulate-in-Editor happened this frame, then unpause for this frame.
+			bWorldIsPaused = !(World->IsCameraMoveable() || World->bDebugFrameStepExecutedThisFrame || World->bToggledBetweenPIEandSIEThisFrame);
+#else
 			bWorldIsPaused = !World->IsCameraMoveable();
-			bWorldIsPaused_IncludingSimulatingInEditor = World->IsPaused();
+#endif
 		}
 	}
 
diff --git a/Engine/Source/Runtime/Engine/Private/World.cpp b/Engine/Source/Runtime/Engine/Private/World.cpp
--- a/Engine/Source/Runtime/Engine/Private/World.cpp
+++ b/Engine/Source/Runtime/Engine/Private/World.cpp
@@ -388,6 +388,10 @@ UWorld::UWorld( const FObjectInitializer& ObjectInitializer )
 , bAllowDeferredPhysicsStateCreation(false)
 , FeatureLevel(GMaxRHIFeatureLevel)
 , bIsBuilt(false)
+#if WITH_EDITOR
+, bDebugFrameStepExecutedThisFrame(false)
+, bToggledBetweenPIEandSIEThisFrame(false)
+#endif
 , bShouldTick(true)
 , bInitializedAndNeedsCleanup(false)
 , ActiveLevelCollectionIndex(INDEX_NONE)
diff --git a/Engine/Source/Runtime/Engine/Public/SceneView.h b/Engine/Source/Runtime/Engine/Public/SceneView.h
--- a/Engine/Source/Runtime/Engine/Public/SceneView.h
+++ b/Engine/Source/Runtime/Engine/Public/SceneView.h
@@ -1731,9 +1731,6 @@ public:
 	 */
 	bool bWorldIsPaused;
 
-	/** GetWorld->IsPaused(). Simulate-in-Editor included. */
-	bool bWorldIsPaused_IncludingSimulatingInEditor;
-
 	/** When enabled, the post processing will output in HDR space */
 	bool bIsHDR;
 
