PR #8485: Fix crash when analyzing material if it hasn't yet been compiled (Contributed by [REDACTED_PERSON])

#rnx

diff --git a/Engine/Source/Runtime/Engine/Private/Materials/MaterialShared.cpp b/Engine/Source/Runtime/Engine/Private/Materials/MaterialShared.cpp
--- a/Engine/Source/Runtime/Engine/Private/Materials/MaterialShared.cpp
+++ b/Engine/Source/Runtime/Engine/Private/Materials/MaterialShared.cpp
@@ -4244,6 +4244,12 @@ void UMaterialInterface::AnalyzeMaterialProperty(EMaterialProperty InProperty, i
 
 	FMaterialCompilationOutput TempOutput;
 	FMaterialResource* MaterialResource = GetMaterialResource(GMaxRHIFeatureLevel);
+	if (MaterialResource == nullptr)
+	{
+		ForceRecompileForRendering(); // Make sure material has a resource to avoid crash
+		MaterialResource = GetMaterialResource(GMaxRHIFeatureLevel);
+	}
+
 	FMaterialShaderMapId ShaderMapID;
 	MaterialResource->GetShaderMapId(GMaxRHIShaderPlatform, nullptr, ShaderMapID);
 	FStaticParameterSet StaticParamSet;
