Fix QualityLevelMinLod, dont do the MinLOD conversion when there a QualityLevelMinLod.Default value != 0

diff --git a/Engine/Source/Runtime/Engine/Private/StaticMesh.cpp b/Engine/Source/Runtime/Engine/Private/StaticMesh.cpp
--- a/Engine/Source/Runtime/Engine/Private/StaticMesh.cpp
+++ b/Engine/Source/Runtime/Engine/Private/StaticMesh.cpp
@@ -5599,7 +5599,7 @@ void UStaticMesh::ExecutePostLoadInternal(FStaticMeshPostLoadContext& Context)
 
 		// Convert PerPlatForm data to PerQuality if perQuality data have not been serialized.
 		// Also test default value, since PerPLatformData can have Default !=0 and no PerPlaform data overrides.
-		bool bConvertMinLODData = (PerQualityLevelData.PerQuality.Num() == 0 && PerPlatformData.PerPlatform.Num() != 0) || (PerPlatformData.PerPlatform.Num() == 0 && PerQualityLevelData.Default != PerPlatformData.Default);
+		bool bConvertMinLODData = (PerQualityLevelData.PerQuality.Num() == 0 && PerQualityLevelData.Default == 0) && (PerPlatformData.PerPlatform.Num() != 0 || PerPlatformData.Default != 0);
 
 		if (GEngine && GEngine->UseStaticMeshMinLODPerQualityLevels && bConvertMinLODData)
 		{
