PR #8821: Provide a delegate for game code to listen to for Replay ID updates (Contributed by [REDACTED_PERSON])

diff --git a/Engine/Source/Runtime/Engine/Private/ReplayHelper.cpp b/Engine/Source/Runtime/Engine/Private/ReplayHelper.cpp
--- a/Engine/Source/Runtime/Engine/Private/ReplayHelper.cpp
+++ b/Engine/Source/Runtime/Engine/Private/ReplayHelper.cpp
@@ -108,6 +108,15 @@ void FReplayHelper::OnStartRecordingComplete(const FStartStreamingResult& Result
 		OnReplayRecordError.Broadcast();
 		return;
 	}
+
+	// Broadcast the replay id for anyone interested as by this point it has been finalized.
+	// Eg. The replay server telling us the replay id in the case of FHttpNetworkReplayStreamer
+	if (ReplayStreamer.IsValid())
+	{
+		const FString ReplayID = ReplayStreamer->GetReplayID();
+
+		FNetworkReplayDelegates::OnReplayIDChanged.Broadcast(World.Get(), ReplayID);
+	}
 }
 
 void FReplayHelper::StartRecording(UNetConnection* Connection)
diff --git a/Engine/Source/Runtime/Engine/Private/UnrealNetwork.cpp b/Engine/Source/Runtime/Engine/Private/UnrealNetwork.cpp
--- a/Engine/Source/Runtime/Engine/Private/UnrealNetwork.cpp
+++ b/Engine/Source/Runtime/Engine/Private/UnrealNetwork.cpp
@@ -13,6 +13,7 @@ FOnReplayScrubCompleteDelegate FNetworkReplayDelegates::OnReplayScrubComplete;
 FOnReplayPlaybackCompleteDelegate FNetworkReplayDelegates::OnReplayPlaybackComplete;
 FOnReplayRecordingCompleteDelegate FNetworkReplayDelegates::OnReplayRecordingComplete;
 FOnPauseChannelsChangedDelegate FNetworkReplayDelegates::OnPauseChannelsChanged;
+FOnReplayIDChangedDelegate FNetworkReplayDelegates::OnReplayIDChanged;
 
 // ----------------------------------------------------------------
 
diff --git a/Engine/Source/Runtime/Engine/Public/Net/UnrealNetwork.h b/Engine/Source/Runtime/Engine/Public/Net/UnrealNetwork.h
--- a/Engine/Source/Runtime/Engine/Public/Net/UnrealNetwork.h
+++ b/Engine/Source/Runtime/Engine/Public/Net/UnrealNetwork.h
@@ -52,6 +52,7 @@ DECLARE_MULTICAST_DELEGATE_OneParam(FOnReplayScrubCompleteDelegate, UWorld* /*Wo
 DECLARE_MULTICAST_DELEGATE_OneParam(FOnReplayPlaybackCompleteDelegate, UWorld* /*World*/);
 DECLARE_MULTICAST_DELEGATE_OneParam(FOnReplayRecordingCompleteDelegate, UWorld* /*World*/);
 DECLARE_MULTICAST_DELEGATE_TwoParams(FOnPauseChannelsChangedDelegate, UWorld* /*World*/, bool /*bPaused*/);
+DECLARE_MULTICAST_DELEGATE_TwoParams(FOnReplayIDChangedDelegate, UWorld* /*World*/, const FString& /*ReplayID*/);
 
 struct ENGINE_API FNetworkReplayDelegates
 {
@@ -83,6 +84,9 @@ struct ENGINE_API FNetworkReplayDelegates
 
 	/** Delegate for external systems to be notified when channels are paused during playback, usually waiting for data to be available. */
 	static FOnPauseChannelsChangedDelegate OnPauseChannelsChanged;
+
+	/** Delegate for external systems to be notified when the SessionName has changed (The replay identifier). */
+	static FOnReplayIDChangedDelegate OnReplayIDChanged;
 };
 
 /**
