Add an interface function to indicate if to use engine built-in morph target compression. Classes derived from UMorphTarget can override with custom compression behaviour.

diff --git a/Engine/Source/Runtime/Engine/Classes/Animation/MorphTarget.h b/Engine/Source/Runtime/Engine/Classes/Animation/MorphTarget.h
--- a/Engine/Source/Runtime/Engine/Classes/Animation/MorphTarget.h
+++ b/Engine/Source/Runtime/Engine/Classes/Animation/MorphTarget.h
@@ -204,6 +204,9 @@ public:
 	/** Discard CPU Buffers after render resources have been created. */
 	ENGINE_API virtual void DiscardVertexData();
 
+	/** Return true if this morph target uses engine built-in compression */
+	ENGINE_API virtual bool UsesBuiltinMorphTargetCompression() const { return true; }
+
 #if WITH_EDITOR
 	/** Populates the given morph target LOD model with the provided deltas */
 	ENGINE_API virtual void PopulateDeltas(const TArray<FMorphTargetDelta>& Deltas, const int32 LODIndex, const TArray<struct FSkelMeshSection>& Sections, const bool bCompareNormal = false, const bool bGeneratedByReductionSetting = false, const float PositionThreshold = THRESH_POINTS_ARE_NEAR);
diff --git a/Engine/Source/Runtime/Engine/Private/Rendering/MorphTargetVertexInfoBuffers.cpp b/Engine/Source/Runtime/Engine/Private/Rendering/MorphTargetVertexInfoBuffers.cpp
--- a/Engine/Source/Runtime/Engine/Private/Rendering/MorphTargetVertexInfoBuffers.cpp
+++ b/Engine/Source/Runtime/Engine/Private/Rendering/MorphTargetVertexInfoBuffers.cpp
@@ -242,7 +242,7 @@ void FMorphTargetVertexInfoBuffers::InitMorphResources(EShaderPlatform ShaderPla
 		int32 NumSrcDeltas = 0;
 		const FMorphTargetDelta* MorphDeltas = MorphTarget->GetMorphTargetDelta(LODIndex, NumSrcDeltas);
 
-		if (NumSrcDeltas == 0)
+		if (NumSrcDeltas == 0 || !MorphTarget->UsesBuiltinMorphTargetCompression())
 		{
 			MaximumValues[0] = 0.0f;
 			MaximumValues[1] = 0.0f;
