Change the Async Load Primary Asset blueprint nodes to return all requested and loaded assets instead of only newly loaded ones, which is what users expect it to do
This works around the fact that LoadPrimaryAssets currently returns null if requested assets are already loaded

diff --git a/Engine/Source/Runtime/Engine/Private/AsyncActionLoadPrimaryAsset.cpp b/Engine/Source/Runtime/Engine/Private/AsyncActionLoadPrimaryAsset.cpp
--- a/Engine/Source/Runtime/Engine/Private/AsyncActionLoadPrimaryAsset.cpp
+++ b/Engine/Source/Runtime/Engine/Private/AsyncActionLoadPrimaryAsset.cpp
@@ -40,6 +40,22 @@ void UAsyncActionLoadPrimaryAssetBase::HandleLoadCompleted()
 	SetReadyToDestroy();
 }
 
+void UAsyncActionLoadPrimaryAssetBase::GetCurrentlyLoadedAssets(TArray<UObject*>& AssetList)
+{
+	if (UAssetManager* Manager = UAssetManager::GetIfValid())
+	{
+		// The assets may have already been loaded but the handle was invalid, check the original list
+		for (const FPrimaryAssetId& IdToLoad : AssetsToLoad)
+		{
+			UObject* LoadedObject = Manager->GetPrimaryAssetObject(IdToLoad);
+			if (LoadedObject)
+			{
+				AssetList.Add(LoadedObject);
+			}
+		}
+	}
+}
+
 UAsyncActionLoadPrimaryAsset* UAsyncActionLoadPrimaryAsset::AsyncLoadPrimaryAsset(UObject* WorldContextObject, FPrimaryAssetId PrimaryAsset, const TArray<FName>& LoadBundles)
 {
 	UAsyncActionLoadPrimaryAsset* Action = NewObject<UAsyncActionLoadPrimaryAsset>();
@@ -54,9 +70,13 @@ UAsyncActionLoadPrimaryAsset* UAsyncActionLoadPrimaryAsset::AsyncLoadPrimaryAsse
 void UAsyncActionLoadPrimaryAsset::HandleLoadCompleted()
 {
 	UObject* AssetLoaded = nullptr;
-	if (LoadHandle.IsValid())
+	TArray<UObject*> AssetList;
+
+	GetCurrentlyLoadedAssets(AssetList);
+
+	if (AssetList.Num() > 0)
 	{
-		AssetLoaded = LoadHandle->GetLoadedAsset();
+		AssetLoaded = AssetList[0];
 	}
 
 	Super::HandleLoadCompleted();
@@ -77,9 +97,13 @@ UAsyncActionLoadPrimaryAssetClass* UAsyncActionLoadPrimaryAssetClass::AsyncLoadP
 void UAsyncActionLoadPrimaryAssetClass::HandleLoadCompleted()
 {
 	TSubclassOf<UObject> AssetLoaded = nullptr;
-	if (LoadHandle.IsValid())
+	TArray<UObject*> AssetList;
+
+	GetCurrentlyLoadedAssets(AssetList);
+
+	if (AssetList.Num() > 0)
 	{
-		AssetLoaded = Cast<UClass>(LoadHandle->GetLoadedAsset());
+		AssetLoaded = Cast<UClass>(AssetList[0]);
 	}
 
 	Super::HandleLoadCompleted();
@@ -101,10 +125,7 @@ void UAsyncActionLoadPrimaryAssetList::HandleLoadCompleted()
 {
 	TArray<UObject*> AssetList;
 
-	if (LoadHandle.IsValid())
-	{
-		LoadHandle->GetLoadedAssets(AssetList);
-	}
+	GetCurrentlyLoadedAssets(AssetList);
 
 	Super::HandleLoadCompleted();
 	Completed.Broadcast(AssetList);
@@ -126,18 +147,15 @@ void UAsyncActionLoadPrimaryAssetClassList::HandleLoadCompleted()
 	TArray<TSubclassOf<UObject>> AssetClassList;
 	TArray<UObject*> AssetList;
 
-	if (LoadHandle.IsValid())
+	GetCurrentlyLoadedAssets(AssetList);
+
+	for (UObject* LoadedAsset : AssetList)
 	{
-		LoadHandle->GetLoadedAssets(AssetList);
+		UClass* LoadedClass = Cast<UClass>(LoadedAsset);
 
-		for (UObject* LoadedAsset : AssetList)
+		if (LoadedClass)
 		{
-			UClass* LoadedClass = Cast<UClass>(LoadedAsset);
-
-			if (LoadedClass)
-			{
-				AssetClassList.Add(LoadedClass);
-			}
+			AssetClassList.Add(LoadedClass);
 		}
 	}
 
diff --git a/Engine/Source/Runtime/Engine/Private/AsyncActionLoadPrimaryAsset.h b/Engine/Source/Runtime/Engine/Private/AsyncActionLoadPrimaryAsset.h
--- a/Engine/Source/Runtime/Engine/Private/AsyncActionLoadPrimaryAsset.h
+++ b/Engine/Source/Runtime/Engine/Private/AsyncActionLoadPrimaryAsset.h
@@ -25,6 +25,9 @@ protected:
 	/** Called from asset manager */
 	virtual void HandleLoadCompleted();
 
+	/** Gets all UObjects matching AssetsToLoad that are already in memory */
+	virtual void GetCurrentlyLoadedAssets(TArray<UObject*>& AssetList);
+
 	/** Specific assets requested */
 	TArray<FPrimaryAssetId> AssetsToLoad;
 
@@ -57,8 +60,9 @@ class UAsyncActionLoadPrimaryAsset : public UAsyncActionLoadPrimaryAssetBase
 
 public:
 	/** 
-	 * Load a primary asset into memory. The completed delegate will go off when the load succeeds or fails, you should cast the Loaded object to verify it is the correct type.
-	 * If LoadBundles is specified, those bundles are loaded along with the asset
+	 * Load a primary asset object into memory, this will cause it to stay loaded until it is explicitly unloaded.
+	 * The completed event will happen when the load succeeds or fails, you should cast the Loaded object to verify it is the correct type.
+	 * If LoadBundles is specified, those bundles are loaded along with the asset.
 	 */
 	UFUNCTION(BlueprintCallable, meta=(BlueprintInternalUseOnly="true", Category = "AssetManager", AutoCreateRefTerm = "LoadBundles", WorldContext = "WorldContextObject"))
 	static UAsyncActionLoadPrimaryAsset* AsyncLoadPrimaryAsset(UObject* WorldContextObject, FPrimaryAssetId PrimaryAsset, const TArray<FName>& LoadBundles);
@@ -80,9 +84,10 @@ class UAsyncActionLoadPrimaryAssetClass : public UAsyncActionLoadPrimaryAssetBas
 	GENERATED_BODY()
 
 public:
-	/** 
-	 * Load a primary asset class into memory. The completed delegate will go off when the load succeeds or fails, you should cast the Loaded class to verify it is the correct type 
-	 * If LoadBundles is specified, those bundles are loaded along with the asset
+	/**
+	 * Load a primary asset class  into memory, this will cause it to stay loaded until it is explicitly unloaded.
+	 * The completed event will happen when the load succeeds or fails, you should cast the Loaded class to verify it is the correct type.
+	 * If LoadBundles is specified, those bundles are loaded along with the asset.
 	 */
 	UFUNCTION(BlueprintCallable, meta = (BlueprintInternalUseOnly = "true", Category = "AssetManager", AutoCreateRefTerm = "LoadBundles", WorldContext = "WorldContextObject"))
 	static UAsyncActionLoadPrimaryAssetClass* AsyncLoadPrimaryAssetClass(UObject* WorldContextObject, FPrimaryAssetId PrimaryAsset, const TArray<FName>& LoadBundles);
@@ -104,9 +109,10 @@ class UAsyncActionLoadPrimaryAssetList : public UAsyncActionLoadPrimaryAssetBase
 	GENERATED_BODY()
 
 public:
-	/** 
-	 * Load a list primary assets into memory. The completed delegate will go off when the load succeeds or fails, you should cast the Loaded object list to verify it is the correct type 
-	 * If LoadBundles is specified, those bundles are loaded along with the asset list
+	/**
+	 * Load a list of primary asset objects into memory, this will cause them to stay loaded until explicitly unloaded.
+	 * The completed event will happen when the load succeeds or fails, and the Loaded list will contain all of the requested assets found at completion.
+	 * If LoadBundles is specified, those bundles are loaded along with the assets.
 	 */
 	UFUNCTION(BlueprintCallable, meta=(BlueprintInternalUseOnly="true", Category = "AssetManager", AutoCreateRefTerm = "LoadBundles", WorldContext = "WorldContextObject"))
 	static UAsyncActionLoadPrimaryAssetList* AsyncLoadPrimaryAssetList(UObject* WorldContextObject, const TArray<FPrimaryAssetId>& PrimaryAssetList, const TArray<FName>& LoadBundles);
@@ -128,9 +134,10 @@ class UAsyncActionLoadPrimaryAssetClassList : public UAsyncActionLoadPrimaryAsse
 	GENERATED_BODY()
 
 public:
-	/** 
-	 * Load a list primary asset classes into memory. The completed delegate will go off when the load succeeds or fails, you should cast the Loaded object list to verify it is the correct type 
-	 * If LoadBundles is specified, those bundles are loaded along with the asset list
+	/**
+	 * Load a list of primary asset classes into memory, this will cause them to stay loaded until explicitly unloaded.
+	 * The completed event will happen when the load succeeds or fails, and the Loaded list will contain all of the requested classes found at completion.
+	 * If LoadBundles is specified, those bundles are loaded along with the assets.
 	 */
 	UFUNCTION(BlueprintCallable, meta = (BlueprintInternalUseOnly = "true", Category = "AssetManager", AutoCreateRefTerm = "LoadBundles", WorldContext = "WorldContextObject"))
 	static UAsyncActionLoadPrimaryAssetClassList* AsyncLoadPrimaryAssetClassList(UObject* WorldContextObject, const TArray<FPrimaryAssetId>& PrimaryAssetList, const TArray<FName>& LoadBundles);
@@ -153,13 +160,16 @@ class UAsyncActionChangePrimaryAssetBundles : public UAsyncActionLoadPrimaryAsse
 
 public:
 	/**
-	 * Change the bundle state of all assets that match OldBundles to instead contain NewBundles. 
+	 * Change the bundle state of all assets that match OldBundles to instead contain NewBundles.
+	 * This will not change the loaded status of primary assets but will load or unload secondary assets based on the bundles.
 	 */
 	UFUNCTION(BlueprintCallable, meta = (BlueprintInternalUseOnly = "true", Category = "AssetManager", WorldContext = "WorldContextObject"))
 	static UAsyncActionChangePrimaryAssetBundles* AsyncChangeBundleStateForMatchingPrimaryAssets(UObject* WorldContextObject, const TArray<FName>& NewBundles, const TArray<FName>& OldBundles);
 
 	/**
-	 * Change the bundle state of assets in PrimaryAssetList. AddBundles are added and RemoveBundles are removed, both must be filled in but an empty array is allowed
+	 * Change the bundle state of specific assets in PrimaryAssetList. 
+	 * AddBundles are added to the final state and RemoveBundles are removed, an empty array will make no change.
+	 * This will not change the loaded status of primary assets but will load or unload secondary assets based on the bundles.
 	 */
 	UFUNCTION(BlueprintCallable, meta = (BlueprintInternalUseOnly = "true", Category = "AssetManager", WorldContext = "WorldContextObject"))
 	static UAsyncActionChangePrimaryAssetBundles* AsyncChangeBundleStateForPrimaryAssetList(UObject* WorldContextObject, const TArray<FPrimaryAssetId>& PrimaryAssetList, const TArray<FName>& AddBundles, const TArray<FName>& RemoveBundles);
