LevelInstance: CanEdit checks if package exists

diff --git a/Engine/Source/Runtime/Engine/Private/LevelInstance/LevelInstanceSubsystem.cpp b/Engine/Source/Runtime/Engine/Private/LevelInstance/LevelInstanceSubsystem.cpp
--- a/Engine/Source/Runtime/Engine/Private/LevelInstance/LevelInstanceSubsystem.cpp
+++ b/Engine/Source/Runtime/Engine/Private/LevelInstance/LevelInstanceSubsystem.cpp
@@ -1542,6 +1542,16 @@ bool ULevelInstanceSubsystem::CanEditLevelInstance(const ALevelInstance* LevelIn
 		return false;
 	}
 
+	FPackagePath WorldAssetPath;
+	if (!FPackagePath::TryFromPackageName(LevelInstanceActor->GetWorldAssetPackage(), WorldAssetPath) || !FPackageName::DoesPackageExist(WorldAssetPath))
+	{
+		if (OutReason)
+		{
+			*OutReason = FText::Format(LOCTEXT("CanEditLevelInstanceInvalidAsset", "Level Instance asset is invalid ({0})."), FText::FromString(LevelInstanceActor->GetWorldAssetPackage()));
+		}
+		return false;
+	}
+
 	return true;
 }
 
