Optimize ULevel::GetLoadedExternalObjectPackages()

* Save heap alloc for string copy while iterating all packages to check for external actor.
* Remove superfluous TSet<>::Contains(), we don't care about the lookup as redundant adds are ok.

This knocks several full seconds off of deleting an actor in a big map.

#rnx

diff --git a/Engine/Source/Runtime/Engine/Private/Level.cpp b/Engine/Source/Runtime/Engine/Private/Level.cpp
--- a/Engine/Source/Runtime/Engine/Private/Level.cpp
+++ b/Engine/Source/Runtime/Engine/Private/Level.cpp
@@ -3224,15 +3224,15 @@ TArray<UPackage*> ULevel::GetLoadedExternalObjectPackages() const
 	
 	for (TObjectIterator<UPackage> It; It; ++It)
 	{
-		if (!ExternalObjectPackages.Contains(*It))
+		TStringBuilder<256> PackageName;
+		It->GetLoadedPath().AppendPackageName(PackageName);
+		FStringView PackageNameStringView(PackageName);
+		for (const FString& ExternalObjectsPath : ExternalObjectsPaths)
 		{
-			for (const FString& ExternalObjectsPath : ExternalObjectsPaths)
+			if (PackageNameStringView.Contains(ExternalActorsPath))
 			{
-				if (It->GetLoadedPath().GetPackageName().Contains(ExternalObjectsPath))
-				{
-					ExternalObjectPackages.Add(*It);
-					break;
-				}
+				ExternalObjectPackages.Add(*It);
+				break;
 			}
 		}
 	}
