Fix variable rename from robomerge conflict

#rnx

diff --git a/Engine/Source/Runtime/Engine/Private/Level.cpp b/Engine/Source/Runtime/Engine/Private/Level.cpp
--- a/Engine/Source/Runtime/Engine/Private/Level.cpp
+++ b/Engine/Source/Runtime/Engine/Private/Level.cpp
@@ -3229,7 +3229,7 @@ TArray<UPackage*> ULevel::GetLoadedExternalObjectPackages() const
 		FStringView PackageNameStringView(PackageName);
 		for (const FString& ExternalObjectsPath : ExternalObjectsPaths)
 		{
-			if (PackageNameStringView.Contains(ExternalActorsPath))
+			if (PackageNameStringView.Contains(ExternalObjectsPath))
 			{
 				ExternalObjectPackages.Add(*It);
 				break;
