Fix the execution of construction scripts that might rely on asset such as static meshes not being properly compiled.

diff --git a/Engine/Source/Runtime/Engine/Classes/Engine/Level.h b/Engine/Source/Runtime/Engine/Classes/Engine/Level.h
--- a/Engine/Source/Runtime/Engine/Classes/Engine/Level.h
+++ b/Engine/Source/Runtime/Engine/Classes/Engine/Level.h
@@ -38,6 +38,7 @@ class FLevelPartitionOperationScope;
 class FRegisterComponentContext;
 class SNotificationItem;
 class UActorFolder;
+struct FAssetCompileData;
 
 UINTERFACE()
 class ULevelPartitionInterface : public UInterface
@@ -825,6 +826,9 @@ private:
 	/** Temporary array containing actor folder objects manually loaded from their external packages (only used while loading the level). */
 	UPROPERTY(Transient)
 	TArray<TObjectPtr<UActorFolder>> LoadedExternalActorFolders;
+
+	/** Actors awaiting outstanding asset compilation prior to running their non trivial construction scripts after loading this level. */
+	TArray<TWeakObjectPtr<AActor>> PendingConstructionScriptActors;
 #endif // #if WITH_EDITORONLY_DATA
 
 	enum class ERouteActorInitializationState : uint8
@@ -1311,6 +1315,18 @@ private:
 
 #if WITH_EDITOR
 private:
+	/**
+	 * Potentially defer the running of an actor's construction script on load
+	 * The running will be deferred if the actor's user construction script is non trivial 
+	 * and if there are outstanding asset compilation currently running (currently checks against static meshes)
+	 * @param InActor The actor to check
+	 * @return true if the running of the construction script was deferred
+	 */
+	bool DeferRunningConstructionScripts(AActor* InActor);
+
+	/** Callback for asset compilation to determine if we can run deferred construction scripts. */
+	void OnAssetPostCompile(const TArray<FAssetCompileData>&);
+
 	/** Attempts to detect and fix any issues with the level script blueprint and associated objects */
 	void RepairLevelScript();
 
diff --git a/Engine/Source/Runtime/Engine/Classes/GameFramework/Actor.h b/Engine/Source/Runtime/Engine/Classes/GameFramework/Actor.h
--- a/Engine/Source/Runtime/Engine/Classes/GameFramework/Actor.h
+++ b/Engine/Source/Runtime/Engine/Classes/GameFramework/Actor.h
@@ -2964,6 +2964,9 @@ public:
 	/** Reset private properties to defaults, and all FRandomStream structs in this Actor, so they will start their sequence of random numbers again. */
 	void ResetPropertiesForConstruction();
 
+	/** Returns true if the actor's class has a non trivial user construction script. */
+	bool HasNonTrivialUserConstructionScript() const;
+
 	/** Rerun construction scripts, destroying all autogenerated components; will attempt to preserve the root component location. */
 	virtual void RerunConstructionScripts();
 
diff --git a/Engine/Source/Runtime/Engine/Private/ActorConstruction.cpp b/Engine/Source/Runtime/Engine/Private/ActorConstruction.cpp
--- a/Engine/Source/Runtime/Engine/Private/ActorConstruction.cpp
+++ b/Engine/Source/Runtime/Engine/Private/ActorConstruction.cpp
@@ -220,6 +220,17 @@ void AActor::DestroyConstructedComponents()
 	}
 }
 
+
+bool AActor::HasNonTrivialUserConstructionScript() const
+{
+	UFunction* UCS = GetClass()->FindFunctionByName(FName(TEXT("UserConstructionScript"))/*UEdGraphSchema_K2::FN_UserConstructionScript*/);
+	if (UCS && UCS->Script.Num())
+	{
+		return true;
+	}
+	return false;
+}
+
 void AActor::RerunConstructionScripts()
 {
 	checkf(!HasAnyFlags(RF_ClassDefaultObject), TEXT("RerunConstructionScripts should never be called on a CDO as it can mutate the transient data on the CDO which then propagates to instances!"));
diff --git a/Engine/Source/Runtime/Engine/Private/Level.cpp b/Engine/Source/Runtime/Engine/Private/Level.cpp
--- a/Engine/Source/Runtime/Engine/Private/Level.cpp
+++ b/Engine/Source/Runtime/Engine/Private/Level.cpp
@@ -53,6 +53,8 @@ Level.cpp: Level-related functions
 #include "UnrealEngine.h"
 #include "Misc/ArchiveMD5.h"
 #if WITH_EDITOR
+#include "AssetCompilingManager.h"
+#include "StaticMeshCompiler.h"
 #include "Components/InstancedStaticMeshComponent.h"
 #include "Kismet2/KismetEditorUtilities.h"
 #include "Kismet2/BlueprintEditorUtils.h"
@@ -941,6 +943,10 @@ void ULevel::PostLoad()
 			}
 		}
 	}
+
+	// Register to the asset async compilation delegate to manage the deferring of non trivial construction scripts on actors
+	FAssetCompilingManager::Get().OnAssetPostCompileEvent().AddUObject(this, &ULevel::OnAssetPostCompile);
+
 #endif
 
 	// Ensure that the level is pointed to the owning world.  For streamed levels, this will be the world of the P map
@@ -1483,6 +1489,13 @@ bool ULevel::IncrementalRunConstructionScripts(bool bProcessAllActors)
 			continue;
 		}
 
+#if WITH_EDITOR
+		if (DeferRunningConstructionScripts(Actor))
+		{
+			continue;
+		}
+#endif
+
 #if PERF_TRACK_DETAILED_ASYNC_STATS
 		FScopeCycleCounterUObject ContextScope(Actor);
 #endif
@@ -1573,6 +1586,42 @@ void ULevel::MarkLevelComponentsRenderStateDirty()
 
 #if WITH_EDITOR
 
+bool ULevel::DeferRunningConstructionScripts(AActor* InActor)
+{
+	// if there are outstanding asset (static meshes) compilation when loading actors
+	// Defer the running of non trivial construction scripts
+	// This is done to ensure that user construction scripts that runs line trace
+	// Would not get an improper result from inflight static mesh compilation for example.
+	if (FStaticMeshCompilingManager::Get().GetNumRemainingMeshes() &&
+		InActor->HasNonTrivialUserConstructionScript())
+	{
+		PendingConstructionScriptActors.Add(InActor);
+		return true;
+	}
+	return false;
+}
+
+void ULevel::OnAssetPostCompile(const TArray<FAssetCompileData>&)
+{
+	// Once there are no more outstanding asset (static meshes), run actors non trivial construction scripts
+	if (FStaticMeshCompilingManager::Get().GetNumRemainingMeshes() == 0)
+	{
+		// since this deferred run of construction script was supposed to be done during level load
+		// temporarily set the global flag to prevent dirtying the level package.
+		// @note: it would quite preferable if we would have a scoped context than touching a global variable... 
+		GIsEditorLoadingPackage = true;
+		for (TWeakObjectPtr<AActor> WeakActor : PendingConstructionScriptActors)
+		{
+			if (AActor* Actor = WeakActor.Get())
+			{
+				Actor->RerunConstructionScripts();
+			}
+		}
+		GIsEditorLoadingPackage = false;
+		PendingConstructionScriptActors.Empty();
+	}
+}
+
 void ULevel::CreateModelComponents()
 {
 	TRACE_CPUPROFILER_EVENT_SCOPE(ULevel::CreateModelComponents)
