Fix initialization of new FBodyInstance properties

diff --git a/Engine/Source/Runtime/Engine/Private/PhysicsEngine/BodyInstance.cpp b/Engine/Source/Runtime/Engine/Private/PhysicsEngine/BodyInstance.cpp
--- a/Engine/Source/Runtime/Engine/Private/PhysicsEngine/BodyInstance.cpp
+++ b/Engine/Source/Runtime/Engine/Private/PhysicsEngine/BodyInstance.cpp
@@ -337,6 +337,8 @@ FBodyInstance::FBodyInstance()
 	, bUseCCD(false)
 	, bIgnoreAnalyticCollisions(false)
 	, bNotifyRigidBodyCollision(false)
+	, bContactModification(false)
+	, bSmoothEdgeCollisions(false)
 	, bLockTranslation(true)
 	, bLockRotation(true)
 	, bLockXTranslation(false)
