Enable GPU-Scene support for SplineMesh

diff --git a/Engine/Shaders/Public/ShaderVersion.ush b/Engine/Shaders/Public/ShaderVersion.ush
--- a/Engine/Shaders/Public/ShaderVersion.ush
+++ b/Engine/Shaders/Public/ShaderVersion.ush
@@ -2,5 +2,5 @@
 // Each time the console command is executed it generates a new GUID. As this file is included
 // in Platform.ush (which should be included in any shader) it allows to invalidate the shader DDC.
 // 
-// GUID = B5734ACE-E057-480F-8E1E-67E08B1AB26D
+// GUID = 03177423-197C-410E-8F15-CC6C2F7304E6
 
diff --git a/Engine/Source/Runtime/Engine/Private/Components/SplineMeshComponent.cpp b/Engine/Source/Runtime/Engine/Private/Components/SplineMeshComponent.cpp
--- a/Engine/Source/Runtime/Engine/Private/Components/SplineMeshComponent.cpp
+++ b/Engine/Source/Runtime/Engine/Private/Components/SplineMeshComponent.cpp
@@ -58,9 +58,9 @@ void FSplineMeshVertexFactoryShaderParameters::GetElementShaderBindings(
 	FVertexInputStreamArray& VertexStreams
 ) const
 {
+	const auto* LocalVertexFactory = static_cast<const FLocalVertexFactory*>(VertexFactory);
 	if (BatchElement.bUserDataIsColorVertexBuffer)
 	{
-		const auto* LocalVertexFactory = static_cast<const FLocalVertexFactory*>(VertexFactory);
 		FColorVertexBuffer* OverrideColorVertexBuffer = (FColorVertexBuffer*)BatchElement.UserData;
 		check(OverrideColorVertexBuffer);
 
@@ -69,10 +69,14 @@ void FSplineMeshVertexFactoryShaderParameters::GetElementShaderBindings(
 			LocalVertexFactory->GetColorOverrideStream(OverrideColorVertexBuffer, VertexStreams);
 		}
 	}
+	if (LocalVertexFactory->SupportsManualVertexFetch(FeatureLevel) || UseGPUScene(GMaxRHIShaderPlatform, FeatureLevel))
+	{
+		ShaderBindings.Add(Shader->GetUniformBufferParameter<FLocalVertexFactoryUniformShaderParameters>(), LocalVertexFactory->GetUniformBuffer());
+	}
 
-		checkSlow(BatchElement.bIsSplineProxy);
-		FSplineMeshSceneProxy* SplineProxy = BatchElement.SplineMeshSceneProxy;
-		FSplineMeshParams& SplineParams = SplineProxy->SplineParams;
+	checkSlow(BatchElement.bIsSplineProxy);
+	FSplineMeshSceneProxy* SplineProxy = BatchElement.SplineMeshSceneProxy;
+	FSplineMeshParams& SplineParams = SplineProxy->SplineParams;
 
 	FVector4f ParamData[10];
 	// LWC_TODO: Precision loss
@@ -109,6 +113,7 @@ IMPLEMENT_VERTEX_FACTORY_TYPE(FSplineMeshVertexFactory, "/Engine/Private/LocalVe
 	| EVertexFactoryFlags::SupportsDynamicLighting
 	| EVertexFactoryFlags::SupportsPrecisePrevWorldPos
 	| EVertexFactoryFlags::SupportsPositionOnly
+	| EVertexFactoryFlags::SupportsPrimitiveIdStream
 );
 
 //////////////////////////////////////////////////////////////////////////
diff --git a/Engine/Source/Runtime/Engine/Private/LocalVertexFactory.cpp b/Engine/Source/Runtime/Engine/Private/LocalVertexFactory.cpp
--- a/Engine/Source/Runtime/Engine/Private/LocalVertexFactory.cpp
+++ b/Engine/Source/Runtime/Engine/Private/LocalVertexFactory.cpp
@@ -200,7 +200,11 @@ bool FLocalVertexFactory::ShouldCompilePermutation(const FVertexFactoryShaderPer
 
 void FLocalVertexFactory::ModifyCompilationEnvironment(const FVertexFactoryShaderPermutationParameters& Parameters, FShaderCompilerEnvironment& OutEnvironment)
 {
-	OutEnvironment.SetDefine(TEXT("VF_SUPPORTS_SPEEDTREE_WIND"),TEXT("1"));
+	// Don't override e.g. SplineMesh's opt-out
+	if (!OutEnvironment.GetDefinitions().Contains("VF_SUPPORTS_SPEEDTREE_WIND"))
+	{
+		OutEnvironment.SetDefine(TEXT("VF_SUPPORTS_SPEEDTREE_WIND"), TEXT("1"));
+	}
 
 	const bool ContainsManualVertexFetch = OutEnvironment.GetDefinitions().Contains("MANUAL_VERTEX_FETCH");
 	if (!ContainsManualVertexFetch && RHISupportsManualVertexFetch(Parameters.Platform))
diff --git a/Engine/Source/Runtime/Engine/Private/SplineMeshSceneProxy.cpp b/Engine/Source/Runtime/Engine/Private/SplineMeshSceneProxy.cpp
--- a/Engine/Source/Runtime/Engine/Private/SplineMeshSceneProxy.cpp
+++ b/Engine/Source/Runtime/Engine/Private/SplineMeshSceneProxy.cpp
@@ -21,8 +21,8 @@ void FSplineMeshVertexFactory::ModifyCompilationEnvironment(const FVertexFactory
 		OutEnvironment.SetDefine(TEXT("MANUAL_VERTEX_FETCH"), TEXT("0"));
 	}
 
-	// We don't call this because we don't actually support speed tree wind, and this advertises support for that
-	//FLocalVertexFactory::ModifyCompilationEnvironment(Type, Platform, Material, OutEnvironment);
+	OutEnvironment.SetDefine(TEXT("VF_SUPPORTS_SPEEDTREE_WIND"), TEXT("0"));
+	FLocalVertexFactory::ModifyCompilationEnvironment(Parameters, OutEnvironment);
 
 	OutEnvironment.SetDefine(TEXT("USE_SPLINEDEFORM"), TEXT("1"));
 }
@@ -32,8 +32,6 @@ FSplineMeshSceneProxy::FSplineMeshSceneProxy(USplineMeshComponent* InComponent)
 {
 	bSupportsDistanceFieldRepresentation = false;
 	bSupportsMeshCardRepresentation = false;
-	bVFRequiresPrimitiveUniformBuffer = true;
-	bSupportsGPUScene = false;
 
 	// make sure all the materials are okay to be rendered as a spline mesh
 	for (FStaticMeshSceneProxy::FLODInfo& LODInfo : LODs)
