Rename Log 'Level Instance Log' to 'Packed Level Actor Log'

diff --git a/Engine/Source/Editor/LevelInstanceEditor/Private/LevelInstanceEditorModule.cpp b/Engine/Source/Editor/LevelInstanceEditor/Private/LevelInstanceEditorModule.cpp
--- a/Engine/Source/Editor/LevelInstanceEditor/Private/LevelInstanceEditorModule.cpp
+++ b/Engine/Source/Editor/LevelInstanceEditor/Private/LevelInstanceEditorModule.cpp
@@ -572,7 +572,7 @@ void FLevelInstanceEditorModule::StartupModule()
 	InitOptions.bShowFilters = true;
 	InitOptions.bShowPages = false;
 	InitOptions.bAllowClear = true;
-	MessageLogModule.RegisterLogListing("LevelInstance", LOCTEXT("LevelInstanceLog", "Level Instance Log"), InitOptions);
+	MessageLogModule.RegisterLogListing("PackedLevelActor", LOCTEXT("PackedLevelActorLog", "Packed Level Actor Log"), InitOptions);
 		
 	FLevelInstanceEditorModeCommands::Register();
 
diff --git a/Engine/Source/Runtime/Engine/Private/PackedLevelActor/PackedLevelActorBuilder.cpp b/Engine/Source/Runtime/Engine/Private/PackedLevelActor/PackedLevelActorBuilder.cpp
--- a/Engine/Source/Runtime/Engine/Private/PackedLevelActor/PackedLevelActorBuilder.cpp
+++ b/Engine/Source/Runtime/Engine/Private/PackedLevelActor/PackedLevelActorBuilder.cpp
@@ -164,8 +164,8 @@ bool FPackedLevelActorBuilder::PackActor(APackedLevelActor* InPackedLevelActor)
 
 bool FPackedLevelActorBuilder::PackActor(APackedLevelActor* InPackedLevelActor, ALevelInstance* InLevelInstanceToPack)
 {
-	FMessageLog LevelInstanceLog("LevelInstance");
-	LevelInstanceLog.Info(FText::Format(LOCTEXT("PackingStarted", "Packing of '{0}' started..."), FText::FromString(InPackedLevelActor->GetWorldAssetPackage())));
+	FMessageLog Log("PackedLevelActor");
+	Log.Info(FText::Format(LOCTEXT("PackingStarted", "Packing of '{0}' started..."), FText::FromString(InPackedLevelActor->GetWorldAssetPackage())));
 	
 	FPackedLevelActorBuilderContext Context(*this, InPackedLevelActor);
 
@@ -177,7 +177,7 @@ bool FPackedLevelActorBuilder::PackActor(APackedLevelActor* InPackedLevelActor,
 	ULevel* SourceLevel = LevelInstanceSubystem->GetLevelInstanceLevel(InLevelInstanceToPack);
 	if (!SourceLevel)
 	{
-		LevelInstanceLog.Error(FText::Format(LOCTEXT("FailedPackingNoLevel", "Packing of '{0}' failed"), FText::FromString(InPackedLevelActor->GetWorldAssetPackage())));
+		Log.Error(FText::Format(LOCTEXT("FailedPackingNoLevel", "Packing of '{0}' failed"), FText::FromString(InPackedLevelActor->GetWorldAssetPackage())));
 		return false;
 	}
 
@@ -211,7 +211,7 @@ bool FPackedLevelActorBuilder::PackActor(APackedLevelActor* InPackedLevelActor,
 		Builder->PackActors(Context, InPackedLevelActor, Pair.Key, Pair.Value);
 	}
 			
-	Context.Report(LevelInstanceLog);
+	Context.Report(Log);
 	return true;
 }
 
@@ -220,7 +220,7 @@ bool FPackedLevelActorBuilderContext::ShouldPackComponent(UActorComponent* Actor
 	return ActorComponent && !ActorComponent->IsVisualizationComponent();
 }
 
-void FPackedLevelActorBuilderContext::Report(FMessageLog& LevelInstanceLog) const
+void FPackedLevelActorBuilderContext::Report(FMessageLog& Log) const
 {
 	TSet<UActorComponent*> NotClusteredComponents;
 	uint32 TotalWarningCount = 0;
@@ -232,20 +232,20 @@ void FPackedLevelActorBuilderContext::Report(FMessageLog& LevelInstanceLog) cons
 
 		if (ActorDiscards.Contains(Actor))
 		{
-			LevelInstanceLog.Info(FText::Format(LOCTEXT("ActorDiscard", "Actor '{0}' ignored (Actor Discard)"), FText::FromString(Actor->GetPathName())));
+			Log.Info(FText::Format(LOCTEXT("ActorDiscard", "Actor '{0}' ignored (Actor Discard)"), FText::FromString(Actor->GetPathName())));
 			continue;
 		}
 
 		if (Actor->GetClass()->HasAnyClassFlags(CLASS_Transient))
 		{
-			LevelInstanceLog.Info(FText::Format(LOCTEXT("ActorTransientClassDiscard", "Actor '{0}' of type '{1}' ignored (Transient Class Discard)"), FText::FromString(Actor->GetPathName()), FText::FromString(Actor->GetClass()->GetPathName())));
+			Log.Info(FText::Format(LOCTEXT("ActorTransientClassDiscard", "Actor '{0}' of type '{1}' ignored (Transient Class Discard)"), FText::FromString(Actor->GetPathName()), FText::FromString(Actor->GetClass()->GetPathName())));
 			continue;
 		}
 
 		// Class must match (not a child)
 		if (ClassDiscards.Contains(Actor->GetClass()))
 		{
-			LevelInstanceLog.Info(FText::Format(LOCTEXT("ActorClassDiscard", "Actor '{0}' of type '{1}' ignored (Class Discard)"), FText::FromString(Actor->GetPathName()), FText::FromString(Actor->GetClass()->GetPathName())));
+			Log.Info(FText::Format(LOCTEXT("ActorClassDiscard", "Actor '{0}' of type '{1}' ignored (Class Discard)"), FText::FromString(Actor->GetPathName()), FText::FromString(Actor->GetClass()->GetPathName())));
 			continue;
 		}
 
@@ -260,38 +260,38 @@ void FPackedLevelActorBuilderContext::Report(FMessageLog& LevelInstanceLog) cons
 
 			if (ClassDiscards.Contains(Component->GetClass()))
 			{
-				LevelInstanceLog.Info(FText::Format(LOCTEXT("ComponentClassDiscard", "Component '{0}' of type '{1}' ignored (Class Discard)"), FText::FromString(Component->GetPathName()), FText::FromString(Component->GetClass()->GetPathName())));
+				Log.Info(FText::Format(LOCTEXT("ComponentClassDiscard", "Component '{0}' of type '{1}' ignored (Class Discard)"), FText::FromString(Component->GetPathName()), FText::FromString(Component->GetClass()->GetPathName())));
 				continue;
 			}
 
 			if (Actor->GetClass()->HasAnyClassFlags(CLASS_Transient))
 			{
-				LevelInstanceLog.Info(FText::Format(LOCTEXT("ComopnentTransientClassDiscard", "Component '{0}' of type '{1}' ignored (Transient Class Discard)"), FText::FromString(Component->GetPathName()), FText::FromString(Component->GetClass()->GetPathName())));
+				Log.Info(FText::Format(LOCTEXT("ComopnentTransientClassDiscard", "Component '{0}' of type '{1}' ignored (Transient Class Discard)"), FText::FromString(Component->GetPathName()), FText::FromString(Component->GetClass()->GetPathName())));
 				continue;
 			}
 
 			WarningCount++;
 
-			LevelInstanceLog.Warning(FText::Format(LOCTEXT("ComponentNotPacked", "Component '{0}' was not packed"), FText::FromString(Component->GetPathName())));
+			Log.Warning(FText::Format(LOCTEXT("ComponentNotPacked", "Component '{0}' was not packed"), FText::FromString(Component->GetPathName())));
 		}
 
 		if (WarningCount)
 		{
-			LevelInstanceLog.Warning(FText::Format(LOCTEXT("ActorNotPacked", "Actor '{0}' was not packed completely ({1} warning(s))"), FText::FromString(Actor->GetPathName()), FText::AsNumber(WarningCount)));
+			Log.Warning(FText::Format(LOCTEXT("ActorNotPacked", "Actor '{0}' was not packed completely ({1} warning(s))"), FText::FromString(Actor->GetPathName()), FText::AsNumber(WarningCount)));
 		}
 		else
 		{
-			LevelInstanceLog.Info(FText::Format(LOCTEXT("ActorPacked", "Actor '{0}' packed successfully"), FText::FromString(Actor->GetPathName())));
+			Log.Info(FText::Format(LOCTEXT("ActorPacked", "Actor '{0}' packed successfully"), FText::FromString(Actor->GetPathName())));
 		}
 		TotalWarningCount += WarningCount;
 	}
 
 	if (TotalWarningCount)
 	{
-		LevelInstanceLog.Warning(LOCTEXT("WarningsReported", "Warnings have been reported. Consider using a regular ALevelInstance instead."));
-		LevelInstanceLog.Open();
+		Log.Warning(LOCTEXT("WarningsReported", "Warnings have been reported. Consider using a regular ALevelInstance instead."));
+		Log.Open();
 	}
-	LevelInstanceLog.Info(FText::Format(LOCTEXT("PackCompleted", "Packing '{0}' completed with {1} warning(s)"), FText::FromString(PackedLevelActor->GetWorldAssetPackage()), FText::AsNumber(TotalWarningCount)));
+	Log.Info(FText::Format(LOCTEXT("PackCompleted", "Packing '{0}' completed with {1} warning(s)"), FText::FromString(PackedLevelActor->GetWorldAssetPackage()), FText::AsNumber(TotalWarningCount)));
 }
 
 ALevelInstance* FPackedLevelActorBuilder::CreateTransientLevelInstanceForPacking(TSoftObjectPtr<UWorld> InWorldAsset, const FVector& InLocation, const FRotator& InRotator)
diff --git a/Engine/Source/Runtime/Engine/Private/PackedLevelActor/PackedLevelActorBuilder.h b/Engine/Source/Runtime/Engine/Private/PackedLevelActor/PackedLevelActorBuilder.h
--- a/Engine/Source/Runtime/Engine/Private/PackedLevelActor/PackedLevelActorBuilder.h
+++ b/Engine/Source/Runtime/Engine/Private/PackedLevelActor/PackedLevelActorBuilder.h
@@ -79,7 +79,7 @@ public:
 	void ClusterLevelActor(AActor* InLevelActor);
 	void FindOrAddCluster(FPackedLevelActorBuilderClusterID&& InClusterID, UActorComponent* InComponent = nullptr);
 	void DiscardActor(AActor* InActor);
-	void Report(FMessageLog& LevelInstanceLog) const;
+	void Report(FMessageLog& Log) const;
 
 	const TMap<FPackedLevelActorBuilderClusterID, TArray<UActorComponent*>>& GetClusters() const { return Clusters; }
 	
