-Make ReflectionCaptures work in WorldPartition maps
-They'll be part of the PersistentLevel, so this is a quickfix until they're fully supported as data loaded from the partition cells.

diff --git a/Engine/Source/Runtime/Engine/Private/Components/ReflectionCaptureComponent.cpp b/Engine/Source/Runtime/Engine/Private/Components/ReflectionCaptureComponent.cpp
--- a/Engine/Source/Runtime/Engine/Private/Components/ReflectionCaptureComponent.cpp
+++ b/Engine/Source/Runtime/Engine/Private/Components/ReflectionCaptureComponent.cpp
@@ -109,6 +109,13 @@ FReflectionCaptureMapBuildData* UReflectionCaptureComponent::GetMapBuildData() c
 			else if (OwnerLevel->MapBuildData)
 			{
 				MapBuildData = OwnerLevel->MapBuildData;
+			}			
+			else if (OwnerLevel->OwningWorld->IsPartitionedWorld())
+			{
+				// Fallback to PersistentLevel in case of missing data in OwnerLevel
+				// This is quite likely for WP maps in PIE/engine until the ReflectionsCaptures
+				// are updated to distribute themselves in the cells (planned for 5.1)
+				MapBuildData = OwnerLevel->OwningWorld->PersistentLevel->MapBuildData;
 			}
 			 
 			if (MapBuildData)
