fix reentrant locking in FSmartNameContainer::Iterate

diff --git a/Engine/Plugins/Animation/ControlRig/Source/ControlRig/Private/AnimNode_ControlRigBase.cpp b/Engine/Plugins/Animation/ControlRig/Source/ControlRig/Private/AnimNode_ControlRigBase.cpp
--- a/Engine/Plugins/Animation/ControlRig/Source/ControlRig/Private/AnimNode_ControlRigBase.cpp
+++ b/Engine/Plugins/Animation/ControlRig/Source/ControlRig/Private/AnimNode_ControlRigBase.cpp
@@ -651,15 +651,15 @@ void FAnimNode_ControlRigBase::CacheBones_AnyThread(const FAnimationCacheBonesCo
 			{
 				const FSmartNameMapping* CurveMapping = RequiredBones.GetSkeletonAsset()->GetSmartNameContainer(USkeleton::AnimCurveMappingName);
 				URigHierarchy* Hierarchy = ControlRig->GetHierarchy();
-				CurveMapping->Iterate([Hierarchy, this](const FSmartNameMapping* Mapping, SmartName::UID_Type ID)
+				CurveMapping->Iterate([Hierarchy, this](const FSmartNameMappingIterator& Iterator)
 				{
 					// see if the curve name exists in the control rig
 					FName CurveName;
-					if (Mapping->GetName(ID, CurveName))
+					if (Iterator.GetName(CurveName))
 					{
 						if (Hierarchy->GetIndex(FRigElementKey(CurveName, ERigElementType::Curve)) != INDEX_NONE)
 						{
-							ControlRigCurveMappingByName.Add(CurveName, ID);
+							ControlRigCurveMappingByName.Add(CurveName, Iterator.GetIndex());
 						}
 					}
 				});
diff --git a/Engine/Source/Runtime/Engine/Classes/Animation/SmartName.h b/Engine/Source/Runtime/Engine/Classes/Animation/SmartName.h
--- a/Engine/Source/Runtime/Engine/Classes/Animation/SmartName.h
+++ b/Engine/Source/Runtime/Engine/Classes/Animation/SmartName.h
@@ -59,6 +59,8 @@ struct FCurveMetaData
 USTRUCT()
 struct ENGINE_API FSmartNameMapping
 {
+	friend struct FSmartNameMappingIterator;
+	
 	GENERATED_USTRUCT_BODY();
 
 	FSmartNameMapping();
@@ -150,7 +152,7 @@ struct ENGINE_API FSmartNameMapping
 	SmartName::UID_Type GetMaxUID() const { return (SmartName::UID_Type)(CurveNameList.Num() - 1); }
 
 	/** Iterate over all Names in this Mapping */
-	void Iterate(TFunction<void(const FSmartNameMapping* Mapping, SmartName::UID_Type ID)> Callback) const;
+	void Iterate(TFunction<void(const struct FSmartNameMappingIterator& Iterator)> Callback) const;
 
 private:
 	/*Internal no lock function to prevent re-entrant locking, see API function GetName for documentation.*/
@@ -164,6 +166,9 @@ private:
 
 	/*Internal no lock function to prevent re-entrant locking, see API function FindUID for documentation.*/
 	SmartName::UID_Type FindUID_NoLock(const FName& Name) const;
+	
+	/*Internal no lock function to prevent re-entrant locking, see API function GetCurveMetaData for documentation.*/
+	const FCurveMetaData* GetCurveMetaData_NoLock(FName CurveName) const;
 
 	// List of curve names, indexed by UID
 	TArray<FName> CurveNameList;
@@ -176,6 +181,42 @@ private:
 	TMap<FName, FCurveMetaData> CurveMetaDataMap;
 };
 
+// Struct for providing access to SmartNameMapping data within FSmartNameMapping::Iterate callback functions
+struct ENGINE_API FSmartNameMappingIterator
+{
+	public:
+		friend struct FSmartNameMapping;
+	
+		bool GetName(FName& OutCurveName) const
+		{
+			return Mapping->GetName_NoLock(Index, OutCurveName);
+		}
+	
+		const FCurveMetaData* GetCurveMetaData() const
+		{
+			FName Name;
+			if (Mapping->GetName_NoLock(Index, Name))
+			{
+				return Mapping->GetCurveMetaData_NoLock(Name);
+			}
+			else
+			{
+				return nullptr;
+			}
+		}
+	
+		SmartName::UID_Type GetIndex() const { return Index; }
+	
+	private:
+		// This class struct should only be crated by FSmartNameMapping::Iterate
+		FSmartNameMappingIterator(const FSmartNameMapping* InMapping, SmartName::UID_Type InIndex):
+			Mapping(InMapping), Index(InIndex)
+		{}
+	
+		const FSmartNameMapping* Mapping;
+		SmartName::UID_Type Index;
+};
+
 USTRUCT()
 struct ENGINE_API FSmartNameContainer
 {
diff --git a/Engine/Source/Runtime/Engine/Private/Animation/SmartName.cpp b/Engine/Source/Runtime/Engine/Private/Animation/SmartName.cpp
--- a/Engine/Source/Runtime/Engine/Private/Animation/SmartName.cpp
+++ b/Engine/Source/Runtime/Engine/Private/Animation/SmartName.cpp
@@ -20,13 +20,14 @@ FSmartNameMapping::FSmartNameMapping()
 {
 }
 
-void FSmartNameMapping::Iterate(TFunction<void(const FSmartNameMapping* Mapping, SmartName::UID_Type ID)> Callback) const
+void FSmartNameMapping::Iterate(TFunction<void(const FSmartNameMappingIterator& Iterator)> Callback) const
 {
 	FReadScopeLock Lock(SmartNameRWLock);
 
 	for (int32 NameIndex = 0; NameIndex < CurveNameList.Num(); ++NameIndex)
 	{
-		Callback(this, NameIndex);	
+		FSmartNameMappingIterator Iterator(this, NameIndex);
+		Callback(Iterator);	
 	}
 }
 
@@ -58,6 +59,11 @@ FCurveMetaData* FSmartNameMapping::GetCurveMetaData(FName CurveName)
 const FCurveMetaData* FSmartNameMapping::GetCurveMetaData(FName CurveName) const
 {
 	FReadScopeLock Lock(SmartNameRWLock);
+	return GetCurveMetaData_NoLock(CurveName);
+}
+
+const FCurveMetaData* FSmartNameMapping::GetCurveMetaData_NoLock(FName CurveName) const
+{
 	checkSlow(Exists_NoLock(CurveName));
 	return CurveMetaDataMap.Find(CurveName);
 }
diff --git a/Engine/Source/Runtime/Engine/Private/BoneContainer.cpp b/Engine/Source/Runtime/Engine/Private/BoneContainer.cpp
--- a/Engine/Source/Runtime/Engine/Private/BoneContainer.cpp
+++ b/Engine/Source/Runtime/Engine/Private/BoneContainer.cpp
@@ -276,10 +276,10 @@ void FBoneContainer::CacheRequiredAnimCurveUids(const FCurveEvaluationOption& Cu
 
 			// if the linked joints don't exists in RequiredBones, remove itself
 			int32 NumAvailableUIDs = 0;
-			Mapping->Iterate([this, &CurveEvalOption, &NumAvailableUIDs](const FSmartNameMapping* Mapping, SmartName::UID_Type CurveNameIndex)
+			Mapping->Iterate([this, &CurveEvalOption, &NumAvailableUIDs](const FSmartNameMappingIterator& Iterator)
 			{
 				FName CurveName;
-				if (Mapping->GetName(CurveNameIndex, CurveName))
+				if (Iterator.GetName(CurveName))
 				{
 					bool bBeingUsed = true;
 					if (!CurveEvalOption.bAllowCurveEvaluation)
@@ -300,7 +300,7 @@ void FBoneContainer::CacheRequiredAnimCurveUids(const FCurveEvaluationOption& Cu
 						}
 						else
 						{
-							const FCurveMetaData* CurveMetaData = Mapping->GetCurveMetaData(CurveName);
+							const FCurveMetaData* CurveMetaData = Iterator.GetCurveMetaData();
 							if (CurveMetaData)
 							{
 								if (CurveMetaData->MaxLOD < CurveEvalOption.LODIndex)
@@ -334,7 +334,7 @@ void FBoneContainer::CacheRequiredAnimCurveUids(const FCurveEvaluationOption& Cu
 
 					if (bBeingUsed)
 					{
-						UIDToArrayIndexLUT[CurveNameIndex] = NumAvailableUIDs++;
+						UIDToArrayIndexLUT[Iterator.GetIndex()] = NumAvailableUIDs++;
 					}
 				}
 				UIDToArrayIndexLUTValidCount = NumAvailableUIDs;
