Network Physics: advanced leash mode
-Start of improvements to physics replication that don't rely on rollback networking

#rnx

diff --git a/Engine/Source/Runtime/Engine/Classes/Components/PrimitiveComponent.h b/Engine/Source/Runtime/Engine/Classes/Components/PrimitiveComponent.h
--- a/Engine/Source/Runtime/Engine/Classes/Components/PrimitiveComponent.h
+++ b/Engine/Source/Runtime/Engine/Classes/Components/PrimitiveComponent.h
@@ -2474,7 +2474,7 @@ public:
 	 * Applies RigidBodyState only if it needs to be updated
 	 * NeedsUpdate flag will be removed from UpdatedState after all velocity corrections are finished
 	 */
-	void SetRigidBodyReplicatedTarget(FRigidBodyState& UpdatedState, const FName BoneName = NAME_None);
+	void SetRigidBodyReplicatedTarget(FRigidBodyState& UpdatedState, const FName BoneName = NAME_None, int32 ServerFrame = 0, int32 ServerHandle = 0);
 
 	/** 
 	 *	Get the state of the rigid body responsible for this Actor's physics, and fill in the supplied FRigidBodyState struct based on it.
diff --git a/Engine/Source/Runtime/Engine/Classes/Engine/EngineTypes.h b/Engine/Source/Runtime/Engine/Classes/Engine/EngineTypes.h
--- a/Engine/Source/Runtime/Engine/Classes/Engine/EngineTypes.h
+++ b/Engine/Source/Runtime/Engine/Classes/Engine/EngineTypes.h
@@ -3507,6 +3507,14 @@ struct ENGINE_API FRepMovement
 	UPROPERTY(Transient)
 	uint8 bRepPhysics : 1;
 
+	/** Server physics step */
+	UPROPERTY(Transient)
+	int32 ServerFrame;
+
+	/** ID assigned by server used to ensure determinism by physics. */
+	UPROPERTY(Transient)
+	int32 ServerPhysicsHandle = INDEX_NONE;
+
 	/** Allows tuning the compression level for the replicated location vector. You should only need to change this from the default if you see visual artifacts. */
 	UPROPERTY(EditDefaultsOnly, Category=Replication, AdvancedDisplay)
 	EVectorQuantization LocationQuantizationLevel;
@@ -3551,10 +3559,12 @@ struct ENGINE_API FRepMovement
 	bool NetSerialize(FArchive& Ar, class UPackageMap* Map, bool& bOutSuccess)
 	{
 		// pack bitfield with flags
-		uint8 Flags = (bSimulatedPhysicSleep << 0) | (bRepPhysics << 1);
-		Ar.SerializeBits(&Flags, 2);
+		uint8 Flags = (bSimulatedPhysicSleep << 0) | (bRepPhysics << 1) | ((ServerFrame > 0) << 2) | ((ServerPhysicsHandle != INDEX_NONE) << 3);
+		Ar.SerializeBits(&Flags, 4);
 		bSimulatedPhysicSleep = ( Flags & ( 1 << 0 ) ) ? 1 : 0;
 		bRepPhysics = ( Flags & ( 1 << 1 ) ) ? 1 : 0;
+		const bool bRepServerFrame = (Flags & (1 << 2)) ? 1 : 0;
+		const bool bRepServerHandle = (Flags & (1 << 3)) ? 1 : 0;
 
 		bOutSuccess = true;
 
@@ -3584,10 +3594,24 @@ struct ENGINE_API FRepMovement
 			bOutSuccess &= SerializeQuantizedVector( Ar, AngularVelocity, VelocityQuantizationLevel );
 		}
 
+		if (bRepServerFrame)
+		{
+			uint32 uServerFrame = (uint32)ServerFrame;
+			Ar.SerializeIntPacked(uServerFrame);
+			ServerFrame = (int32)uServerFrame;
+		}
+
+		if (bRepServerHandle)
+		{
+			uint32 uServerPhysicsHandle = (uint32)ServerPhysicsHandle;
+			Ar.SerializeIntPacked(uServerPhysicsHandle);
+			ServerPhysicsHandle = (int32)uServerPhysicsHandle;
+		}
+
 		return true;
 	}
 
-	void FillFrom(const struct FRigidBodyState& RBState, const AActor* const Actor = nullptr)
+	void FillFrom(const struct FRigidBodyState& RBState, const AActor* const Actor = nullptr, int32 InServerFrame = 0)
 	{
 		Location = RebaseOntoZeroOrigin(RBState.Position, Actor);
 		Rotation = RBState.Quaternion.Rotator();
@@ -3595,6 +3619,7 @@ struct ENGINE_API FRepMovement
 		AngularVelocity = RBState.AngVel;
 		bSimulatedPhysicSleep = (RBState.Flags & ERigidBodyFlags::Sleeping) != 0;
 		bRepPhysics = true;
+		ServerFrame = InServerFrame;
 	}
 
 	void CopyTo(struct FRigidBodyState& RBState, const AActor* const Actor = nullptr) const
diff --git a/Engine/Source/Runtime/Engine/Classes/GameFramework/PlayerController.h b/Engine/Source/Runtime/Engine/Classes/GameFramework/PlayerController.h
--- a/Engine/Source/Runtime/Engine/Classes/GameFramework/PlayerController.h
+++ b/Engine/Source/Runtime/Engine/Classes/GameFramework/PlayerController.h
@@ -2118,6 +2118,8 @@ public:
 
 		int8 QuantizedTimeDilation = 1; // Server sent this to this client, telling them to dilate local time either catch up or slow down
 		float TargetNumBufferedCmds = 0.f;
+
+		int32 GetLocalFrameOffset() const { return LastProcessedInputFrame - LastRecvServerFrame; }
 	};	
 
 	// Client pushes input data locally. RPC is sent here but also includes redundant data
@@ -2128,6 +2130,7 @@ public:
 	void ServerRecvClientInputFrame(int32 RecvClientInputFrame, const TArray<uint8>& Data);
 
 	const FClientFrameInfo& GetClientFrameInfo() const { return ClientFrameInfo; }
+	FClientFrameInfo& GetClientFrameInfo() { return ClientFrameInfo; }
 
 	// -------------------------------------------------------------------------
 	// Server
diff --git a/Engine/Source/Runtime/Engine/Private/ActorReplication.cpp b/Engine/Source/Runtime/Engine/Private/ActorReplication.cpp
--- a/Engine/Source/Runtime/Engine/Private/ActorReplication.cpp
+++ b/Engine/Source/Runtime/Engine/Private/ActorReplication.cpp
@@ -21,6 +21,12 @@
 #include "Net/Core/PushModel/PushModel.h"
 #include "Interfaces/Interface_ActorSubobject.h"
 
+#if WITH_CHAOS
+#include "Physics/Experimental/PhysScene_Chaos.h"
+#include "PBDRigidsSolver.h"
+#include "ChaosSolversModule.h"
+#endif
+
 /*-----------------------------------------------------------------------------
 	AActor networking implementation.
 -----------------------------------------------------------------------------*/
@@ -259,11 +265,12 @@ void AActor::PostNetReceivePhysicState()
 	UPrimitiveComponent* RootPrimComp = Cast<UPrimitiveComponent>(RootComponent);
 	if (RootPrimComp)
 	{
+		const FRepMovement& ThisReplicatedMovement = GetReplicatedMovement();
 		FRigidBodyState NewState;
-		GetReplicatedMovement().CopyTo(NewState, this);
+		ThisReplicatedMovement.CopyTo(NewState, this);
 
 		FVector DeltaPos(FVector::ZeroVector);
-		RootPrimComp->SetRigidBodyReplicatedTarget(NewState);
+		RootPrimComp->SetRigidBodyReplicatedTarget(NewState, NAME_None, ThisReplicatedMovement.ServerFrame, ThisReplicatedMovement.ServerPhysicsHandle);
 	}
 }
 
@@ -354,10 +361,26 @@ void AActor::GatherCurrentMovement()
 		UPrimitiveComponent* RootPrimComp = Cast<UPrimitiveComponent>(GetRootComponent());
 		if (RootPrimComp && RootPrimComp->IsSimulatingPhysics())
 		{
-			FRigidBodyState RBState;
-			RootPrimComp->GetRigidBodyState(RBState);
+			bool bFoundInCache = false;
+#if WITH_CHAOS
+			UWorld* World = GetWorld();
+			if (FPhysScene_Chaos* Scene = static_cast<FPhysScene_Chaos*>(World->GetPhysicsScene()))
+			{
+				if (FRigidBodyState* FoundState = Scene->ReplicationCache.Map.Find(FObjectKey(RootPrimComp)))
+				{
+					ReplicatedMovement.FillFrom(*FoundState, this, Scene->ReplicationCache.ServerFrame);
+					bFoundInCache = true;
+				}
+			}
+#endif
+			if (!bFoundInCache)
+			{
+				// fallback to GT data
+				FRigidBodyState RBState;
+				RootPrimComp->GetRigidBodyState(RBState);
+				ReplicatedMovement.FillFrom(RBState, this, 0);
+			}
 
-			ReplicatedMovement.FillFrom(RBState, this);
 			// Don't replicate movement if we're welded to another parent actor.
 			// Their replication will affect our position indirectly since we are attached.
 			ReplicatedMovement.bRepPhysics = !RootPrimComp->IsWelded();
diff --git a/Engine/Source/Runtime/Engine/Private/EngineTypes.cpp b/Engine/Source/Runtime/Engine/Private/EngineTypes.cpp
--- a/Engine/Source/Runtime/Engine/Private/EngineTypes.cpp
+++ b/Engine/Source/Runtime/Engine/Private/EngineTypes.cpp
@@ -593,6 +593,7 @@ FRepMovement::FRepMovement()
 	, Rotation(ForceInit)
 	, bSimulatedPhysicSleep(false)
 	, bRepPhysics(false)
+	, ServerFrame(0)
 	, LocationQuantizationLevel(EVectorQuantization::RoundWholeNumber)
 	, VelocityQuantizationLevel(EVectorQuantization::RoundWholeNumber)
 	, RotationQuantizationLevel(ERotatorQuantization::ByteComponents)
diff --git a/Engine/Source/Runtime/Engine/Private/PhysicsEngine/PhysicsReplication.cpp b/Engine/Source/Runtime/Engine/Private/PhysicsEngine/PhysicsReplication.cpp
--- a/Engine/Source/Runtime/Engine/Private/PhysicsEngine/PhysicsReplication.cpp
+++ b/Engine/Source/Runtime/Engine/Private/PhysicsEngine/PhysicsReplication.cpp
@@ -143,7 +143,7 @@ FPhysicsReplication::~FPhysicsReplication()
 #endif
 }
 
-bool FPhysicsReplication::ApplyRigidBodyState(float DeltaSeconds, FBodyInstance* BI, FReplicatedPhysicsTarget& PhysicsTarget, const FRigidBodyErrorCorrection& ErrorCorrection, const float PingSecondsOneWay)
+bool FPhysicsReplication::ApplyRigidBodyState(float DeltaSeconds, FBodyInstance* BI, FReplicatedPhysicsTarget& PhysicsTarget, const FRigidBodyErrorCorrection& ErrorCorrection, const float InPingSecondsOneWay, int32 LocalFrame, int32 NumPredictedFrames)
 {
 	if (CharacterMovementCVars::SkipPhysicsReplication)
 	{
@@ -155,6 +155,39 @@ bool FPhysicsReplication::ApplyRigidBodyState(float DeltaSeconds, FBodyInstance*
 		return false;
 	}
 
+	// LocalFrame the local frame number we should use to access past state in the rewind data
+	// NumPredictedFrames is how many frames (steps) ahead "now" is for the client compared to the latest data we've received from the server (use this to determine accurately how far ahead this object should extrapolate from its "physics target"
+	
+	Chaos::FPhysicsSolver* Solver = PhysScene->GetSolver();
+	Chaos::FRewindData* RewindData = Solver->GetRewindData();
+
+	if (LocalFrame > RewindData->GetEarliestFrame_Internal() && LocalFrame < RewindData->CurrentFrame())
+	{
+		auto Proxy = BI->GetPhysicsActorHandle();
+		const auto P = RewindData->GetPastStateAtFrame(*Proxy->GetHandle_LowLevel(), LocalFrame);
+		
+//#if 0
+		// Debugging/test: compare and print out if locations differ
+		if (FVector::DistSquared(PhysicsTarget.TargetState.Position, P.X()) > 1.f)
+		{
+			const int32 EarliestFrame = RewindData->GetEarliestFrame_Internal();
+			const int32 CurrentFrame = RewindData->CurrentFrame();
+
+			UE_LOG(LogTemp, Warning, TEXT(""));
+			UE_LOG(LogTemp, Warning, TEXT("Location differs"));
+			UE_LOG(LogTemp, Warning, TEXT("   Replicated: %s"), *PhysicsTarget.TargetState.Position.ToString());
+			UE_LOG(LogTemp, Warning, TEXT("   Historic: %s"), *P.X().ToString());
+		}
+//#endif
+	}
+
+	//if (true)
+	{
+		return false;
+	}
+
+	float PingSecondsOneWay = (Solver->GetLastDt() * NumPredictedFrames) * 0.5f; //InPingSecondsOneWay;
+
 	//
 	// NOTES:
 	//
@@ -460,6 +493,25 @@ float FPhysicsReplication::GetOwnerPing(const AActor* const Owner, const FReplic
 
 void FPhysicsReplication::OnTick(float DeltaSeconds, TMap<TWeakObjectPtr<UPrimitiveComponent>, FReplicatedPhysicsTarget>& ComponentsToTargets)
 {
+	int32 LocalFrameOffset = 0;		// LocalFrame = ServerFrame + LocalFrameOffset;
+	int32 NumPredictedFrames = 0;	// How many frames "ahead" of the server we are predicting. That is, how many frames are in flight between us and the server.
+
+	if (UWorld* World = GetOwningWorld())
+	{
+		if (World->GetNetMode() == NM_Client)
+		{
+			if (APlayerController* PlayerController = World->GetFirstPlayerController())
+			{
+				Chaos::FPhysicsSolver* Solver = PhysScene->GetSolver();
+				check(Solver);
+
+				const APlayerController::FClientFrameInfo& ClientFrameInfo = PlayerController->GetClientFrameInfo();
+				LocalFrameOffset = ClientFrameInfo.GetLocalFrameOffset();
+				NumPredictedFrames = Solver->GetCurrentFrame() - ClientFrameInfo.LastProcessedInputFrame;
+			}
+		}
+	}
+
 	const FRigidBodyErrorCorrection& PhysicErrorCorrection = UPhysicsSettings::Get()->PhysicErrorCorrection;
 #if WITH_CHAOS
 	using namespace Chaos;
@@ -500,7 +552,8 @@ void FPhysicsReplication::OnTick(float DeltaSeconds, TMap<TWeakObjectPtr<UPrimit
 
 						if (UpdatedState.Flags & ERigidBodyFlags::NeedsUpdate)
 						{
-							const bool bRestoredState = ApplyRigidBodyState(DeltaSeconds, BI, PhysicsTarget, PhysicErrorCorrection, PingSecondsOneWay);
+							const int32 LocalFrame = PhysicsTarget.ServerFrame + LocalFrameOffset;
+							const bool bRestoredState = ApplyRigidBodyState(DeltaSeconds, BI, PhysicsTarget, PhysicErrorCorrection, PingSecondsOneWay, LocalFrame, NumPredictedFrames);
 
 							// Need to update the component to match new position.
 							if (PhysicsReplicationCVars::SkipSkeletalRepOptimization == 0 || Cast<USkeletalMeshComponent>(PrimComp) == nullptr)	//simulated skeletal mesh does its own polling of physics results so we don't need to call this as it'll happen at the end of the physics sim
@@ -631,7 +684,7 @@ void FPhysicsReplication::ApplyAsyncDesiredState(const float DeltaSeconds, const
 }
 #endif
 
-void FPhysicsReplication::SetReplicatedTarget(UPrimitiveComponent* Component, FName BoneName, const FRigidBodyState& ReplicatedTarget)
+void FPhysicsReplication::SetReplicatedTarget(UPrimitiveComponent* Component, FName BoneName, const FRigidBodyState& ReplicatedTarget, int32 ServerFrame)
 {
 	if (UWorld* OwningWorld = GetOwningWorld())
 	{
@@ -648,6 +701,7 @@ void FPhysicsReplication::SetReplicatedTarget(UPrimitiveComponent* Component, FN
 			Target->PrevPosTarget = ReplicatedTarget.Position;
 		}
 
+		Target->ServerFrame = ServerFrame;
 		Target->TargetState = ReplicatedTarget;
 		Target->BoneName = BoneName;
 		Target->ArrivedTimeSeconds = OwningWorld->GetTimeSeconds();
diff --git a/Engine/Source/Runtime/Engine/Private/PrimitiveComponentPhysics.cpp b/Engine/Source/Runtime/Engine/Private/PrimitiveComponentPhysics.cpp
--- a/Engine/Source/Runtime/Engine/Private/PrimitiveComponentPhysics.cpp
+++ b/Engine/Source/Runtime/Engine/Private/PrimitiveComponentPhysics.cpp
@@ -15,6 +15,10 @@
 #include "Physics/PhysicsInterfaceCore.h"
 #include "UObject/UObjectThreadContext.h"
 
+#if WITH_CHAOS
+#include "PhysicsProxy/SingleParticlePhysicsProxy.h"
+#endif
+
 //////////////// PRIMITIVECOMPONENT ///////////////
 
 #define LOCTEXT_NAMESPACE "PrimitiveComponent"
@@ -30,7 +34,7 @@ DECLARE_CYCLE_STAT(TEXT("PrimComp SetCollisionProfileName"), STAT_PrimComp_SetCo
 	#define WarnInvalidPhysicsOperations(Text, BodyInstance, BoneName)
 #endif
 
-void UPrimitiveComponent::SetRigidBodyReplicatedTarget(FRigidBodyState& UpdatedState, FName BoneName)
+void UPrimitiveComponent::SetRigidBodyReplicatedTarget(FRigidBodyState& UpdatedState, FName BoneName, int32 ServerFrame, int32 ServerHandle)
 {
 	if (UWorld* World = GetWorld())
 	{
@@ -41,7 +45,10 @@ void UPrimitiveComponent::SetRigidBodyReplicatedTarget(FRigidBodyState& UpdatedS
 				FBodyInstance* BI = GetBodyInstance(BoneName);
 				if (BI && BI->IsValidBodyInstance())
 				{
-					PhysicsReplication->SetReplicatedTarget(this, BoneName, UpdatedState);
+					PhysicsReplication->SetReplicatedTarget(this, BoneName, UpdatedState, ServerFrame);
+#if WITH_CHAOS
+					BI->GetPhysicsActorHandle();// ->GetGameThreadAPI().SetParticleID(Chaos::FParticleID{ ServerPhysicsHandle, INDEX_NONE });
+#endif
 				}
 			}
 		}
diff --git a/Engine/Source/Runtime/Engine/Public/Physics/Experimental/PhysScene_Chaos.h b/Engine/Source/Runtime/Engine/Public/Physics/Experimental/PhysScene_Chaos.h
--- a/Engine/Source/Runtime/Engine/Public/Physics/Experimental/PhysScene_Chaos.h
+++ b/Engine/Source/Runtime/Engine/Public/Physics/Experimental/PhysScene_Chaos.h
@@ -237,6 +237,17 @@ public:
 	void RemoveDeferredPhysicsStateCreation(UPrimitiveComponent* Component);
 	void ProcessDeferredCreatePhysicsState();
 
+
+	// Storage structure for replication data
+	// probably should just expose read/write API not the structure directly itself like this.
+	struct FPrimitiveComponentReplicationCache
+	{
+		int32 ServerFrame = 0;
+		TMap<FObjectKey, FRigidBodyState>	Map;
+	};
+
+	FPrimitiveComponentReplicationCache ReplicationCache;
+
 private:
 	UPROPERTY()
 	TArray<UPrimitiveComponent*> CollisionEventRegistrations;
@@ -351,6 +362,9 @@ private:
 	// Maps Component to PhysicsProxy that is created
 	TMap<UPrimitiveComponent*, TArray<IPhysicsProxyBase*>> ComponentToPhysicsProxyMap;
 
+	//TSet<UActorComponent*> FixedTickComponents;
+	//TSet<AActor*> FixedTickActors;
+
 	/** The SolverActor that spawned and owns this scene */
 	TWeakObjectPtr<AActor> SolverActor;
 
diff --git a/Engine/Source/Runtime/Engine/Public/PhysicsReplication.h b/Engine/Source/Runtime/Engine/Public/PhysicsReplication.h
--- a/Engine/Source/Runtime/Engine/Public/PhysicsReplication.h
+++ b/Engine/Source/Runtime/Engine/Public/PhysicsReplication.h
@@ -25,7 +25,8 @@ struct FReplicatedPhysicsTarget
 {
 	FReplicatedPhysicsTarget() :
 		ArrivedTimeSeconds(0.0f),
-		AccumulatedErrorSeconds(0.0f)
+		AccumulatedErrorSeconds(0.0f),
+		ServerFrame(0)
 	{ }
 
 	/** The target state replicated by server */
@@ -44,6 +45,9 @@ struct FReplicatedPhysicsTarget
 	FVector PrevPosTarget;
 	FVector PrevPos;
 
+	/** ServerFrame this target was replicated on (must be converted to local frame prior to client-side use) */
+	int32 ServerFrame;
+
 #if !UE_BUILD_SHIPPING
 	FDebugFloatHistory ErrorHistory;
 #endif
@@ -59,6 +63,7 @@ struct FAsyncPhysicsDesiredState
 	Chaos::FSingleParticlePhysicsProxy* Proxy;
 	Chaos::EObjectStateType ObjectState;
 	bool bShouldSleep;
+	int32 ServerFrame;
 };
 #endif
 
@@ -79,7 +84,8 @@ public:
 	void Tick(float DeltaSeconds);
 
 	/** Sets the latest replicated target for a body instance */
-	virtual void SetReplicatedTarget(UPrimitiveComponent* Component, FName BoneName, const FRigidBodyState& ReplicatedTarget);
+	virtual void SetReplicatedTarget(UPrimitiveComponent* Component, FName BoneName, const FRigidBodyState& ReplicatedTarget) { SetReplicatedTarget(Component, BoneName, ReplicatedTarget, 0); }
+	virtual void SetReplicatedTarget(UPrimitiveComponent* Component, FName BoneName, const FRigidBodyState& ReplicatedTarget, int32 ServerFrame);
 
 	/** Remove the replicated target*/
 	virtual void RemoveReplicatedTarget(UPrimitiveComponent* Component);
@@ -91,7 +97,7 @@ protected:
 	virtual void OnTargetRestored(TWeakObjectPtr<UPrimitiveComponent> Component, const FReplicatedPhysicsTarget& Target) {}
 
 	/** Called when a dynamic rigid body receives a physics update */
-	virtual bool ApplyRigidBodyState(float DeltaSeconds, FBodyInstance* BI, FReplicatedPhysicsTarget& PhysicsTarget, const FRigidBodyErrorCorrection& ErrorCorrection, const float PingSecondsOneWay);
+	virtual bool ApplyRigidBodyState(float DeltaSeconds, FBodyInstance* BI, FReplicatedPhysicsTarget& PhysicsTarget, const FRigidBodyErrorCorrection& ErrorCorrection, const float PingSecondsOneWay, int32 LocalFrame, int32 NumPredictedFrames);
 
 	UWorld* GetOwningWorld();
 	const UWorld* GetOwningWorld() const;
diff --git a/Engine/Source/Runtime/Experimental/Chaos/Public/Chaos/ParticleHandle.h b/Engine/Source/Runtime/Experimental/Chaos/Public/Chaos/ParticleHandle.h
--- a/Engine/Source/Runtime/Experimental/Chaos/Public/Chaos/ParticleHandle.h
+++ b/Engine/Source/Runtime/Experimental/Chaos/Public/Chaos/ParticleHandle.h
@@ -1918,6 +1918,11 @@ public:
 	const FParticleID& ParticleID() const { return MNonFrequentData.Read().ParticleID(); }
 	void SetParticleID(const FParticleID& ParticleID, bool bInvalidate = true)
 	{
+		if (this->ParticleID() == ParticleID)
+		{
+			return;
+		}
+
 		MNonFrequentData.Modify(bInvalidate, MDirtyFlags, Proxy, [ParticleID](auto& Data) { Data.SetParticleID(ParticleID); });
 	}
 
