Allow DDC2 texture build path to not force stream all streaming mips up front in the editor.

diff --git a/Engine/Source/Runtime/Engine/Private/TextureDerivedDataTask.cpp b/Engine/Source/Runtime/Engine/Private/TextureDerivedDataTask.cpp
--- a/Engine/Source/Runtime/Engine/Private/TextureDerivedDataTask.cpp
+++ b/Engine/Source/Runtime/Engine/Private/TextureDerivedDataTask.cpp
@@ -1051,7 +1051,14 @@ public:
 					this->DerivedData.ResultMetadata = *InFetchFirstMetadata;
 				}
 
-				BuildSession.Get().Build(FetchDefinition, {}, EBuildPolicy::Cache, *Owner,
+				FBuildPolicyBuilder BuildPolicyBuilder(bInlineMips ? EBuildPolicy::Cache : (EBuildPolicy::CacheQuery | EBuildPolicy::SkipData));
+				if (!bInlineMips)
+				{
+					BuildPolicyBuilder.AddValuePolicy(FValueId::FromName("Description"_ASV), EBuildPolicy::Cache);
+					BuildPolicyBuilder.AddValuePolicy(FValueId::FromName("MipTail"_ASV), EBuildPolicy::Cache);
+				}
+
+				BuildSession.Get().Build(FetchDefinition, {}, BuildPolicyBuilder.Build(), *Owner,
 					[this, 
 					 FetchOrBuildDefinition = MoveTemp(FetchOrBuildDefinition), 
 					 Flags,
@@ -1105,11 +1112,24 @@ private:
 	void BeginBuild(const UE::DerivedData::FBuildDefinition& Definition, ETextureCacheFlags Flags)
 	{
 		using namespace UE::DerivedData;
-		EBuildPolicy BuildPolicy = EBuildPolicy::Default;
+
+		FBuildPolicy BuildPolicy;
 		if (EnumHasAnyFlags(Flags, ETextureCacheFlags::ForceRebuild))
 		{
-			EnumRemoveFlags(BuildPolicy, EBuildPolicy::CacheQuery);
+			BuildPolicy = EBuildPolicy::Default & ~EBuildPolicy::CacheQuery;
 		}
+		else if (bInlineMips)
+		{
+			BuildPolicy = EBuildPolicy::Default;
+		}
+		else
+		{
+			FBuildPolicyBuilder BuildPolicyBuilder(EBuildPolicy::Build | EBuildPolicy::CacheQuery | EBuildPolicy::CacheStoreOnBuild | EBuildPolicy::SkipData);
+			BuildPolicyBuilder.AddValuePolicy(FValueId::FromName("Description"_ASV), EBuildPolicy::Default);
+			BuildPolicyBuilder.AddValuePolicy(FValueId::FromName("MipTail"_ASV), EBuildPolicy::Default);
+			BuildPolicy = BuildPolicyBuilder.Build();
+		}
+
 		BuildSession.Get().Build(Definition, {}, BuildPolicy, *Owner,
 			[this](FBuildCompleteParams&& Params) { EndBuild(MoveTemp(Params)); });
 	}
