Deprecation of SpokenText field in SoundWave

diff --git a/Engine/Source/Editor/AudioEditor/Private/Factories/DialogueWaveFactory.cpp b/Engine/Source/Editor/AudioEditor/Private/Factories/DialogueWaveFactory.cpp
--- a/Engine/Source/Editor/AudioEditor/Private/Factories/DialogueWaveFactory.cpp
+++ b/Engine/Source/Editor/AudioEditor/Private/Factories/DialogueWaveFactory.cpp
@@ -19,7 +19,7 @@ UObject* UDialogueWaveFactory::FactoryCreateNew(UClass* Class, UObject* InParent
 
 	if (InitialSoundWave)
 	{
-		DialogueWave->SpokenText = InitialSoundWave->SpokenText;
+		DialogueWave->SpokenText = InitialSoundWave->SpokenText_DEPRECATED;
 		DialogueWave->bMature = InitialSoundWave->bMature;
 	}
 
diff --git a/Engine/Source/Runtime/Engine/Classes/Sound/SoundWave.h b/Engine/Source/Runtime/Engine/Classes/Sound/SoundWave.h
--- a/Engine/Source/Runtime/Engine/Classes/Sound/SoundWave.h
+++ b/Engine/Source/Runtime/Engine/Classes/Sound/SoundWave.h
@@ -585,8 +585,8 @@ public:
 	mutable ESoundWaveLoadingBehavior LoadingBehavior;
 
 	/** A localized version of the text that is actually spoken phonetically in the audio. */
-	UPROPERTY(EditAnywhere, BlueprintReadOnly, Category = Subtitles)
-	FString SpokenText;
+	UPROPERTY(meta = (DeprecatedProperty, DeprecationMessage = "Use Subtitles instead."))
+	FString SpokenText_DEPRECATED;
 
 	/** The priority of the subtitle. */
 	UPROPERTY(EditAnywhere, BlueprintReadOnly, Category = Subtitles)
@@ -642,8 +642,7 @@ public:
 	int32 TrackedMemoryUsage;
 
 	/**
-	 * Subtitle cues.  If empty, use SpokenText as the subtitle.  Will often be empty,
-	 * as the contents of the subtitle is commonly identical to what is spoken.
+	 * Subtitle cues. 
 	 */
 	UPROPERTY(EditAnywhere, BlueprintReadOnly, Category = Subtitles)
 	TArray<struct FSubtitleCue> Subtitles;
diff --git a/Engine/Source/Runtime/Engine/Private/AudioDebug.cpp b/Engine/Source/Runtime/Engine/Private/AudioDebug.cpp
--- a/Engine/Source/Runtime/Engine/Private/AudioDebug.cpp
+++ b/Engine/Source/Runtime/Engine/Private/AudioDebug.cpp
@@ -2449,7 +2449,7 @@ namespace Audio
 
 			if (Subtitle.Len() == 0)
 			{
-				Subtitle = InSoundWave.SpokenText;
+				Subtitle = InSoundWave.SpokenText_DEPRECATED;
 			}
 
 			if (Subtitle.Len() == 0)
