Network Physics
-remove early out from client correction path, didn't mean to check in.
-remove debug logging code from being compiled in
-restore old ApplyRigidBodyState function signature as backwards compat mode so other projects that override PhysicsReplication remain the same.

diff --git a/Engine/Source/Runtime/Engine/Private/PhysicsEngine/PhysicsReplication.cpp b/Engine/Source/Runtime/Engine/Private/PhysicsEngine/PhysicsReplication.cpp
--- a/Engine/Source/Runtime/Engine/Private/PhysicsEngine/PhysicsReplication.cpp
+++ b/Engine/Source/Runtime/Engine/Private/PhysicsEngine/PhysicsReplication.cpp
@@ -157,7 +157,7 @@ bool FPhysicsReplication::ApplyRigidBodyState(float DeltaSeconds, FBodyInstance*
 
 	// LocalFrame the local frame number we should use to access past state in the rewind data
 	// NumPredictedFrames is how many frames (steps) ahead "now" is for the client compared to the latest data we've received from the server (use this to determine accurately how far ahead this object should extrapolate from its "physics target"
-	
+
 	Chaos::FPhysicsSolver* Solver = PhysScene->GetSolver();
 	Chaos::FRewindData* RewindData = Solver->GetRewindData();
 
@@ -165,8 +165,8 @@ bool FPhysicsReplication::ApplyRigidBodyState(float DeltaSeconds, FBodyInstance*
 	{
 		auto Proxy = BI->GetPhysicsActorHandle();
 		const auto P = RewindData->GetPastStateAtFrame(*Proxy->GetHandle_LowLevel(), LocalFrame);
-		
-//#if 0
+
+#if 0
 		// Debugging/test: compare and print out if locations differ
 		if (FVector::DistSquared(PhysicsTarget.TargetState.Position, P.X()) > 1.f)
 		{
@@ -178,15 +178,27 @@ bool FPhysicsReplication::ApplyRigidBodyState(float DeltaSeconds, FBodyInstance*
 			UE_LOG(LogTemp, Warning, TEXT("   Replicated: %s"), *PhysicsTarget.TargetState.Position.ToString());
 			UE_LOG(LogTemp, Warning, TEXT("   Historic: %s"), *P.X().ToString());
 		}
-//#endif
+#endif
 	}
 
-	//if (true)
+	// Call into the old ApplyRigidBodyState function for now,
+	// "new leash mode" should replace this.
+	// Note that old ApplyRigidBodyState is overridden in other projects, so consider backwards compat path
+	float PingSecondsOneWay = (Solver->GetLastDt() * NumPredictedFrames) * 0.5f;
+	return ApplyRigidBodyState(DeltaSeconds, BI, PhysicsTarget, ErrorCorrection, PingSecondsOneWay);
+}
+
+bool FPhysicsReplication::ApplyRigidBodyState(float DeltaSeconds, FBodyInstance* BI, FReplicatedPhysicsTarget& PhysicsTarget, const FRigidBodyErrorCorrection& ErrorCorrection, const float PingSecondsOneWay)
+{
+	if (CharacterMovementCVars::SkipPhysicsReplication)
 	{
 		return false;
 	}
 
-	float PingSecondsOneWay = (Solver->GetLastDt() * NumPredictedFrames) * 0.5f; //InPingSecondsOneWay;
+	if (!BI->IsInstanceSimulatingPhysics())
+	{
+		return false;
+	}
 
 	//
 	// NOTES:
diff --git a/Engine/Source/Runtime/Engine/Public/PhysicsReplication.h b/Engine/Source/Runtime/Engine/Public/PhysicsReplication.h
--- a/Engine/Source/Runtime/Engine/Public/PhysicsReplication.h
+++ b/Engine/Source/Runtime/Engine/Public/PhysicsReplication.h
@@ -98,6 +98,7 @@ protected:
 
 	/** Called when a dynamic rigid body receives a physics update */
 	virtual bool ApplyRigidBodyState(float DeltaSeconds, FBodyInstance* BI, FReplicatedPhysicsTarget& PhysicsTarget, const FRigidBodyErrorCorrection& ErrorCorrection, const float PingSecondsOneWay, int32 LocalFrame, int32 NumPredictedFrames);
+	virtual bool ApplyRigidBodyState(float DeltaSeconds, FBodyInstance* BI, FReplicatedPhysicsTarget& PhysicsTarget, const FRigidBodyErrorCorrection& ErrorCorrection, const float PingSecondsOneWay); // deprecated path with no localframe/numpredicted
 
 	UWorld* GetOwningWorld();
 	const UWorld* GetOwningWorld() const;
