Still try and get complete shader maps from the DDC if they exist even with running with per-shader DDC.

diff --git a/Engine/Source/Runtime/Engine/Private/Materials/MaterialShared.cpp b/Engine/Source/Runtime/Engine/Private/Materials/MaterialShared.cpp
--- a/Engine/Source/Runtime/Engine/Private/Materials/MaterialShared.cpp
+++ b/Engine/Source/Runtime/Engine/Private/Materials/MaterialShared.cpp
@@ -2227,18 +2227,12 @@ bool FMaterial::CacheShaders(const FMaterialShaderMapId& ShaderMapId, EShaderPla
 				}
 			}
 
-			// If we are loading individual shaders from the shader job cache don't attempt to load full maps.
-			const bool bSkipCompilationOnPostLoad = IsShaderJobCacheDDCEnabled();
-
 			// Attempt to load from the derived data cache if we are uncooked and don't have any shadermap.
 			// If we have an incomplete shadermap, continue with it to prevent creation of duplicate shadermaps for the same ShaderMapId
 			if (!ShaderMap && !FPlatformProperties::RequiresCookedData())
 			{
 				TRefCountPtr<FMaterialShaderMap> LoadedShaderMap;
-				if (bSkipCompilationOnPostLoad == false || IsDefaultMaterial())
-				{
-					FMaterialShaderMap::LoadFromDerivedDataCache(this, ShaderMapId, Platform, TargetPlatform, LoadedShaderMap, DDCKeyHash);
-				}
+				FMaterialShaderMap::LoadFromDerivedDataCache(this, ShaderMapId, Platform, TargetPlatform, LoadedShaderMap, DDCKeyHash);
 				ShaderMap = LoadedShaderMap;
 			}
 
