Removed unused UPrimitiveComponent::CachedLODParentPrimitive
Prevent render state being marked dirty when calling SetLODParentPrimitive() with the existing value

diff --git a/Engine/Source/Runtime/Engine/Classes/Components/PrimitiveComponent.h b/Engine/Source/Runtime/Engine/Classes/Components/PrimitiveComponent.h
--- a/Engine/Source/Runtime/Engine/Classes/Components/PrimitiveComponent.h
+++ b/Engine/Source/Runtime/Engine/Classes/Components/PrimitiveComponent.h
@@ -1774,17 +1774,11 @@ private:
 	UPROPERTY(NonPIEDuplicateTransient)
 	TObjectPtr<class UPrimitiveComponent> LODParentPrimitive;
 
-	UPROPERTY(Transient)
-	TObjectPtr<class UPrimitiveComponent> CachedLODParentPrimitive;
-
 public:
 
 	/** Set the LOD parent component */
 	void SetLODParentPrimitive(UPrimitiveComponent* InLODParentPrimitive);
 
-	/** Set the cached LOD parent component */
-	void SetCachedLODParentPrimitive(UPrimitiveComponent* InCachedLODParentPrimitive);
-
 	/** Gets the LOD Parent, which is used to compute visibility when hierarchical LOD is enabled */
 	UPrimitiveComponent* GetLODParentPrimitive() const;
 
diff --git a/Engine/Source/Runtime/Engine/Private/Components/PrimitiveComponent.cpp b/Engine/Source/Runtime/Engine/Private/Components/PrimitiveComponent.cpp
--- a/Engine/Source/Runtime/Engine/Private/Components/PrimitiveComponent.cpp
+++ b/Engine/Source/Runtime/Engine/Private/Components/PrimitiveComponent.cpp
@@ -1268,9 +1268,6 @@ void UPrimitiveComponent::PostLoad()
 	}
 #endif
 
-	// Ensure the cached LOD parent primitive matches the loaded one
-	CachedLODParentPrimitive = LODParentPrimitive;
-
 	// Make sure cached cull distance is up-to-date.
 	if( LDMaxDrawDistance > 0.f )
 	{
@@ -4065,6 +4062,12 @@ void UPrimitiveComponent::SetHiddenInSceneCapture(bool bValue)
 
 void UPrimitiveComponent::SetLODParentPrimitive(UPrimitiveComponent * InLODParentPrimitive)
 {
+	if (LODParentPrimitive == InLODParentPrimitive)
+	{
+		return;
+	}
+
+
 #if WITH_EDITOR	
 	const ALODActor* ParentLODActor = [&]() -> const ALODActor*
 	{
@@ -4082,21 +4085,14 @@ void UPrimitiveComponent::SetLODParentPrimitive(UPrimitiveComponent * InLODParen
 	if (!GIsEditor || !InLODParentPrimitive || ShouldGenerateAutoLOD(ParentLODActor ? ParentLODActor->LODLevel - 1 : INDEX_NONE))
 #endif
 	{
-		// @todo, what do we do with old parent. We can't just reset undo parent because the parent might be used by other primitive
 		LODParentPrimitive = InLODParentPrimitive;
-		SetCachedLODParentPrimitive(LODParentPrimitive);
+		MarkRenderStateDirty();
 	}
 }
 
-void UPrimitiveComponent::SetCachedLODParentPrimitive(UPrimitiveComponent* InCachedLODParentPrimitive)
-{
-	CachedLODParentPrimitive = InCachedLODParentPrimitive;
-	MarkRenderStateDirty();
-}
-
 UPrimitiveComponent* UPrimitiveComponent::GetLODParentPrimitive() const
 {
-	return CachedLODParentPrimitive;
+	return LODParentPrimitive;
 }
 
 #if WITH_EDITOR
