Fixed a bug where distance fields would no longer work correctly with lwc coordinates

diff --git a/Engine/Source/Runtime/Engine/Private/Particles/ParticleGpuSimulation.cpp b/Engine/Source/Runtime/Engine/Private/Particles/ParticleGpuSimulation.cpp
--- a/Engine/Source/Runtime/Engine/Private/Particles/ParticleGpuSimulation.cpp
+++ b/Engine/Source/Runtime/Engine/Private/Particles/ParticleGpuSimulation.cpp
@@ -3251,16 +3251,22 @@ FGPUSpriteParticleEmitterInstance(FFXSystem* InFXSystem, FGPUSpriteEmitterInfo&
 		}
 		const FMatrix ComponentToWorldMatrix = ComponentTransform.ToMatrixWithScale();
 		const FMatrix ComponentToWorld = (bLocalSpace || EmitterInfo.LocalVectorField.bIgnoreComponentTransform) ? FMatrix::Identity : ComponentToWorldMatrix;
-
-		const FRotationMatrix VectorFieldTransform(LocalVectorFieldRotation);
-		const FMatrix VectorFieldToWorld = VectorFieldTransform * EmitterInfo.LocalVectorField.Transform.ToMatrixWithScale() * ComponentToWorld;
+		
 		FGPUSpriteDynamicEmitterData* DynamicData = new FGPUSpriteDynamicEmitterData(EmitterInfo.RequiredModule);
 		DynamicData->FXSystem = FXSystem;
 		DynamicData->Resources = EmitterInfo.Resources;
 		DynamicData->MaterialProxy = GetCurrentMaterial()->GetRenderProxy();
 		DynamicData->bIsMaterialTranslucent = IsTranslucentBlendMode(GetCurrentMaterial()->GetBlendMode());
 		DynamicData->Simulation = Simulation;
-		DynamicData->SimulationBounds = Template->bUseFixedRelativeBoundingBox ? Template->FixedRelativeBoundingBox.TransformBy(ComponentToWorldMatrix) : Component->Bounds.GetBox();
+		DynamicData->SimulationBounds = Template->bUseFixedRelativeBoundingBox ? Template->FixedRelativeBoundingBox.TransformBy(Component->GetComponentTransform()) : Component->Bounds.GetBox();
+		DynamicData->SortMode = EmitterInfo.RequiredModule->SortMode;
+		DynamicData->bSelected = bSelected;
+		DynamicData->bUseLocalSpace = EmitterInfo.RequiredModule->bUseLocalSpace;
+
+		// set up vector field data
+		const FMatrix VectorFieldComponentToWorld = (bLocalSpace || EmitterInfo.LocalVectorField.bIgnoreComponentTransform) ? FMatrix::Identity : Component->GetComponentTransform().ToMatrixWithScale();
+		const FRotationMatrix VectorFieldTransform(LocalVectorFieldRotation);
+		const FMatrix VectorFieldToWorld = VectorFieldTransform * EmitterInfo.LocalVectorField.Transform.ToMatrixWithScale() * VectorFieldComponentToWorld;
 		DynamicData->LocalVectorFieldToWorld = VectorFieldToWorld;
 		DynamicData->LocalVectorFieldIntensity = EmitterInfo.LocalVectorField.Intensity;
 		DynamicData->LocalVectorFieldTightness = EmitterInfo.LocalVectorField.Tightness;	
@@ -3268,9 +3274,6 @@ FGPUSpriteParticleEmitterInstance(FFXSystem* InFXSystem, FGPUSpriteEmitterInfo&
 		DynamicData->bLocalVectorFieldTileY = EmitterInfo.LocalVectorField.bTileY;	
 		DynamicData->bLocalVectorFieldTileZ = EmitterInfo.LocalVectorField.bTileZ;	
 		DynamicData->bLocalVectorFieldUseFixDT = EmitterInfo.LocalVectorField.bUseFixDT;
-		DynamicData->SortMode = EmitterInfo.RequiredModule->SortMode;
-		DynamicData->bSelected = bSelected;
-		DynamicData->bUseLocalSpace = EmitterInfo.RequiredModule->bUseLocalSpace;
 
 		// Get LWC tile
 		DynamicData->EmitterDynamicParameters.LWCTile = bUseTileOffset ? Component->GetLWCTile() : FVector3f::ZeroVector;
