MaterialInstance parameter overrides are supposed to be ignored for nullptr textures

diff --git a/Engine/Source/Runtime/Engine/Private/Materials/MaterialInstanceSupport.h b/Engine/Source/Runtime/Engine/Private/Materials/MaterialInstanceSupport.h
--- a/Engine/Source/Runtime/Engine/Private/Materials/MaterialInstanceSupport.h
+++ b/Engine/Source/Runtime/Engine/Private/Materials/MaterialInstanceSupport.h
@@ -156,11 +156,12 @@ public:
 		const TArray<TNamedParameter<ValueType> >& ValueArray = GetValueArray<ValueType>();
 		int Index = RenderThread_FindParameterByNameInternal<ValueType>(ParameterInfo, bWasFound);
 
-		if (bWasFound)
+		if (bWasFound && IsValidParameterValue(ValueArray[Index].Value))
 		{
 			OutValue = ValueArray[Index].Value;
+			return true;
 		}
-		return bWasFound;
+		return false;
 	}
 
 private:
