Fix no velocity written for unlit material

diff --git a/Engine/Shaders/Private/BasePassPixelShader.usf b/Engine/Shaders/Private/BasePassPixelShader.usf
--- a/Engine/Shaders/Private/BasePassPixelShader.usf
+++ b/Engine/Shaders/Private/BasePassPixelShader.usf
@@ -1817,8 +1817,8 @@ void FPixelShaderInOut_MainPS(
 				Out.MRT[1] = 0;
 				SetGBufferForUnlit(Out.MRT[2]);
 				Out.MRT[3] = 0;
-				Out.MRT[4] = 0;
-				Out.MRT[5] = 0;
+				Out.MRT[GBUFFER_HAS_VELOCITY ? 5 : 4] = 0;
+				Out.MRT[GBUFFER_HAS_VELOCITY ? 6 : 5] = 0;
 			}
 		}
 		#else
diff --git a/Engine/Source/Runtime/Engine/Private/ShaderCompiler/ShaderGenerationUtil.cpp b/Engine/Source/Runtime/Engine/Private/ShaderCompiler/ShaderGenerationUtil.cpp
--- a/Engine/Source/Runtime/Engine/Private/ShaderCompiler/ShaderGenerationUtil.cpp
+++ b/Engine/Source/Runtime/Engine/Private/ShaderCompiler/ShaderGenerationUtil.cpp
@@ -1715,6 +1715,7 @@ static void DetermineUsedMaterialSlots(
 	if (Mat.MATERIAL_SHADINGMODEL_UNLIT)
 	{
 		Slots[GBS_SceneColor] = true;
+		Slots[GBS_Velocity] = bHasVelocity;
 	}
 
 	if (Mat.MATERIAL_SHADINGMODEL_DEFAULT_LIT)
