Network Physics
-Added REP_MOVEMENT_DISABLE_QUANTIZATION compile time define for disabling quanitzation of movement replication. This is intended for development/testing, not as a shipping setting

#rnx none

diff --git a/Engine/Source/Runtime/Engine/Classes/Engine/EngineTypes.h b/Engine/Source/Runtime/Engine/Classes/Engine/EngineTypes.h
--- a/Engine/Source/Runtime/Engine/Classes/Engine/EngineTypes.h
+++ b/Engine/Source/Runtime/Engine/Classes/Engine/EngineTypes.h
@@ -3556,60 +3556,7 @@ struct ENGINE_API FRepMovement
 		}
 	}
 
-	bool NetSerialize(FArchive& Ar, class UPackageMap* Map, bool& bOutSuccess)
-	{
-		// pack bitfield with flags
-		uint8 Flags = (bSimulatedPhysicSleep << 0) | (bRepPhysics << 1) | ((ServerFrame > 0) << 2) | ((ServerPhysicsHandle != INDEX_NONE) << 3);
-		Ar.SerializeBits(&Flags, 4);
-		bSimulatedPhysicSleep = ( Flags & ( 1 << 0 ) ) ? 1 : 0;
-		bRepPhysics = ( Flags & ( 1 << 1 ) ) ? 1 : 0;
-		const bool bRepServerFrame = (Flags & (1 << 2)) ? 1 : 0;
-		const bool bRepServerHandle = (Flags & (1 << 3)) ? 1 : 0;
-
-		bOutSuccess = true;
-
-		// update location, rotation, linear velocity
-		bOutSuccess &= SerializeQuantizedVector( Ar, Location, LocationQuantizationLevel );
-		
-		switch(RotationQuantizationLevel)
-		{
-			case ERotatorQuantization::ByteComponents:
-			{
-				Rotation.SerializeCompressed( Ar );
-				break;
-			}
-
-			case ERotatorQuantization::ShortComponents:
-			{
-				Rotation.SerializeCompressedShort( Ar );
-				break;
-			}
-		}
-		
-		bOutSuccess &= SerializeQuantizedVector( Ar, LinearVelocity, VelocityQuantizationLevel );
-
-		// update angular velocity if required
-		if ( bRepPhysics )
-		{
-			bOutSuccess &= SerializeQuantizedVector( Ar, AngularVelocity, VelocityQuantizationLevel );
-		}
-
-		if (bRepServerFrame)
-		{
-			uint32 uServerFrame = (uint32)ServerFrame;
-			Ar.SerializeIntPacked(uServerFrame);
-			ServerFrame = (int32)uServerFrame;
-		}
-
-		if (bRepServerHandle)
-		{
-			uint32 uServerPhysicsHandle = (uint32)ServerPhysicsHandle;
-			Ar.SerializeIntPacked(uServerPhysicsHandle);
-			ServerPhysicsHandle = (int32)uServerPhysicsHandle;
-		}
-
-		return true;
-	}
+	bool NetSerialize(FArchive& Ar, class UPackageMap* Map, bool& bOutSuccess);
 
 	void FillFrom(const struct FRigidBodyState& RBState, const AActor* const Actor = nullptr, int32 InServerFrame = 0)
 	{
diff --git a/Engine/Source/Runtime/Engine/Private/EngineTypes.cpp b/Engine/Source/Runtime/Engine/Private/EngineTypes.cpp
--- a/Engine/Source/Runtime/Engine/Private/EngineTypes.cpp
+++ b/Engine/Source/Runtime/Engine/Private/EngineTypes.cpp
@@ -600,6 +600,77 @@ FRepMovement::FRepMovement()
 {
 }
 
+// compile switch for disabling quantization of replicated movement, meant for testing.
+#ifndef REP_MOVEMENT_DISABLE_QUANTIZATION
+#define REP_MOVEMENT_DISABLE_QUANTIZATION 0
+#endif
+
+bool FRepMovement::NetSerialize(FArchive& Ar, class UPackageMap* Map, bool& bOutSuccess)
+{
+	// pack bitfield with flags
+	uint8 Flags = (bSimulatedPhysicSleep << 0) | (bRepPhysics << 1) | ((ServerFrame > 0) << 2) | ((ServerPhysicsHandle != INDEX_NONE) << 3);
+	Ar.SerializeBits(&Flags, 4);
+	bSimulatedPhysicSleep = ( Flags & ( 1 << 0 ) ) ? 1 : 0;
+	bRepPhysics = ( Flags & ( 1 << 1 ) ) ? 1 : 0;
+	const bool bRepServerFrame = (Flags & (1 << 2)) ? 1 : 0;
+	const bool bRepServerHandle = (Flags & (1 << 3)) ? 1 : 0;
+
+	bOutSuccess = true;
+
+#if REP_MOVEMENT_DISABLE_QUANTIZATION
+	Ar << Location;
+	Ar << LinearVelocity;
+	Ar << Rotation;
+	if (bRepPhysics)
+	{
+		Ar << AngularVelocity;
+	}
+#else
+
+	// update location, rotation, linear velocity
+	bOutSuccess &= SerializeQuantizedVector( Ar, Location, LocationQuantizationLevel );
+
+	switch(RotationQuantizationLevel)
+	{
+	case ERotatorQuantization::ByteComponents:
+	{
+		Rotation.SerializeCompressed( Ar );
+		break;
+	}
+
+	case ERotatorQuantization::ShortComponents:
+	{
+		Rotation.SerializeCompressedShort( Ar );
+		break;
+	}
+	}
+
+	bOutSuccess &= SerializeQuantizedVector( Ar, LinearVelocity, VelocityQuantizationLevel );
+
+	// update angular velocity if required
+	if ( bRepPhysics )
+	{
+		bOutSuccess &= SerializeQuantizedVector( Ar, AngularVelocity, VelocityQuantizationLevel );
+	}
+#endif
+
+	if (bRepServerFrame)
+	{
+		uint32 uServerFrame = (uint32)ServerFrame;
+		Ar.SerializeIntPacked(uServerFrame);
+		ServerFrame = (int32)uServerFrame;
+	}
+
+	if (bRepServerHandle)
+	{
+		uint32 uServerPhysicsHandle = (uint32)ServerPhysicsHandle;
+		Ar.SerializeIntPacked(uServerPhysicsHandle);
+		ServerPhysicsHandle = (int32)uServerPhysicsHandle;
+	}
+
+	return true;
+}
+
 /** Rebase zero-origin position onto local world origin value. */
 FVector FRepMovement::RebaseOntoLocalOrigin(const FVector& Location, const struct FIntVector& LocalOrigin)
 {
