LocalPlayer: Handle ViewState allocation in case of rendering views in non-sequential StereoIndex order.

This may happen if nDisplay needs to re-order its viewfamily renders.

#rnx

diff --git a/Engine/Source/Runtime/Engine/Private/LocalPlayer.cpp b/Engine/Source/Runtime/Engine/Private/LocalPlayer.cpp
--- a/Engine/Source/Runtime/Engine/Private/LocalPlayer.cpp
+++ b/Engine/Source/Runtime/Engine/Private/LocalPlayer.cpp
@@ -784,10 +784,24 @@ bool ULocalPlayer::CalcSceneViewInitOptions(
 	check(PlayerController && PlayerController->GetWorld());
 
 	const uint32 ViewIndex = StereoViewIndex != INDEX_NONE ? StereoViewIndex : 0;
-	if (!ViewStates.IsValidIndex(ViewIndex))
+
+	// Make sure the ViewStates array has enough elements for the given ViewIndex.
 	{
-		ViewStates.EmplaceAt(ViewIndex);
-		ViewStates[ViewIndex].Allocate(GetWorld()->FeatureLevel);
+		const int32 RequiredViewStates = (ViewIndex + 1) - ViewStates.Num();
+		
+		if (RequiredViewStates > 0)
+		{
+			ViewStates.AddDefaulted(RequiredViewStates);		
+		}
+	}
+
+	// Allocate the current ViewState if necessary
+	if (ViewStates[ViewIndex].GetReference() == nullptr)
+	{
+		const UWorld* CurrentWorld = GetWorld();
+		const ERHIFeatureLevel::Type FeatureLevel = CurrentWorld ? CurrentWorld->FeatureLevel.GetValue() : GMaxRHIFeatureLevel;
+
+		ViewStates[ViewIndex].Allocate(FeatureLevel);
 	}
 
 	ViewInitOptions.SceneViewStateInterface = ViewStates[ViewIndex].GetReference();
@@ -1436,8 +1450,10 @@ bool ULocalPlayer::Exec(UWorld* InWorld, const TCHAR* Cmd,FOutputDevice& Ar)
 		// Reset states (e.g. TemporalAA index) to make rendering more deterministic (for automated screenshot verification)
 		for (auto& State : ViewStates)
 		{
-			FSceneViewStateInterface* Ref = State.GetReference();
-			Ref->ResetViewState();
+			if (FSceneViewStateInterface* Ref = State.GetReference())
+			{
+				Ref->ResetViewState();
+			}
 		}
 
 		return true;
