Allow compiling special materials asynchronously on startup.

- Saves few seconds (5%) of startup time with job cache DDC on. Potentially saves time without it as well.

diff --git a/Engine/Source/Runtime/Engine/Private/Materials/Material.cpp b/Engine/Source/Runtime/Engine/Private/Materials/Material.cpp
--- a/Engine/Source/Runtime/Engine/Private/Materials/Material.cpp
+++ b/Engine/Source/Runtime/Engine/Private/Materials/Material.cpp
@@ -567,6 +567,12 @@ void UMaterialInterface::PostCDOContruct()
 	}
 }
 
+// We can save time if instead of blocking after compilation of each synchronous material we block after scheduling all of them
+bool GPoolSpecialMaterialsCompileJobs = true;
+bool PoolSpecialMaterialsCompileJobs()
+{
+	return GPoolSpecialMaterialsCompileJobs;
+}
 
 void UMaterialInterface::PostLoadDefaultMaterials()
 {
@@ -610,6 +616,14 @@ void UMaterialInterface::PostLoadDefaultMaterials()
 				bPostLoaded = false;
 			}
 		}
+
+		// Block after scheduling for compilation all (hopefully) default materials.
+		// Even if not all of them ended up being post-loaded, block here just out of extra caution
+		if (GPoolSpecialMaterialsCompileJobs == true)
+		{
+			GPoolSpecialMaterialsCompileJobs = false;
+			GShaderCompilingManager->FinishAllCompilation();
+		}
 	}
 }
 
diff --git a/Engine/Source/Runtime/Engine/Private/Materials/MaterialShader.cpp b/Engine/Source/Runtime/Engine/Private/Materials/MaterialShader.cpp
--- a/Engine/Source/Runtime/Engine/Private/Materials/MaterialShader.cpp
+++ b/Engine/Source/Runtime/Engine/Private/Materials/MaterialShader.cpp
@@ -60,6 +60,8 @@ TMap<FMaterialShaderMapId,FMaterialShaderMap*> FMaterialShaderMap::GIdToMaterial
 TArray<FMaterialShaderMap*> FMaterialShaderMap::AllMaterialShaderMaps;
 FCriticalSection FMaterialShaderMap::AllMaterialShaderMapsGuard;
 #endif
+// defined in the same module (Material.cpp)
+bool PoolSpecialMaterialsCompileJobs();
 
 /** 
  * Tracks material resources and their shader maps that are being compiled.
@@ -2066,7 +2068,8 @@ void FMaterialShaderMap::Compile(
 	}
 
 	// Compile the shaders for this shader map now if the material is not deferring and deferred compiles are not enabled globally
-	if (PrecompileMode == EMaterialShaderPrecompileMode::Synchronous)
+	// If we're early in the startup we can save some time by compiling all special/default materials asynchronously, even if normally they are synchronous
+	if (PrecompileMode == EMaterialShaderPrecompileMode::Synchronous && !PoolSpecialMaterialsCompileJobs())
 	{
 		TArray<int32> CurrentShaderMapId;
 		CurrentShaderMapId.Add(CompilingId);
diff --git a/Engine/Source/Runtime/Engine/Private/Materials/MaterialShared.cpp b/Engine/Source/Runtime/Engine/Private/Materials/MaterialShared.cpp
--- a/Engine/Source/Runtime/Engine/Private/Materials/MaterialShared.cpp
+++ b/Engine/Source/Runtime/Engine/Private/Materials/MaterialShared.cpp
@@ -124,6 +124,9 @@ static FAutoConsoleCommand GFlushMaterialUniforms(
 })
 );
 
+// defined in the same module (Material.cpp)
+bool PoolSpecialMaterialsCompileJobs();
+
 bool AllowDitheredLODTransition(ERHIFeatureLevel::Type FeatureLevel)
 {
 	// On mobile support for 'Dithered LOD Transition' has to be explicitly enabled in projects settings
@@ -2497,7 +2500,8 @@ bool FMaterial::BeginCompileShaderMap(
 		// Compile the shaders for the material.
 		NewShaderMap->Compile(this, ShaderMapId, MaterialEnvironment, NewCompilationOutput, Platform, PrecompileMode);
 
-		if (PrecompileMode == EMaterialShaderPrecompileMode::Synchronous)
+		// early in the startup we can save some time by compiling all special/default materials asynchronously, even if normally they are synchronous
+		if (PrecompileMode == EMaterialShaderPrecompileMode::Synchronous && !PoolSpecialMaterialsCompileJobs())
 		{
 			// If this is a synchronous compile, assign the compile result to the output
 			check(NewShaderMap->GetCompilingId() == 0u);
