Remove deprecated bExtendDefaultLuminanceRangeInAutoExposureSettings from project settings menu.
- We don't want new projects disabling it.

diff --git a/Engine/Source/Runtime/Engine/Classes/Engine/RendererSettings.h b/Engine/Source/Runtime/Engine/Classes/Engine/RendererSettings.h
--- a/Engine/Source/Runtime/Engine/Classes/Engine/RendererSettings.h
+++ b/Engine/Source/Runtime/Engine/Classes/Engine/RendererSettings.h
@@ -620,13 +620,6 @@ class ENGINE_API URendererSettings : public UDeveloperSettings
 		ToolTip = "Default Value for auto exposure bias."))
 	float DefaultFeatureAutoExposureBias;
 
-	UE_DEPRECATED(4.26, "Extend Default Luminance Range is deprecated, and will be forced to ON at all times in future revisions.")
-	UPROPERTY(config, EditAnywhere, Category = DefaultSettings, meta = (
-		ConsoleVariable = "r.DefaultFeature.AutoExposure.ExtendDefaultLuminanceRange", DisplayName = "Extend default luminance range in Auto Exposure settings",
-		ToolTip = "Whether the default values for AutoExposure should support an extended range of scene luminance. Also changes the exposure settings to be expressed in EV100.",
-		ConfigRestartRequired=true))
-	uint32 bExtendDefaultLuminanceRangeInAutoExposureSettings: 1;
-
 	UPROPERTY(config, EditAnywhere, Category = DefaultSettings, meta = (
 		ConsoleVariable = "r.DefaultFeature.MotionBlur", DisplayName = "Motion Blur",
 		ToolTip = "Whether the default for MotionBlur is enabled or not (postprocess volume/camera/game setting can still override and enable or disable it independently)"))
