Fix auto-exposure toggle in project settings not working

diff --git a/Engine/Source/Runtime/Engine/Private/SceneView.cpp b/Engine/Source/Runtime/Engine/Private/SceneView.cpp
--- a/Engine/Source/Runtime/Engine/Private/SceneView.cpp
+++ b/Engine/Source/Runtime/Engine/Private/SceneView.cpp
@@ -1892,7 +1892,17 @@ void FSceneView::StartFinalPostprocessSettings(FVector InViewLocation)
 		{
 			FinalPostProcessSettings.AmbientOcclusionIntensity = 0;
 		}
-		if (CVarDefaultAutoExposure.GetValueOnGameThread())
+		if (!CVarDefaultAutoExposure.GetValueOnGameThread())
+		{
+			FinalPostProcessSettings.AutoExposureMinBrightness = 1;
+			FinalPostProcessSettings.AutoExposureMaxBrightness = 1;
+			if (CVarDefaultAutoExposureExtendDefaultLuminanceRange.GetValueOnGameThread())
+			{
+				FinalPostProcessSettings.AutoExposureMinBrightness = LuminanceToEV100(FinalPostProcessSettings.AutoExposureMinBrightness);
+				FinalPostProcessSettings.AutoExposureMaxBrightness = LuminanceToEV100(FinalPostProcessSettings.AutoExposureMaxBrightness);
+			}
+		}
+		else
 		{
 			int32 Value = CVarDefaultAutoExposureMethod.GetValueOnGameThread();
 			if (Value >= 0 && Value < AEM_MAX)
