Added blend profile support to inertialization requests.
Shelf CL: 15427774 and 15596715 from The Coalition by [REDACTED_PERSON]

diff --git a/Engine/Config/BaseEngine.ini b/Engine/Config/BaseEngine.ini
--- a/Engine/Config/BaseEngine.ini
+++ b/Engine/Config/BaseEngine.ini
@@ -956,6 +956,8 @@ NearClipPlane=10.0
 +PropertyRedirects=(OldName="GeometryCollectionComponent.Simulating", NewName="SimulatePhysics")
 +PropertyRedirects=(OldName="ControlRigComponent.OnPreUpdateDelegate", NewName="OnPreForwardsSolveDelegate")
 +PropertyRedirects=(OldName="ControlRigComponent.OnPostUpdateDelegate", NewName="OnPostForwardsSolveDelegate")
++PropertyRedirects=(OldName="AnimNode_Mirror.BlendTimeOnMirrorStateChange",NewName="AnimNode_Mirror.BlendTime")
++PropertyRedirects=(OldName="AnimNode_Mirror.bResetChildOnMirrorStateChange",NewName="AnimNode_Mirror.bResetChild")
 +PropertyRedirects=(OldName="PlayerState.Ping",NewName="CompressedPing")
 ; Sequencer LWC support
 +ClassRedirects=(OldName="/Script/MovieSceneTracks.MovieSceneVectorSection",NewName="/Script/MovieSceneTracks.MovieSceneFloatVectorSection")
diff --git a/Engine/Source/Editor/DetailCustomizations/Private/AnimTransitionNodeDetails.cpp b/Engine/Source/Editor/DetailCustomizations/Private/AnimTransitionNodeDetails.cpp
--- a/Engine/Source/Editor/DetailCustomizations/Private/AnimTransitionNodeDetails.cpp
+++ b/Engine/Source/Editor/DetailCustomizations/Private/AnimTransitionNodeDetails.cpp
@@ -228,7 +228,7 @@ void FAnimTransitionNodeDetails::CustomizeDetails( IDetailLayoutBuilder& DetailB
 		CrossfadeCategory.AddProperty(GET_MEMBER_NAME_CHECKED(UAnimStateTransitionNode, CrossfadeDuration)).DisplayName( LOCTEXT("DurationLabel", "Duration") );
 		CrossfadeCategory.AddProperty(GET_MEMBER_NAME_CHECKED(UAnimStateTransitionNode, BlendMode)).DisplayName(LOCTEXT("ModeLabel", "Mode")).IsEnabled(BlendSettingsEnabledAttribute);
 		CrossfadeCategory.AddProperty(GET_MEMBER_NAME_CHECKED(UAnimStateTransitionNode, CustomBlendCurve)).DisplayName(LOCTEXT("CurveLabel", "Custom Blend Curve")).IsEnabled(BlendSettingsEnabledAttribute);
-		CrossfadeCategory.AddProperty(GET_MEMBER_NAME_CHECKED(UAnimStateTransitionNode, BlendProfile)).DisplayName(LOCTEXT("BlendProfileLabel", "Blend Profile")).IsEnabled(BlendSettingsEnabledAttribute);
+		CrossfadeCategory.AddProperty(GET_MEMBER_NAME_CHECKED(UAnimStateTransitionNode, BlendProfile)).DisplayName(LOCTEXT("BlendProfileLabel", "Blend Profile"));
 
 		//////////////////////////////////////////////////////////////////////////
 
diff --git a/Engine/Source/Runtime/AnimGraphRuntime/Private/AnimNodes/AnimNode_BlendListBase.cpp b/Engine/Source/Runtime/AnimGraphRuntime/Private/AnimNodes/AnimNode_BlendListBase.cpp
--- a/Engine/Source/Runtime/AnimGraphRuntime/Private/AnimNodes/AnimNode_BlendListBase.cpp
+++ b/Engine/Source/Runtime/AnimGraphRuntime/Private/AnimNodes/AnimNode_BlendListBase.cpp
@@ -117,7 +117,7 @@ void FAnimNode_BlendListBase::Update_AnyThread(const FAnimationUpdateContext& Co
 				UE::Anim::IInertializationRequester* InertializationRequester = Context.GetMessage<UE::Anim::IInertializationRequester>();
 				if (InertializationRequester)
 				{
-					InertializationRequester->RequestInertialization(CurrentBlendTimes[ChildIndex]);
+					InertializationRequester->RequestInertialization(CurrentBlendTimes[ChildIndex], CurrentBlendProfile);
 					InertializationRequester->AddDebugRecord(*Context.AnimInstanceProxy, Context.GetCurrentNodeId());
 				}
 				else
diff --git a/Engine/Source/Runtime/AnimGraphRuntime/Private/AnimNodes/AnimNode_Mirror.cpp b/Engine/Source/Runtime/AnimGraphRuntime/Private/AnimNodes/AnimNode_Mirror.cpp
--- a/Engine/Source/Runtime/AnimGraphRuntime/Private/AnimNodes/AnimNode_Mirror.cpp
+++ b/Engine/Source/Runtime/AnimGraphRuntime/Private/AnimNodes/AnimNode_Mirror.cpp
@@ -3,6 +3,7 @@
 #include "AnimNodes/AnimNode_Mirror.h"
 #include "Animation/AnimInstanceProxy.h"
 #include "Animation/AnimNode_Inertialization.h"
+#include "Animation/AnimTrace.h"
 #include "AnimationRuntime.h"
 #include "Animation/AttributesRuntime.h"
 #include "Animation/MirrorSyncScope.h"
@@ -30,6 +31,11 @@ bool FAnimNode_MirrorBase::GetMirror() const
 	return false;
 }
 
+UBlendProfile* FAnimNode_MirrorBase::GetBlendProfile() const
+{
+	return nullptr;
+}
+
 float FAnimNode_MirrorBase::GetBlendTimeOnMirrorStateChange() const
 {
 	return 0.0;
@@ -60,6 +66,11 @@ bool FAnimNode_MirrorBase::SetMirror(bool bInMirror)
 	return false;
 }
 
+bool FAnimNode_MirrorBase::SetBlendProfile(UBlendProfile* BlendProfile)
+{
+	return false;
+}
+
 bool FAnimNode_MirrorBase::SetBlendTimeOnMirrorStateChange(float InBlendTime)
 {
 	return false;
@@ -140,7 +151,7 @@ void FAnimNode_MirrorBase::Update_AnyThread(const FAnimationUpdateContext& Conte
 			UE::Anim::IInertializationRequester* InertializationRequester = Context.GetMessage<UE::Anim::IInertializationRequester>();
 			if (InertializationRequester)
 			{
-				InertializationRequester->RequestInertialization(GetBlendTimeOnMirrorStateChange());
+				InertializationRequester->RequestInertialization(GetBlendTimeOnMirrorStateChange(), GetBlendProfile());
 				InertializationRequester->AddDebugRecord(*Context.AnimInstanceProxy, Context.GetCurrentNodeId());
 			}
 			else
@@ -163,6 +174,8 @@ void FAnimNode_MirrorBase::Update_AnyThread(const FAnimationUpdateContext& Conte
 	bMirrorStateIsValid = true;
 
 	Source.Update(Context);
+
+	TRACE_ANIM_NODE_VALUE(Context, TEXT("Mirrored"), bMirrorState);
 }
 
 void FAnimNode_MirrorBase::Evaluate_AnyThread(FPoseContext& Output)
@@ -235,9 +248,14 @@ bool FAnimNode_Mirror::GetMirror() const
 	return GET_ANIM_NODE_DATA(bool, bMirror);
 }
 
+UBlendProfile* FAnimNode_Mirror::GetBlendProfile() const
+{
+	return GET_ANIM_NODE_DATA(TObjectPtr<UBlendProfile>, BlendProfile);
+}
+
 float FAnimNode_Mirror::GetBlendTimeOnMirrorStateChange() const
 {
-	return GET_ANIM_NODE_DATA(float, BlendTimeOnMirrorStateChange);
+	return GET_ANIM_NODE_DATA(float, BlendTime);
 }
 
 bool FAnimNode_Mirror::GetBoneMirroring() const
@@ -257,7 +275,7 @@ bool FAnimNode_Mirror::GetAttributeMirroring() const
 
 bool FAnimNode_Mirror::GetResetChildOnMirrorStateChange() const
 {
-	return GET_ANIM_NODE_DATA(bool, bResetChildOnMirrorStateChange);
+	return GET_ANIM_NODE_DATA(bool, bResetChild);
 }
 
 bool FAnimNode_Mirror::SetMirror(bool bInMirror)
@@ -274,13 +292,29 @@ bool FAnimNode_Mirror::SetMirror(bool bInMirror)
 	return false;
 }
 
+bool FAnimNode_Mirror::SetBlendProfile(UBlendProfile* InBlendProfile)
+{
+#if WITH_EDITORONLY_DATA
+	BlendProfile = InBlendProfile;
+	GET_MUTABLE_ANIM_NODE_DATA(TObjectPtr<UBlendProfile>, BlendProfile) = InBlendProfile;
+#endif
+
+	if (TObjectPtr<UBlendProfile>* BlendProfilePtr = GET_INSTANCE_ANIM_NODE_DATA_PTR(TObjectPtr<UBlendProfile>, BlendProfile))
+	{
+		*BlendProfilePtr = InBlendProfile;
+		return true;
+	}
+
+	return false;
+}
+
 bool FAnimNode_Mirror::SetBlendTimeOnMirrorStateChange(float InBlendTime)
 {
 #if WITH_EDITORONLY_DATA
-	BlendTimeOnMirrorStateChange = InBlendTime;
+	BlendTime = InBlendTime;
 #endif
 
-	if (float* BlendTimePtr = GET_INSTANCE_ANIM_NODE_DATA_PTR(float, BlendTimeOnMirrorStateChange))
+	if (float* BlendTimePtr = GET_INSTANCE_ANIM_NODE_DATA_PTR(float, BlendTime))
 	{
 		*BlendTimePtr = InBlendTime;
 		return true;
@@ -336,12 +370,12 @@ bool FAnimNode_Mirror::SetAttributeMirroring(bool bInAttributeMirroring)
 bool FAnimNode_Mirror::SetResetChildOnMirrorStateChange(bool bInResetChildOnMirrorStateChange)
 {
 #if WITH_EDITORONLY_DATA
-	bResetChildOnMirrorStateChange = bInResetChildOnMirrorStateChange;
+	bResetChild = bInResetChildOnMirrorStateChange;
 #endif
 
-	if (bool* bResetChildOnMirrorStateChangePtr = GET_INSTANCE_ANIM_NODE_DATA_PTR(bool, bResetChildOnMirrorStateChange))
+	if (bool* bResetChildPtr = GET_INSTANCE_ANIM_NODE_DATA_PTR(bool, bResetChild))
 	{
-		*bResetChildOnMirrorStateChangePtr = bInResetChildOnMirrorStateChange;
+		*bResetChildPtr = bInResetChildOnMirrorStateChange;
 		return true;
 	}
 
@@ -369,9 +403,14 @@ bool FAnimNode_Mirror_Standalone::GetMirror() const
 	return bMirror;
 }
 
+UBlendProfile* FAnimNode_Mirror_Standalone::GetBlendProfile() const
+{
+	return BlendProfile.Get();
+}
+
 float FAnimNode_Mirror_Standalone::GetBlendTimeOnMirrorStateChange() const
 {
-	return BlendTimeOnMirrorStateChange;
+	return BlendTime;
 }
 
 bool FAnimNode_Mirror_Standalone::GetBoneMirroring() const
@@ -391,7 +430,7 @@ bool FAnimNode_Mirror_Standalone::GetAttributeMirroring() const
 
 bool FAnimNode_Mirror_Standalone::GetResetChildOnMirrorStateChange() const
 {
-	return bResetChildOnMirrorStateChange;
+	return bResetChild;
 }
 
 bool FAnimNode_Mirror_Standalone::SetMirror(bool bInMirror)
@@ -400,9 +439,15 @@ bool FAnimNode_Mirror_Standalone::SetMirror(bool bInMirror)
 	return true;
 }
 
+bool FAnimNode_Mirror_Standalone::SetBlendProfile(UBlendProfile* InBlendProfile)
+{
+	BlendProfile = InBlendProfile;
+	return true;
+}
+
 bool FAnimNode_Mirror_Standalone::SetBlendTimeOnMirrorStateChange(float InBlendTime)
 {
-	BlendTimeOnMirrorStateChange = InBlendTime;
+	BlendTime = InBlendTime;
 	return true;
 }
 
@@ -427,7 +472,7 @@ bool FAnimNode_Mirror_Standalone::SetAttributeMirroring(bool bInAttributeMirrori
 
 bool FAnimNode_Mirror_Standalone::SetResetChildOnMirrorStateChange(bool bInResetChildOnMirrorStateChange)
 {
-	bResetChildOnMirrorStateChange = bInResetChildOnMirrorStateChange;
+	bResetChild = bInResetChildOnMirrorStateChange;
 	return true;
 }
 
diff --git a/Engine/Source/Runtime/AnimGraphRuntime/Public/AnimNodes/AnimNode_Mirror.h b/Engine/Source/Runtime/AnimGraphRuntime/Public/AnimNodes/AnimNode_Mirror.h
--- a/Engine/Source/Runtime/AnimGraphRuntime/Public/AnimNodes/AnimNode_Mirror.h
+++ b/Engine/Source/Runtime/AnimGraphRuntime/Public/AnimNodes/AnimNode_Mirror.h
@@ -30,6 +30,9 @@ public:
 	// Get Mirror State
 	virtual bool GetMirror() const;
 
+	// Get the BlendProfile to use when transitioning between mirrored and unmirrored states
+	virtual UBlendProfile* GetBlendProfile() const;
+
 	// How long to blend using inertialization when switching  mirrored state
 	virtual float GetBlendTimeOnMirrorStateChange() const;
 
@@ -48,6 +51,9 @@ public:
 	// Set Mirror State
 	virtual bool SetMirror(bool bInMirror);
 
+	// Set the BlendProfile to use when transitioning between mirrored and unmirrored states
+	virtual bool SetBlendProfile(UBlendProfile* BlendProfile);
+
 	// Set how long to blend using inertialization when switching  mirrored state
 	// @return true if the value was set (it is dynamic), or false if it could not (it is not dynamic or pin exposed)
 	virtual bool SetBlendTimeOnMirrorStateChange(float InBlendTime);
@@ -104,6 +110,7 @@ public:
 	virtual bool SetMirrorDataTable(UMirrorDataTable* MirrorTable) override;
 
 	virtual bool GetMirror() const override;
+	virtual UBlendProfile* GetBlendProfile() const;
 	virtual float GetBlendTimeOnMirrorStateChange() const override;
 	virtual bool GetBoneMirroring() const override;
 	virtual bool GetCurveMirroring() const override;
@@ -111,6 +118,7 @@ public:
 	virtual bool GetResetChildOnMirrorStateChange() const override;
 
 	virtual bool SetMirror(bool bInMirror) override;
+	virtual bool SetBlendProfile(UBlendProfile* BlendProfile);
 	virtual bool SetBlendTimeOnMirrorStateChange(float InBlendTime) override;
 	virtual bool SetBoneMirroring(bool bInBoneMirroring) override;
 	virtual bool SetCurveMirroring(bool bInCurveMirroring) override;
@@ -126,8 +134,17 @@ protected:
 	UPROPERTY(EditAnywhere, Category = Settings, meta = (FoldProperty))
 	TObjectPtr<UMirrorDataTable> MirrorDataTable = nullptr;
 
-	UPROPERTY(EditAnywhere, Category = Settings, meta = (PinHiddenByDefault, FoldProperty))
-	float BlendTimeOnMirrorStateChange = 0.0f;
+	// Inertialization blend time to use when transitioning between mirrored and unmirrored states
+	UPROPERTY(EditAnywhere, Category = MirrorTransition, meta = (PinHiddenByDefault, FoldProperty))
+	float BlendTime = 0.0f;
+
+	// Blend profile to use when transitioning between mirrored and unmirrored states
+	UPROPERTY(EditAnywhere, Category = MirrorTransition, meta = (PinHiddenByDefault, UseAsBlendProfile = true, FoldProperty))
+	TObjectPtr<UBlendProfile> BlendProfile = nullptr;
+
+	// Whether to reset (reinitialize) the child (source) pose when the mirror state changes
+	UPROPERTY(EditAnywhere, Category = MirrorTransition, meta = (FoldProperty))
+	bool bResetChild = false;
 
 	UPROPERTY(EditAnywhere, Category = MirroredChannels, meta=(DisplayName="Bone", FoldProperty))
 	bool bBoneMirroring = true;
@@ -137,10 +154,6 @@ protected:
 
 	UPROPERTY(EditAnywhere, Category = MirroredChannels, meta=(DisplayName = "Attributes", FoldProperty))
 	bool bAttributeMirroring = true;
-
-	// Whether to reset (reinitialize) the child (source)pose when the mirror state changes
-	UPROPERTY(EditAnywhere, Category = Settings, meta = ( FoldProperty))
-	bool bResetChildOnMirrorStateChange = false;
 #endif
 };
 
@@ -159,6 +172,7 @@ public:
 	virtual bool SetMirrorDataTable(UMirrorDataTable* MirrorTable) override;
 
 	virtual bool GetMirror() const override;
+	virtual UBlendProfile* GetBlendProfile() const;
 	virtual float GetBlendTimeOnMirrorStateChange() const override;
 	virtual bool GetBoneMirroring() const override;
 	virtual bool GetCurveMirroring() const override;
@@ -166,6 +180,7 @@ public:
 	virtual bool GetResetChildOnMirrorStateChange() const override;
 
 	virtual bool SetMirror(bool bInMirror) override;
+	virtual bool SetBlendProfile(UBlendProfile* BlendProfile);
 	virtual bool SetBlendTimeOnMirrorStateChange(float InBlendTime) override;
 	virtual bool SetBoneMirroring(bool bInBoneMirroring) override;
 	virtual bool SetCurveMirroring(bool bInCurveMirroring) override;
@@ -180,8 +195,17 @@ protected:
 	UPROPERTY(EditAnywhere, Category = Settings)
 	TObjectPtr<UMirrorDataTable> MirrorDataTable = nullptr;
 
-	UPROPERTY(EditAnywhere, Category = Settings, meta = (PinHiddenByDefault, FoldProperty))
-	float BlendTimeOnMirrorStateChange = 0.0f;
+	// Inertialization blend time to use when transitioning between mirrored and unmirrored states
+	UPROPERTY(EditAnywhere, Category = MirrorTransition, meta = (PinHiddenByDefault, FoldProperty))
+	float BlendTime = 0.0f;
+
+	// Blend profile to use when transitioning between mirrored and unmirrored states
+	UPROPERTY(EditAnywhere, Category = MirrorTransition, meta = (PinHiddenByDefault, UseAsBlendProfile = true, FoldProperty))
+	TObjectPtr<UBlendProfile> BlendProfile = nullptr;
+
+	// Whether to reset (reinitialize) the child (source) pose when the mirror state changes
+	UPROPERTY(EditAnywhere, Category = MirrorTransition, meta = (FoldProperty))
+	bool bResetChild = false;
 
 	UPROPERTY(EditAnywhere, Category = MirroredChannels, meta = (DisplayName = "Bone", FoldProperty))
 	bool bBoneMirroring = true;
@@ -191,8 +215,4 @@ protected:
 
 	UPROPERTY(EditAnywhere, Category = MirroredChannels, meta = (DisplayName = "Attributes", FoldProperty))
 	bool bAttributeMirroring = true;
-
-	// Whether to reset (reinitialize) the child (source)pose when the mirror state changes
-	UPROPERTY(EditAnywhere, Category = Settings, meta = (FoldProperty))
-	bool bResetChildOnMirrorStateChange = false;
 };
\ No newline at end of file
diff --git a/Engine/Source/Runtime/Engine/Classes/Animation/AnimClassInterface.h b/Engine/Source/Runtime/Engine/Classes/Animation/AnimClassInterface.h
--- a/Engine/Source/Runtime/Engine/Classes/Animation/AnimClassInterface.h
+++ b/Engine/Source/Runtime/Engine/Classes/Animation/AnimClassInterface.h
@@ -147,6 +147,12 @@ struct FAnimGraphBlendOptions
 	UPROPERTY(EditAnywhere, Category = GraphBlending)
 	float BlendInTime;       
 
+	/**
+	* Optional blend profile to use when blending this graph in (if BlendInTime > 0)
+	*/
+	UPROPERTY(EditAnywhere, Category = GraphBlending, meta = (UseAsBlendProfile = true))
+	TObjectPtr<UBlendProfile> BlendInProfile;
+
 	/**
 	* Time to blend this graph out using Inertialization. Specify -1.0 to defer to the BlendInTime of the next graph.
 	* To blend this graph out you must place an Inertialization node after the Linked Anim Graph node or Linked Anim Layer node that uses this graph.
@@ -154,9 +160,17 @@ struct FAnimGraphBlendOptions
 	UPROPERTY(EditAnywhere, Category = GraphBlending)
 	float BlendOutTime;
 
+	/**
+	* Optional blend profile to use when blending this graph out (if BlendOutTime > 0)
+	*/
+	UPROPERTY(EditAnywhere, Category = GraphBlending, meta = (UseAsBlendProfile = true))
+	TObjectPtr<UBlendProfile> BlendOutProfile;
+
 	FAnimGraphBlendOptions()
 		: BlendInTime(-1.0f)
+		, BlendInProfile(nullptr)
 		, BlendOutTime(-1.0f)
+		, BlendOutProfile(nullptr)
 	{}
 };
 
diff --git a/Engine/Source/Runtime/Engine/Classes/Animation/AnimNode_Inertialization.h b/Engine/Source/Runtime/Engine/Classes/Animation/AnimNode_Inertialization.h
--- a/Engine/Source/Runtime/Engine/Classes/Animation/AnimNode_Inertialization.h
+++ b/Engine/Source/Runtime/Engine/Classes/Animation/AnimNode_Inertialization.h
@@ -30,7 +30,7 @@ public:
 
 	// Request to activate inertialization for a duration.
 	// If multiple requests are made on the same inertialization node, the minimum requested time will be used.
-	virtual void RequestInertialization(float InRequestedDuration) = 0;
+	virtual void RequestInertialization(float InRequestedDuration, const UBlendProfile* InBlendProfile = nullptr) = 0;
 
 	// Add a record of this request
 	virtual void AddDebugRecord(const FAnimInstanceProxy& InSourceProxy, int32 InSourceNodeId) = 0;
@@ -112,6 +112,47 @@ struct ENGINE_API FInertializationCurve
 	}
 };
 
+USTRUCT()
+struct FInertializationRequest
+{
+	GENERATED_BODY()
+
+	FInertializationRequest()
+		: Duration(-1.0f)
+		, BlendProfile(nullptr)
+	{
+	}
+
+	FInertializationRequest(float InDuration, const UBlendProfile* InBlendProfile)
+		: Duration(InDuration)
+		, BlendProfile(InBlendProfile)
+	{
+	}
+
+	void Clear()
+	{
+		Duration = -1.0f;
+		BlendProfile = nullptr;
+	}
+
+	friend bool operator==(const FInertializationRequest& A, const FInertializationRequest& B)
+	{
+		return A.Duration == B.Duration && A.BlendProfile == B.BlendProfile;
+	}
+
+	friend bool operator!=(const FInertializationRequest& A, const FInertializationRequest& B)
+	{
+		return !(A == B);
+	}
+
+	UPROPERTY(Transient)
+	float Duration;
+
+	UPROPERTY(Transient)
+	TObjectPtr<const UBlendProfile> BlendProfile;
+};
+
+
 USTRUCT()
 struct FInertializationPose
 {
@@ -245,18 +286,24 @@ struct FInertializationPoseDiff
 
 	// Initialize the pose difference from the current pose and the two previous snapshots
 	//
-	// Pose					the current frame's pose
-	// ComponentTransform	the current frame's component to world transform
-	// AttachParentName		the current frame's attach parent name (for checking if the attachment has changed)
-	// Prev1				the previous frame's pose
-	// Prev2				the pose from two frames before
-	// FilteredCurvesUIDs	list of curves we don't want to inertialize
+	// Pose								the current frame's pose
+	// ComponentTransform				the current frame's component to world transform
+	// AttachParentName					the current frame's attach parent name (for checking if the attachment has changed)
+	// Prev1							the previous frame's pose
+	// Prev2							the pose from two frames before
+	// FilteredCurvesUIDs				list of curves we don't want to inertialize
 	//
 	void InitFrom(const FCompactPose& Pose, const FBlendedCurve& Curves, const FTransform& ComponentTransform, const FName& AttachParentName, const FInertializationPose& Prev1, const FInertializationPose& Prev2, const TSet<SmartName::UID_Type>& FilteredCurvesUIDs);
 
 	// Apply this difference to a pose, decaying over time as InertializationElapsedTime approaches InertializationDuration
 	//
-	void ApplyTo(FCompactPose& Pose, FBlendedCurve& Curves, float InertializationElapsedTime, float InertializationDuration) const;
+	// Pose								[in/out] the current frame's pose
+	// Curves							[in/out] the current frame's animation curves
+	// InertializationElapsedTime		time elapsed since the start of the inertialization
+	// InertializationDuration			total inertialization duration (used for curves)
+	// InertializationDurationPerBone	inertialization duration per bone (indexed by skeleton bone index) (used for pose)
+	//
+	void ApplyTo(FCompactPose& Pose, FBlendedCurve& Curves, float InertializationElapsedTime, float InertializationDuration, TArrayView<const float> InertializationDurationPerBone) const;
 
 	// Get the inertialization space for this pose diff (for debug display)
 	//
@@ -285,11 +332,17 @@ struct ENGINE_API FAnimNode_Inertialization : public FAnimNode_Base
 {
 	GENERATED_BODY()
 
-	UPROPERTY(EditAnywhere, BlueprintReadWrite, Category = Links)
+	UPROPERTY(EditAnywhere, Category = Links)
 	FPoseLink Source;
 
+private:
+
+	// Optional default blend profile to use when no blend profile is supplied with the inertialization request
+	UPROPERTY(EditAnywhere, Category = BlendProfile, meta = (UseAsBlendProfile = true))
+	TObjectPtr<UBlendProfile> DefaultBlendProfile = nullptr;
+
 	// List of curves that should not use inertial blending. These curves will instantly change when inertialization begins.
-	UPROPERTY(EditAnywhere, BlueprintReadWrite, Category = Filter)
+	UPROPERTY(EditAnywhere, Category = Filter)
 	TArray<FName> FilteredCurves;
 
 public: // FAnimNode_Inertialization
@@ -299,9 +352,7 @@ public: // FAnimNode_Inertialization
 	// Request to activate inertialization for a duration.
 	// If multiple requests are made on the same inertialization node, the minimum requested time will be used.
 	//
-	virtual void RequestInertialization(float Duration);
-
-	virtual float GetRequestedDuration() const { return RequestedDuration; }
+	virtual void RequestInertialization(float Duration, const UBlendProfile* BlendProfile);
 
 	// Log an error when a node wants to inertialize but no inertialization ancestor node exists
 	//
@@ -320,28 +371,20 @@ public: // FAnimNode_Base
 
 protected:
 
-	// Consume Inertialization Request
-	//
-	// Returns any pending inertialization request and removes it from future processing.  Returns zero if there is no pending request.
-	// This function is virtual so that a derived class could optionally hook into other external sources of inertialization requests
-	// (for example from the owning actor for requests triggered from game code).
-	//
-	virtual float ConsumeInertializationRequest(FPoseContext& Context);
-
 	// Start Inertialization
 	//
 	// Computes the inertialization pose difference from the current pose and the two previous poses (to capture velocity).  This function
 	// is virtual so that a derived class could optionally regularize the pose snapshots to align better with the current frame's pose
 	// before computing the inertial difference (for example to correct for instantaneous changes in the root relative to its children).
 	//
-	virtual void StartInertialization(FPoseContext& Context, FInertializationPose& PreviousPose1, FInertializationPose& PreviousPose2, float Duration, /*OUT*/ FInertializationPoseDiff& OutPoseDiff);
+	virtual void StartInertialization(FPoseContext& Context, FInertializationPose& PreviousPose1, FInertializationPose& PreviousPose2, float Duration, TArrayView<const float> DurationPerBone, /*OUT*/ FInertializationPoseDiff& OutPoseDiff);
 
 	// Apply Inertialization
 	//
 	// Applies the inertialization pose difference to the current pose (feathering down to zero as ElapsedTime approaches Duration).  This
 	// function is virtual so that a derived class could optionally adjust the pose based on any regularization done in StartInertialization.
 	//
-	virtual void ApplyInertialization(FPoseContext& Context, const FInertializationPoseDiff& PoseDiff, float ElapsedTime, float Duration);
+	virtual void ApplyInertialization(FPoseContext& Context, const FInertializationPoseDiff& PoseDiff, float ElapsedTime, float Duration, TArrayView<const float> DurationPerBone);
 
 
 private:
@@ -355,8 +398,9 @@ private:
 	// Elapsed delta time between calls to evaluate
 	float DeltaTime;
 
-	// Pending inertialization request
-	float RequestedDuration;
+	// Pending inertialization requests
+	UPROPERTY(Transient)
+	TArray<FInertializationRequest> RequestQueue;
 
 	// Teleport type
 	ETeleportType TeleportType;
@@ -364,7 +408,17 @@ private:
 	// Inertialization state
 	EInertializationState InertializationState;
 	float InertializationElapsedTime;
+
+	// Inertialization duration for the main inertialization request (used for curve blending and deficit tracking)
 	float InertializationDuration;
+
+	// Inertialization durations indexed by skeleton bone index (used for per-bone blending)
+	TCustomBoneIndexArray<float, FSkeletonPoseBoneIndex> InertializationDurationPerBone;
+
+	// Maximum of InertializationDuration and all entries in InertializationDurationPerBone (used for knowing when to shutdown the inertialization)
+	float InertializationMaxDuration;
+
+	// Inertialization deficit (for tracking and reducing 'pose melting' when thrashing inertialization requests)
 	float InertializationDeficit;
 
 	// Inertialization pose differences
diff --git a/Engine/Source/Runtime/Engine/Classes/Animation/AnimNode_LinkedAnimGraph.h b/Engine/Source/Runtime/Engine/Classes/Animation/AnimNode_LinkedAnimGraph.h
--- a/Engine/Source/Runtime/Engine/Classes/Animation/AnimNode_LinkedAnimGraph.h
+++ b/Engine/Source/Runtime/Engine/Classes/Animation/AnimNode_LinkedAnimGraph.h
@@ -56,8 +56,19 @@ public:
 	int32 CachedLinkedNodeIndex;
 
 protected:
-	// Inertial blending duration to request next update
-	float PendingBlendDuration;
+	// Inertial blending duration to request next update (pulled from the prior state's blend out)
+	float PendingBlendOutDuration;
+
+	// Optional blend profile to use during inertial blending (pulled from the prior state's blend out)
+	UPROPERTY(Transient)
+	TObjectPtr<const UBlendProfile> PendingBlendOutProfile;
+
+	// Inertial blending duration to request next update (pulled from the new state's blend in)
+	float PendingBlendInDuration;
+
+	// Optional blend profile to use during inertial blending (pulled from the new state's blend in)
+	UPROPERTY(Transient)
+	TObjectPtr<const UBlendProfile> PendingBlendInProfile;
 
 public:
 	/** Whether named notifies will be received by this linked instance from other instances (outer or other linked instances) */
diff --git a/Engine/Source/Runtime/Engine/Classes/Animation/BlendProfile.h b/Engine/Source/Runtime/Engine/Classes/Animation/BlendProfile.h
--- a/Engine/Source/Runtime/Engine/Classes/Animation/BlendProfile.h
+++ b/Engine/Source/Runtime/Engine/Classes/Animation/BlendProfile.h
@@ -143,6 +143,12 @@ public:
 	 */
 	void FillBoneScalesArray(TArray<float>& OutBoneBlendProfileFactors, const FBoneContainer& BoneContainer) const;
 
+	/** Fill an array of floats with the bone duration values. One for each bone in the skeleton pose.
+	 * @param OutDurationPerBone Must be sized to the number bones in the skeleton pose. It will be filled with the durations of each bone as setup in the blend profile editor.
+	 * @param Duration The duration of the blend.
+	 */
+	void FillSkeletonBoneDurationsArray(TCustomBoneIndexArrayView<float, FSkeletonPoseBoneIndex> OutDurationPerBone, float Duration) const;
+
 	// IInterpolationIndexProvider
 	virtual int32 GetPerBoneInterpolationIndex(const FCompactPoseBoneIndex& InCompactPoseBoneIndex, const FBoneContainer& BoneContainer, const IInterpolationIndexProvider::FPerBoneInterpolationData* Data) const override;
 	// End IInterpolationIndexProvider
diff --git a/Engine/Source/Runtime/Engine/Private/Animation/AnimNode_Inertialization.cpp b/Engine/Source/Runtime/Engine/Private/Animation/AnimNode_Inertialization.cpp
--- a/Engine/Source/Runtime/Engine/Private/Animation/AnimNode_Inertialization.cpp
+++ b/Engine/Source/Runtime/Engine/Private/Animation/AnimNode_Inertialization.cpp
@@ -5,6 +5,8 @@
 #include "AnimationRuntime.h"
 #include "Animation/AnimNodeMessages.h"
 #include "Animation/AnimNode_SaveCachedPose.h"
+#include "Animation/BlendProfile.h"
+#include "Algo/MaxElement.h"
 #include "HAL/LowLevelMemTracker.h"
 
 LLM_DEFINE_TAG(Animation_Inertialization);
@@ -37,9 +39,9 @@ public:
 
 private:
 	// IInertializationRequester interface
-	virtual void RequestInertialization(float InRequestedDuration) override
+	virtual void RequestInertialization(float InRequestedDuration, const UBlendProfile* InBlendProfile) override
 	{ 
-		Node.RequestInertialization(InRequestedDuration); 
+		Node.RequestInertialization(InRequestedDuration, InBlendProfile); 
 	}
 
 	virtual void AddDebugRecord(const FAnimInstanceProxy& InSourceProxy, int32 InSourceNodeId)
@@ -60,28 +62,40 @@ private:
 	FAnimInstanceProxy& Proxy;
 };
 
+
+static int32 GetNumSkeletonBones(const FBoneContainer& BoneContainer)
+{
+	const USkeleton* SkeletonAsset = BoneContainer.GetSkeletonAsset();
+	check(SkeletonAsset);
+
+	const FReferenceSkeleton& RefSkeleton = SkeletonAsset->GetReferenceSkeleton();
+	return RefSkeleton.GetNum();
+}
+
 }}	// namespace UE::Anim
 
+
 FAnimNode_Inertialization::FAnimNode_Inertialization()
 	: DeltaTime(0.0f)
-	, RequestedDuration(-1.0f)
 	, TeleportType(ETeleportType::None)
 	, InertializationState(EInertializationState::Inactive)
 	, InertializationElapsedTime(0.0f)
 	, InertializationDuration(0.0f)
+	, InertializationMaxDuration(0.0f)
 	, InertializationDeficit(0.0f)
 {
 }
 
 
-void FAnimNode_Inertialization::RequestInertialization(float Duration)
+void FAnimNode_Inertialization::RequestInertialization(float Duration, const UBlendProfile* BlendProfile)
 {
-	if (RequestedDuration < 0.0f || Duration < RequestedDuration)
+	if (Duration >= 0.0f)
 	{
-		RequestedDuration = Duration;
+		RequestQueue.AddUnique(FInertializationRequest(Duration, BlendProfile));
 	}
 }
 
+
 /*static*/ void FAnimNode_Inertialization::LogRequestError(const FAnimationUpdateContext& Context, const FPoseLinkBase& RequesterPoseLink)
 {
 #if WITH_EDITORONLY_DATA	
@@ -103,7 +117,12 @@ void FAnimNode_Inertialization::Initialize_AnyThread(const FAnimationInitializeC
 	FAnimNode_Base::Initialize_AnyThread(Context);
 	Source.Initialize(Context);
 
+
+	const int32 NumSkeletonBones = UE::Anim::GetNumSkeletonBones(Context.AnimInstanceProxy->GetRequiredBones());
+
 	PoseSnapshots.Empty(INERTIALIZATION_MAX_POSE_SNAPSHOTS);
+	RequestQueue.Reserve(8);
+	InertializationDurationPerBone.Reserve(NumSkeletonBones);
 
 	DeltaTime = 0.0f;
 
@@ -153,16 +172,21 @@ void FAnimNode_Inertialization::Update_AnyThread(const FAnimationUpdateContext&
 	{
 		// If we have a pending request forward the request to other Inertialization nodes
 		// that were skipped due to pose caching.
-		if(RequestedDuration >= 0.0f)
+		if(RequestQueue.Num() > 0)
 		{
 			// Cached poses have their Update function called once even though there may be multiple UseCachedPose nodes for the same pose.
 			// Because of this, there may be Inertialization ancestors of the UseCachedPose nodes that missed out on requests.
 			// So here we forward 'this' node's requests to the ancestors of those skipped UseCachedPose nodes.
+			// Note that in some cases, we may be forwarding the requests back to this same node.  Those duplicate requests will ultimately
+			// be ignored by the 'AddUnique' in the body of FAnimNode_Inertialization::RequestInertialization.
 			for (const UE::Anim::FMessageStack& Stack : InSkippedUpdates)
 			{
 				Stack.ForEachMessage<UE::Anim::IInertializationRequester>([this, NodeId, &Proxy](UE::Anim::IInertializationRequester& InMessage)
 				{
-					InMessage.RequestInertialization(RequestedDuration);
+					for (const FInertializationRequest& Request : RequestQueue)
+					{
+						InMessage.RequestInertialization(Request.Duration, Request.BlendProfile);
+					}
  					InMessage.AddDebugRecord(Proxy, NodeId);
 
 					return UE::Anim::FMessageStack::EEnumerate::Stop;
@@ -177,7 +201,6 @@ void FAnimNode_Inertialization::Update_AnyThread(const FAnimationUpdateContext&
 	DeltaTime += Context.GetDeltaTime();
 }
 
-
 void FAnimNode_Inertialization::Evaluate_AnyThread(FPoseContext& Output)
 {
 	LLM_SCOPE_BYNAME(TEXT("Animation/Inertialization"));
@@ -185,12 +208,13 @@ void FAnimNode_Inertialization::Evaluate_AnyThread(FPoseContext& Output)
 
 	Source.Evaluate(Output);
 
-	// Extract any pending inertialization request
-	const float Duration = ConsumeInertializationRequest(Output);
-
 	// Disable inertialization if requested (for testing / debugging)
 	if (!CVarAnimInertializationEnable.GetValueOnAnyThread())
 	{
+		// Clear any pending inertialization requests
+		RequestQueue.Reset();
+
+		// Clear the inertialization state
 		Deactivate();
 
 		// Clear the pose history
@@ -204,36 +228,100 @@ void FAnimNode_Inertialization::Evaluate_AnyThread(FPoseContext& Output)
 	}
 
 	// Update the inertialization state if a new inertialization request is pending
-	if (Duration > SMALL_NUMBER && PoseSnapshots.Num() > 0)
+	const int32 NumRequests = RequestQueue.Num();
+	if (NumRequests > 0 && PoseSnapshots.Num() > 0)
 	{
 		float AppliedDeficit = 0.0f;
 		if (InertializationState == EInertializationState::Active)
 		{
 			// An active inertialization is being interrupted. Keep track of the lost inertialization time
-			// and reduce future durations if interruptions continue. Without this mitigation, 
+			// and reduce future durations if interruptions continue. Without this mitigation,
 			// repeated interruptions will lead to a degenerate pose because the pose target is unstable.
 			bool bApplyDeficit = InertializationDeficit > 0.0f && !CVarAnimInertializationIgnoreDeficit.GetValueOnAnyThread();
 			InertializationDeficit = InertializationDuration - InertializationElapsedTime;
-			AppliedDeficit = bApplyDeficit ? FMath::Min(Duration, InertializationDeficit) : 0.0f;
+			AppliedDeficit = bApplyDeficit ? InertializationDeficit : 0.0f;
 		}
+
 		InertializationState = EInertializationState::Pending;
 		InertializationElapsedTime = 0.0f;
-		InertializationDuration = Duration - AppliedDeficit;
+		
+		const int32 NumSkeletonBones = UE::Anim::GetNumSkeletonBones(Output.AnimInstanceProxy->GetRequiredBones());
+
+		auto FillSkeletonBoneDurationsArray = [this, NumSkeletonBones](auto& DurationPerBone, float Duration, const UBlendProfile* BlendProfile) {
+			if (BlendProfile == nullptr)
+			{
+				BlendProfile = DefaultBlendProfile;
+			}
+
+			if (BlendProfile != nullptr)
+			{
+				DurationPerBone.SetNum(NumSkeletonBones);
+				BlendProfile->FillSkeletonBoneDurationsArray(DurationPerBone, Duration);
+			}
+			else
+			{
+				DurationPerBone.Init(Duration, NumSkeletonBones);
+			}
+		};
+
+		// Handle the first inertialization request in the queue
+		InertializationDuration = FMath::Max(RequestQueue[0].Duration - AppliedDeficit, 0.0f);
+		FillSkeletonBoneDurationsArray(InertializationDurationPerBone, InertializationDuration, RequestQueue[0].BlendProfile);
+
+		// Handle all subsequent inertialization requests (often there will be only a single request)
+		if (NumRequests > 1)
+		{
+			UE::Anim::TTypedIndexArray<FSkeletonPoseBoneIndex, float, FAnimStackAllocator> RequestDurationPerBone;
+			for (int32 RequestIndex = 1; RequestIndex < NumRequests; ++RequestIndex)
+			{
+				const FInertializationRequest& Request = RequestQueue[RequestIndex];
+				const float RequestDuration = FMath::Max(Request.Duration - AppliedDeficit, 0.0f);
+
+				// Merge this request in with the previous requests (using the minimum requested time per bone)
+				InertializationDuration = FMath::Min(InertializationDuration, RequestDuration);
+				if (Request.BlendProfile != nullptr)
+				{
+					FillSkeletonBoneDurationsArray(RequestDurationPerBone, RequestDuration, Request.BlendProfile);
+					for (int32 SkeletonBoneIndex = 0; SkeletonBoneIndex < NumSkeletonBones; ++SkeletonBoneIndex)
+					{
+						InertializationDurationPerBone[SkeletonBoneIndex] = FMath::Min(InertializationDurationPerBone[SkeletonBoneIndex], RequestDurationPerBone[SkeletonBoneIndex]);
+					}
+				}
+				else
+				{
+					for (int32 SkeletonBoneIndex = 0; SkeletonBoneIndex < NumSkeletonBones; ++SkeletonBoneIndex)
+					{
+						InertializationDurationPerBone[SkeletonBoneIndex] = FMath::Min(InertializationDurationPerBone[SkeletonBoneIndex], RequestDuration);
+					}
+				}
+			}
+		}
+
+		// Cache the maximum duration across all bones (so we know when to deactivate the inertialization request)
+		InertializationMaxDuration = FMath::Max(InertializationDuration, *Algo::MaxElement(InertializationDurationPerBone));
 	}
 
+	RequestQueue.Reset();
+
 	// Update the inertialization timer
 	if (InertializationState != EInertializationState::Inactive)
 	{
 		InertializationElapsedTime += DeltaTime;
 		if (InertializationElapsedTime >= InertializationDuration)
 		{
-			Deactivate();
+			// Reset the deficit accumulator
+			InertializationDeficit = 0.0f;
 		}
 		else
 		{
 			// Pay down the accumulated deficit caused by interruptions
 			InertializationDeficit -= FMath::Min(InertializationDeficit, DeltaTime);
 		}
+
+		if (InertializationElapsedTime >= InertializationMaxDuration)
+		{
+			Deactivate();
+		}
 	}
 
 	const FTransform ComponentTransform = Output.AnimInstanceProxy->GetComponentTransform();
@@ -283,12 +371,17 @@ void FAnimNode_Inertialization::Evaluate_AnyThread(FPoseContext& Output)
 			PoseSnapshots[PoseSnapshots.Num() - 1],
 			PoseSnapshots[FMath::Max(PoseSnapshots.Num() - 2, 0)],
 			InertializationDuration,
+			InertializationDurationPerBone,
 			InertializationPoseDiff);
 		InertializationState = EInertializationState::Active;
 	}
 	if (InertializationState == EInertializationState::Active)
 	{
-		ApplyInertialization(Output, InertializationPoseDiff, InertializationElapsedTime, InertializationDuration);
+		ApplyInertialization(Output,
+			InertializationPoseDiff,
+			InertializationElapsedTime,
+			InertializationDuration,
+			InertializationDurationPerBone);
 	}
 
 	// Get the parent actor attachment information (to detect and counteract discontinuities when changing parents)
@@ -333,6 +426,7 @@ void FAnimNode_Inertialization::Evaluate_AnyThread(FPoseContext& Output)
 	TRACE_ANIM_NODE_VALUE(Output, TEXT("State"), *UEnum::GetValueAsString(InertializationState));
 	TRACE_ANIM_NODE_VALUE(Output, TEXT("Elapsed Time"), InertializationElapsedTime);
 	TRACE_ANIM_NODE_VALUE(Output, TEXT("Duration"), InertializationDuration);
+	TRACE_ANIM_NODE_VALUE(Output, TEXT("Max Duration"), InertializationMaxDuration);
 	TRACE_ANIM_NODE_VALUE(Output, TEXT("Normalized Time"), InertializationDuration > KINDA_SMALL_NUMBER ? (InertializationElapsedTime / InertializationDuration) : 0.0f);
 }
 
@@ -381,15 +475,7 @@ void FAnimNode_Inertialization::ResetDynamics(ETeleportType InTeleportType)
 }
 
 
-float FAnimNode_Inertialization::ConsumeInertializationRequest(FPoseContext& Context)
-{
-	const float Duration = RequestedDuration;
-	RequestedDuration = -1.0f;
-	return Duration;
-}
-
-
-void FAnimNode_Inertialization::StartInertialization(FPoseContext& Context, FInertializationPose& PreviousPose1, FInertializationPose& PreviousPose2, float Duration, /*OUT*/ FInertializationPoseDiff& OutPoseDiff)
+void FAnimNode_Inertialization::StartInertialization(FPoseContext& Context, FInertializationPose& PreviousPose1, FInertializationPose& PreviousPose2, float Duration, TArrayView<const float> DurationPerBone, /*OUT*/ FInertializationPoseDiff& OutPoseDiff)
 {
 	// Determine if this skeletal mesh's actor is attached to another actor
 	FName AttachParentName = NAME_None;
@@ -405,16 +491,19 @@ void FAnimNode_Inertialization::StartInertialization(FPoseContext& Context, FIne
 }
 
 
-void FAnimNode_Inertialization::ApplyInertialization(FPoseContext& Context, const FInertializationPoseDiff& PoseDiff, float ElapsedTime, float Duration)
+void FAnimNode_Inertialization::ApplyInertialization(FPoseContext& Context, const FInertializationPoseDiff& PoseDiff, float ElapsedTime, float Duration, TArrayView<const float> DurationPerBone)
 {
-	PoseDiff.ApplyTo(Context.Pose, Context.Curve, ElapsedTime, Duration);
+	PoseDiff.ApplyTo(Context.Pose, Context.Curve, ElapsedTime, Duration, DurationPerBone);
 }
 
+
 void FAnimNode_Inertialization::Deactivate()
 {
 	InertializationState = EInertializationState::Inactive;
 	InertializationElapsedTime = 0.0f;
 	InertializationDuration = 0.0f;
+	InertializationDurationPerBone.Reset();
+	InertializationMaxDuration = 0.0f;
 	InertializationDeficit = 0.0f;
 }
 
@@ -423,7 +512,7 @@ void FInertializationPose::InitFrom(const FCompactPose& Pose, const FBlendedCurv
 {
 	const FBoneContainer& BoneContainer = Pose.GetBoneContainer();
 
-	const int32 NumSkeletonBones = BoneContainer.GetSkeletonAsset()->GetReferenceSkeleton().GetNum();
+	const int32 NumSkeletonBones = UE::Anim::GetNumSkeletonBones(BoneContainer);
 	BoneTransforms.Reset(NumSkeletonBones);
 	BoneTransforms.AddZeroed(NumSkeletonBones);
 	BoneStates.Reset(NumSkeletonBones);
@@ -479,8 +568,7 @@ void FInertializationPoseDiff::InitFrom(const FCompactPose& Pose, const FBlended
 	}
 
 	// Compute the inertialization differences for each bone
-	const FReferenceSkeleton& RefSkeleton = BoneContainer.GetSkeletonAsset()->GetReferenceSkeleton();
-	const int32 NumSkeletonBones = RefSkeleton.GetNum();
+	const int32 NumSkeletonBones = UE::Anim::GetNumSkeletonBones(BoneContainer);
 	BoneDiffs.Empty(NumSkeletonBones);
 	BoneDiffs.AddZeroed(NumSkeletonBones);
 	for (FCompactPoseBoneIndex BoneIndex : Pose.ForEachBoneIndex())
@@ -662,7 +750,7 @@ void FInertializationPoseDiff::InitFrom(const FCompactPose& Pose, const FBlended
 
 // Apply this difference to a pose and a set of curves, decaying over time as InertializationElapsedTime approaches InertializationDuration
 //
-void FInertializationPoseDiff::ApplyTo(FCompactPose& Pose, FBlendedCurve& Curves, float InertializationElapsedTime, float InertializationDuration) const
+void FInertializationPoseDiff::ApplyTo(FCompactPose& Pose, FBlendedCurve& Curves, float InertializationElapsedTime, float InertializationDuration, TArrayView<const float> InertializationDurationPerBone) const
 {
 	const FBoneContainer& BoneContainer = Pose.GetBoneContainer();
 
@@ -674,20 +762,21 @@ void FInertializationPoseDiff::ApplyTo(FCompactPose& Pose, FBlendedCurve& Curves
 		if (SkeletonPoseBoneIndex != INDEX_NONE)
 		{
 			const FInertializationBoneDiff& BoneDiff = BoneDiffs[SkeletonPoseBoneIndex];
+			const float Duration = InertializationDurationPerBone[SkeletonPoseBoneIndex];
 
 			// Apply the bone translation difference
 			const FVector T = BoneDiff.TranslationDirection *
-				CalcInertialFloat(BoneDiff.TranslationMagnitude, BoneDiff.TranslationSpeed, InertializationElapsedTime, InertializationDuration);
+				CalcInertialFloat(BoneDiff.TranslationMagnitude, BoneDiff.TranslationSpeed, InertializationElapsedTime, Duration);
 			Pose[BoneIndex].AddToTranslation(T);
 
 			// Apply the bone rotation difference
 			const FQuat Q = FQuat(BoneDiff.RotationAxis,
-				CalcInertialFloat(BoneDiff.RotationAngle, BoneDiff.RotationSpeed, InertializationElapsedTime, InertializationDuration));
+				CalcInertialFloat(BoneDiff.RotationAngle, BoneDiff.RotationSpeed, InertializationElapsedTime, Duration));
 			Pose[BoneIndex].SetRotation(Q * Pose[BoneIndex].GetRotation());
 
 			// Apply the bone scale difference
 			const FVector S = BoneDiff.ScaleAxis *
-				CalcInertialFloat(BoneDiff.ScaleMagnitude, BoneDiff.ScaleSpeed, InertializationElapsedTime, InertializationDuration);
+				CalcInertialFloat(BoneDiff.ScaleMagnitude, BoneDiff.ScaleSpeed, InertializationElapsedTime, Duration);
 			Pose[BoneIndex].SetScale3D(S + Pose[BoneIndex].GetScale3D());
 		}
 	}
diff --git a/Engine/Source/Runtime/Engine/Private/Animation/AnimNode_LinkedAnimGraph.cpp b/Engine/Source/Runtime/Engine/Private/Animation/AnimNode_LinkedAnimGraph.cpp
--- a/Engine/Source/Runtime/Engine/Private/Animation/AnimNode_LinkedAnimGraph.cpp
+++ b/Engine/Source/Runtime/Engine/Private/Animation/AnimNode_LinkedAnimGraph.cpp
@@ -8,27 +8,6 @@
 #include "Animation/AnimNode_Root.h"
 #include "Animation/AnimTrace.h"
 
-static float GetBlendDuration(const IAnimClassInterface* PriorAnimBPClass, const IAnimClassInterface* NewAnimBPClass, FName Layer)
-{
-	const FAnimGraphBlendOptions* PriorBlendOptions = PriorAnimBPClass ? PriorAnimBPClass->GetGraphBlendOptions().Find(Layer) : nullptr;
-	const FAnimGraphBlendOptions* NewBlendOptions = NewAnimBPClass ? NewAnimBPClass->GetGraphBlendOptions().Find(Layer) : nullptr;
-
-	float BlendOutTime = PriorBlendOptions ? PriorBlendOptions->BlendOutTime : -1.0f;
-	float BlendInTime = NewBlendOptions ? NewBlendOptions->BlendInTime : -1.0f;
-
-	if (BlendInTime < 0.0f)
-	{
-		return BlendOutTime;
-	}
-
-	if (BlendOutTime < 0.0f)
-	{
-		return BlendInTime;
-	}
-
-	return FMath::Min(BlendInTime, BlendOutTime);
-}
-
 FAnimNode_LinkedAnimGraph::FAnimNode_LinkedAnimGraph()
 	: InstanceClass(nullptr)
 #if WITH_EDITORONLY_DATA
@@ -37,7 +16,10 @@ FAnimNode_LinkedAnimGraph::FAnimNode_LinkedAnimGraph()
 	, LinkedRoot(nullptr)
 	, NodeIndex(INDEX_NONE)
 	, CachedLinkedNodeIndex(INDEX_NONE)
-	, PendingBlendDuration(-1.0f)
+	, PendingBlendOutDuration(-1.0f)
+	, PendingBlendOutProfile(nullptr)
+	, PendingBlendInDuration(-1.0f)
+	, PendingBlendInProfile(nullptr)
 	, bReceiveNotifiesFromLinkedInstances(false)
 	, bPropagateNotifiesToLinkedInstances(false)
 {
@@ -125,22 +107,32 @@ void FAnimNode_LinkedAnimGraph::Update_AnyThread(const FAnimationUpdateContext&
 		InputPoses[0].Update(InContext);
 	}
 
-	// Consume pending inertial blend request
-	if(PendingBlendDuration >= 0.0f)
+	// Consume pending inertial blend requests
+	if(PendingBlendOutDuration >= 0.0f || PendingBlendInDuration >= 0.0f)
 	{
 		UE::Anim::IInertializationRequester* InertializationRequester = InContext.GetMessage<UE::Anim::IInertializationRequester>();
 		if(InertializationRequester)
 		{
-			InertializationRequester->RequestInertialization(PendingBlendDuration);
+			// Issue the pending inertialization requests (which will get merged together by the inertialization node itself)
+			if (PendingBlendOutDuration >= 0.0f)
+			{
+				InertializationRequester->RequestInertialization(PendingBlendOutDuration, PendingBlendOutProfile);
+			}
+			if (PendingBlendInDuration >= 0.0f)
+			{
+				InertializationRequester->RequestInertialization(PendingBlendInDuration, PendingBlendInProfile);
+			}
 			InertializationRequester->AddDebugRecord(*InContext.AnimInstanceProxy, InContext.GetCurrentNodeId());
 		}
-		else if ((PendingBlendDuration != 0.0f) && (InputPoses.Num() > 0))
+		else if ((PendingBlendOutDuration != 0.0f) && (PendingBlendInDuration != 0.0f) && (InputPoses.Num() > 0))
 		{
 			FAnimNode_Inertialization::LogRequestError(InContext, InputPoses[0]);
 		}
-
-		PendingBlendDuration = -1.0f;
 	}
+	PendingBlendOutDuration = -1.0f;
+	PendingBlendOutProfile = nullptr;
+	PendingBlendInDuration = -1.0f;
+	PendingBlendInProfile = nullptr;
 
 	TRACE_ANIM_NODE_VALUE(InContext, TEXT("Name"), GetDynamicLinkFunctionName());
 	TRACE_ANIM_NODE_VALUE(InContext, TEXT("Target Class"), InstanceClass.Get());
@@ -453,5 +445,30 @@ int32 FAnimNode_LinkedAnimGraph::FindFunctionInputIndex(const FAnimBlueprintFunc
 
 void FAnimNode_LinkedAnimGraph::RequestBlend(const IAnimClassInterface* PriorAnimBPClass, const IAnimClassInterface* NewAnimBPClass)
 {
-	PendingBlendDuration = GetBlendDuration(PriorAnimBPClass, NewAnimBPClass, GetDynamicLinkFunctionName());
+	const FName Layer = GetDynamicLinkFunctionName();
+
+	const FAnimGraphBlendOptions* PriorBlendOptions = PriorAnimBPClass ? PriorAnimBPClass->GetGraphBlendOptions().Find(Layer) : nullptr;
+	const FAnimGraphBlendOptions* NewBlendOptions = NewAnimBPClass ? NewAnimBPClass->GetGraphBlendOptions().Find(Layer) : nullptr;
+
+	if (PriorBlendOptions && PriorBlendOptions->BlendOutTime >= 0.0f)
+	{
+		PendingBlendOutDuration = PriorBlendOptions->BlendOutTime;
+		PendingBlendOutProfile = PriorBlendOptions->BlendOutProfile;
+	}
+	else
+	{
+		PendingBlendOutDuration = -1.0f;
+		PendingBlendOutProfile = nullptr;
+	}
+
+	if (NewBlendOptions && NewBlendOptions->BlendInTime >= 0.0f)
+	{
+		PendingBlendInDuration = NewBlendOptions->BlendInTime;
+		PendingBlendInProfile = NewBlendOptions->BlendInProfile;
+	}
+	else
+	{
+		PendingBlendInDuration = -1.0f;
+		PendingBlendInProfile = nullptr;
+	}
 }
diff --git a/Engine/Source/Runtime/Engine/Private/Animation/AnimNode_StateMachine.cpp b/Engine/Source/Runtime/Engine/Private/Animation/AnimNode_StateMachine.cpp
--- a/Engine/Source/Runtime/Engine/Private/Animation/AnimNode_StateMachine.cpp
+++ b/Engine/Source/Runtime/Engine/Private/Animation/AnimNode_StateMachine.cpp
@@ -462,7 +462,7 @@ void FAnimNode_StateMachine::Update_AnyThread(const FAnimationUpdateContext& Con
 					UE::Anim::IInertializationRequester* InertializationRequester = Context.GetMessage<UE::Anim::IInertializationRequester>();
 					if (InertializationRequester)
 					{
-						InertializationRequester->RequestInertialization(ReferenceTransition.CrossfadeDuration);
+						InertializationRequester->RequestInertialization(ReferenceTransition.CrossfadeDuration, ReferenceTransition.BlendProfile);
 						InertializationRequester->AddDebugRecord(*Context.AnimInstanceProxy, Context.GetCurrentNodeId());
 					}
 					else
diff --git a/Engine/Source/Runtime/Engine/Private/Animation/BlendProfile.cpp b/Engine/Source/Runtime/Engine/Private/Animation/BlendProfile.cpp
--- a/Engine/Source/Runtime/Engine/Private/Animation/BlendProfile.cpp
+++ b/Engine/Source/Runtime/Engine/Private/Animation/BlendProfile.cpp
@@ -209,6 +209,51 @@ void UBlendProfile::FillBoneScalesArray(TArray<float>& OutBoneBlendProfileFactor
 	}
 }
 
+void UBlendProfile::FillSkeletonBoneDurationsArray(TCustomBoneIndexArrayView<float, FSkeletonPoseBoneIndex> OutDurationPerBone, float Duration) const
+{
+	check(OwningSkeleton != nullptr);
+	const FReferenceSkeleton& RefSkeleton = OwningSkeleton->GetReferenceSkeleton();
+	const int32 NumSkeletonBones = RefSkeleton.GetNum();
+	check(OutDurationPerBone.Num() == NumSkeletonBones);
+
+	for(float& BoneDuration: OutDurationPerBone)
+	{
+		BoneDuration = Duration;
+	}
+
+	switch (Mode)
+	{
+		case EBlendProfileMode::TimeFactor:
+		{
+			for (const FBlendProfileBoneEntry& Entry : ProfileEntries)
+			{
+				const FSkeletonPoseBoneIndex SkeletonBoneIndex(Entry.BoneReference.BoneIndex);
+				OutDurationPerBone[SkeletonBoneIndex] *= Entry.BlendScale;
+			}
+		}
+		break;
+
+		case EBlendProfileMode::WeightFactor:
+		{
+			for (const FBlendProfileBoneEntry& Entry : ProfileEntries)
+			{
+				const FSkeletonPoseBoneIndex SkeletonBoneIndex(Entry.BoneReference.BoneIndex);
+				if (Entry.BlendScale > SMALL_NUMBER)
+				{
+					OutDurationPerBone[SkeletonBoneIndex] /= Entry.BlendScale;
+				}
+			}
+		}
+		break;
+
+		default:
+		{
+			checkf(false, TEXT("The selected Blend Profile Mode is not supported (Mode=%d)"), Mode);
+		}
+		break;
+	}
+}
+
 float UBlendProfile::CalculateBoneWeight(float BoneFactor, EBlendProfileMode Mode, const FAlphaBlend& BlendInfo, float BlendStartAlpha, float MainWeight, bool bInverse)
 {
 	switch (Mode)
