Starting to add a physics thread version of hit result structures so we can do SQ in fixed tick and physics thread.
Not actually ready to use yet, but nothing is calling the non-gt version

diff --git a/Engine/Source/Runtime/Engine/Private/Collision/CollisionConversions.cpp b/Engine/Source/Runtime/Engine/Private/Collision/CollisionConversions.cpp
--- a/Engine/Source/Runtime/Engine/Private/Collision/CollisionConversions.cpp
+++ b/Engine/Source/Runtime/Engine/Private/Collision/CollisionConversions.cpp
@@ -22,6 +22,7 @@
 #endif
 #include "PhysicsEngine/CollisionQueryFilterCallback.h"
 #include "GameFramework/LightWeightInstanceManager.h"
+#include "PhysicsProxy/SingleParticlePhysicsProxy.h"
 
 // Used to place overlaps into a TMap when deduplicating them
 struct FOverlapKey
@@ -50,7 +51,8 @@ uint32 GetTypeHash(const FOverlapKey& Key)
 extern int32 CVarShowInitialOverlaps;
 
 // Forward declare, I don't want to move the entire function right now or we lose change history.
-static bool ConvertOverlappedShapeToImpactHit(const UWorld* World, const FHitLocation& Hit, const FVector& StartLoc, const FVector& EndLoc, FHitResult& OutResult, const FPhysicsGeometry& Geom, const FTransform& QueryTM, const FCollisionFilterData& QueryFilter, bool bReturnPhysMat);
+template <typename THitLocation>
+static bool ConvertOverlappedShapeToImpactHit(const UWorld* World, const THitLocation& Hit, const FVector& StartLoc, const FVector& EndLoc, FHitResult& OutResult, const FPhysicsGeometry& Geom, const FTransform& QueryTM, const FCollisionFilterData& QueryFilter, bool bReturnPhysMat);
 
 DECLARE_CYCLE_STAT(TEXT("ConvertQueryHit"), STAT_ConvertQueryImpactHit, STATGROUP_Collision);
 DECLARE_CYCLE_STAT(TEXT("ConvertOverlapToHit"), STAT_CollisionConvertOverlapToHit, STATGROUP_Collision);
@@ -97,7 +99,12 @@ static FVector FindSimpleOpposingNormal(const FHitLocation& Hit, const FVector&
  * @param InNormal - default value in case no new normal is computed.
  * @return New normal we compute for geometry.
  */
+#if WITH_CHAOS
+template <typename THitLocation>
+static FVector FindGeomOpposingNormal(ECollisionShapeType QueryGeomType, const THitLocation& Hit, const FVector& TraceDirectionDenorm, const FVector InNormal)
+#else
 static FVector FindGeomOpposingNormal(ECollisionShapeType QueryGeomType, const FHitLocation& Hit, const FVector& TraceDirectionDenorm, const FVector InNormal)
+#endif
 {
 	// TODO: can we support other shapes here as well?
 	if (QueryGeomType == ECollisionShapeType::Capsule || QueryGeomType == ECollisionShapeType::Sphere)
@@ -233,7 +240,29 @@ static void SetHitResultFromShapeAndFaceIndex(const FPhysicsShape& Shape,  const
 	OutResult.FaceIndex = INDEX_NONE;
 }
 
-EConvertQueryResult ConvertQueryImpactHit(const UWorld* World, const FHitLocation& Hit, FHitResult& OutResult, float CheckLength, const FCollisionFilterData& QueryFilter, const FVector& StartLoc, const FVector& EndLoc, const FPhysicsGeometry* Geom, const FTransform& QueryTM, bool bReturnFaceIndex, bool bReturnPhysMat)
+const FPhysicsActor* GetGTActor(const FPhysicsActor* GTActor) { return GTActor; }
+
+template <bool bGT>
+const FPhysicsShape* GetGTShape(const FPhysicsShape* GTShape) { return GTShape; }
+
+template <>
+const FPhysicsShape* GetGTShape<false>(const FPhysicsShape* GTShape) { return GTShape; ensure(false); }	//TODO: implement this
+
+#if WITH_CHAOS
+const FPhysicsActor* GetGTActor(const Chaos::FGeometryParticleHandle* PTActor)
+{
+	//TODO: need to pass in context so that in PT we always return null
+	//In frozen GT this is ok because object can't be unregistered while we're holding on to this
+	//On PT this is only true if we acquire a gt lock which is not great
+	auto Proxy = static_cast<const Chaos::FSingleParticlePhysicsProxy*>(PTActor->PhysicsProxy());
+	return Proxy->GetParticle_LowLevel();
+}
+
+const FPhysicsShape* GetGTShape(const Chaos::FPerShapeData* GTShape) { return GTShape; }
+#endif
+
+template <typename THitLocation>
+EConvertQueryResult ConvertQueryImpactHitImp(const UWorld* World, const THitLocation& Hit, FHitResult& OutResult, float CheckLength, const FCollisionFilterData& QueryFilter, const FVector& StartLoc, const FVector& EndLoc, const FPhysicsGeometry* Geom, const FTransform& QueryTM, bool bReturnFaceIndex, bool bReturnPhysMat)
 {
 	SCOPE_CYCLE_COUNTER(STAT_ConvertQueryImpactHit);
 
@@ -251,7 +280,7 @@ EConvertQueryResult ConvertQueryImpactHit(const UWorld* World, const FHitLocatio
 	checkSlow(Flags & EHitFlags::Distance);
 
 	const FPhysicsShape* pHitShape = GetShape(Hit);
-	const FPhysicsActor* pHitActor = GetActor(Hit);
+	const auto pHitActor = GetActor(Hit);
 
 	const uint32 InternalFaceIndex = GetInternalFaceIndex(Hit);
 	const bool bInitialOverlap = HadInitialOverlap(Hit);
@@ -280,7 +309,13 @@ EConvertQueryResult ConvertQueryImpactHit(const UWorld* World, const FHitLocatio
 	}
 
 	const FPhysicsShape& HitShape = *pHitShape;
-	const FPhysicsActor& HitActor = *pHitActor;
+	const auto& HitActor = *pHitActor;
+
+	//Get the GT actor/shape for accessing GT specific data, should not be used for things like hit location
+	const FPhysicsActor* GTActor = GetGTActor(pHitActor);
+	//Note that a non-gt query could still happen on gt if we're in a fixed tick
+	constexpr bool bIsGTQuery = std::is_same<THitLocation, FHitLocation>::value;
+	const FPhysicsShape* GTShape = GetGTShape<bIsGTQuery>(pHitShape);
 
 	// See if this is a 'blocking' hit
 	const FCollisionFilterData ShapeFilter = GetQueryFilterData(HitShape);
@@ -303,10 +338,10 @@ EConvertQueryResult ConvertQueryImpactHit(const UWorld* World, const FHitLocatio
 	if (bUseReturnedPoint)
 	{
 		Position = GetPosition(Hit);
-		if (Position.ContainsNaN())
+		if (Position.ContainsNaN() && GTActor && GTShape)
 		{
 #if ENABLE_NAN_DIAGNOSTIC
-			SetHitResultFromShapeAndFaceIndex(HitShape, HitActor, InternalFaceIndex, OutResult.ImpactPoint, OutResult, bReturnPhysMat);
+			SetHitResultFromShapeAndFaceIndex(*GTShape, *GTActor, InternalFaceIndex, OutResult.ImpactPoint, OutResult, bReturnPhysMat);
 			UE_LOG(LogCore, Error, TEXT("ConvertQueryImpactHit() NaN details:\n>> Actor:%s (%s)\n>> Component:%s\n>> Item:%d\n>> BoneName:%s\n>> Time:%f\n>> Distance:%f\n>> Location:%s\n>> bIsBlocking:%d\n>> bStartPenetrating:%d"),
 				*OutResult.GetHitObjectHandle().GetName(), OutResult.HasValidHitObjectHandle() ? *OutResult.GetHitObjectHandle().FetchActor()->GetPathName() : TEXT("no path"),
 				*GetNameSafe(OutResult.GetComponent()), OutResult.Item, *OutResult.BoneName.ToString(),
@@ -323,10 +358,10 @@ EConvertQueryResult ConvertQueryImpactHit(const UWorld* World, const FHitLocatio
 	// Caution: we may still have an initial overlap, but with null Geom. This is the case for RayCast results.
 	const bool bUseReturnedNormal = ((Flags & EHitFlags::Normal) && !bInitialOverlap);
 	const FVector HitNormal = GetNormal(Hit);
-	if (bUseReturnedNormal && HitNormal.ContainsNaN())
+	if (bUseReturnedNormal && HitNormal.ContainsNaN() && GTActor && GTShape)
 	{
 #if ENABLE_NAN_DIAGNOSTIC
-		SetHitResultFromShapeAndFaceIndex(HitShape, HitActor, InternalFaceIndex, OutResult.ImpactPoint, OutResult, bReturnPhysMat);
+		SetHitResultFromShapeAndFaceIndex(*GTShape, *GTActor, InternalFaceIndex, OutResult.ImpactPoint, OutResult, bReturnPhysMat);
 		UE_LOG(LogCore, Error, TEXT("ConvertQueryImpactHit() NaN details:\n>> Actor:%s (%s)\n>> Component:%s\n>> Item:%d\n>> BoneName:%s\n>> Time:%f\n>> Distance:%f\n>> Location:%s\n>> bIsBlocking:%d\n>> bStartPenetrating:%d"),
 			*OutResult.GetHitObjectHandle().GetName(), OutResult.HasValidHitObjectHandle() ? *OutResult.GetHitObjectHandle().FetchActor()->GetPathName() : TEXT("no path"),
 			*GetNameSafe(OutResult.GetComponent()), OutResult.Item, *OutResult.BoneName.ToString(),
@@ -347,7 +382,10 @@ EConvertQueryResult ConvertQueryImpactHit(const UWorld* World, const FHitLocatio
 	OutResult.TraceEnd = EndLoc;
 
 	// Fill in Actor, Component, material, etc.
-	SetHitResultFromShapeAndFaceIndex(HitShape, HitActor, InternalFaceIndex, OutResult.ImpactPoint, OutResult, bReturnPhysMat);
+	if(GTActor && GTShape)
+	{
+		SetHitResultFromShapeAndFaceIndex(*GTShape, *GTActor, InternalFaceIndex, OutResult.ImpactPoint, OutResult, bReturnPhysMat);
+	}
 
 #if ENABLE_CHECK_HIT_NORMAL
 	CheckHitResultNormal(OutResult, TEXT("Invalid Normal from ConvertQueryImpactHit"), StartLoc, EndLoc, Geom, OriginalNormal);
@@ -369,22 +407,31 @@ EConvertQueryResult ConvertQueryImpactHit(const UWorld* World, const FHitLocatio
 	if(GeomType == ECollisionShapeType::Heightfield)
 	{
 		// Lookup physical material for heightfields
-		if (bReturnPhysMat && InternalFaceIndex != GetInvalidPhysicsFaceIndex())
+		if (bReturnPhysMat && InternalFaceIndex != GetInvalidPhysicsFaceIndex() && GTActor && GTShape)
 		{
-			if (const FPhysicsMaterial* Material = GetMaterialFromInternalFaceIndex(HitShape, HitActor, InternalFaceIndex))
+			if (const FPhysicsMaterial* Material = GetMaterialFromInternalFaceIndex(*GTShape, *GTActor, InternalFaceIndex))
 			{
 				OutResult.PhysMaterial = GetUserData(*Material);
 			}
 		}
 	}
-	else if (bReturnFaceIndex && GeomType == ECollisionShapeType::Trimesh && InternalFaceIndex != GetInvalidPhysicsFaceIndex())
+	else if (bReturnFaceIndex && GeomType == ECollisionShapeType::Trimesh && InternalFaceIndex != GetInvalidPhysicsFaceIndex() && GTShape)
 	{
-		OutResult.FaceIndex = GetTriangleMeshExternalFaceIndex(HitShape, InternalFaceIndex);
+		OutResult.FaceIndex = GetTriangleMeshExternalFaceIndex(*GTShape, InternalFaceIndex);
 	}
 
 	return EConvertQueryResult::Valid;
 }
 
+EConvertQueryResult ConvertQueryImpactHit(const UWorld* World, const FHitLocation& Hit, FHitResult& OutResult, float CheckLength, const FCollisionFilterData& QueryFilter, const FVector& StartLoc, const FVector& EndLoc, const FPhysicsGeometry* Geom, const FTransform& QueryTM, bool bReturnFaceIndex, bool bReturnPhysMat)
+{
+	return ConvertQueryImpactHitImp(World, Hit, OutResult, CheckLength, QueryFilter, StartLoc, EndLoc, Geom, QueryTM, bReturnFaceIndex, bReturnPhysMat);
+}
+EConvertQueryResult ConvertQueryImpactHit(const UWorld* World, const FPTLocationHit& Hit, FHitResult& OutResult, float CheckLength, const FCollisionFilterData& QueryFilter, const FVector& StartLoc, const FVector& EndLoc, const FPhysicsGeometry* Geom, const FTransform& QueryTM, bool bReturnFaceIndex, bool bReturnPhysMat)
+{
+	return ConvertQueryImpactHitImp(World, Hit, OutResult, CheckLength, QueryFilter, StartLoc, EndLoc, Geom, QueryTM, bReturnFaceIndex, bReturnPhysMat);
+}
+
 template <typename HitType>
 EConvertQueryResult ConvertTraceResults(bool& OutHasValidBlockingHit, const UWorld* World, int32 NumHits, HitType* Hits, float CheckLength, const FCollisionFilterData& QueryFilter, TArray<FHitResult>& OutHits, const FVector& StartLoc, const FVector& EndLoc, const FPhysicsGeometry& Geom, const FTransform& QueryTM, float MaxDistance, bool bReturnFaceIndex, bool bReturnPhysMat)
 {
@@ -449,12 +496,13 @@ template EConvertQueryResult ConvertTraceResults<FHitRaycast>(bool& OutHasValidB
 template EConvertQueryResult ConvertTraceResults<FHitRaycast>(bool& OutHasValidBlockingHit, const UWorld* World, int32 NumHits, FHitRaycast* Hits, float CheckLength, const FCollisionFilterData& QueryFilter, FHitResult& OutHit, const FVector& StartLoc, const FVector& EndLoc, const FPhysicsGeometry& Geom, const FTransform& QueryTM, float MaxDistance, bool bReturnFaceIndex, bool bReturnPhysMat);
 
 /** Util to convert an overlapped shape into a sweep hit result, returns whether it was a blocking hit. */
-static bool ConvertOverlappedShapeToImpactHit(const UWorld* World, const FHitLocation& Hit, const FVector& StartLoc, const FVector& EndLoc, FHitResult& OutResult, const FPhysicsGeometry& Geom, const FTransform& QueryTM, const FCollisionFilterData& QueryFilter, bool bReturnPhysMat)
+template <typename THitLocation>
+static bool ConvertOverlappedShapeToImpactHit(const UWorld* World, const THitLocation& Hit, const FVector& StartLoc, const FVector& EndLoc, FHitResult& OutResult, const FPhysicsGeometry& Geom, const FTransform& QueryTM, const FCollisionFilterData& QueryFilter, bool bReturnPhysMat)
 {
 	SCOPE_CYCLE_COUNTER(STAT_CollisionConvertOverlapToHit);
 
 	const FPhysicsShape* pHitShape = GetShape(Hit);
-	const FPhysicsActor* pHitActor = GetActor(Hit);
+	const auto pHitActor = GetActor(Hit);
 	if ((pHitShape == nullptr) || (pHitActor == nullptr))
 	{
 		OutResult.Reset();
@@ -462,7 +510,7 @@ static bool ConvertOverlappedShapeToImpactHit(const UWorld* World, const FHitLoc
 	}
 
 	const FPhysicsShape& HitShape = *pHitShape;
-	const FPhysicsActor& HitActor = *pHitActor;
+	const auto& HitActor = *pHitActor;
 
 	// See if this is a 'blocking' hit
 	FCollisionFilterData ShapeFilter = GetQueryFilterData(HitShape);
@@ -549,7 +597,15 @@ static bool ConvertOverlappedShapeToImpactHit(const UWorld* World, const FHitLoc
 
 	OutResult.Normal = OutResult.ImpactNormal;
 	
-	SetHitResultFromShapeAndFaceIndex(HitShape, HitActor, GetInternalFaceIndex(Hit), OutResult.ImpactPoint, OutResult, bReturnPhysMat);
+	const FPhysicsActor* GTHitActor = GetGTActor(&HitActor);
+
+	//Note that a non-gt query could still happen on gt if we're in a fixed tick
+	constexpr bool bIsGTQuery = std::is_same<THitLocation, FHitLocation>::value;
+	const FPhysicsShape* GTHitShape = GetGTShape<bIsGTQuery>(&HitShape);
+	if(GTHitActor && GTHitShape)
+	{
+		SetHitResultFromShapeAndFaceIndex(*GTHitShape, *GTHitActor, GetInternalFaceIndex(Hit), OutResult.ImpactPoint, OutResult, bReturnPhysMat);
+	}
 
 	return bBlockingHit;
 }
diff --git a/Engine/Source/Runtime/Engine/Private/Collision/CollisionConversions.h b/Engine/Source/Runtime/Engine/Private/Collision/CollisionConversions.h
--- a/Engine/Source/Runtime/Engine/Private/Collision/CollisionConversions.h
+++ b/Engine/Source/Runtime/Engine/Private/Collision/CollisionConversions.h
@@ -31,6 +31,7 @@ enum class EConvertQueryResult
  * @return	Whether result passed NaN/Inf checks.
  */
 EConvertQueryResult ConvertQueryImpactHit(const UWorld* World, const FHitLocation& PHit, FHitResult& OutResult, float CheckLength, const FCollisionFilterData& QueryFilter, const FVector& StartLoc, const FVector& EndLoc, const FPhysicsGeometry* Geom, const FTransform& QueryTM, bool bReturnFaceIndex, bool bReturnPhysMat);
+EConvertQueryResult ConvertQueryImpactHit(const UWorld* World, const FPTLocationHit& PHit, FHitResult& OutResult, float CheckLength, const FCollisionFilterData& QueryFilter, const FVector& StartLoc, const FVector& EndLoc, const FPhysicsGeometry* Geom, const FTransform& QueryTM, bool bReturnFaceIndex, bool bReturnPhysMat);
 
 /** 
  * Util to convert physX sweep results to unreal hit results and add to array
diff --git a/Engine/Source/Runtime/Engine/Public/Physics/Experimental/PhysInterface_Chaos.h b/Engine/Source/Runtime/Engine/Public/Physics/Experimental/PhysInterface_Chaos.h
--- a/Engine/Source/Runtime/Engine/Public/Physics/Experimental/PhysInterface_Chaos.h
+++ b/Engine/Source/Runtime/Engine/Public/Physics/Experimental/PhysInterface_Chaos.h
@@ -148,6 +148,16 @@ FORCEINLINE void ComputeZeroDistanceImpactNormalAndPenetration(const UWorld* Wor
 	//TODO_SQ_IMPLEMENTATION
 }
 
+FORCEINLINE void DrawOverlappingTris(const UWorld* World, const FPTLocationHit& Hit, const Chaos::FImplicitObject& Geom, const FTransform& QueryTM)
+{
+	//TODO_SQ_IMPLEMENTATION
+}
+
+FORCEINLINE void ComputeZeroDistanceImpactNormalAndPenetration(const UWorld* World, const FPTLocationHit& Hit, const Chaos::FImplicitObject& Geom, const FTransform& QueryTM, FHitResult& OutResult)
+{
+	//TODO_SQ_IMPLEMENTATION
+}
+
 Chaos::FChaosPhysicsMaterial* GetMaterialFromInternalFaceIndex(const FPhysicsShape& Shape, const FPhysicsActor& Actor, uint32 InternalFaceIndex);
 
 Chaos::FChaosPhysicsMaterial* GetMaterialFromInternalFaceIndexAndHitLocation(const FPhysicsShape& Shape, const FPhysicsActor& Actor, uint32 InternalFaceIndex, const FVector& HitLocation);
diff --git a/Engine/Source/Runtime/PhysicsCore/Public/ChaosInterfaceWrapperCore.h b/Engine/Source/Runtime/PhysicsCore/Public/ChaosInterfaceWrapperCore.h
--- a/Engine/Source/Runtime/PhysicsCore/Public/ChaosInterfaceWrapperCore.h
+++ b/Engine/Source/Runtime/PhysicsCore/Public/ChaosInterfaceWrapperCore.h
@@ -116,51 +116,102 @@ inline bool HadInitialOverlap(const FLocationHit& Hit)
 	return Hit.Distance <= 0.f;
 }
 
+inline bool HadInitialOverlap(const FPTLocationHit& Hit)
+{
+	return Hit.Distance <= 0.f;
+}
+
 inline const Chaos::FPerShapeData* GetShape(const FActorShape& Hit)
 {
 	return Hit.Shape;
 }
 
+inline const Chaos::FPerShapeData* GetShape(const FPTActorShape& Hit)
+{
+	return Hit.Shape;
+}
+
 inline Chaos::FGeometryParticle* GetActor(const FActorShape& Hit)
 {
 	return Hit.Actor;
 }
 
+inline Chaos::FGeometryParticleHandle* GetActor(const FPTActorShape& Hit)
+{
+	return Hit.Actor;
+}
+
 inline Chaos::FReal GetDistance(const FLocationHit& Hit)
 {
 	return Hit.Distance;
 }
 
+inline Chaos::FReal GetDistance(const FPTLocationHit& Hit)
+{
+	return Hit.Distance;
+}
+
 inline FVector GetPosition(const FLocationHit& Hit)
 {
 	return Hit.WorldPosition;
 }
 
+inline FVector GetPosition(const FPTLocationHit& Hit)
+{
+	return Hit.WorldPosition;
+}
+
 inline FVector GetNormal(const FLocationHit& Hit)
 {
 	return Hit.WorldNormal;
 }
 
+inline FVector GetNormal(const FPTLocationHit& Hit)
+{
+	return Hit.WorldNormal;
+}
+
 inline FHitFlags GetFlags(const FLocationHit& Hit)
 {
 	return Hit.Flags;
 }
 
+inline FHitFlags GetFlags(const FPTLocationHit& Hit)
+{
+	return Hit.Flags;
+}
+
+
 FORCEINLINE void SetFlags(FLocationHit& Hit, FHitFlags Flags)
 {
 	Hit.Flags = Flags;
 }
 
+FORCEINLINE void SetFlags(FPTLocationHit& Hit, FHitFlags Flags)
+{
+	Hit.Flags = Flags;
+}
+
 inline uint32 GetInternalFaceIndex(const FQueryHit& Hit)
 {
 	return Hit.FaceIndex;
 }
 
+inline uint32 GetInternalFaceIndex(const FPTQueryHit& Hit)
+{
+	return Hit.FaceIndex;
+}
+
 inline void SetInternalFaceIndex(FQueryHit& Hit, uint32 FaceIndex)
 {
 	Hit.FaceIndex = FaceIndex;
 }
 
+inline void SetInternalFaceIndex(FPTQueryHit& Hit, uint32 FaceIndex)
+{
+	Hit.FaceIndex = FaceIndex;
+}
+
 
 inline uint32 GetInvalidPhysicsFaceIndex()
 {
diff --git a/Engine/Source/Runtime/PhysicsCore/Public/ChaosSQTypes.h b/Engine/Source/Runtime/PhysicsCore/Public/ChaosSQTypes.h
--- a/Engine/Source/Runtime/PhysicsCore/Public/ChaosSQTypes.h
+++ b/Engine/Source/Runtime/PhysicsCore/Public/ChaosSQTypes.h
@@ -101,6 +101,47 @@ namespace ChaosInterface
 		return Ar;
 	}
 
+	struct FPTActorShape
+	{
+		Chaos::FGeometryParticleHandle* Actor;
+		const Chaos::FPerShapeData* Shape;
+	};
+
+	struct FPTQueryHit : public FPTActorShape
+	{
+		/**
+		Face index of touched triangle, for triangle meshes, convex meshes and height fields. Defaults to -1 if face index is not available
+		*/
+
+		int32 FaceIndex = -1; // Signed int to match TArray's size type, and so INDEX_NONE/-1 doesn't underflow.
+
+		int32 ElementIndex = -1; // Currently used to indicate which shape was hit for a particle with multiple shapes.
+	};
+
+	struct FPTLocationHit : public FPTQueryHit
+	{
+		FHitFlags Flags;
+		FVector WorldPosition;
+		FVector WorldNormal;
+		float Distance;	// LWC_TODO: Should be FVector::FReal, but that causes precision issues resulting in collision failures. Investigate!
+
+		bool operator<(const FPTLocationHit& Other) const { return Distance < Other.Distance; }
+	};
+
+	struct FPTRaycastHit : public FPTLocationHit
+	{
+		float U;
+		float V;
+	};
+	
+	struct FPTOverlapHit : public FPTQueryHit
+	{
+	};
+
+	struct FPTSweepHit : public FPTLocationHit
+	{
+	};
+
 	inline void FinishQueryHelper(TArray<FOverlapHit>& Hits, const FOverlapHit& BlockingHit, bool bHasBlockingHit)
 	{
 		if (bHasBlockingHit)
@@ -109,6 +150,14 @@ namespace ChaosInterface
 		}
 	}
 
+	inline void FinishQueryHelper(TArray<FPTOverlapHit>& Hits, const FPTOverlapHit& BlockingHit, bool bHasBlockingHit)
+	{
+		if (bHasBlockingHit)
+		{
+			Hits.Add(BlockingHit);
+		}
+	}
+
 	template <typename HitType>
 	void FinishQueryHelper(TArray<HitType>& Hits, const HitType& BlockingHit, bool bHasBlockingHit)
 	{
