[Backout] - CL18792501
Original CL Desc
-----------------------------------------------------------------
Allow DDC2 texture build path to not force stream all streaming mips up front in the editor.

diff --git a/Engine/Source/Runtime/Engine/Private/TextureDerivedDataTask.cpp b/Engine/Source/Runtime/Engine/Private/TextureDerivedDataTask.cpp
--- a/Engine/Source/Runtime/Engine/Private/TextureDerivedDataTask.cpp
+++ b/Engine/Source/Runtime/Engine/Private/TextureDerivedDataTask.cpp
@@ -1051,14 +1051,7 @@ public:
 					this->DerivedData.ResultMetadata = *InFetchFirstMetadata;
 				}
 
-				FBuildPolicyBuilder BuildPolicyBuilder(bInlineMips ? EBuildPolicy::Cache : (EBuildPolicy::CacheQuery | EBuildPolicy::SkipData));
-				if (!bInlineMips)
-				{
-					BuildPolicyBuilder.AddValuePolicy(FValueId::FromName("Description"_ASV), EBuildPolicy::Cache);
-					BuildPolicyBuilder.AddValuePolicy(FValueId::FromName("MipTail"_ASV), EBuildPolicy::Cache);
-				}
-
-				BuildSession.Get().Build(FetchDefinition, {}, BuildPolicyBuilder.Build(), *Owner,
+				BuildSession.Get().Build(FetchDefinition, {}, EBuildPolicy::Cache, *Owner,
 					[this, 
 					 FetchOrBuildDefinition = MoveTemp(FetchOrBuildDefinition), 
 					 Flags,
@@ -1112,24 +1105,11 @@ private:
 	void BeginBuild(const UE::DerivedData::FBuildDefinition& Definition, ETextureCacheFlags Flags)
 	{
 		using namespace UE::DerivedData;
-
-		FBuildPolicy BuildPolicy;
+		EBuildPolicy BuildPolicy = EBuildPolicy::Default;
 		if (EnumHasAnyFlags(Flags, ETextureCacheFlags::ForceRebuild))
 		{
-			BuildPolicy = EBuildPolicy::Default & ~EBuildPolicy::CacheQuery;
+			EnumRemoveFlags(BuildPolicy, EBuildPolicy::CacheQuery);
 		}
-		else if (bInlineMips)
-		{
-			BuildPolicy = EBuildPolicy::Default;
-		}
-		else
-		{
-			FBuildPolicyBuilder BuildPolicyBuilder(EBuildPolicy::Build | EBuildPolicy::CacheQuery | EBuildPolicy::CacheStoreOnBuild | EBuildPolicy::SkipData);
-			BuildPolicyBuilder.AddValuePolicy(FValueId::FromName("Description"_ASV), EBuildPolicy::Default);
-			BuildPolicyBuilder.AddValuePolicy(FValueId::FromName("MipTail"_ASV), EBuildPolicy::Default);
-			BuildPolicy = BuildPolicyBuilder.Build();
-		}
-
 		BuildSession.Get().Build(Definition, {}, BuildPolicy, *Owner,
 			[this](FBuildCompleteParams&& Params) { EndBuild(MoveTemp(Params)); });
 	}
