removed value copy from camera obfuscations

diff --git a/Engine/Source/Runtime/Engine/Classes/Camera/PlayerCameraManager.h b/Engine/Source/Runtime/Engine/Classes/Camera/PlayerCameraManager.h
--- a/Engine/Source/Runtime/Engine/Classes/Camera/PlayerCameraManager.h
+++ b/Engine/Source/Runtime/Engine/Classes/Camera/PlayerCameraManager.h
@@ -322,10 +322,10 @@ public:
 	virtual void SetLastFrameCameraCachePOV(const FMinimalViewInfo& InPOV);
 
 	/** Gets value of CameraCachePrivate.POV */
-	virtual FMinimalViewInfo GetCameraCachePOV() const;
+	virtual const FMinimalViewInfo& GetCameraCachePOV() const;
 
 	/** Gets value of LastFrameCameraCachePrivate.POV */
-	virtual FMinimalViewInfo GetLastFrameCameraCachePOV() const;
+	virtual const FMinimalViewInfo& GetLastFrameCameraCachePOV() const;
 
 	/** Get value of CameraCachePrivate.Time  */
 	float GetCameraCacheTime() const { return CameraCachePrivate.TimeStamp; }
diff --git a/Engine/Source/Runtime/Engine/Private/PlayerCameraManager.cpp b/Engine/Source/Runtime/Engine/Private/PlayerCameraManager.cpp
--- a/Engine/Source/Runtime/Engine/Private/PlayerCameraManager.cpp
+++ b/Engine/Source/Runtime/Engine/Private/PlayerCameraManager.cpp
@@ -890,7 +890,7 @@ void APlayerCameraManager::UnlockOrthoWidth()
 
 void APlayerCameraManager::GetCameraViewPoint(FVector& OutCamLoc, FRotator& OutCamRot) const
 {
-	const FMinimalViewInfo CurrentPOV = GetCameraCachePOV();
+	const FMinimalViewInfo& CurrentPOV = GetCameraCachePOV();
 	OutCamLoc = CurrentPOV.Location;
 	OutCamRot = CurrentPOV.Rotation;
 }
@@ -953,8 +953,8 @@ void APlayerCameraManager::UpdateCamera(float DeltaTime)
 			const float ClientNetCamUpdateDeltaTime = GameNetworkManager->ClientNetCamUpdateDeltaTime;
 			const float ClientNetCamUpdatePositionLimit = GameNetworkManager->ClientNetCamUpdatePositionLimit;
 
-			FMinimalViewInfo CurrentPOV = GetCameraCachePOV();
-			FMinimalViewInfo LastPOV = GetLastFrameCameraCachePOV();
+			const FMinimalViewInfo& CurrentPOV = GetCameraCachePOV();
+			const FMinimalViewInfo& LastPOV = GetLastFrameCameraCachePOV();
 
 			FVector ClientCameraPosition = FRepMovement::RebaseOntoZeroOrigin(CurrentPOV.Location, this);
 			FVector PrevClientCameraPosition = FRepMovement::RebaseOntoZeroOrigin(LastPOV.Location, this);
@@ -1240,7 +1240,7 @@ void APlayerCameraManager::LimitViewYaw(FRotator& ViewRotation, float InViewYawM
 
 void APlayerCameraManager::DisplayDebug(class UCanvas* Canvas, const FDebugDisplayInfo& DebugDisplay, float& YL, float& YPos)
 {
-	FMinimalViewInfo CurrentPOV = GetCameraCachePOV();
+	const FMinimalViewInfo& CurrentPOV = GetCameraCachePOV();
 
 	FDisplayDebugManager& DisplayDebugManager = Canvas->DisplayDebugManager;
 	DisplayDebugManager.SetDrawColor(FColor(255, 255, 255));
@@ -1568,12 +1568,12 @@ void APlayerCameraManager::SetLastFrameCameraCachePOV(const FMinimalViewInfo& In
 	LastFrameCameraCachePrivate.POV = InPOV;
 }
 
-FMinimalViewInfo APlayerCameraManager::GetCameraCachePOV() const
+const FMinimalViewInfo& APlayerCameraManager::GetCameraCachePOV() const
 {
 	return CameraCachePrivate.POV;
 }
 
-FMinimalViewInfo APlayerCameraManager::GetLastFrameCameraCachePOV() const
+const FMinimalViewInfo& APlayerCameraManager::GetLastFrameCameraCachePOV() const
 {
 	return LastFrameCameraCachePrivate.POV;
 }
