Additional demo driver logging for consider times and hitting the limit
Moved the prereplication control actor properties to advanced

diff --git a/Engine/Source/Runtime/Engine/Classes/GameFramework/Actor.h b/Engine/Source/Runtime/Engine/Classes/GameFramework/Actor.h
--- a/Engine/Source/Runtime/Engine/Classes/GameFramework/Actor.h
+++ b/Engine/Source/Runtime/Engine/Classes/GameFramework/Actor.h
@@ -219,10 +219,10 @@ private:
 	UPROPERTY(ReplicatedUsing=OnRep_ReplicateMovement, Category=Replication, EditDefaultsOnly)
 	uint8 bReplicateMovement:1;    
 
-	UPROPERTY(EditDefaultsOnly, Category = Replication)
+	UPROPERTY(EditDefaultsOnly, Category = Replication, AdvancedDisplay)
 	uint8 bCallPreReplication:1;
 
-	UPROPERTY(EditDefaultsOnly, Category = Replication)
+	UPROPERTY(EditDefaultsOnly, Category = Replication, AdvancedDisplay)
 	uint8 bCallPreReplicationForReplay:1;
 
 	/**
@@ -309,7 +309,7 @@ public:
 	 * Otherwise, RewindForReplay will be called if we detect the actor needs to be reset.
 	 * Note, this Actor must not be destroyed by gamecode, and RollbackViaDeletion may not be used. 
 	 */
-	UPROPERTY(Category=Replication, EditDefaultsOnly)
+	UPROPERTY(Category=Replication, EditDefaultsOnly, AdvancedDisplay)
 	uint8 bReplayRewindable:1;
 
 	/**
diff --git a/Engine/Source/Runtime/Engine/Private/DemoNetDriver.cpp b/Engine/Source/Runtime/Engine/Private/DemoNetDriver.cpp
--- a/Engine/Source/Runtime/Engine/Private/DemoNetDriver.cpp
+++ b/Engine/Source/Runtime/Engine/Private/DemoNetDriver.cpp
@@ -1685,6 +1685,7 @@ void UDemoNetDriver::TickDemoRecordFrame(float DeltaSeconds)
 
 				if (HasConsiderTimeBeenExhausted())
 				{
+					UE_LOG(LogDemo, Verbose, TEXT("Consider time exhaused prioritizing destruction infos."));
 					break;
 				}
 			}
@@ -1724,10 +1725,14 @@ void UDemoNetDriver::TickDemoRecordFrame(float DeltaSeconds)
 
 			const float CurrentTime = GetDemoCurrentTime();
 
+			int32 ProcessedCount = 0;
+
 			for (const TSharedPtr<FNetworkObjectInfo>& ObjectInfo : ActiveObjectSet)
 			{
 				FNetworkObjectInfo* ActorInfo = ObjectInfo.Get();
 
+				++ProcessedCount;
+
 				if (GetDemoCurrentTime() > ActorInfo->NextUpdateTime)
 				{
 					AActor* Actor = ActorInfo->Actor;
@@ -1818,6 +1823,7 @@ void UDemoNetDriver::TickDemoRecordFrame(float DeltaSeconds)
 
 				if (HasConsiderTimeBeenExhausted())
 				{
+					UE_LOG(LogDemo, Verbose, TEXT("Consider time exhaused while iterating the active object list [%d/%d]"), ProcessedCount, ActiveObjectSet.Num());
 					break;
 				}
 			}
@@ -5130,10 +5136,12 @@ void UDemoNetDriver::AdjustConsiderTime(const float ReplicatedPercent)
 		if (ReplicatedPercent > IncreaseThreshold)
 		{
 			RecordBuildConsiderAndPrioritizeTimeSlice += 0.1f;
+			UE_LOG(LogDemo, Verbose, TEXT("AdjustConsiderTime: RecordBuildConsiderAndPrioritizeTimeSlice is now %0.1f"), RecordBuildConsiderAndPrioritizeTimeSlice)
 		}
 		else if (ReplicatedPercent < DecreaseThreshold)
 		{
 			RecordBuildConsiderAndPrioritizeTimeSlice *= (1.f - ReplicatedPercent) * 0.5f;
+			UE_LOG(LogDemo, Verbose, TEXT("AdjustConsiderTime: RecordBuildConsiderAndPrioritizeTimeSlice is now %0.1f"), RecordBuildConsiderAndPrioritizeTimeSlice)
 		}
 
 		RecordBuildConsiderAndPrioritizeTimeSlice = FMath::Clamp<float>(RecordBuildConsiderAndPrioritizeTimeSlice, MinRepTime, MaxRepTime);
