Added estimated compressed size (total and nanite-only) for static mesh
Added estimated sizes to asset audit and static mesh viewport

diff --git a/Engine/Source/Editor/StaticMeshEditor/Private/StaticMeshEditorViewportClient.cpp b/Engine/Source/Editor/StaticMeshEditor/Private/StaticMeshEditorViewportClient.cpp
--- a/Engine/Source/Editor/StaticMeshEditor/Private/StaticMeshEditorViewportClient.cpp
+++ b/Engine/Source/Editor/StaticMeshEditor/Private/StaticMeshEditorViewportClient.cpp
@@ -907,6 +907,19 @@ void FStaticMeshEditorViewportClient::DrawCanvas( FViewport& InViewport, FSceneV
 		TextItems.Emplace(FText::Format(NSLOCTEXT("UnrealEd", "NumPrimitives_F", "Num Collision Primitives:  {0}"), FText::AsNumber(StaticMesh->GetBodySetup()->AggGeom.GetElementCount())));
 	}
 
+	// Estimated compressed size
+	if(StaticMesh->GetRenderData())
+	{
+		FNumberFormattingOptions NumberOptions;
+		NumberOptions.MinimumFractionalDigits = 2;
+		NumberOptions.MaximumFractionalDigits = 2;
+
+		TextItems.Emplace(
+			FText::Format(NSLOCTEXT("UnrealEd", "EstimatedCompressedSize", "Estimated Compressed Disk Size: {0} MB ({1} MB Nanite)"),
+			FText::AsNumber(StaticMesh->GetRenderData()->EstimatedCompressedSize / 1048576.0f, &NumberOptions),
+			FText::AsNumber(StaticMesh->GetRenderData()->EstimatedNaniteTotalCompressedSize / 1048576.0f, &NumberOptions)));
+	}
+
 	if (StaticMeshComponent && StaticMeshComponent->SectionIndexPreview != INDEX_NONE)
 	{
 		TextItems.Emplace(NSLOCTEXT("UnrealEd", "MeshSectionsHiddenWarning",  "Mesh Sections Hidden"));
diff --git a/Engine/Source/Runtime/Engine/Private/StaticMesh.cpp b/Engine/Source/Runtime/Engine/Private/StaticMesh.cpp
--- a/Engine/Source/Runtime/Engine/Private/StaticMesh.cpp
+++ b/Engine/Source/Runtime/Engine/Private/StaticMesh.cpp
@@ -77,6 +77,7 @@
 #include "DerivedDataRequestOwner.h"
 #include "DerivedDataCacheRecord.h"
 #include "DerivedDataValue.h"
+#include "Compression/OodleDataCompression.h"
 #endif // #if WITH_EDITOR
 
 #include "Engine/StaticMeshSocket.h"
@@ -1583,6 +1584,8 @@ void FStaticMeshRenderData::Serialize(FArchive& Ar, UStaticMesh* Owner, bool bCo
 	if (!bCooked)
 	{
 		Ar << MaterialIndexToImportIndex;
+		Ar << EstimatedNaniteTotalCompressedSize;
+		Ar << EstimatedNaniteStreamingCompressedSize;
 	}
 
 #endif // #if WITH_EDITORONLY_DATA
@@ -2393,7 +2396,7 @@ static void SerializeBuildSettingsForDDC(FArchive& Ar, FMeshBuildSettings& Build
 // differences, etc.) replace the version GUID below with a new one.
 // In case of merge conflicts with DDC versions, you MUST generate a new GUID
 // and set this new GUID as the version.
-#define STATICMESH_DERIVEDDATA_VER TEXT("24983DE8033B49EE8B6566C299172A6C")
+#define STATICMESH_DERIVEDDATA_VER TEXT("F6219E656D7345DAA8585A76F41AD167")
 
 const FString& GetStaticMeshDerivedDataVersion()
 {
@@ -2735,11 +2738,13 @@ void FStaticMeshRenderData::Cache(const ITargetPlatform* TargetPlatform, UStatic
 
 		static const FValueId MeshDataId = FValueId::FromName("MeshData");
 		static const FValueId NaniteStreamingDataId = FValueId::FromName("NaniteStreamingData");
+		static const double DDCSizeToEstimateFactor = 0.9;	// Hack: DDC uses fast compression. Compressed data is typically at least 10% smaller in shipped build, so use that as our estimate.
 
 		FCacheKey CacheKey;
 		CacheKey.Bucket = FCacheBucket(TEXT("StaticMesh"));
 		CacheKey.Hash = FIoHash::HashBuffer(MakeMemoryView(FTCHARToUTF8(DerivedDataKey)));
 
+		uint64 EstimatedMeshDataCompressedSize = 0;
 		FSharedBuffer MeshDataBuffer;
 		FIoHash NaniteStreamingDataHash;
 		{
@@ -2752,11 +2757,13 @@ void FStaticMeshRenderData::Cache(const ITargetPlatform* TargetPlatform, UStatic
 			Request.Policy = PolicyBuilder.Build();
 
 			FRequestOwner RequestOwner(EPriority::Blocking);
-			
 			GetCache().Get(MakeArrayView(&Request, 1), RequestOwner,
-				[&MeshDataBuffer, &NaniteStreamingDataHash](FCacheGetResponse&& Response)
+				[&MeshDataBuffer, &EstimatedMeshDataCompressedSize, &NaniteStreamingDataHash](FCacheGetResponse&& Response)
 				{
-					MeshDataBuffer = Response.Record.GetValue(MeshDataId).GetData().Decompress();
+					const FCompressedBuffer& CompressedBuffer = Response.Record.GetValue(MeshDataId).GetData();
+					MeshDataBuffer = CompressedBuffer.Decompress();
+					EstimatedMeshDataCompressedSize = uint64(CompressedBuffer.GetCompressedSize() * DDCSizeToEstimateFactor);
+
 					NaniteStreamingDataHash = Response.Record.GetValue(NaniteStreamingDataId).GetRawHash();
 				});
 			RequestOwner.Wait();
@@ -2769,6 +2776,10 @@ void FStaticMeshRenderData::Cache(const ITargetPlatform* TargetPlatform, UStatic
 			FMemoryReaderView Ar(MakeArrayView((const uint8*)MeshDataBuffer.GetData(), MeshDataBuffer.GetSize()), /*bIsPersistent=*/ true);
 			Serialize(Ar, Owner, /*bCooked=*/ false);
 
+			// Reconstruct EstimatedCompressedSize
+			// It is not serialized as it is not known before the serialized data has been compressed and we don't want to compress twice.
+			EstimatedCompressedSize = EstimatedMeshDataCompressedSize + EstimatedNaniteStreamingCompressedSize;
+
 			check(NaniteResources.StreamablePages.GetBulkDataSize() == 0);
 			if (NaniteResources.ResourceFlags & NANITE_RESOURCE_FLAG_STREAMING_DATA_IN_DDC)
 			{
@@ -2842,19 +2853,30 @@ void FStaticMeshRenderData::Cache(const ITargetPlatform* TargetPlatform, UStatic
 			
 			int64 TotalPushedBytes = 0;
 			FCacheRecordBuilder RecordBuilder(CacheKey);
-			if (NaniteResources.HasStreamingData())
+			if (NaniteResources.PageStreamingStates.Num() > 0)
 			{
-				// Compress streaming data, add it to record builder and remove it from BulkData
-				FByteBulkData& BulkData = NaniteResources.StreamablePages;
-				TotalPushedBytes += BulkData.GetBulkDataSize();
-				FValue Value = FValue::Compress(FSharedBuffer::MakeView(BulkData.LockReadOnly(), BulkData.GetBulkDataSize()));
-				RecordBuilder.AddValue(NaniteStreamingDataId, Value);
-				BulkData.Unlock();
-				BulkData.RemoveBulkData();
-
-				NaniteResources.ResourceFlags |= NANITE_RESOURCE_FLAG_STREAMING_DATA_IN_DDC;
-				NaniteResources.DDCKeyHash = CacheKey.Hash;
-				NaniteResources.DDCRawHash = Value.GetRawHash();
+				if (NaniteResources.HasStreamingData())
+				{
+					// Compress streaming data, add it to record builder and remove it from BulkData
+					FByteBulkData& BulkData = NaniteResources.StreamablePages;
+					TotalPushedBytes += BulkData.GetBulkDataSize();
+
+					FValue Value = FValue::Compress(FSharedBuffer::MakeView(BulkData.LockReadOnly(), BulkData.GetBulkDataSize()));
+					RecordBuilder.AddValue(NaniteStreamingDataId, Value);
+					EstimatedNaniteStreamingCompressedSize = uint64(Value.GetData().GetCompressedSize() * DDCSizeToEstimateFactor);
+					BulkData.Unlock();
+					BulkData.RemoveBulkData();
+
+					NaniteResources.ResourceFlags |= NANITE_RESOURCE_FLAG_STREAMING_DATA_IN_DDC;
+					NaniteResources.DDCKeyHash = CacheKey.Hash;
+					NaniteResources.DDCRawHash = Value.GetRawHash();
+				}
+				
+				// Compress non-streaming data for size estimation
+				FLargeMemoryWriter Ar(0, /*bIsPersistent=*/ true);
+				NaniteResources.Serialize(Ar, Owner, /*bCooked=*/ false);
+				FCompressedBuffer CompressedBuffer = FCompressedBuffer::Compress(FSharedBuffer::MakeView(Ar.GetData(), Ar.TotalSize()));
+				EstimatedNaniteTotalCompressedSize = EstimatedNaniteStreamingCompressedSize + uint64(CompressedBuffer.GetCompressedSize() * DDCSizeToEstimateFactor);
 			}
 			
 			bLODsShareStaticLighting = Owner->CanLODsShareStaticLighting();
@@ -2888,8 +2910,10 @@ void FStaticMeshRenderData::Cache(const ITargetPlatform* TargetPlatform, UStatic
 
 			if (bSaveDDC)
 			{
-				RecordBuilder.AddValue(MeshDataId, FSharedBuffer::MakeView(Ar.GetData(), Ar.TotalSize()));
+				FValue Value = FValue::Compress(FSharedBuffer::MakeView(Ar.GetData(), Ar.TotalSize()));
+				RecordBuilder.AddValue(MeshDataId, Value);
 				TotalPushedBytes += Ar.TotalSize();
+				EstimatedCompressedSize = uint64(Value.GetData().GetCompressedSize() * DDCSizeToEstimateFactor) + EstimatedNaniteStreamingCompressedSize;
 
 				FRequestOwner RequestOwner(UE::DerivedData::EPriority::Blocking);
 				const FCachePutRequest PutRequest = { FSharedString(TEXT("StaticMesh")), RecordBuilder.Build(), ECachePolicy::Default | ECachePolicy::KeepAlive };
@@ -4201,6 +4225,16 @@ void UStaticMesh::GetAssetRegistryTags(TArray<FAssetRegistryTag>& OutTags) const
 		DistanceFieldSize += VolumeData.StreamableMips.GetBulkDataSize();
 	}
 
+	uint64 EstimatedCompressedSize = 0;
+	uint64 EstimatedNaniteCompressedSize = 0;
+#if WITH_EDITORONLY_DATA
+	if (GetRenderData())
+	{
+		EstimatedCompressedSize = (int32)GetRenderData()->EstimatedCompressedSize;
+		EstimatedNaniteCompressedSize =  (int32)GetRenderData()->EstimatedNaniteTotalCompressedSize;
+	}
+#endif
+
 	OutTags.Add(FAssetRegistryTag("NaniteTriangles", FString::FromInt(NumNaniteTriangles), FAssetRegistryTag::TT_Numerical));
 	OutTags.Add(FAssetRegistryTag("NaniteVertices", FString::FromInt(NumNaniteVertices), FAssetRegistryTag::TT_Numerical));
 	OutTags.Add(FAssetRegistryTag("Triangles", FString::FromInt(NumTriangles), FAssetRegistryTag::TT_Numerical) );
@@ -4215,6 +4249,8 @@ void UStaticMesh::GetAssetRegistryTags(TArray<FAssetRegistryTag>& OutTags) const
 	OutTags.Add(FAssetRegistryTag("DefaultCollision", DefaultCollisionName.ToString(), FAssetRegistryTag::TT_Alphabetical));
 	OutTags.Add(FAssetRegistryTag("CollisionComplexity", ComplexityString, FAssetRegistryTag::TT_Alphabetical));
 	OutTags.Add(FAssetRegistryTag("DistanceFieldSize", FString::FromInt(DistanceFieldSize), FAssetRegistryTag::TT_Numerical, FAssetRegistryTag::TD_Memory));
+	OutTags.Add(FAssetRegistryTag("EstTotalCompressedSize", FString::Printf(TEXT("%llu"), EstimatedCompressedSize), FAssetRegistryTag::TT_Numerical, FAssetRegistryTag::TD_Memory));
+	OutTags.Add(FAssetRegistryTag("EstNaniteCompressedSize", FString::Printf(TEXT("%llu"), EstimatedNaniteCompressedSize), FAssetRegistryTag::TT_Numerical, FAssetRegistryTag::TD_Memory));
 
 	Super::GetAssetRegistryTags(OutTags);
 }
diff --git a/Engine/Source/Runtime/Engine/Public/StaticMeshResources.h b/Engine/Source/Runtime/Engine/Public/StaticMeshResources.h
--- a/Engine/Source/Runtime/Engine/Public/StaticMeshResources.h
+++ b/Engine/Source/Runtime/Engine/Public/StaticMeshResources.h
@@ -668,6 +668,15 @@ public:
 	/** The next cached derived data in the list. */
 	TUniquePtr<class FStaticMeshRenderData> NextCachedRenderData;
 
+	/** Estimate of total compressed size of all rendering data, including Nanite data. */
+	uint64 EstimatedCompressedSize = 0;
+
+	/** Estimate of total compressed size of Nanite data. Includes streaming and non-streaming data. */
+	uint64 EstimatedNaniteTotalCompressedSize = 0;
+
+	/** Estimate of compressed size of Nanite streaming data. */
+	uint64 EstimatedNaniteStreamingCompressedSize = 0;
+
 
 	void SyncUVChannelData(const TArray<FStaticMaterial>& ObjectData);
 
