Added cm/s metadata to OutofWaterZ
#skipci

diff --git a/Engine/Source/Runtime/Engine/Classes/GameFramework/CharacterMovementComponent.h b/Engine/Source/Runtime/Engine/Classes/GameFramework/CharacterMovementComponent.h
--- a/Engine/Source/Runtime/Engine/Classes/GameFramework/CharacterMovementComponent.h
+++ b/Engine/Source/Runtime/Engine/Classes/GameFramework/CharacterMovementComponent.h
@@ -537,7 +537,7 @@ public:
 	float MaxOutOfWaterStepHeight;
 
 	/** Z velocity applied when pawn tries to get out of water */
-	UPROPERTY(Category="Character Movement: Swimming", EditAnywhere, BlueprintReadWrite, AdvancedDisplay)
+	UPROPERTY(Category="Character Movement: Swimming", EditAnywhere, BlueprintReadWrite, AdvancedDisplay, meta=(ForceUnits="cm/s"))
 	float OutofWaterZ;
 
 	/** Mass of pawn (for when momentum is imparted to it). */
