Large world coordinates support in quantized vector net serialization.

diff --git a/Engine/Plugins/NetcodeUnitTest/NetcodeUnitTest/Source/NetcodeUnitTest/Classes/UnitTests/Engine/PackedVectorTest.h b/Engine/Plugins/NetcodeUnitTest/NetcodeUnitTest/Source/NetcodeUnitTest/Classes/UnitTests/Engine/PackedVectorTest.h
--- a/Engine/Plugins/NetcodeUnitTest/NetcodeUnitTest/Source/NetcodeUnitTest/Classes/UnitTests/Engine/PackedVectorTest.h
+++ b/Engine/Plugins/NetcodeUnitTest/NetcodeUnitTest/Source/NetcodeUnitTest/Classes/UnitTests/Engine/PackedVectorTest.h
@@ -19,4 +19,10 @@ class UPackedVectorTest : public UUnitTest
 
 private:
 	virtual bool ExecuteUnitTest() override;
+
+	bool ExecuteFloatTest(TMap<FString,bool>& TestResults);
+	bool ExecuteDoubleTest(TMap<FString,bool>& TestResults);
+	bool ExecuteWriteDoubleReadFloatTest(TMap<FString,bool>& TestResults);
+	bool ExecuteWriteFloatReadDoubleTest(TMap<FString,bool>& TestResults);
+	void VerifyTestResults(const TMap<FString,bool>& TestResults);
 };
diff --git a/Engine/Plugins/NetcodeUnitTest/NetcodeUnitTest/Source/NetcodeUnitTest/Private/UnitTests/Engine/PackedVectorTest.cpp b/Engine/Plugins/NetcodeUnitTest/NetcodeUnitTest/Source/NetcodeUnitTest/Private/UnitTests/Engine/PackedVectorTest.cpp
--- a/Engine/Plugins/NetcodeUnitTest/NetcodeUnitTest/Source/NetcodeUnitTest/Private/UnitTests/Engine/PackedVectorTest.cpp
+++ b/Engine/Plugins/NetcodeUnitTest/NetcodeUnitTest/Source/NetcodeUnitTest/Private/UnitTests/Engine/PackedVectorTest.cpp
@@ -3,7 +3,7 @@
 #include "UnitTests/Engine/PackedVectorTest.h"
 #include "Engine/NetSerialization.h"
 #include "UnitTestEnvironment.h"
-#include <cmath>
+#include <limits>
 
 namespace PackedVectorTest
 {
@@ -37,7 +37,7 @@ UPackedVectorTest::UPackedVectorTest(const FObjectInitializer& ObjectInitializer
 
 	UnitTestDate = FDateTime(2018, 9, 19);
 
-	ExpectedResult.Add(TEXT("ShooterGame"), EUnitTestVerification::VerifiedFixed);
+	ExpectedResult.Add(TEXT("NullUnitEnv"), EUnitTestVerification::VerifiedFixed);
 
 	UnitTestTimeout = 5;
 }
@@ -48,13 +48,24 @@ bool UPackedVectorTest::ExecuteUnitTest()
 
 	TestResults.Add(TEXT("Commencing Read/WritePackedVector tests. Only fails will be shown in log."), "Success" && true);
 
+	const bool bFloatTestSuccess = ExecuteFloatTest(TestResults);
+	const bool bDoubleTestSuccess = ExecuteDoubleTest(TestResults);
+	const bool bWriteDoubleReadFloatTestSuccess = ExecuteWriteDoubleReadFloatTest(TestResults);
+
+	VerifyTestResults(TestResults);
+
+	return bFloatTestSuccess && bDoubleTestSuccess;
+}
+
+bool UPackedVectorTest::ExecuteFloatTest(TMap<FString, bool>& TestResults)
+{
 	static const float Quantize10_Values[] =
 	{
 		0.0f,
 		-180817.42f,
 		47.11f,
-		-FMath::Exp2(25.0f), // overflow
-		INFINITY, // non-finite
+		-FMath::Exp2(25.0f),
+		std::numeric_limits<float>::infinity(), // non-finite
 	};
 
 	static const float Quantize100_Values[] =
@@ -63,19 +74,18 @@ bool UPackedVectorTest::ExecuteUnitTest()
 		+180720.42f,
 		-19751216.0f,
 		FMath::Exp2(31.0f),
-		-INFINITY,
+		-std::numeric_limits<float>::infinity(),
 	};
 
 	struct TestCase
 	{
 		int ScaleFactor;
-		int BitsPerComponent;
 
 		const float* TestValues;
 		int TestValueCount;
 
-		TFunction<bool(FVector&, FArchive&)> Reader;
-		TFunction<bool(FVector, FArchive&)> Writer;
+		TFunction<bool(FVector3f&, FArchive&)> Reader;
+		TFunction<bool(FVector3f, FArchive&)> Writer;
 
 	};
 
@@ -83,27 +93,25 @@ bool UPackedVectorTest::ExecuteUnitTest()
 	{
 		{
 			10,
-			24,
 			Quantize10_Values,
 			sizeof(Quantize10_Values) / sizeof(Quantize10_Values[0]),
-			[](FVector& Value, FArchive& Ar) { return ReadPackedVector<10, 24>(Value, Ar); },
-			[](FVector Value, FArchive& Ar) { return WritePackedVector<10, 24>(Value, Ar); },
+			[](FVector3f& Value, FArchive& Ar) { return ReadPackedVector<10, 24>(Value, Ar); },
+			[](FVector3f Value, FArchive& Ar) { return WritePackedVector<10, 24>(Value, Ar); },
 		},
 
 		{
 			100,
-			30,
 			Quantize100_Values,
 			sizeof(Quantize100_Values) / sizeof(Quantize100_Values[0]),
-			[](FVector& Value, FArchive& Ar) { return ReadPackedVector<100, 30>(Value, Ar); },
-			[](FVector Value, FArchive& Ar) { return WritePackedVector<100, 30>(Value, Ar); },
+			[](FVector3f& Value, FArchive& Ar) { return ReadPackedVector<100, 30>(Value, Ar); },
+			[](FVector3f Value, FArchive& Ar) { return WritePackedVector<100, 30>(Value, Ar); },
 		},
 	};
 
 	constexpr bool bAllowResize = false;
 	FBitWriter Writer(128, bAllowResize);
 
-	for (size_t TestCaseIt = 0, TestCaseEndIt = sizeof(TestCases) / sizeof(TestCases[0]); TestCaseIt != TestCaseEndIt; ++TestCaseIt)
+	for (size_t TestCaseIt = 0, TestCaseEndIt = sizeof(TestCases)/sizeof(TestCases[0]); TestCaseIt != TestCaseEndIt; ++TestCaseIt)
 	{
 		const TestCase& Test = TestCases[TestCaseIt];
 		for (size_t ValueIt = 0, ValueEndIt = Test.TestValueCount; ValueIt != ValueEndIt; ++ValueIt)
@@ -111,8 +119,8 @@ bool UPackedVectorTest::ExecuteUnitTest()
 			Writer.Reset();
 
 			const float ScalarValue = Test.TestValues[ValueIt];
-			const FVector WriteValue(ScalarValue);
-			FVector ReadValue;
+			const FVector3f WriteValue(ScalarValue);
+			FVector3f ReadValue;
 
 			const bool bOverflowOrNan = !Test.Writer(WriteValue, Writer);
 			bool LocalSuccess = !Writer.GetError();
@@ -130,29 +138,233 @@ bool UPackedVectorTest::ExecuteUnitTest()
 					if (bOverflowOrNan)
 					{
 						if (WriteValue.ContainsNaN())
-							LocalSuccess &= (ReadValue == FVector::ZeroVector);
+						{
+							LocalSuccess &= (ReadValue == FVector3f::ZeroVector);
+						}
 						else
 						{
-							// Overflowed value. Should probably be close to range start or end [-2^(ComponentBitCount)/ScaleFactor, 2^(ComponentBitCount)/ScaleFactor]
-							const float MaxValue = FMath::Exp2(float(Test.BitsPerComponent))/Test.ScaleFactor;
-							LocalSuccess &= PackedVectorTest::AlmostEqualUlps(FMath::Abs(ReadValue.X), MaxValue, 1);
+							LocalSuccess = false;
 						}
 					}
 					else
 					{
 						const float ValueDiff = FMath::Abs(ReadValue.X - WriteValue.X);
-						LocalSuccess &= ValueDiff < 2.0f / Test.ScaleFactor; // The diff test might need some adjustment
+						LocalSuccess &= ValueDiff < 2.0f/Test.ScaleFactor; // The diff test might need some adjustment
+					}
+				}
+			}
+
+			if (!LocalSuccess)
+			{
+				TestResults.Add(FString::Printf(TEXT("Read/WritePackedVector failed with scale %d and value %f. Got %f"), Test.ScaleFactor, ScalarValue, ReadValue.X), LocalSuccess);
+			}
+		}
+	}
+
+	return true;
+}
+
+bool UPackedVectorTest::ExecuteDoubleTest(TMap<FString, bool>& TestResults)
+{
+	static const double Quantize10_Values[] =
+	{
+		0.0f,
+		-180817.42,
+		47.11f,
+		-FMath::Exp2(25.0),
+		std::numeric_limits<double>::infinity(), // non-finite
+	};
+
+	static const double Quantize100_Values[] =
+	{
+		0.0f,
+		+180720.42,
+		-19751216.0,
+		FMath::Exp2(40.0),
+		FMath::Exp2(59.0),
+		-std::numeric_limits<double>::infinity(),
+	};
+
+	struct TestCase
+	{
+		int ScaleFactor;
+
+		const double* TestValues;
+		int TestValueCount;
+
+		TFunction<bool(FVector3d&, FArchive&)> Reader;
+		TFunction<bool(FVector3d, FArchive&)> Writer;
+
+	};
+
+	TestCase TestCases[] =
+	{
+		{
+			10,
+			Quantize10_Values,
+			sizeof(Quantize10_Values) / sizeof(Quantize10_Values[0]),
+			[](FVector3d& Value, FArchive& Ar) { return ReadPackedVector<10, 24>(Value, Ar); },
+			[](FVector3d Value, FArchive& Ar) { return WritePackedVector<10, 24>(Value, Ar); },
+		},
+
+		{
+			100,
+			Quantize100_Values,
+			sizeof(Quantize100_Values) / sizeof(Quantize100_Values[0]),
+			[](FVector3d& Value, FArchive& Ar) { return ReadPackedVector<100, 30>(Value, Ar); },
+			[](FVector3d Value, FArchive& Ar) { return WritePackedVector<100, 30>(Value, Ar); },
+		},
+	};
+
+	constexpr bool bAllowResize = false;
+	FBitWriter Writer(512, bAllowResize);
+
+	for (size_t TestCaseIt = 0, TestCaseEndIt = sizeof(TestCases)/sizeof(TestCases[0]); TestCaseIt != TestCaseEndIt; ++TestCaseIt)
+	{
+		const TestCase& Test = TestCases[TestCaseIt];
+		for (size_t ValueIt = 0, ValueEndIt = Test.TestValueCount; ValueIt != ValueEndIt; ++ValueIt)
+		{
+			Writer.Reset();
+
+			const double ScalarValue = Test.TestValues[ValueIt];
+			const FVector3d WriteValue(ScalarValue);
+			FVector3d ReadValue;
+
+			const bool bOverflowOrNan = !Test.Writer(WriteValue, Writer);
+			bool LocalSuccess = !Writer.GetError();
+
+			if (LocalSuccess)
+			{
+				FBitReader Reader(Writer.GetData(), Writer.GetNumBits());
+
+				Test.Reader(ReadValue, Reader);
+				LocalSuccess &= !Reader.GetError();
+				LocalSuccess = (ReadValue.X == ReadValue.Y) && (ReadValue.X == ReadValue.Z);
+				if (LocalSuccess)
+				{
+					// At this point we should have similar values as the original ones, except for NaN and overflowed values
+					if (bOverflowOrNan)
+					{
+						if (WriteValue.ContainsNaN())
+						{
+							LocalSuccess &= (ReadValue == FVector3d::ZeroVector);
+						}
+						else
+						{
+							LocalSuccess = false;
+						}
+					}
+					else
+					{
+						const double ValueDiff = FMath::Abs(ReadValue.X - WriteValue.X);
+						LocalSuccess &= ValueDiff < 2.0/Test.ScaleFactor; // The diff test might need some adjustment
+					}
+				}
+			}
+
+			if (!LocalSuccess)
+			{
+				TestResults.Add(FString::Printf(TEXT("Read/WritePackedVector failed with scale %d and value %f. Got %f"), Test.ScaleFactor, ScalarValue, ReadValue.X), LocalSuccess);
+			}
+		}
+	}
+
+	return true;
+}
+
+bool UPackedVectorTest::ExecuteWriteDoubleReadFloatTest(TMap<FString, bool>& TestResults)
+{
+	static const double Quantize100_Values[] =
+	{
+		0.0f,
+		+180720.42,
+		-19751216.0,
+		FMath::Exp2(40.0),
+		FMath::Exp2(59.0),
+		-std::numeric_limits<double>::infinity(),
+	};
+
+	struct TestCase
+	{
+		int ScaleFactor;
+
+		const double* TestValues;
+		int TestValueCount;
+
+		TFunction<bool(FVector3f&, FArchive&)> Reader;
+		TFunction<bool(FVector3d, FArchive&)> Writer;
+
+	};
+
+	TestCase TestCases[] =
+	{
+		{
+			100,
+			Quantize100_Values,
+			sizeof(Quantize100_Values) / sizeof(Quantize100_Values[0]),
+			[](FVector3f& Value, FArchive& Ar) { return ReadPackedVector<100, 30>(Value, Ar); },
+			[](FVector3d Value, FArchive& Ar) { return WritePackedVector<100, 30>(Value, Ar); },
+		},
+	};
+
+	constexpr bool bAllowResize = false;
+	FBitWriter Writer(512, bAllowResize);
+
+	for (size_t TestCaseIt = 0, TestCaseEndIt = sizeof(TestCases)/sizeof(TestCases[0]); TestCaseIt != TestCaseEndIt; ++TestCaseIt)
+	{
+		const TestCase& Test = TestCases[TestCaseIt];
+		for (size_t ValueIt = 0, ValueEndIt = Test.TestValueCount; ValueIt != ValueEndIt; ++ValueIt)
+		{
+			Writer.Reset();
+
+			const double ScalarValue = Test.TestValues[ValueIt];
+			const FVector3d WriteValue(ScalarValue);
+			FVector3f ReadValue;
+
+			const bool bOverflowOrNan = !Test.Writer(WriteValue, Writer);
+			bool LocalSuccess = !Writer.GetError();
+
+			if (LocalSuccess)
+			{
+				FBitReader Reader(Writer.GetData(), Writer.GetNumBits());
+
+				Test.Reader(ReadValue, Reader);
+				LocalSuccess &= !Reader.GetError();
+				LocalSuccess = (ReadValue.X == ReadValue.Y) && (ReadValue.X == ReadValue.Z);
+				if (LocalSuccess)
+				{
+					// At this point we should have similar values as the original ones, except for NaN and overflowed values
+					if (bOverflowOrNan)
+					{
+						if (WriteValue.ContainsNaN())
+						{
+							LocalSuccess &= (ReadValue == FVector::ZeroVector);
+						}
+						else
+						{
+							LocalSuccess = false;
+						}
+					}
+					else
+					{
+						const double ValueDiff = FMath::Abs(ReadValue.X - WriteValue.X);
+						LocalSuccess &= ValueDiff < 2.0/Test.ScaleFactor; // The diff test might need some adjustment
 					}
 				}
 			}
 
 			if (!LocalSuccess)
 			{
-				TestResults.Add(FString::Printf(TEXT("Read/WritePackedVector failed with scale %d, bit count %d and value %f. Got %f"), Test.ScaleFactor, Test.BitsPerComponent, ScalarValue, ReadValue.X), LocalSuccess);
+				TestResults.Add(FString::Printf(TEXT("Read/WritePackedVector failed with scale %d and value %f. Got %f"), Test.ScaleFactor, ScalarValue, ReadValue.X), LocalSuccess);
 			}
 		}
 	}
 
+	return true;
+}
+
+void UPackedVectorTest::VerifyTestResults(const TMap<FString,bool>& TestResults)
+{
 	// Verify the results
 	for (TMap<FString, bool>::TConstIterator It(TestResults); It; ++It)
 	{
@@ -168,6 +380,4 @@ bool UPackedVectorTest::ExecuteUnitTest()
 	{
 		VerificationState = EUnitTestVerification::VerifiedFixed;
 	}
-
-	return true;
 }
diff --git a/Engine/Plugins/Runtime/NetworkPrediction/Source/NetworkPrediction/NetworkPrediction.Build.cs b/Engine/Plugins/Runtime/NetworkPrediction/Source/NetworkPrediction/NetworkPrediction.Build.cs
--- a/Engine/Plugins/Runtime/NetworkPrediction/Source/NetworkPrediction/NetworkPrediction.Build.cs
+++ b/Engine/Plugins/Runtime/NetworkPrediction/Source/NetworkPrediction/NetworkPrediction.Build.cs
@@ -35,8 +35,9 @@ namespace UnrealBuildTool.Rules
 			PrivateDependencyModuleNames.AddRange(
 				new string[]
 				{
-					"TraceLog"					
-                }
+					"NetCore",
+					"TraceLog"
+				}
 				);
 
 			DynamicallyLoadedModuleNames.AddRange(
diff --git a/Engine/Plugins/Runtime/NetworkPredictionExtras/Source/NetworkPredictionExtras/NetworkPredictionExtras.Build.cs b/Engine/Plugins/Runtime/NetworkPredictionExtras/Source/NetworkPredictionExtras/NetworkPredictionExtras.Build.cs
--- a/Engine/Plugins/Runtime/NetworkPredictionExtras/Source/NetworkPredictionExtras/NetworkPredictionExtras.Build.cs
+++ b/Engine/Plugins/Runtime/NetworkPredictionExtras/Source/NetworkPredictionExtras/NetworkPredictionExtras.Build.cs
@@ -36,7 +36,8 @@ namespace UnrealBuildTool.Rules
 			PrivateDependencyModuleNames.AddRange(
 				new string[]
 				{
-					"TraceLog",
+					"NetCore",
+ 					"TraceLog",
 					"Engine"
                 }
 				);
diff --git a/Engine/Source/Runtime/Core/Public/Misc/NetworkVersion.h b/Engine/Source/Runtime/Core/Public/Misc/NetworkVersion.h
--- a/Engine/Source/Runtime/Core/Public/Misc/NetworkVersion.h
+++ b/Engine/Source/Runtime/Core/Public/Misc/NetworkVersion.h
@@ -54,6 +54,7 @@ enum EEngineNetworkVersionHistory
 	HISTORY_INTERFACE_PROPERTY_SERIALIZATION = 20,	// Bump version to support net serialization of FInterfaceProperty
 	HISTORY_MONTAGE_PLAY_INST_ID_SERIALIZATION = 21,// Bump version to support net serialization of FGameplayAbilityRepAnimMontage, addition of PlayInstanceId and removal of bForcePlayBit
 	HISTORY_SERIALIZE_DOUBLE_VECTORS_AS_DOUBLES	= 22,// Bump version to support net serialization of double vector types
+	HISTORY_PACKED_VECTOR_LWC_SUPPORT = 23,			// Bump version to support quantized LWC FVector net serialization
 	// New history items go above here.
 
 	HISTORY_ENGINENETVERSION_PLUS_ONE,
diff --git a/Engine/Source/Runtime/Engine/Classes/Engine/NetSerialization.h b/Engine/Source/Runtime/Engine/Classes/Engine/NetSerialization.h
--- a/Engine/Source/Runtime/Engine/Classes/Engine/NetSerialization.h
+++ b/Engine/Source/Runtime/Engine/Classes/Engine/NetSerialization.h
@@ -11,13 +11,11 @@
 #include "Stats/Stats.h"
 #include "UObject/ObjectMacros.h"
 #include "UObject/Class.h"
-#include "Serialization/BitReader.h"
-#include "Misc/NetworkGuid.h"
+#include "Misc/NetworkVersion.h"
 #include "UObject/CoreNet.h"
 #include "EngineLogs.h"
+#include "Net/Core/Serialization/QuantizedVectorSerialization.h"
 #include "Net/Serialization/FastArraySerializer.h"
-#include "Containers/ArrayView.h"
-#include "HAL/IConsoleManager.h"
 
 #include "NetSerialization.generated.h"
 
@@ -171,57 +169,14 @@ struct TStructOpsTypeTraits< FExampleStruct > : public TStructOpsTypeTraitsBase2
  *
 */
 
-// LWC_TODO: Proper double FVector serialization required.
-template<int32 ScaleFactor, int32 MaxBitsPerComponent>
-bool WritePackedVector(FVector3f Value, FArchive& Ar)	// Note Value is intended to not be a reference since we are scaling it before serializing!
+namespace UE::Net::Private
 {
-	check(Ar.IsSaving());
-
-	// Scale vector by quant factor first
-	Value *= ScaleFactor;
-
-	// Nan Check
-	if( Value.ContainsNaN() )
-	{
-		logOrEnsureNanError(TEXT("WritePackedVector: Value contains NaN, clearing for safety."));
-		FVector	Dummy(0, 0, 0);
-		WritePackedVector<ScaleFactor, MaxBitsPerComponent>(Dummy, Ar);
-		return false;
-	}
-
-	// Some platforms have RoundToInt implementations that essentially reduces the allowed inputs to 2^31.
-	const FVector3f ClampedValue = ClampVector(Value, FVector3f(-1073741824.0f), FVector3f(1073741760.0f));
-	bool bClamp = ClampedValue != Value;
-
-	// Do basically FVector::SerializeCompressed
-	int32 IntX	= FMath::RoundToInt(ClampedValue.X);
-	int32 IntY	= FMath::RoundToInt(ClampedValue.Y);
-	int32 IntZ	= FMath::RoundToInt(ClampedValue.Z);
-			
-	uint32 Bits	= FMath::Clamp<uint32>( FMath::CeilLogTwo( 1 + FMath::Max3( FMath::Abs(IntX), FMath::Abs(IntY), FMath::Abs(IntZ) ) ), 1, MaxBitsPerComponent ) - 1;
-
-	// Serialize how many bits each component will have
-	Ar.SerializeInt( Bits, MaxBitsPerComponent );
-
-	int32  Bias	= 1<<(Bits+1);
-	uint32 Max	= 1<<(Bits+2);
-	uint32 DX	= IntX + Bias;
-	uint32 DY	= IntY + Bias;
-	uint32 DZ	= IntZ + Bias;
-
-	if (DX >= Max) { bClamp=true; DX = static_cast<int32>(DX) > 0 ? Max-1 : 0; }
-	if (DY >= Max) { bClamp=true; DY = static_cast<int32>(DY) > 0 ? Max-1 : 0; }
-	if (DZ >= Max) { bClamp=true; DZ = static_cast<int32>(DZ) > 0 ? Max-1 : 0; }
-	
-	Ar.SerializeInt( DX, Max );
-	Ar.SerializeInt( DY, Max );
-	Ar.SerializeInt( DZ, Max );
 
-	return !bClamp;
-}
-
-template<uint32 ScaleFactor, int32 MaxBitsPerComponent>
-bool ReadPackedVector(FVector3f &Value, FArchive& Ar)
+// ScaleFactor is multiplied before send and divided by post receive. A higher ScaleFactor means more precision.
+// MaxBitsPerComponent is the maximum number of bits to use per component. This is only a maximum. A header is
+// written (size = Log2 (MaxBitsPerComponent)) to indicate how many bits are actually used. 
+template<int32 ScaleFactor, int32 MaxBitsPerComponent>
+bool LegacyReadPackedVector(FVector3f& Value, FArchive& Ar)
 {
 	uint32 Bits	= 0;
 
@@ -238,7 +193,6 @@ bool ReadPackedVector(FVector3f &Value, FArchive& Ar)
 	Ar.SerializeInt( DY, Max );
 	Ar.SerializeInt( DZ, Max );
 	
-	
 	float fact = (float)ScaleFactor;
 
 	Value.X = (float)(static_cast<int32>(DX)-Bias) / fact;
@@ -248,45 +202,83 @@ bool ReadPackedVector(FVector3f &Value, FArchive& Ar)
 	return true;
 }
 
-template<uint32 ScaleFactor, int32 MaxBitsPerComponent>
-bool ReadPackedVector(FVector3d& Value, FArchive& Ar)
+}
+
+template<int32 ScaleFactor, int32 MaxBitsPerComponent>
+bool WritePackedVector(FVector3f Value, FArchive& Ar)
 {
-	// LWC_TODO: Proper double FVector serialization required.
-	FVector3f AsFloat;
-	bool bRet = ReadPackedVector<ScaleFactor, MaxBitsPerComponent>(AsFloat, Ar);
-	Value = AsFloat;
-	return bRet;
+	return UE::Net::WriteQuantizedVector(ScaleFactor, Value, Ar);
 }
 
-// ScaleFactor is multiplied before send and divided by post receive. A higher ScaleFactor means more precision.
-// MaxBitsPerComponent is the maximum number of bits to use per component. This is only a maximum. A header is
-// written (size = Log2 (MaxBitsPerComponent)) to indicate how many bits are actually used. 
+template<int32 ScaleFactor, int32 MaxBitsPerComponent>
+bool WritePackedVector(FVector3d Vector, FArchive& Ar)
+{
+	return UE::Net::WriteQuantizedVector(ScaleFactor, Vector, Ar);
+}
 
-template<uint32 ScaleFactor, int32 MaxBitsPerComponent>
-bool SerializePackedVector(FVector3f &Vector, FArchive& Ar)
+template<int32 ScaleFactor, int32 MaxBitsPerComponent>
+bool ReadPackedVector(FVector3f& Value, FArchive& Ar)
 {
-	if (Ar.IsSaving())
+	if (Ar.EngineNetVer() >= HISTORY_PACKED_VECTOR_LWC_SUPPORT)
+	{
+		return UE::Net::ReadQuantizedVector(ScaleFactor, Value, Ar);
+	}
+	else
 	{
-		return  WritePackedVector<ScaleFactor, MaxBitsPerComponent>(Vector, Ar);
+		return UE::Net::Private::LegacyReadPackedVector<ScaleFactor, MaxBitsPerComponent>(Value, Ar);
 	}
+}
 
-	ReadPackedVector<ScaleFactor, MaxBitsPerComponent>(Vector, Ar);
-	return true;
+template<int32 ScaleFactor, int32 MaxBitsPerComponent>
+bool ReadPackedVector(FVector3d& Value, FArchive& Ar)
+{
+	if (Ar.EngineNetVer() >= HISTORY_PACKED_VECTOR_LWC_SUPPORT)
+	{
+		return UE::Net::ReadQuantizedVector(ScaleFactor, Value, Ar);
+	}
+	else
+	{
+		FVector3f AsFloat;
+		const bool bRet =  UE::Net::Private::LegacyReadPackedVector<ScaleFactor, MaxBitsPerComponent>(AsFloat, Ar);
+		Value = FVector3d(AsFloat);
+		return bRet;
+	}
 }
 
-template<uint32 ScaleFactor, int32 MaxBitsPerComponent>
-bool SerializePackedVector(FVector3d& Vector, FArchive& Ar)
+template<int32 ScaleFactor, int32 MaxBitsPerComponent>
+bool SerializePackedVector(FVector3f& Value, FArchive& Ar)
 {
-	// LWC_TODO: Proper double FVector serialization required.
-	FVector3f AsFloat(Vector);
-	bool bRet = SerializePackedVector<ScaleFactor, MaxBitsPerComponent>(AsFloat, Ar);
-	Vector = AsFloat;
-	return bRet;
+	if (Ar.EngineNetVer() >= HISTORY_PACKED_VECTOR_LWC_SUPPORT)
+	{
+		return UE::Net::SerializeQuantizedVector<ScaleFactor>(Value, Ar);
+	}
+	else
+	{
+		check(Ar.IsLoading());
+		return UE::Net::Private::LegacyReadPackedVector<ScaleFactor, MaxBitsPerComponent>(Value, Ar);
+	}
+}
+
+template<int32 ScaleFactor, int32 MaxBitsPerComponent>
+bool SerializePackedVector(FVector3d& Value, FArchive& Ar)
+{
+	if (Ar.EngineNetVer() >= HISTORY_PACKED_VECTOR_LWC_SUPPORT)
+	{
+		return UE::Net::SerializeQuantizedVector<ScaleFactor>(Value, Ar);
+	}
+	else
+	{
+		check(Ar.IsLoading());
+		FVector3f AsFloat(Value);
+		const bool bRet = UE::Net::Private::LegacyReadPackedVector<ScaleFactor, MaxBitsPerComponent>(AsFloat, Ar);
+		Value = FVector3d(AsFloat);
+		return bRet;
+	}
 }
 
 // --------------------------------------------------------------
 
-template<int32 MaxValue, int32 NumBits>
+template<int32 MaxValue, uint32 NumBits>
 struct TFixedCompressedFloatDetails
 {
 	                                                                // NumBits = 8:
@@ -294,26 +286,10 @@ struct TFixedCompressedFloatDetails
 	static constexpr int32 Bias = (1 << (NumBits - 1));             //   1000 0000 - Bias to pivot around (in order to support signed values)
 	static constexpr int32 SerIntMax = (1 << (NumBits - 0));        // 1 0000 0000 - What we pass into SerializeInt
 	static constexpr int32 MaxDelta = (1 << (NumBits - 0)) - 1;     //   1111 1111 - Max delta is
-
-#if !PLATFORM_COMPILER_HAS_IF_CONSTEXPR
-	static constexpr float GetInvScale()
-	{
-		if (MaxValue > MaxBitValue)
-		{
-			// We have to scale down, scale needs to be a float:
-			return (float)MaxValue / (float)MaxBitValue;
-		}
-		else
-		{
-			int32 scale = MaxBitValue / MaxValue;
-			return 1.f / scale;
-		}
-	}
-#endif
 };
 
-template<int32 MaxValue, int32 NumBits>
-bool WriteFixedCompressedFloat(const float Value, FArchive& Ar)
+template<int32 MaxValue, uint32 NumBits, typename T, TEMPLATE_REQUIRES(TIsFloatingPoint<T>::Value), TEMPLATE_REQUIRES(NumBits < 32)>
+bool WriteFixedCompressedFloat(const T Value, FArchive& Ar)
 {
 	using Details = TFixedCompressedFloatDetails<MaxValue, NumBits>;
 
@@ -321,15 +297,15 @@ bool WriteFixedCompressedFloat(const float Value, FArchive& Ar)
 	int32 ScaledValue;
 	if ( MaxValue > Details::MaxBitValue )
 	{
-		// We have to scale this down, scale needs to be a float:
-		const float scale = (float)Details::MaxBitValue / (float)MaxValue;
-		ScaledValue = FMath::TruncToInt(scale * Value);
+		// We have to scale this down
+		const T Scale = T(Details::MaxBitValue)/MaxValue;
+		ScaledValue = FMath::TruncToInt(Scale * Value);
 	}
 	else
 	{
 		// We will scale up to get extra precision. But keep is a whole number preserve whole values
-		enum { scale = Details::MaxBitValue / MaxValue };
-		ScaledValue = FMath::RoundToInt( scale * Value );
+		enum { Scale = Details::MaxBitValue / MaxValue };
+		ScaledValue = FMath::RoundToInt( Scale * Value );
 	}
 
 	uint32 Delta = static_cast<uint32>(ScaledValue + Details::Bias);
@@ -345,33 +321,28 @@ bool WriteFixedCompressedFloat(const float Value, FArchive& Ar)
 	return !clamp;
 }
 
-template<int32 MaxValue, int32 NumBits>
-bool ReadFixedCompressedFloat(float &Value, FArchive& Ar)
+template<int32 MaxValue, uint32 NumBits, typename T, TEMPLATE_REQUIRES(TIsFloatingPoint<T>::Value), TEMPLATE_REQUIRES(NumBits < 32)>
+bool ReadFixedCompressedFloat(T& Value, FArchive& Ar)
 {
 	using Details = TFixedCompressedFloatDetails<MaxValue, NumBits>;
 
 	uint32 Delta;
 	Ar.SerializeInt(Delta, Details::SerIntMax);
-	float UnscaledValue = static_cast<float>( static_cast<int32>(Delta) - Details::Bias );
+	T UnscaledValue = static_cast<T>( static_cast<int32>(Delta) - Details::Bias );
 
-#if PLATFORM_COMPILER_HAS_IF_CONSTEXPR
 	if constexpr (MaxValue > Details::MaxBitValue)
 	{
 		// We have to scale down, scale needs to be a float:
-		const float InvScale = MaxValue / (float)Details::MaxBitValue;
+		const T InvScale = MaxValue / (T)Details::MaxBitValue;
 		Value = UnscaledValue * InvScale;
 	}
 	else
 	{
-		enum { scale = Details::MaxBitValue / MaxValue };
-		const float InvScale = 1.f / (float)scale;
+		enum { Scale = Details::MaxBitValue / MaxValue };
+		const T InvScale = T(1) / (T)Scale;
 
 		Value = UnscaledValue * InvScale;
 	}
-#else
-	constexpr float InvScale = Details::GetInvScale();
-	Value = UnscaledValue * InvScale;
-#endif
 
 	return true;
 }
@@ -381,7 +352,7 @@ bool ReadFixedCompressedFloat(float &Value, FArchive& Ar)
 // NumBits is the total number of bits to use - this includes the sign bit!
 //
 // So passing in NumBits = 8, and MaxValue = 2^8, you will scale down to fit into 7 bits so you can leave 1 for the sign bit.
-template<int32 MaxValue, int32 NumBits>
+template<int32 MaxValue, uint32 NumBits>
 bool SerializeFixedVector(FVector3f &Vector, FArchive& Ar)
 {
 	if (Ar.IsSaving())
@@ -399,14 +370,24 @@ bool SerializeFixedVector(FVector3f &Vector, FArchive& Ar)
 	return true;
 }
 
-template<int32 MaxValue, int32 NumBits>
-bool SerializeFixedVector(FVector3d& Vector, FArchive& Ar)
+template<int32 MaxValue, uint32 NumBits>
+bool SerializeFixedVector(FVector3d &Vector, FArchive& Ar)
 {
-	// LWC_TODO: SerializeFixedVector workaround. Double versions of Read/WriteFixedCompressedFloat for improved performance?
-	FVector3f AsFloat(Vector);
-	bool bResult = SerializeFixedVector<MaxValue, NumBits>(AsFloat, Ar);
-	Vector = AsFloat;
-	return bResult;
+	if (Ar.IsSaving())
+	{
+		bool success = true;
+		success &= WriteFixedCompressedFloat<MaxValue, NumBits>(Vector.X, Ar);
+		success &= WriteFixedCompressedFloat<MaxValue, NumBits>(Vector.Y, Ar);
+		success &= WriteFixedCompressedFloat<MaxValue, NumBits>(Vector.Z, Ar);
+		return success;
+	}
+	else
+	{
+		ReadFixedCompressedFloat<MaxValue, NumBits>(Vector.X, Ar);
+		ReadFixedCompressedFloat<MaxValue, NumBits>(Vector.Y, Ar);
+		ReadFixedCompressedFloat<MaxValue, NumBits>(Vector.Z, Ar);
+		return true;
+	}
 }
 // --------------------------------------------------------------
 
diff --git a/Engine/Source/Runtime/Net/Core/Private/Net/Core/Serialization/QuantizedVectorSerialization.cpp b/Engine/Source/Runtime/Net/Core/Private/Net/Core/Serialization/QuantizedVectorSerialization.cpp
new file mode 100644
--- /dev/null
+++ b/Engine/Source/Runtime/Net/Core/Private/Net/Core/Serialization/QuantizedVectorSerialization.cpp
@@ -0,0 +1,218 @@
+// Copyright Epic Games, Inc. All Rights Reserved.
+
+#include "Net/Core/Serialization/QuantizedVectorSerialization.h"
+#include "HAL/PlatformMath.h"
+#include "Logging/LogMacros.h"
+#include "Math/Vector.h"
+#include "Traits/IntType.h"
+
+namespace UE::Net::Private
+{
+
+/* Returns the number of bits needed for the int32 to be replicated properly. At least 1 for the sign bit. */
+inline uint32 GetBitsNeeded(const int32 Value)
+{
+	const uint32 MassagedValue = uint32(Value ^ (Value >> 31U));
+	return 33U - static_cast<uint32>(FPlatformMath::CountLeadingZeros(MassagedValue));
+}
+
+/* Returns the number of bits needed for the int32 to be replicated properly. At least 1 for the sign bit. */
+inline uint32 GetBitsNeeded(const int64 Value)
+{
+	const uint64 MassagedValue = uint64(Value ^ (Value >> 63U));
+	return 65U - static_cast<uint32>(FPlatformMath::CountLeadingZeros64(MassagedValue));
+}
+
+/* Round to integer */
+inline int32 RoundFloatToInt(float F)
+{
+	return int32(F + FPlatformMath::Sign(F)*0.5f);
+}
+
+inline int64 RoundFloatToInt(double F)
+{
+	return int64(F + FPlatformMath::Sign(F)*0.5);
+}
+
+template<class T>
+bool WriteQuantizedVector(const int32 Scale, const T& Value, FArchive& Ar)
+{
+	using ScalarType = decltype(T::X);
+	constexpr SIZE_T ScalarTypeSize = sizeof(ScalarType);
+	using IntType = typename TSignedIntType<ScalarTypeSize>::Type;
+
+	static_assert(ScalarTypeSize == 4U || ScalarTypeSize == 8U, "Unknown floating point type.");
+
+	// Beyond 2^MaxExponentForScaling scaling cannot improve the precision as the next floating point value is at least 1.0 more. 
+	constexpr uint32 MaxExponentForScaling = ScalarTypeSize == 4 ? 23U : 52U;
+	constexpr ScalarType MaxValueToScale = ScalarType(IntType(1) << MaxExponentForScaling);
+
+	// Rounding of large values can introduce additional precision errors and the extra cost to serialize with full precision is small.
+	constexpr uint32 MaxExponentAfterScaling = ScalarTypeSize == 4 ? 30U : 62U;
+	constexpr ScalarType MaxScaledValue = ScalarType(IntType(1) << MaxExponentAfterScaling);
+
+	// NaN values can be properly serialized using the full precision path, but they typically cause lots of errors
+	// for the typical engine use case.
+	if (Value.ContainsNaN())
+	{
+		logOrEnsureNanError(TEXT("%s"), TEXT("WriteQuantizedVector: Value isn't finite. Clearing for safety."));
+		WriteQuantizedVector(Scale, T{0,0,0}, Ar);
+		return false;
+	}
+
+	const ScalarType Factor = Scale;
+	T ScaledValue;
+	ScaledValue.X = Value.X*Factor;
+	ScaledValue.Y = Value.Y*Factor;
+	ScaledValue.Z = Value.Z*Factor;
+
+	// If the component values are within bounds then we optimize the bandwidth, otherwise we use full precision.
+	if (ScaledValue.GetAbsMax() < MaxScaledValue)
+	{
+		const bool bUseScaledValue = Value.GetAbsMin() < MaxValueToScale;
+
+		IntType X;
+		IntType Y;
+		IntType Z;
+		if (bUseScaledValue)
+		{
+			X = RoundFloatToInt(ScaledValue.X);
+			Y = RoundFloatToInt(ScaledValue.Y);
+			Z = RoundFloatToInt(ScaledValue.Z);
+		}
+		else
+		{
+			X = RoundFloatToInt(Value.X);
+			Y = RoundFloatToInt(Value.Y);
+			Z = RoundFloatToInt(Value.Z);
+		}
+
+		const uint32 ComponentBitCount = FPlatformMath::Max(GetBitsNeeded(X), FPlatformMath::Max(GetBitsNeeded(Y), GetBitsNeeded(Z)));
+		uint32 ComponentBitCountAndScaleInfo = (bUseScaledValue ? (1U << 6U) : 0U) | ComponentBitCount;
+		Ar.SerializeInt(ComponentBitCountAndScaleInfo, 1U << 7U);
+
+		Ar.SerializeBits(&X, ComponentBitCount);
+		Ar.SerializeBits(&Y, ComponentBitCount);
+		Ar.SerializeBits(&Z, ComponentBitCount);
+	}
+	else
+	{
+		// A component bit count of 0 indicates full precision.
+		constexpr uint32 ComponentBitCount = 0;
+		uint32 ComponentBitCountAndTypeInfo = (ScalarTypeSize == 8U ? (1U << 6U) : 0U) | ComponentBitCount;
+		Ar.SerializeInt(ComponentBitCountAndTypeInfo, 1U << 7U);
+		Ar.SerializeBits(const_cast<T*>(&Value), ScalarTypeSize*8U*3U);
+	}
+	
+	return true;
+}
+
+template<class T>
+bool ReadQuantizedVector(const int32 Scale, T& Value, FArchive& Ar)
+{
+	using ScalarType = decltype(T::X);
+	constexpr SIZE_T ScalarTypeSize = sizeof(ScalarType);
+	static_assert(ScalarTypeSize == 4 || ScalarTypeSize == 8, "Unknown floating point type.");
+	using IntType = typename TSignedIntType<ScalarTypeSize>::Type;
+
+	uint32 ComponentBitCountAndExtraInfo = 0;
+	Ar.SerializeInt(ComponentBitCountAndExtraInfo, 1U << 7U);
+	const uint32 ComponentBitCount = ComponentBitCountAndExtraInfo & 63U;
+	const uint32 ExtraInfo = ComponentBitCountAndExtraInfo >> 6U;
+
+	if (ComponentBitCount > 0U)
+	{
+		int64 X = 0;
+		int64 Y = 0;
+		int64 Z = 0;
+
+		Ar.SerializeBits(&X, ComponentBitCount);
+		Ar.SerializeBits(&Y, ComponentBitCount);
+		Ar.SerializeBits(&Z, ComponentBitCount);
+
+		// Sign-extend the values. The most significant bit read indicates the sign.
+		const uint64 SignBit = (1ULL << (ComponentBitCount - 1U));
+		X = (X ^ SignBit) - SignBit;
+		Y = (Y ^ SignBit) - SignBit;
+		Z = (Z ^ SignBit) - SignBit;
+
+		T TempValue;
+		TempValue.X = ScalarType(X);
+		TempValue.Y = ScalarType(Y);
+		TempValue.Z = ScalarType(Z);
+
+		// Apply scaling if needed.
+		if (ExtraInfo)
+		{
+			Value = TempValue/ScalarType(Scale);
+		}
+		else
+		{
+			Value = TempValue;
+		}
+
+		return true;
+	}
+	else
+	{
+		const uint32 ReceivedScalarTypeSize = ExtraInfo ? 8U : 4U;
+		if (ReceivedScalarTypeSize == 8U)
+		{
+			FVector3d TempValue;
+			Ar.SerializeBits(&TempValue, 8U*8U*3U);
+			if (TempValue.ContainsNaN())
+			{
+				logOrEnsureNanError(TEXT("%s"), TEXT("ReadQuantizedVector: Value isn't finite. Clearing for safety."));
+				Value = T{0,0,0};
+				return false;
+			}
+
+			Value = T(TempValue);
+			return true;
+		}
+		else
+		{
+			FVector3f TempValue;
+			Ar.SerializeBits(&TempValue, 4U*8U*3U);
+			if (TempValue.ContainsNaN())
+			{
+				logOrEnsureNanError(TEXT("%s"), TEXT("ReadQuantizedVector: Value isn't finite. Clearing for safety."));
+				Value = T{0,0,0};
+				return false;
+			}
+
+			Value = T(TempValue);
+			return true;
+		}
+	}
+
+	// Should not get here so something is very wrong.
+	return false;
+}
+
+}
+
+namespace UE::Net
+{
+
+bool WriteQuantizedVector(int32 Scale, const FVector3d& Value, FArchive& Ar)
+{
+	return Private::WriteQuantizedVector(Scale, Value, Ar);
+}
+
+bool ReadQuantizedVector(int32 Scale, FVector3d& Value, FArchive& Ar)
+{
+	return Private::ReadQuantizedVector(Scale, Value, Ar);
+}
+
+bool WriteQuantizedVector(int32 Scale, const FVector3f& Value, FArchive& Ar)
+{
+	return Private::WriteQuantizedVector(Scale, Value, Ar);
+}
+
+bool ReadQuantizedVector(int32 Scale, FVector3f& Value, FArchive& Ar)
+{
+	return Private::ReadQuantizedVector(Scale, Value, Ar);
+}
+
+}
diff --git a/Engine/Source/Runtime/Net/Core/Public/Net/Core/Serialization/QuantizedVectorSerialization.h b/Engine/Source/Runtime/Net/Core/Public/Net/Core/Serialization/QuantizedVectorSerialization.h
new file mode 100644
--- /dev/null
+++ b/Engine/Source/Runtime/Net/Core/Public/Net/Core/Serialization/QuantizedVectorSerialization.h
@@ -0,0 +1,39 @@
+// Copyright Epic Games, Inc. All Rights Reserved.
+
+#pragma once
+
+#include "CoreTypes.h"
+#include "Math/MathFwd.h"
+#include "Serialization/Archive.h"
+
+namespace UE::Net
+{
+
+/* Scales the supplied vector and then rounds to an integer value which is serialized.
+ * If scaling cannot improve the precision of any of the components the value will not be scaled before rounding.
+ * If scaling produces a value that would cause any of the components to be clamped then serialization
+ * will fallback to send full precision components.
+ * N.B. The Vector3f and Vector3d variants are bit compatible with eachother, but it's discouraged to mix types
+ * when writing and reading.
+ */
+NETCORE_API bool WriteQuantizedVector(const int32 Scale, const FVector3d& Value, FArchive& Ar);
+NETCORE_API bool WriteQuantizedVector(const int32 Scale, const FVector3f& Value, FArchive& Ar);
+
+/* Reads a quantized vector that was written using WriteQuantizedVector with the same scale. */
+NETCORE_API bool ReadQuantizedVector(const int32 Scale, FVector3d& Value, FArchive& Ar);
+NETCORE_API bool ReadQuantizedVector(const int32 Scale, FVector3f& Value, FArchive& Ar);
+
+template<int32 Scale>
+bool SerializeQuantizedVector(FVector& Value, FArchive& Ar)
+{
+	if (Ar.IsLoading())
+	{
+		return ReadQuantizedVector(Scale, Value, Ar);
+	}
+	else
+	{
+		return WriteQuantizedVector(Scale, Value, Ar);
+	}
+}
+
+}
