Deprecating EInternalObjectFlags::PendingKill. Making sure iterators use the appropriate flags based on the current state of PendingKill being enabled or not.

diff --git a/Engine/Plugins/Animation/ControlRig/Source/ControlRigEditor/Private/ControlRigEditorModule.cpp b/Engine/Plugins/Animation/ControlRig/Source/ControlRigEditor/Private/ControlRigEditorModule.cpp
--- a/Engine/Plugins/Animation/ControlRig/Source/ControlRigEditor/Private/ControlRigEditorModule.cpp
+++ b/Engine/Plugins/Animation/ControlRig/Source/ControlRigEditor/Private/ControlRigEditorModule.cpp
@@ -2222,7 +2222,7 @@ void FControlRigEditorModule::PreChange(const UUserDefinedStruct* Changed,
 		// List of unique classes and structs to regenerate bytecode and property referenced objects list
 		TSet<UStruct*> StructsToRegenerateReferencesFor;
 
-		for (TAllFieldsIterator<FStructProperty> FieldIt(RF_NoFlags, EInternalObjectFlags::PendingKill); FieldIt; ++FieldIt)
+		for (TAllFieldsIterator<FStructProperty> FieldIt(RF_NoFlags, EInternalObjectFlags::Garbage); FieldIt; ++FieldIt)
 		{
 			FStructProperty* StructProperty = *FieldIt;
 			if (StructProperty && (StructureToReinstance == StructProperty->Struct))
@@ -2264,7 +2264,7 @@ void FControlRigEditorModule::PreChange(const UUserDefinedStruct* Changed,
 	}
 	
 	// in the future we could only invalidate caches on affected rig instances, it shouldn't make too much of a difference though
-	for (TObjectIterator<UControlRig> It(RF_Transient | RF_ClassDefaultObject, /** bIncludeDerivedClasses */ true, /** InternalExcludeFlags */ EInternalObjectFlags::PendingKill); It; ++It)
+	for (TObjectIterator<UControlRig> It(RF_Transient | RF_ClassDefaultObject, /** bIncludeDerivedClasses */ true, /** InternalExcludeFlags */ EInternalObjectFlags::Garbage); It; ++It)
 	{
 		UControlRig* Rig = *It;
 		// rebuild property list and property path list
@@ -2276,7 +2276,7 @@ void FControlRigEditorModule::PostChange(const UUserDefinedStruct* Changed,
                                    FStructureEditorUtils::EStructureEditorChangeInfo ChangedType)
 {
 	TArray<UControlRigBlueprint*> BlueprintsToRefresh;
-	for (TObjectIterator<URigVMPin> It(RF_Transient | RF_ClassDefaultObject, /** bIncludeDerivedClasses */ true, /** InternalExcludeFlags */ EInternalObjectFlags::PendingKill); It; ++It)
+	for (TObjectIterator<URigVMPin> It(RF_Transient | RF_ClassDefaultObject, /** bIncludeDerivedClasses */ true, /** InternalExcludeFlags */ EInternalObjectFlags::Garbage); It; ++It)
 	{
 		URigVMPin* Pin = *It;
 		// GetCPPTypeObject also makes sure the pin's type information is update to date
diff --git a/Engine/Plugins/Developer/TraceSourceFiltering/Source/SourceFilteringTrace/Private/ActorFiltering.h b/Engine/Plugins/Developer/TraceSourceFiltering/Source/SourceFilteringTrace/Private/ActorFiltering.h
--- a/Engine/Plugins/Developer/TraceSourceFiltering/Source/SourceFilteringTrace/Private/ActorFiltering.h
+++ b/Engine/Plugins/Developer/TraceSourceFiltering/Source/SourceFilteringTrace/Private/ActorFiltering.h
@@ -168,7 +168,7 @@ public:
 
 		/** Ignore CDOs and objects marked as pending kill */
 		EObjectFlags ExcludeFlags = RF_ClassDefaultObject;
-		EInternalObjectFlags InternalExcludeFlags = EInternalObjectFlags::PendingKill;
+		EInternalObjectFlags InternalExcludeFlags = EInternalObjectFlags::Garbage;
 
 		ForEachObjectOfClasses(AllFilteredClasses, [this](UObject* Object)
 		{
diff --git a/Engine/Plugins/Experimental/GPULightmass/Source/GPULightmass/Private/GPULightmassSettings.cpp b/Engine/Plugins/Experimental/GPULightmass/Source/GPULightmass/Private/GPULightmassSettings.cpp
--- a/Engine/Plugins/Experimental/GPULightmass/Source/GPULightmass/Private/GPULightmassSettings.cpp
+++ b/Engine/Plugins/Experimental/GPULightmass/Source/GPULightmass/Private/GPULightmassSettings.cpp
@@ -318,7 +318,7 @@ void UGPULightmassSubsystem::Launch()
 				int32 NumPrimitiveComponents = 0;
 				int32 NumLightComponents = 0;
 
-				for (UPrimitiveComponent* Component : TObjectRange<UPrimitiveComponent>(RF_ClassDefaultObject | RF_ArchetypeObject, true, EInternalObjectFlags::PendingKill))
+				for (UPrimitiveComponent* Component : TObjectRange<UPrimitiveComponent>(RF_ClassDefaultObject | RF_ArchetypeObject, true, EInternalObjectFlags::Garbage))
 				{
 					if (Component->HasValidSettingsForStaticLighting(false))
 					{
@@ -326,7 +326,7 @@ void UGPULightmassSubsystem::Launch()
 					}
 				}
 
-				for (ULightComponentBase* Component : TObjectRange<ULightComponentBase>(RF_ClassDefaultObject | RF_ArchetypeObject, true, EInternalObjectFlags::PendingKill))
+				for (ULightComponentBase* Component : TObjectRange<ULightComponentBase>(RF_ClassDefaultObject | RF_ArchetypeObject, true, EInternalObjectFlags::Garbage))
 				{
 					if (Component->bAffectsWorld && Component->HasStaticShadowing())
 					{
@@ -337,7 +337,7 @@ void UGPULightmassSubsystem::Launch()
 				FScopedSlowTask SubSlowTask(NumPrimitiveComponents + NumLightComponents, LOCTEXT("RegisteringComponentsWithStaticLightingSystem", "Registering components with static lighting system"));
 				SubSlowTask.MakeDialog();
 
-				for (UPrimitiveComponent* Component : TObjectRange<UPrimitiveComponent>(RF_ClassDefaultObject | RF_ArchetypeObject, true, EInternalObjectFlags::PendingKill))
+				for (UPrimitiveComponent* Component : TObjectRange<UPrimitiveComponent>(RF_ClassDefaultObject | RF_ArchetypeObject, true, EInternalObjectFlags::Garbage))
 				{
 					if (Component->HasValidSettingsForStaticLighting(false))
 					{
@@ -347,7 +347,7 @@ void UGPULightmassSubsystem::Launch()
 					}
 				}
 
-				for (ULightComponentBase* Component : TObjectRange<ULightComponentBase>(RF_ClassDefaultObject | RF_ArchetypeObject, true, EInternalObjectFlags::PendingKill))
+				for (ULightComponentBase* Component : TObjectRange<ULightComponentBase>(RF_ClassDefaultObject | RF_ArchetypeObject, true, EInternalObjectFlags::Garbage))
 				{
 					if (Component->bAffectsWorld && Component->HasStaticShadowing())
 					{
@@ -389,12 +389,12 @@ void UGPULightmassSubsystem::Stop()
 			int32 NumPrimitiveComponents = 0;
 			int32 NumLightComponents = 0;
 
-			for (UPrimitiveComponent* Component : TObjectRange<UPrimitiveComponent>(RF_ClassDefaultObject | RF_ArchetypeObject, true, EInternalObjectFlags::PendingKill))
+			for (UPrimitiveComponent* Component : TObjectRange<UPrimitiveComponent>(RF_ClassDefaultObject | RF_ArchetypeObject, true, EInternalObjectFlags::Garbage))
 			{
 				NumPrimitiveComponents++;
 			}
 
-			for (ULightComponentBase* Component : TObjectRange<ULightComponentBase>(RF_ClassDefaultObject | RF_ArchetypeObject, true, EInternalObjectFlags::PendingKill))
+			for (ULightComponentBase* Component : TObjectRange<ULightComponentBase>(RF_ClassDefaultObject | RF_ArchetypeObject, true, EInternalObjectFlags::Garbage))
 			{
 				NumLightComponents++;
 			}
@@ -402,19 +402,19 @@ void UGPULightmassSubsystem::Stop()
 			FScopedSlowTask SubSlowTask(NumPrimitiveComponents + NumLightComponents, LOCTEXT("UnregisteringComponentsWithStaticLightingSystem", "Unregistering components with static lighting system"));
 
 			// Unregister all landscapes first to prevent grass picking up landscape lightmaps
-			for (ULandscapeComponent* Component : TObjectRange<ULandscapeComponent>(RF_ClassDefaultObject | RF_ArchetypeObject, true, EInternalObjectFlags::PendingKill))
+			for (ULandscapeComponent* Component : TObjectRange<ULandscapeComponent>(RF_ClassDefaultObject | RF_ArchetypeObject, true, EInternalObjectFlags::Garbage))
 			{
 				FStaticLightingSystemInterface::OnPrimitiveComponentUnregistered.Broadcast(Component);
 			}
 
-			for (UPrimitiveComponent* Component : TObjectRange<UPrimitiveComponent>(RF_ClassDefaultObject | RF_ArchetypeObject, true, EInternalObjectFlags::PendingKill))
+			for (UPrimitiveComponent* Component : TObjectRange<UPrimitiveComponent>(RF_ClassDefaultObject | RF_ArchetypeObject, true, EInternalObjectFlags::Garbage))
 			{
 				FStaticLightingSystemInterface::OnPrimitiveComponentUnregistered.Broadcast(Component);
 
 				SubSlowTask.EnterProgressFrame(1, LOCTEXT("UnregisteringComponentsWithStaticLightingSystem", "Unregistering components with static lighting system"));
 			}
 
-			for (ULightComponentBase* Component : TObjectRange<ULightComponentBase>(RF_ClassDefaultObject | RF_ArchetypeObject, true, EInternalObjectFlags::PendingKill))
+			for (ULightComponentBase* Component : TObjectRange<ULightComponentBase>(RF_ClassDefaultObject | RF_ArchetypeObject, true, EInternalObjectFlags::Garbage))
 			{
 				FStaticLightingSystemInterface::OnLightComponentUnregistered.Broadcast(Component);
 
diff --git a/Engine/Plugins/Experimental/VirtualCamera/Source/VCamCoreEditor/Private/VCamEditorLibrary.cpp b/Engine/Plugins/Experimental/VirtualCamera/Source/VCamCoreEditor/Private/VCamEditorLibrary.cpp
--- a/Engine/Plugins/Experimental/VirtualCamera/Source/VCamCoreEditor/Private/VCamEditorLibrary.cpp
+++ b/Engine/Plugins/Experimental/VirtualCamera/Source/VCamCoreEditor/Private/VCamEditorLibrary.cpp
@@ -15,7 +15,7 @@ void UVCamEditorLibrary::GetAllVCamComponentsInLevel(TArray<UVCamComponent*>& VC
 
 	// Loop all VCamComponents which are not pending kill or CDOs
 	const EObjectFlags ExcludeFlags = RF_ClassDefaultObject;
-	for (TObjectIterator<UVCamComponent> It(ExcludeFlags, true, EInternalObjectFlags::PendingKill); It; ++It)
+	for (TObjectIterator<UVCamComponent> It(ExcludeFlags, true, EInternalObjectFlags::Garbage); It; ++It)
 	{
 		UVCamComponent* VCamComponent = *It;
 		if (IsValid(VCamComponent))
diff --git a/Engine/Plugins/Experimental/VirtualHeightfieldMesh/Source/VirtualHeightfieldMeshEditor/Private/WorldPartitionVirtualHeightfieldMeshBuilder.cpp b/Engine/Plugins/Experimental/VirtualHeightfieldMesh/Source/VirtualHeightfieldMeshEditor/Private/WorldPartitionVirtualHeightfieldMeshBuilder.cpp
--- a/Engine/Plugins/Experimental/VirtualHeightfieldMesh/Source/VirtualHeightfieldMeshEditor/Private/WorldPartitionVirtualHeightfieldMeshBuilder.cpp
+++ b/Engine/Plugins/Experimental/VirtualHeightfieldMesh/Source/VirtualHeightfieldMeshEditor/Private/WorldPartitionVirtualHeightfieldMeshBuilder.cpp
@@ -20,7 +20,7 @@ bool UWorldPartitionVirtualHeightfieldMeshBuilder::RunInternal(UWorld* World, co
 	// Find UVirtualHeightfieldMeshComponent and update the associated UHeightfieldMinMaxTexture assets.
 	// todo[vhm]: Convert builder to sequentially load and build sections of the world.
 	TSet<UObject*> ModifiedObjects;
-	for (TObjectIterator<UVirtualHeightfieldMeshComponent> It(RF_ClassDefaultObject, false, EInternalObjectFlags::PendingKill); It; ++It)
+	for (TObjectIterator<UVirtualHeightfieldMeshComponent> It(RF_ClassDefaultObject, false, EInternalObjectFlags::Garbage); It; ++It)
 	{
 		if (It->GetWorld() == World)
 		{
diff --git a/Engine/Plugins/Runtime/GameplayAbilities/Source/GameplayAbilities/Private/AbilitySystemComponent_Abilities.cpp b/Engine/Plugins/Runtime/GameplayAbilities/Source/GameplayAbilities/Private/AbilitySystemComponent_Abilities.cpp
--- a/Engine/Plugins/Runtime/GameplayAbilities/Source/GameplayAbilities/Private/AbilitySystemComponent_Abilities.cpp
+++ b/Engine/Plugins/Runtime/GameplayAbilities/Source/GameplayAbilities/Private/AbilitySystemComponent_Abilities.cpp
@@ -65,7 +65,7 @@ void UAbilitySystemComponent::InitializeComponent()
 	}
 
 	TArray<UObject*> ChildObjects;
-	GetObjectsWithOuter(Owner, ChildObjects, false, RF_NoFlags, EInternalObjectFlags::PendingKill);
+	GetObjectsWithOuter(Owner, ChildObjects, false, RF_NoFlags, EInternalObjectFlags::Garbage);
 
 	for (UObject* Obj : ChildObjects)
 	{
diff --git a/Engine/Plugins/VirtualProduction/CameraCalibration/Source/CameraCalibrationEditor/Private/Calibrators/CameraLensDistortionAlgoPoints.cpp b/Engine/Plugins/VirtualProduction/CameraCalibration/Source/CameraCalibrationEditor/Private/Calibrators/CameraLensDistortionAlgoPoints.cpp
--- a/Engine/Plugins/VirtualProduction/CameraCalibration/Source/CameraCalibrationEditor/Private/Calibrators/CameraLensDistortionAlgoPoints.cpp
+++ b/Engine/Plugins/VirtualProduction/CameraCalibration/Source/CameraCalibrationEditor/Private/Calibrators/CameraLensDistortionAlgoPoints.cpp
@@ -536,7 +536,7 @@ AActor* UCameraLensDistortionAlgoPoints::FindFirstCalibrator() const
 	const UWorld* World = StepsController->GetWorld();
 	const EObjectFlags ExcludeFlags = RF_ClassDefaultObject; // We don't want the calibrator CDOs.
 
-	for (TObjectIterator<UCalibrationPointComponent> It(ExcludeFlags, true, EInternalObjectFlags::PendingKill); It; ++It)
+	for (TObjectIterator<UCalibrationPointComponent> It(ExcludeFlags, true, EInternalObjectFlags::Garbage); It; ++It)
 	{
 		AActor* Actor = It->GetOwner();
 
diff --git a/Engine/Plugins/VirtualProduction/CameraCalibration/Source/CameraCalibrationEditor/Private/Calibrators/CameraNodalOffsetAlgoPoints.cpp b/Engine/Plugins/VirtualProduction/CameraCalibration/Source/CameraCalibrationEditor/Private/Calibrators/CameraNodalOffsetAlgoPoints.cpp
--- a/Engine/Plugins/VirtualProduction/CameraCalibration/Source/CameraCalibrationEditor/Private/Calibrators/CameraNodalOffsetAlgoPoints.cpp
+++ b/Engine/Plugins/VirtualProduction/CameraCalibration/Source/CameraCalibrationEditor/Private/Calibrators/CameraNodalOffsetAlgoPoints.cpp
@@ -2310,7 +2310,7 @@ AActor* UCameraNodalOffsetAlgoPoints::FindFirstCalibrator() const
 	const UWorld* World = StepsController->GetWorld();
 	const EObjectFlags ExcludeFlags = RF_ClassDefaultObject; // We don't want the calibrator CDOs.
 
-	for (TObjectIterator<UCalibrationPointComponent> It(ExcludeFlags, true, EInternalObjectFlags::PendingKill); It; ++It)
+	for (TObjectIterator<UCalibrationPointComponent> It(ExcludeFlags, true, EInternalObjectFlags::Garbage); It; ++It)
 	{
 		AActor* Actor = It->GetOwner();
 
diff --git a/Engine/Plugins/VirtualProduction/DMX/DMXEngine/Source/DMXBlueprintGraph/Private/DMXBlueprintGraphModule.cpp b/Engine/Plugins/VirtualProduction/DMX/DMXEngine/Source/DMXBlueprintGraph/Private/DMXBlueprintGraphModule.cpp
--- a/Engine/Plugins/VirtualProduction/DMX/DMXEngine/Source/DMXBlueprintGraph/Private/DMXBlueprintGraphModule.cpp
+++ b/Engine/Plugins/VirtualProduction/DMX/DMXEngine/Source/DMXBlueprintGraph/Private/DMXBlueprintGraphModule.cpp
@@ -70,7 +70,7 @@ void FDMXBlueprintGraphModule::RegisterCustomClassLayout(FName ClassName, FOnGet
 
 void FDMXBlueprintGraphModule::OnFixtureTypeChanged(const UDMXEntityFixtureType* InFixtureType)
 {
-	for (TObjectIterator<UK2Node_GetDMXAttributeValues> It(RF_Transient | RF_ClassDefaultObject, /** bIncludeDerivedClasses */ true, /** InternalExcludeFlags */ EInternalObjectFlags::PendingKill); It; ++It)
+	for (TObjectIterator<UK2Node_GetDMXAttributeValues> It(RF_Transient | RF_ClassDefaultObject, /** bIncludeDerivedClasses */ true, /** InternalExcludeFlags */ EInternalObjectFlags::Garbage); It; ++It)
 	{
 		if (It->HasValidBlueprint())
 		{
diff --git a/Engine/Plugins/VirtualProduction/DMX/DMXPixelMapping/Source/DMXPixelMappingEditor/Private/Toolkits/DMXPixelMappingToolkit.cpp b/Engine/Plugins/VirtualProduction/DMX/DMXPixelMapping/Source/DMXPixelMappingEditor/Private/Toolkits/DMXPixelMappingToolkit.cpp
--- a/Engine/Plugins/VirtualProduction/DMX/DMXPixelMapping/Source/DMXPixelMappingEditor/Private/Toolkits/DMXPixelMappingToolkit.cpp
+++ b/Engine/Plugins/VirtualProduction/DMX/DMXPixelMapping/Source/DMXPixelMappingEditor/Private/Toolkits/DMXPixelMappingToolkit.cpp
@@ -515,7 +515,7 @@ void FDMXPixelMappingToolkit::UpdateBlueprintNodes(UDMXPixelMapping* InDMXPixelM
 {
 	if (InDMXPixelMapping != nullptr)
 	{
-		for (TObjectIterator<UK2Node_PixelMappingBaseComponent> It(RF_Transient | RF_ClassDefaultObject, /** bIncludeDerivedClasses */ true, /** InternalExcludeFlags */ EInternalObjectFlags::PendingKill); It; ++It)
+		for (TObjectIterator<UK2Node_PixelMappingBaseComponent> It(RF_Transient | RF_ClassDefaultObject, /** bIncludeDerivedClasses */ true, /** InternalExcludeFlags */ EInternalObjectFlags::Garbage); It; ++It)
 		{
 			It->OnPixelMappingChanged(InDMXPixelMapping);
 		}
diff --git a/Engine/Plugins/VirtualProduction/RemoteControl/Source/RemoteControl/Private/RemoteControlPresetRebindingManager.cpp b/Engine/Plugins/VirtualProduction/RemoteControl/Source/RemoteControl/Private/RemoteControlPresetRebindingManager.cpp
--- a/Engine/Plugins/VirtualProduction/RemoteControl/Source/RemoteControl/Private/RemoteControlPresetRebindingManager.cpp
+++ b/Engine/Plugins/VirtualProduction/RemoteControl/Source/RemoteControl/Private/RemoteControlPresetRebindingManager.cpp
@@ -126,7 +126,7 @@ namespace RCPresetRebindingManager
 		}
 
 		const EObjectFlags ExcludeFlags = RF_ClassDefaultObject | RF_ArchetypeObject | RF_Transient;
-		for (TObjectIterator<UActorComponent> It(ExcludeFlags, true, EInternalObjectFlags::PendingKill); It; ++It)
+		for (TObjectIterator<UActorComponent> It(ExcludeFlags, true, EInternalObjectFlags::Garbage); It; ++It)
 		{
 			AddRelevantClassesForObject(*It);
 		}
diff --git a/Engine/Source/Editor/BlueprintGraph/Private/BlueprintNodeTemplateCache.cpp b/Engine/Source/Editor/BlueprintGraph/Private/BlueprintNodeTemplateCache.cpp
--- a/Engine/Source/Editor/BlueprintGraph/Private/BlueprintNodeTemplateCache.cpp
+++ b/Engine/Source/Editor/BlueprintGraph/Private/BlueprintNodeTemplateCache.cpp
@@ -123,7 +123,7 @@ static UEdGraph* BlueprintNodeTemplateCacheImpl::FindCompatibleGraph(UBlueprint*
 	UEdGraph* FoundGraph = nullptr;
 
 	TArray<UObject*> BlueprintChildObjs;
-	GetObjectsWithOuter(BlueprintOuter, BlueprintChildObjs, /*bIncludeNestedObjects =*/false, /*ExclusionFlags =*/ RF_NoFlags, /** InternalExcludeFlags */ EInternalObjectFlags::PendingKill);
+	GetObjectsWithOuter(BlueprintOuter, BlueprintChildObjs, /*bIncludeNestedObjects =*/false, /*ExclusionFlags =*/ RF_NoFlags, /** InternalExcludeFlags */ EInternalObjectFlags::Garbage);
 
 	for (UObject* Child : BlueprintChildObjs)
 	{
diff --git a/Engine/Source/Editor/DetailCustomizations/Private/LandscapeProxyUIDetails.cpp b/Engine/Source/Editor/DetailCustomizations/Private/LandscapeProxyUIDetails.cpp
--- a/Engine/Source/Editor/DetailCustomizations/Private/LandscapeProxyUIDetails.cpp
+++ b/Engine/Source/Editor/DetailCustomizations/Private/LandscapeProxyUIDetails.cpp
@@ -173,7 +173,7 @@ static void GetMissingRuntimeVirtualTextureVolumes(ALandscapeProxy* LandscapePro
 	}
 
 	TArray<URuntimeVirtualTexture*> FoundVolumes;
-	for (TObjectIterator<URuntimeVirtualTextureComponent> It(RF_ClassDefaultObject, false, EInternalObjectFlags::PendingKill); It; ++It)
+	for (TObjectIterator<URuntimeVirtualTextureComponent> It(RF_ClassDefaultObject, false, EInternalObjectFlags::Garbage); It; ++It)
 	{
 		if (It->GetWorld() == World)
 		{
diff --git a/Engine/Source/Editor/Kismet/Private/SMyBlueprint.cpp b/Engine/Source/Editor/Kismet/Private/SMyBlueprint.cpp
--- a/Engine/Source/Editor/Kismet/Private/SMyBlueprint.cpp
+++ b/Engine/Source/Editor/Kismet/Private/SMyBlueprint.cpp
@@ -2831,7 +2831,7 @@ void SMyBlueprint::OnDeleteGraph(UEdGraph* InGraph, EEdGraphSchemaAction_K2Graph
 		FBlueprintEditorUtils::RemoveGraph(GetBlueprintObj(), InGraph, EGraphRemoveFlags::Recompile);
 		BlueprintEditorPtr.Pin()->CloseDocumentTab(InGraph);
 
-		for (TObjectIterator<UK2Node_CreateDelegate> It(RF_ClassDefaultObject, /** bIncludeDerivedClasses */ true, /** InternalExcludeFlags */ EInternalObjectFlags::PendingKill); It; ++It)
+		for (TObjectIterator<UK2Node_CreateDelegate> It(RF_ClassDefaultObject, /** bIncludeDerivedClasses */ true, /** InternalExcludeFlags */ EInternalObjectFlags::Garbage); It; ++It)
 		{
 			if (It->GetGraph() != InGraph)
 			{
@@ -2891,7 +2891,7 @@ void SMyBlueprint::OnDeleteDelegate(FEdGraphSchemaAction_K2Delegate* InDelegateA
 		FBlueprintEditorUtils::RemoveMemberVariable(BlueprintObj, GraphToActOn->GetFName());
 		FBlueprintEditorUtils::RemoveGraph(BlueprintObj, GraphToActOn, EGraphRemoveFlags::Recompile);
 
-		for (TObjectIterator<UK2Node_CreateDelegate> It(RF_ClassDefaultObject, /** bIncludeDerivedClasses */ true, /** InternalExcludeFlags */ EInternalObjectFlags::PendingKill); It; ++It)
+		for (TObjectIterator<UK2Node_CreateDelegate> It(RF_ClassDefaultObject, /** bIncludeDerivedClasses */ true, /** InternalExcludeFlags */ EInternalObjectFlags::Garbage); It; ++It)
 		{
 			if (IsValid(*It) && IsValid(It->GetGraph()))
 			{
diff --git a/Engine/Source/Editor/KismetCompiler/Private/UserDefinedStructureCompilerUtils.cpp b/Engine/Source/Editor/KismetCompiler/Private/UserDefinedStructureCompilerUtils.cpp
--- a/Engine/Source/Editor/KismetCompiler/Private/UserDefinedStructureCompilerUtils.cpp
+++ b/Engine/Source/Editor/KismetCompiler/Private/UserDefinedStructureCompilerUtils.cpp
@@ -82,7 +82,7 @@ struct FUserDefinedStructureCompilerInner
 			// List of unique classes and structs to regenerate bytecode and property referenced objects list
 			TSet<UStruct*> StructsToRegenerateReferencesFor;
 
-			for (TAllFieldsIterator<FStructProperty> FieldIt(RF_NoFlags, EInternalObjectFlags::PendingKill); FieldIt; ++FieldIt)
+			for (TAllFieldsIterator<FStructProperty> FieldIt(RF_NoFlags, EInternalObjectFlags::Garbage); FieldIt; ++FieldIt)
 			{
 				FStructProperty* StructProperty = *FieldIt;
 				if (StructProperty && (StructureToReinstance == StructProperty->Struct))
@@ -130,7 +130,7 @@ struct FUserDefinedStructureCompilerInner
 
 			DuplicatedStruct->RemoveFromRoot();
 
-			for (UBlueprint* Blueprint : TObjectRange<UBlueprint>(RF_ClassDefaultObject, /** bIncludeDerivedClasses */ true, /** InternalExcludeFlags */ EInternalObjectFlags::PendingKill))
+			for (UBlueprint* Blueprint : TObjectRange<UBlueprint>(RF_ClassDefaultObject, /** bIncludeDerivedClasses */ true, /** InternalExcludeFlags */ EInternalObjectFlags::Garbage))
 			{
 				if (Blueprint && !BlueprintsToRecompile.Contains(Blueprint))
 				{
diff --git a/Engine/Source/Editor/LandscapeEditor/Private/LandscapeEdMode.cpp b/Engine/Source/Editor/LandscapeEditor/Private/LandscapeEdMode.cpp
--- a/Engine/Source/Editor/LandscapeEditor/Private/LandscapeEdMode.cpp
+++ b/Engine/Source/Editor/LandscapeEditor/Private/LandscapeEdMode.cpp
@@ -4035,7 +4035,7 @@ bool FEdModeLandscape::NeedToFillEmptyMaterialLayers() const
 void FEdModeLandscape::UpdateBrushList()
 {
 	BrushList.Empty();
-	for (TObjectIterator<ALandscapeBlueprintBrushBase> BrushIt(RF_Transient|RF_ClassDefaultObject|RF_ArchetypeObject, true, EInternalObjectFlags::PendingKill); BrushIt; ++BrushIt)
+	for (TObjectIterator<ALandscapeBlueprintBrushBase> BrushIt(RF_Transient|RF_ClassDefaultObject|RF_ArchetypeObject, true, EInternalObjectFlags::Garbage); BrushIt; ++BrushIt)
 	{
 		ALandscapeBlueprintBrushBase* Brush = *BrushIt;
 		if (Brush->GetTypedOuter<UPackage>() != GetTransientPackage())
diff --git a/Engine/Source/Editor/UnrealEd/Private/BaseWidgetBlueprint.cpp b/Engine/Source/Editor/UnrealEd/Private/BaseWidgetBlueprint.cpp
--- a/Engine/Source/Editor/UnrealEd/Private/BaseWidgetBlueprint.cpp
+++ b/Engine/Source/Editor/UnrealEd/Private/BaseWidgetBlueprint.cpp
@@ -102,6 +102,6 @@ void UBaseWidgetBlueprint::ForEachSourceWidgetImpl(TFunctionRef<void(UWidget*)>
 		},
 		!bIncludeNestedObjects,// Reject nested objects
 		RF_NoFlags,
-		EInternalObjectFlags::PendingKill // Exclude Pending Kills
+		EInternalObjectFlags::Garbage // Exclude Pending Kills
 	);
 }
diff --git a/Engine/Source/Editor/UnrealEd/Private/Commandlets/GatherTextFromAssetsCommandlet.cpp b/Engine/Source/Editor/UnrealEd/Private/Commandlets/GatherTextFromAssetsCommandlet.cpp
--- a/Engine/Source/Editor/UnrealEd/Private/Commandlets/GatherTextFromAssetsCommandlet.cpp
+++ b/Engine/Source/Editor/UnrealEd/Private/Commandlets/GatherTextFromAssetsCommandlet.cpp
@@ -320,7 +320,7 @@ void UGatherTextFromAssetsCommandlet::PurgeGarbage(const bool bPurgeReferencedPa
 					{
 						ObjectsToKeepAlive.Add(InPackageInner);
 					}
-				}, true, RF_NoFlags, EInternalObjectFlags::PendingKill);
+				}, true, RF_NoFlags, EInternalObjectFlags::Garbage);
 			}
 		}
 	}
diff --git a/Engine/Source/Editor/UnrealEd/Private/Commandlets/StabilizeLocalizationKeys.cpp b/Engine/Source/Editor/UnrealEd/Private/Commandlets/StabilizeLocalizationKeys.cpp
--- a/Engine/Source/Editor/UnrealEd/Private/Commandlets/StabilizeLocalizationKeys.cpp
+++ b/Engine/Source/Editor/UnrealEd/Private/Commandlets/StabilizeLocalizationKeys.cpp
@@ -80,7 +80,7 @@ public:
 		this->SetIsSaving(true);
 
 		TArray<UObject*> AllObjectsInPackage;
-		GetObjectsWithOuter(InPackage, AllObjectsInPackage, true, RF_Transient, EInternalObjectFlags::PendingKill);
+		GetObjectsWithOuter(InPackage, AllObjectsInPackage, true, RF_Transient, EInternalObjectFlags::Garbage);
 
 		for (UObject* Obj : AllObjectsInPackage)
 		{
diff --git a/Engine/Source/Editor/UnrealEd/Private/Commandlets/WorldPartitionConvertCommandlet.cpp b/Engine/Source/Editor/UnrealEd/Private/Commandlets/WorldPartitionConvertCommandlet.cpp
--- a/Engine/Source/Editor/UnrealEd/Private/Commandlets/WorldPartitionConvertCommandlet.cpp
+++ b/Engine/Source/Editor/UnrealEd/Private/Commandlets/WorldPartitionConvertCommandlet.cpp
@@ -490,7 +490,7 @@ void UWorldPartitionConvertCommandlet::FixupSoftObjectPaths(UPackage* OuterPacka
 			Object->Serialize(FixupSerializer);
 		}
 		return true;
-	}, true, RF_NoFlags, EInternalObjectFlags::PendingKill);
+	}, true, RF_NoFlags, EInternalObjectFlags::Garbage);
 }
 
 bool UWorldPartitionConvertCommandlet::DetachDependantLevelPackages(ULevel* Level)
diff --git a/Engine/Source/Editor/UnrealEd/Private/DataTableEditorUtils.cpp b/Engine/Source/Editor/UnrealEd/Private/DataTableEditorUtils.cpp
--- a/Engine/Source/Editor/UnrealEd/Private/DataTableEditorUtils.cpp
+++ b/Engine/Source/Editor/UnrealEd/Private/DataTableEditorUtils.cpp
@@ -746,7 +746,7 @@ void FDataTableEditorUtils::BroadcastPostChange(UDataTable* DataTable, EDataTabl
 {
 	if (DataTable && (EDataTableChangeInfo::RowList == Info))
 	{
-		for (TObjectIterator<UK2Node_GetDataTableRow> It(RF_Transient | RF_ClassDefaultObject, /** bIncludeDerivedClasses */ true, /** InternalExcludeFlags */ EInternalObjectFlags::PendingKill); It; ++It)
+		for (TObjectIterator<UK2Node_GetDataTableRow> It(RF_Transient | RF_ClassDefaultObject, /** bIncludeDerivedClasses */ true, /** InternalExcludeFlags */ EInternalObjectFlags::Garbage); It; ++It)
 		{
 			It->OnDataTableRowListChanged(DataTable);
 		}
diff --git a/Engine/Source/Editor/UnrealEd/Private/Fbx/FbxSkeletalMeshImport.cpp b/Engine/Source/Editor/UnrealEd/Private/Fbx/FbxSkeletalMeshImport.cpp
--- a/Engine/Source/Editor/UnrealEd/Private/Fbx/FbxSkeletalMeshImport.cpp
+++ b/Engine/Source/Editor/UnrealEd/Private/Fbx/FbxSkeletalMeshImport.cpp
@@ -1672,7 +1672,7 @@ USkeletalMesh* UnFbx::FFbxImporter::ImportSkeletalMesh(FImportSkeletalMeshArgs &
 	USkeletalMesh* ExistingSkelMesh = nullptr;
 	if ( !ImportSkeletalMeshArgs.FbxShapeArray  )
 	{
-		UObject* ExistingObject = StaticFindObjectFast(UObject::StaticClass(), ImportSkeletalMeshArgs.InParent, ImportSkeletalMeshArgs.Name, false, false, RF_NoFlags, EInternalObjectFlags::PendingKill);
+		UObject* ExistingObject = StaticFindObjectFast(UObject::StaticClass(), ImportSkeletalMeshArgs.InParent, ImportSkeletalMeshArgs.Name, false, false, RF_NoFlags, EInternalObjectFlags::Garbage);
 		ExistingSkelMesh = Cast<USkeletalMesh>(ExistingObject);
 
 		if (!ExistingSkelMesh && ExistingObject)
diff --git a/Engine/Source/Editor/UnrealEd/Private/Kismet2/BlueprintEditorUtils.cpp b/Engine/Source/Editor/UnrealEd/Private/Kismet2/BlueprintEditorUtils.cpp
--- a/Engine/Source/Editor/UnrealEd/Private/Kismet2/BlueprintEditorUtils.cpp
+++ b/Engine/Source/Editor/UnrealEd/Private/Kismet2/BlueprintEditorUtils.cpp
@@ -8778,7 +8778,7 @@ void FBlueprintEditorUtils::PostEditChangeBlueprintActors(UBlueprint* Blueprint,
 		const bool bIncludeDerivedClasses = false;
 
 		TArray<UObject*> MatchingBlueprintObjects;
-		GetObjectsOfClass(Blueprint->GeneratedClass, MatchingBlueprintObjects, bIncludeDerivedClasses, RF_ClassDefaultObject, EInternalObjectFlags::PendingKill);
+		GetObjectsOfClass(Blueprint->GeneratedClass, MatchingBlueprintObjects, bIncludeDerivedClasses, RF_ClassDefaultObject, EInternalObjectFlags::Garbage);
 
 		for (UObject* MatchingObj : MatchingBlueprintObjects)
 		{
@@ -10174,7 +10174,7 @@ namespace
 	template <typename TObject, bool bIsFind, typename FindExisting>
 	void FindOrUpdateNodes(FBlueprintEditorUtils::FOnNodeFoundOrUpdated InOnNodeFoundOrUpdated, FindExisting InFindExisting)
 	{
-		for (TObjectIterator<UK2Node> It(RF_Transient | RF_ClassDefaultObject, /** bIncludeDerivedClasses */ true, /** InternalExcludeFlags */ EInternalObjectFlags::PendingKill); It; ++It)
+		for (TObjectIterator<UK2Node> It(RF_Transient | RF_ClassDefaultObject, /** bIncludeDerivedClasses */ true, /** InternalExcludeFlags */ EInternalObjectFlags::Garbage); It; ++It)
 		{
 			UK2Node* Node = *It;
 
diff --git a/Engine/Source/Editor/UnrealEd/Private/PackageTools.cpp b/Engine/Source/Editor/UnrealEd/Private/PackageTools.cpp
--- a/Engine/Source/Editor/UnrealEd/Private/PackageTools.cpp
+++ b/Engine/Source/Editor/UnrealEd/Private/PackageTools.cpp
@@ -415,7 +415,7 @@ UPackageTools::UPackageTools(const FObjectInitializer& ObjectInitializer)
 					bScriptPackageWasUnloaded = true;
 				}
 
-				GetObjectsWithPackage(PackageBeingUnloaded, ObjectsInPackage, true, RF_Transient, EInternalObjectFlags::PendingKill);
+				GetObjectsWithPackage(PackageBeingUnloaded, ObjectsInPackage, true, RF_Transient, EInternalObjectFlags::Garbage);
 				// Notify any Blueprints that are about to be unloaded, and destroy any leftover worlds.
 				for (UObject* Obj : ObjectsInPackage)
 				{
@@ -929,7 +929,7 @@ UPackageTools::UPackageTools(const FObjectInitializer& ObjectInitializer)
 					}
 				}
 				return true;
-			}, true, RF_Transient, EInternalObjectFlags::PendingKill);
+			}, true, RF_Transient, EInternalObjectFlags::Garbage);
 		}
 
 		if (InPackageReloadPhase == EPackageReloadPhase::OnPackageFixup)
diff --git a/Engine/Source/Editor/UnrealEd/Private/PlayLevel.cpp b/Engine/Source/Editor/UnrealEd/Private/PlayLevel.cpp
--- a/Engine/Source/Editor/UnrealEd/Private/PlayLevel.cpp
+++ b/Engine/Source/Editor/UnrealEd/Private/PlayLevel.cpp
@@ -401,7 +401,7 @@ void UEditorEngine::EndPlayMap()
 		{
 			Object->MarkAsGarbage();
 		};
-		ForEachObjectWithOuter(*It, MarkObjectPendingKill, true, RF_NoFlags, EInternalObjectFlags::PendingKill);
+		ForEachObjectWithOuter(*It, MarkObjectPendingKill, true, RF_NoFlags, EInternalObjectFlags::Garbage);
 	}
 
 	// Flush any render commands and released accessed UTextures and materials to give them a chance to be collected.
diff --git a/Engine/Source/Editor/UnrealEd/Private/StaticLightingSystem/StaticLightingSystem.cpp b/Engine/Source/Editor/UnrealEd/Private/StaticLightingSystem/StaticLightingSystem.cpp
--- a/Engine/Source/Editor/UnrealEd/Private/StaticLightingSystem/StaticLightingSystem.cpp
+++ b/Engine/Source/Editor/UnrealEd/Private/StaticLightingSystem/StaticLightingSystem.cpp
@@ -669,7 +669,7 @@ bool FStaticLightingSystem::BeginLightmassProcess()
 				}
 				}
 
-				for (TObjectIterator<ULightComponentBase> LightIt(RF_ClassDefaultObject, /** bIncludeDerivedClasses */ true, /** InternalExcludeFlags */ EInternalObjectFlags::PendingKill); LightIt; ++LightIt)
+				for (TObjectIterator<ULightComponentBase> LightIt(RF_ClassDefaultObject, /** bIncludeDerivedClasses */ true, /** InternalExcludeFlags */ EInternalObjectFlags::Garbage); LightIt; ++LightIt)
 				{
 					ULightComponentBase* const Light = *LightIt;
 					const bool bLightIsInWorld = IsValid(Light->GetOwner()) 
diff --git a/Engine/Source/Editor/UnrealEd/Private/Subsystems/EditorActorSubsystem.cpp b/Engine/Source/Editor/UnrealEd/Private/Subsystems/EditorActorSubsystem.cpp
--- a/Engine/Source/Editor/UnrealEd/Private/Subsystems/EditorActorSubsystem.cpp
+++ b/Engine/Source/Editor/UnrealEd/Private/Subsystems/EditorActorSubsystem.cpp
@@ -57,7 +57,7 @@ namespace InternalActorUtilitiesSubsystemLibrary
 		}
 
 		const EObjectFlags ExcludeFlags = RF_ClassDefaultObject;
-		for (TObjectIterator<T> It(ExcludeFlags, true, EInternalObjectFlags::PendingKill); It; ++It)
+		for (TObjectIterator<T> It(ExcludeFlags, true, EInternalObjectFlags::Garbage); It; ++It)
 		{
 			T* Obj = *It;
 			if (InternalActorUtilitiesSubsystemLibrary::IsEditorLevelActor(Obj))
diff --git a/Engine/Source/Editor/VirtualTexturingEditor/Private/RuntimeVirtualTextureSetBounds.cpp b/Engine/Source/Editor/VirtualTexturingEditor/Private/RuntimeVirtualTextureSetBounds.cpp
--- a/Engine/Source/Editor/VirtualTexturingEditor/Private/RuntimeVirtualTextureSetBounds.cpp
+++ b/Engine/Source/Editor/VirtualTexturingEditor/Private/RuntimeVirtualTextureSetBounds.cpp
@@ -31,7 +31,7 @@ namespace RuntimeVirtualTexture
 
 		// Expand bounds for the BoundsAlignActor and all primitive components that write to this virtual texture.
 		FBox Bounds(ForceInit);
-		for (TObjectIterator<UPrimitiveComponent> It(RF_ClassDefaultObject, true, EInternalObjectFlags::PendingKill); It; ++It)
+		for (TObjectIterator<UPrimitiveComponent> It(RF_ClassDefaultObject, true, EInternalObjectFlags::Garbage); It; ++It)
 		{
 			bool bUseBounds = BoundsAlignActor.IsValid() && It->GetOwner() == BoundsAlignActor.Get();
 
diff --git a/Engine/Source/Runtime/CoreUObject/Private/Serialization/PropertyLocalizationDataGathering.cpp b/Engine/Source/Runtime/CoreUObject/Private/Serialization/PropertyLocalizationDataGathering.cpp
--- a/Engine/Source/Runtime/CoreUObject/Private/Serialization/PropertyLocalizationDataGathering.cpp
+++ b/Engine/Source/Runtime/CoreUObject/Private/Serialization/PropertyLocalizationDataGathering.cpp
@@ -27,7 +27,7 @@ FPropertyLocalizationDataGatherer::FPropertyLocalizationDataGatherer(TArray<FGat
 	{
 		AllObjectsInPackage.Add(Object);
 		return true;
-	}, true, RF_Transient, EInternalObjectFlags::PendingKill);
+	}, true, RF_Transient, EInternalObjectFlags::Garbage);
 
 	// Iterate over each root object in the package
 	for (const UObject* Object : AllObjectsInPackage)
@@ -236,7 +236,7 @@ void FPropertyLocalizationDataGatherer::GatherLocalizationDataFromObject(const U
 	if (!(GatherTextFlags & EPropertyLocalizationGathererTextFlags::SkipSubObjects))
 	{
 		TArray<UObject*> InnerObjects;
-		GetObjectsWithOuter(Object, InnerObjects, false, RF_Transient, EInternalObjectFlags::PendingKill);
+		GetObjectsWithOuter(Object, InnerObjects, false, RF_Transient, EInternalObjectFlags::Garbage);
 		for (UObject* ChildObject : InnerObjects)
 		{
 			// if the child object as a package set, do not gather from it
diff --git a/Engine/Source/Runtime/CoreUObject/Private/Serialization/TextReferenceCollector.cpp b/Engine/Source/Runtime/CoreUObject/Private/Serialization/TextReferenceCollector.cpp
--- a/Engine/Source/Runtime/CoreUObject/Private/Serialization/TextReferenceCollector.cpp
+++ b/Engine/Source/Runtime/CoreUObject/Private/Serialization/TextReferenceCollector.cpp
@@ -24,7 +24,7 @@ public:
 		// Build up the list of objects that are within our package - we won't follow object references to things outside of our package
 		{
 			TArray<UObject*> AllObjectsInPackageArray;
-			GetObjectsWithOuter(InPackage, AllObjectsInPackageArray, true, RF_Transient, EInternalObjectFlags::PendingKill);
+			GetObjectsWithOuter(InPackage, AllObjectsInPackageArray, true, RF_Transient, EInternalObjectFlags::Garbage);
 
 			AllObjectsInPackage.Reserve(AllObjectsInPackageArray.Num());
 			for (UObject* Object : AllObjectsInPackageArray)
@@ -34,7 +34,7 @@ public:
 		}
 
 		TArray<UObject*> RootObjectsInPackage;
-		GetObjectsWithOuter(InPackage, RootObjectsInPackage, false, RF_Transient, EInternalObjectFlags::PendingKill);
+		GetObjectsWithOuter(InPackage, RootObjectsInPackage, false, RF_Transient, EInternalObjectFlags::Garbage);
 
 		// Iterate over each root object in the package
 		for (UObject* Obj : RootObjectsInPackage)
diff --git a/Engine/Source/Runtime/CoreUObject/Private/UObject/GarbageCollection.cpp b/Engine/Source/Runtime/CoreUObject/Private/UObject/GarbageCollection.cpp
--- a/Engine/Source/Runtime/CoreUObject/Private/UObject/GarbageCollection.cpp
+++ b/Engine/Source/Runtime/CoreUObject/Private/UObject/GarbageCollection.cpp
@@ -696,7 +696,9 @@ public:
 				FUObjectItem* ReferencedMutableObjectItem = GUObjectArray.IndexToObjectUnsafeForGC(ReferencedMutableObjectIndex);
 				if (IsParallel())
 				{
+					PRAGMA_DISABLE_DEPRECATION_WARNINGS
 					if (!ReferencedMutableObjectItem->HasAnyFlags(EInternalObjectFlags::PendingKill | EInternalObjectFlags::Garbage))
+					PRAGMA_ENABLE_DEPRECATION_WARNINGS
 					{
 						if (ReferencedMutableObjectItem->IsUnreachable())
 						{
@@ -737,7 +739,9 @@ public:
 						bAddClusterObjectsToSerialize = true;
 					}
 				}
+				PRAGMA_DISABLE_DEPRECATION_WARNINGS
 				else if (!ReferencedMutableObjectItem->HasAnyFlags(EInternalObjectFlags::PendingKill | EInternalObjectFlags::Garbage))
+				PRAGMA_ENABLE_DEPRECATION_WARNINGS
 				{
 					if (ReferencedMutableObjectItem->IsUnreachable())
 					{
@@ -793,7 +797,9 @@ public:
 			if (ReferncedClusterIndex >= 0) // Pending Kill support
 			{
 				FUObjectItem* ReferencedClusterRootObjectItem = GUObjectArray.IndexToObjectUnsafeForGC(ReferncedClusterIndex);
+				PRAGMA_DISABLE_DEPRECATION_WARNINGS
 				if (!ReferencedClusterRootObjectItem->HasAnyFlags(EInternalObjectFlags::PendingKill | EInternalObjectFlags::Garbage))
+				PRAGMA_ENABLE_DEPRECATION_WARNINGS
 				{
 					// This condition should get collapsed by the compiler based on the template argument
 					if (IsParallel())
@@ -1340,7 +1346,9 @@ public:
 						{
 							bMarkAsUnreachable = false;
 						}
+						PRAGMA_DISABLE_DEPRECATION_WARNINGS
 						else if (ObjectItem->HasAnyFlags(EInternalObjectFlags::PendingKill | EInternalObjectFlags::Garbage) && bWithClusters && ObjectItem->HasAnyFlags(EInternalObjectFlags::ClusterRoot))
+						PRAGMA_ENABLE_DEPRECATION_WARNINGS
 						{
 							ClustersToDissolveList.Push(ObjectItem);
 						}
diff --git a/Engine/Source/Runtime/CoreUObject/Private/UObject/GarbageCollectionVerification.cpp b/Engine/Source/Runtime/CoreUObject/Private/UObject/GarbageCollectionVerification.cpp
--- a/Engine/Source/Runtime/CoreUObject/Private/UObject/GarbageCollectionVerification.cpp
+++ b/Engine/Source/Runtime/CoreUObject/Private/UObject/GarbageCollectionVerification.cpp
@@ -357,7 +357,9 @@ void VerifyObjectFlagMirroring()
 			{
 				UObjectBaseUtility* Object = (UObjectBaseUtility*)ObjectItem.Object;
 				bool bHasObjectFlag = Object->HasAnyFlags(RF_InternalPendingKill);
+				PRAGMA_DISABLE_DEPRECATION_WARNINGS
 				bool bHasInternalFlag = ObjectItem.HasAnyFlags(EInternalObjectFlags::PendingKill);
+				PRAGMA_ENABLE_DEPRECATION_WARNINGS
 				if (bHasObjectFlag != bHasInternalFlag)
 				{
 					UE_LOG(LogGarbage, Warning, TEXT("RF_PendingKill (%d) and EInternalObjectFlags::PendingKill (%d) flag mismatch on %s"),
diff --git a/Engine/Source/Runtime/CoreUObject/Private/UObject/Obj.cpp b/Engine/Source/Runtime/CoreUObject/Private/UObject/Obj.cpp
--- a/Engine/Source/Runtime/CoreUObject/Private/UObject/Obj.cpp
+++ b/Engine/Source/Runtime/CoreUObject/Private/UObject/Obj.cpp
@@ -748,7 +748,7 @@ void UObject::GetArchetypeInstances( TArray<UObject*>& Instances )
 				{
 					Instances.Add(Obj);
 				}
-			}, bIncludeNestedObjects, RF_NoFlags, EInternalObjectFlags::PendingKill); // we need to evaluate CDOs as well, but nothing pending kill
+			}, bIncludeNestedObjects, RF_NoFlags, EInternalObjectFlags::Garbage); // we need to evaluate CDOs as well, but nothing pending kill
 		}
 		else
 		{
@@ -759,7 +759,7 @@ void UObject::GetArchetypeInstances( TArray<UObject*>& Instances )
 				{
 					Instances.Add(Obj);
 				}
-			}, bIncludeNestedObjects, RF_NoFlags, EInternalObjectFlags::PendingKill); // we need to evaluate CDOs as well, but nothing pending kill
+			}, bIncludeNestedObjects, RF_NoFlags, EInternalObjectFlags::Garbage); // we need to evaluate CDOs as well, but nothing pending kill
 
 		}
 	}
@@ -1661,7 +1661,7 @@ void UObject::BuildSubobjectMapping(UObject* OtherObject, TMap<UObject*, UObject
 		{
 			InSubObject->BuildSubobjectMapping(OtherSubObject, ObjectMapping);
 		}
-	}, false, RF_NoFlags, EInternalObjectFlags::PendingKill);
+	}, false, RF_NoFlags, EInternalObjectFlags::Garbage);
 }
 
 
diff --git a/Engine/Source/Runtime/CoreUObject/Private/UObject/ObjectBaseUtility.cpp b/Engine/Source/Runtime/CoreUObject/Private/UObject/ObjectBaseUtility.cpp
--- a/Engine/Source/Runtime/CoreUObject/Private/UObject/ObjectBaseUtility.cpp
+++ b/Engine/Source/Runtime/CoreUObject/Private/UObject/ObjectBaseUtility.cpp
@@ -136,6 +136,27 @@ bool UObjectBaseUtility::bPendingKillDisabled = !GPendingKillEnabled;
 
 void InitNoPendingKill()
 {
+#if !UE_BUILD_SHIPPING
+	PRAGMA_DISABLE_DEPRECATION_WARNINGS
+	{
+		UObjectBaseUtility::bPendingKillDisabled = false;
+		check(UObjectBaseUtility::FixGarbageOrPendingKillInternalObjectFlags(EInternalObjectFlags::None) == EInternalObjectFlags::None);
+		check(UObjectBaseUtility::FixGarbageOrPendingKillInternalObjectFlags(EInternalObjectFlags::PendingKill) == EInternalObjectFlags::PendingKill);
+		check(UObjectBaseUtility::FixGarbageOrPendingKillInternalObjectFlags(EInternalObjectFlags::Garbage) == EInternalObjectFlags::PendingKill);
+		check(UObjectBaseUtility::FixGarbageOrPendingKillInternalObjectFlags(EInternalObjectFlags::PendingKill | EInternalObjectFlags::Garbage) == EInternalObjectFlags::PendingKill);
+		check(UObjectBaseUtility::FixGarbageOrPendingKillInternalObjectFlags(EInternalObjectFlags::RootSet | EInternalObjectFlags::PendingKill) == (EInternalObjectFlags::RootSet | EInternalObjectFlags::PendingKill));
+		check(UObjectBaseUtility::FixGarbageOrPendingKillInternalObjectFlags(EInternalObjectFlags::RootSet | EInternalObjectFlags::Garbage) == (EInternalObjectFlags::RootSet | EInternalObjectFlags::PendingKill));
+		UObjectBaseUtility::bPendingKillDisabled = true;
+		check(UObjectBaseUtility::FixGarbageOrPendingKillInternalObjectFlags(EInternalObjectFlags::None) == EInternalObjectFlags::None);
+		check(UObjectBaseUtility::FixGarbageOrPendingKillInternalObjectFlags(EInternalObjectFlags::PendingKill) == EInternalObjectFlags::Garbage);
+		check(UObjectBaseUtility::FixGarbageOrPendingKillInternalObjectFlags(EInternalObjectFlags::Garbage) == EInternalObjectFlags::Garbage);
+		check(UObjectBaseUtility::FixGarbageOrPendingKillInternalObjectFlags(EInternalObjectFlags::PendingKill | EInternalObjectFlags::Garbage) == EInternalObjectFlags::Garbage);
+		check(UObjectBaseUtility::FixGarbageOrPendingKillInternalObjectFlags(EInternalObjectFlags::RootSet | EInternalObjectFlags::PendingKill) == (EInternalObjectFlags::RootSet | EInternalObjectFlags::Garbage));
+		check(UObjectBaseUtility::FixGarbageOrPendingKillInternalObjectFlags(EInternalObjectFlags::RootSet | EInternalObjectFlags::Garbage) == (EInternalObjectFlags::RootSet | EInternalObjectFlags::Garbage));
+	}
+	PRAGMA_ENABLE_DEPRECATION_WARNINGS
+#endif // !UE_BUILD_SHIPPING
+
 	check(GConfig);
 	bool bPendingKillEnabled = false;
 	GConfig->GetBool(TEXT("/Script/Engine.GarbageCollectionSettings"), TEXT("gc.PendingKillEnabled"), bPendingKillEnabled, GEngineIni);
diff --git a/Engine/Source/Runtime/CoreUObject/Private/UObject/PackageReload.cpp b/Engine/Source/Runtime/CoreUObject/Private/UObject/PackageReload.cpp
--- a/Engine/Source/Runtime/CoreUObject/Private/UObject/PackageReload.cpp
+++ b/Engine/Source/Runtime/CoreUObject/Private/UObject/PackageReload.cpp
@@ -587,7 +587,7 @@ void ReloadPackages(const TArrayView<FReloadPackageData>& InPackagesToReload, TA
 
 			// Main pass to go through and fix-up any references pointing to data from the old package to point to data from the new package
 			// todo: multi-thread this like FHotReloadModule::ReplaceReferencesToReconstructedCDOs?
-			for (FThreadSafeObjectIterator ObjIter(UObject::StaticClass(), false, RF_NoFlags, EInternalObjectFlags::PendingKill); ObjIter; ++ObjIter)
+			for (FThreadSafeObjectIterator ObjIter(UObject::StaticClass(), false, RF_NoFlags, EInternalObjectFlags::Garbage); ObjIter; ++ObjIter)
 			{
 				UObject* PotentialReferencer = *ObjIter;
 
@@ -693,7 +693,7 @@ void ReloadPackages(const TArrayView<FReloadPackageData>& InPackagesToReload, TA
 				//{
 				//	PackageReloadInternal::DumpExternalReferences(InExistingObject, ExistingPackage);
 				//	return true; // continue
-				//}, true, RF_NoFlags, EInternalObjectFlags::PendingKill);
+				//}, true, RF_NoFlags, EInternalObjectFlags::Garbage);
 			}
 		}
 	}
diff --git a/Engine/Source/Runtime/CoreUObject/Private/UObject/ReferenceChainSearch.cpp b/Engine/Source/Runtime/CoreUObject/Private/UObject/ReferenceChainSearch.cpp
--- a/Engine/Source/Runtime/CoreUObject/Private/UObject/ReferenceChainSearch.cpp
+++ b/Engine/Source/Runtime/CoreUObject/Private/UObject/ReferenceChainSearch.cpp
@@ -223,10 +223,12 @@ FString FReferenceChainSearch::GetObjectFlags(UObject* InObject)
 		Flags += TEXT("(native) ");
 	}
 
+	PRAGMA_DISABLE_DEPRECATION_WARNINGS
 	if (InObject->HasAnyInternalFlags(EInternalObjectFlags::PendingKill))
 	{
 		Flags += TEXT("(PendingKill) ");
 	}
+	PRAGMA_ENABLE_DEPRECATION_WARNINGS
 
 	if (InObject->HasAnyInternalFlags(EInternalObjectFlags::Garbage))
 	{
diff --git a/Engine/Source/Runtime/CoreUObject/Private/UObject/SavePackage/SavePackageUtilities.cpp b/Engine/Source/Runtime/CoreUObject/Private/UObject/SavePackage/SavePackageUtilities.cpp
--- a/Engine/Source/Runtime/CoreUObject/Private/UObject/SavePackage/SavePackageUtilities.cpp
+++ b/Engine/Source/Runtime/CoreUObject/Private/UObject/SavePackage/SavePackageUtilities.cpp
@@ -2426,8 +2426,8 @@ bool UPackage::IsEmptyPackage(UPackage* Package, const UObject* LastReferencer)
 				return false;
 			}
 			return true;
-		// Don't consider transient, class default or pending kill / garbage objects
-		}, false, RF_Transient | RF_ClassDefaultObject, EInternalObjectFlags::PendingKill | EInternalObjectFlags::Garbage);
+		// Don't consider transient, class default or garbage objects
+		}, false, RF_Transient | RF_ClassDefaultObject, EInternalObjectFlags::Garbage);
 		return bIsEmpty;
 	}
 
diff --git a/Engine/Source/Runtime/CoreUObject/Private/UObject/UObjectHash.cpp b/Engine/Source/Runtime/CoreUObject/Private/UObject/UObjectHash.cpp
--- a/Engine/Source/Runtime/CoreUObject/Private/UObject/UObjectHash.cpp
+++ b/Engine/Source/Runtime/CoreUObject/Private/UObject/UObjectHash.cpp
@@ -984,6 +984,8 @@ void GetObjectsWithOuter(const class UObjectBase* Outer, TArray<UObject *>& Resu
 	{
 		ExclusionInternalFlags |= EInternalObjectFlags::AsyncLoading;
 	}
+	ExclusionInternalFlags = UObjectBaseUtility::FixGarbageOrPendingKillInternalObjectFlags(ExclusionInternalFlags);
+
 	int32 StartNum = Results.Num();
 	FUObjectHashTables& ThreadHash = FUObjectHashTables::Get();
 	FHashTableLock HashLock(ThreadHash);
@@ -1043,6 +1045,8 @@ void ForEachObjectWithOuterBreakable(const class UObjectBase* Outer, TFunctionRe
 	{
 		ExclusionInternalFlags |= EInternalObjectFlags::AsyncLoading;
 	}
+	ExclusionInternalFlags = UObjectBaseUtility::FixGarbageOrPendingKillInternalObjectFlags(ExclusionInternalFlags);
+
 	FUObjectHashTables& ThreadHash = FUObjectHashTables::Get();
 	FHashTableLock HashLock(ThreadHash);
 	TArray<FHashBucket*, TInlineAllocator<1> > AllInners;
@@ -1137,6 +1141,8 @@ void ForEachObjectWithPackage(const class UPackage* Package, TFunctionRef<bool(U
 	{
 		ExclusionInternalFlags |= EInternalObjectFlags::AsyncLoading;
 	}
+	ExclusionInternalFlags = UObjectBaseUtility::FixGarbageOrPendingKillInternalObjectFlags(ExclusionInternalFlags);
+
 	FUObjectHashTables& ThreadHash = FUObjectHashTables::Get();
 	FHashTableLock HashLock(ThreadHash);
 	TArray<FHashBucket*, TInlineAllocator<1> > AllInners;
@@ -1242,6 +1248,7 @@ FORCEINLINE void ForEachObjectOfClasses_Implementation(FUObjectHashTables& Threa
 	{
 		ExclusionInternalFlags |= EInternalObjectFlags::AsyncLoading;
 	}
+	ExclusionInternalFlags = UObjectBaseUtility::FixGarbageOrPendingKillInternalObjectFlags(ExclusionInternalFlags);
 
 	TBucketMapLock ClassToObjectListMapLock(ThreadHash.ClassToObjectListMap);
 
diff --git a/Engine/Source/Runtime/CoreUObject/Public/UObject/Object.h b/Engine/Source/Runtime/CoreUObject/Public/UObject/Object.h
--- a/Engine/Source/Runtime/CoreUObject/Public/UObject/Object.h
+++ b/Engine/Source/Runtime/CoreUObject/Public/UObject/Object.h
@@ -1624,7 +1624,9 @@ struct FInternalUObjectBaseUtilityIsValidFlagsChecker
 	FORCEINLINE static bool CheckObjectValidBasedOnItsFlags(const UObject* Test)
 	{
 		// Here we don't really check if the flags match but if the end result is the same
+		PRAGMA_DISABLE_DEPRECATION_WARNINGS
 		checkSlow(GUObjectArray.IndexToObject(Test->InternalIndex)->HasAnyFlags(EInternalObjectFlags::PendingKill | EInternalObjectFlags::Garbage) == Test->HasAnyFlags(RF_InternalPendingKill | RF_InternalGarbage));
+		PRAGMA_ENABLE_DEPRECATION_WARNINGS
 		return !Test->HasAnyFlags(RF_InternalPendingKill | RF_InternalGarbage);
 	}
 };
diff --git a/Engine/Source/Runtime/CoreUObject/Public/UObject/ObjectMacros.h b/Engine/Source/Runtime/CoreUObject/Public/UObject/ObjectMacros.h
--- a/Engine/Source/Runtime/CoreUObject/Public/UObject/ObjectMacros.h
+++ b/Engine/Source/Runtime/CoreUObject/Public/UObject/ObjectMacros.h
@@ -569,15 +569,17 @@ enum class EInternalObjectFlags : int32
 	Async = 1 << 26, ///< Object exists only on a different thread than the game thread.
 	AsyncLoading = 1 << 27, ///< Object is being asynchronously loaded.
 	Unreachable = 1 << 28, ///< Object is not reachable on the object graph.
-	PendingKill = 1 << 29, ///< Objects that are pending destruction (invalid for gameplay but valid objects). This flag is mirrored in EObjectFlags as RF_PendingKill for performance
+	PendingKill UE_DEPRECATED(5.0, "PendingKill flag should no longer be used. Use Garbage flag instead.") = 1 << 29, ///< Objects that are pending destruction (invalid for gameplay but valid objects). This flag is mirrored in EObjectFlags as RF_PendingKill for performance
 	RootSet = 1 << 30, ///< Object will not be garbage collected, even if unreferenced.
 	PendingConstruction = 1 << 31, ///< Object didn't have its class constructor called yet (only the UObjectBase one to initialize its most basic members)
 
 	GarbageCollectionKeepFlags = Native | Async | AsyncLoading | LoaderImport,
+	PRAGMA_DISABLE_DEPRECATION_WARNINGS
 	MirroredFlags = Garbage | PendingKill, /// Flags mirrored in EObjectFlags
 
 	//~ Make sure this is up to date!
 	AllFlags = LoaderImport | Garbage | PersistentGarbage | ReachableInCluster | ClusterRoot | Native | Async | AsyncLoading | Unreachable | PendingKill | RootSet | PendingConstruction
+	PRAGMA_ENABLE_DEPRECATION_WARNINGS
 };
 ENUM_CLASS_FLAGS(EInternalObjectFlags);
 
diff --git a/Engine/Source/Runtime/CoreUObject/Public/UObject/UObjectArray.h b/Engine/Source/Runtime/CoreUObject/Public/UObject/UObjectArray.h
--- a/Engine/Source/Runtime/CoreUObject/Public/UObject/UObjectArray.h
+++ b/Engine/Source/Runtime/CoreUObject/Public/UObject/UObjectArray.h
@@ -183,15 +183,21 @@ struct FUObjectItem
 
 	FORCEINLINE void SetPendingKill()
 	{
+		PRAGMA_DISABLE_DEPRECATION_WARNINGS
 		ThisThreadAtomicallySetFlag(EInternalObjectFlags::PendingKill);
+		PRAGMA_ENABLE_DEPRECATION_WARNINGS
 	}
 	FORCEINLINE void ClearPendingKill()
 	{
+		PRAGMA_DISABLE_DEPRECATION_WARNINGS
 		ThisThreadAtomicallyClearedFlag(EInternalObjectFlags::PendingKill);
+		PRAGMA_ENABLE_DEPRECATION_WARNINGS
 	}
 	FORCEINLINE bool IsPendingKill() const
 	{
+		PRAGMA_DISABLE_DEPRECATION_WARNINGS
 		return !!(Flags & int32(EInternalObjectFlags::PendingKill));
+		PRAGMA_ENABLE_DEPRECATION_WARNINGS
 	}
 
 	FORCEINLINE void SetRootSet()
@@ -727,7 +733,9 @@ public:
 		FUObjectItem* ObjectItem = IndexToObject(Index);
 		if (ObjectItem && ObjectItem->Object)
 		{
+			PRAGMA_DISABLE_DEPRECATION_WARNINGS
 			if (!bEvenIfPendingKill && ObjectItem->HasAnyFlags(EInternalObjectFlags::PendingKill | EInternalObjectFlags::Garbage))
+			PRAGMA_ENABLE_DEPRECATION_WARNINGS
 			{
 				ObjectItem = nullptr;;
 			}
@@ -745,7 +753,9 @@ public:
 	{
 		if (ObjectItem)
 		{
+			PRAGMA_DISABLE_DEPRECATION_WARNINGS
 			return bEvenIfPendingKill ? !ObjectItem->IsUnreachable() : !(ObjectItem->HasAnyFlags(EInternalObjectFlags::Unreachable | EInternalObjectFlags::PendingKill | EInternalObjectFlags::Garbage));
+			PRAGMA_ENABLE_DEPRECATION_WARNINGS
 		}
 		return false;
 	}
@@ -766,7 +776,9 @@ public:
 	FORCEINLINE bool IsStale(FUObjectItem* ObjectItem, bool bEvenIfPendingKill)
 	{
 		// This method assumes ObjectItem is valid.
+		PRAGMA_DISABLE_DEPRECATION_WARNINGS
 		return bEvenIfPendingKill ? (ObjectItem->HasAnyFlags(EInternalObjectFlags::Unreachable | EInternalObjectFlags::PendingKill | EInternalObjectFlags::Garbage)) : (ObjectItem->IsUnreachable());
+		PRAGMA_ENABLE_DEPRECATION_WARNINGS
 	}
 
 	FORCEINLINE bool IsStale(int32 Index, bool bEvenIfPendingKill)
diff --git a/Engine/Source/Runtime/CoreUObject/Public/UObject/UObjectBaseUtility.h b/Engine/Source/Runtime/CoreUObject/Public/UObject/UObjectBaseUtility.h
--- a/Engine/Source/Runtime/CoreUObject/Public/UObject/UObjectBaseUtility.h
+++ b/Engine/Source/Runtime/CoreUObject/Public/UObject/UObjectBaseUtility.h
@@ -199,7 +199,9 @@ public:
 		}
 		else
 		{
+			PRAGMA_DISABLE_DEPRECATION_WARNINGS
 			checkSlow(GUObjectArray.IndexToObject(InternalIndex)->HasAnyFlags(EInternalObjectFlags::PendingKill) == HasAnyFlags(RF_InternalPendingKill));
+			PRAGMA_ENABLE_DEPRECATION_WARNINGS
 			return HasAnyFlags(RF_InternalPendingKill);
 		}
 	}
@@ -313,7 +315,9 @@ public:
 	UE_DEPRECATED(5.0, "IsPendingKillOrUnreachable() should no longer be used. Use IsValid(Object), IsValidChecked(Object), GetValid(Object) and/or IsUnreachable() instead.")
 	FORCEINLINE bool IsPendingKillOrUnreachable() const
 	{
+		PRAGMA_DISABLE_DEPRECATION_WARNINGS
 		return GUObjectArray.IndexToObject(InternalIndex)->HasAnyFlags(EInternalObjectFlags::PendingKill | EInternalObjectFlags::Garbage | EInternalObjectFlags::Unreachable);
+		PRAGMA_ENABLE_DEPRECATION_WARNINGS
 	}
 
 	/** Checks if the object is native. */
@@ -331,7 +335,9 @@ public:
 	FORCEINLINE void SetInternalFlags(EInternalObjectFlags FlagsToSet) const
 	{
 		FUObjectItem* ObjectItem = GUObjectArray.IndexToObject(InternalIndex);
+		PRAGMA_DISABLE_DEPRECATION_WARNINGS
 		checkf(!(FlagsToSet & (EInternalObjectFlags::PendingKill | EInternalObjectFlags::Garbage)) || int32(FlagsToSet & (EInternalObjectFlags::PendingKill | EInternalObjectFlags::Garbage)) == (ObjectItem->Flags & int32(EInternalObjectFlags::PendingKill | EInternalObjectFlags::Garbage)), TEXT("SetInternalFlags should not set the PendingKill or Garbage flag. Use MarkPendingKill or MarkAsGarbage instead"));
+		PRAGMA_ENABLE_DEPRECATION_WARNINGS
 		ObjectItem->SetFlags(FlagsToSet);
 	}
 
@@ -366,7 +372,9 @@ public:
 	FORCEINLINE void ClearInternalFlags(EInternalObjectFlags FlagsToClear) const
 	{
 		FUObjectItem* ObjectItem = GUObjectArray.IndexToObject(InternalIndex);
+		PRAGMA_DISABLE_DEPRECATION_WARNINGS
 		checkf(!(FlagsToClear & (EInternalObjectFlags::PendingKill | EInternalObjectFlags::Garbage)) || (ObjectItem->Flags & int32(FlagsToClear & (EInternalObjectFlags::PendingKill | EInternalObjectFlags::Garbage))) == 0, TEXT("ClearInternalFlags should not clear PendingKill or Garbage flag. Use ClearGarbage() instead"));
+		PRAGMA_ENABLE_DEPRECATION_WARNINGS
 		ObjectItem->ClearFlags(FlagsToClear);
 	}
 
@@ -379,7 +387,9 @@ public:
 	FORCEINLINE bool AtomicallyClearInternalFlags(EInternalObjectFlags FlagsToClear) const
 	{
 		FUObjectItem* ObjectItem = GUObjectArray.IndexToObject(InternalIndex);
+		PRAGMA_DISABLE_DEPRECATION_WARNINGS
 		checkf((ObjectItem->Flags & int32(FlagsToClear & (EInternalObjectFlags::PendingKill | EInternalObjectFlags::Garbage))) == 0, TEXT("ClearInternalFlags should not clear PendingKill or Garbage flag. Use ClearGarbage() instead"));
+		PRAGMA_ENABLE_DEPRECATION_WARNINGS
 		return ObjectItem->ThisThreadAtomicallyClearedFlag(FlagsToClear);
 	}
 
@@ -428,6 +438,24 @@ public:
 		return !bPendingKillDisabled;
 	}
 
+	/** Helper function that sets the appropriate flag based on PK being enabled or not */
+	FORCEINLINE static EInternalObjectFlags FixGarbageOrPendingKillInternalObjectFlags(const EInternalObjectFlags InFlags)
+	{
+		PRAGMA_DISABLE_DEPRECATION_WARNINGS
+		if (!(InFlags & (EInternalObjectFlags::Garbage | EInternalObjectFlags::PendingKill)))
+		{
+			// Pass through
+			return InFlags;
+		}
+		else
+		{
+			return bPendingKillDisabled ?
+				((InFlags & ~EInternalObjectFlags::PendingKill) | EInternalObjectFlags::Garbage) : // Replace PK with Garbage
+				((InFlags & ~EInternalObjectFlags::Garbage) | EInternalObjectFlags::PendingKill); // Replace Garbage with PK
+		}
+		PRAGMA_ENABLE_DEPRECATION_WARNINGS
+	}
+
 public:
 	/**
 	* Called after load to determine if the object can be a cluster root
diff --git a/Engine/Source/Runtime/CoreUObject/Public/UObject/UObjectIterator.h b/Engine/Source/Runtime/CoreUObject/Public/UObject/UObjectIterator.h
--- a/Engine/Source/Runtime/CoreUObject/Public/UObject/UObjectIterator.h
+++ b/Engine/Source/Runtime/CoreUObject/Public/UObject/UObjectIterator.h
@@ -97,7 +97,7 @@ public:
 		: FUObjectArray::TIterator(GUObjectArray, bOnlyGCedObjects)
 		, Class(InClass)
 		, ExclusionFlags(AdditionalExclusionFlags)
-		, InternalExclusionFlags(InInternalExclusionFlags)
+		, InternalExclusionFlags(UObjectBaseUtility::FixGarbageOrPendingKillInternalObjectFlags(InInternalExclusionFlags))
 	{
 		// We don't want to return any objects that are currently being background loaded unless we're using the object iterator during async loading.
 		InternalExclusionFlags |= EInternalObjectFlags::Unreachable | EInternalObjectFlags::PendingConstruction;
diff --git a/Engine/Source/Runtime/Engine/Private/Actor.cpp b/Engine/Source/Runtime/Engine/Private/Actor.cpp
--- a/Engine/Source/Runtime/Engine/Private/Actor.cpp
+++ b/Engine/Source/Runtime/Engine/Private/Actor.cpp
@@ -333,7 +333,7 @@ void AActor::ResetOwnedComponents()
 				ReplicatedComponents.Add(Component);
 			}
 		}
-	}, true, RF_NoFlags, EInternalObjectFlags::PendingKill);
+	}, true, RF_NoFlags, EInternalObjectFlags::Garbage);
 }
 
 void AActor::PostInitProperties()
@@ -3666,7 +3666,7 @@ void AActor::PostActorConstruction()
 	}
 	else
 	{
-		// Set IsPendingKill() to true so that when the initial undo record is made,
+		// Invalidate the object so that when the initial undo record is made,
 		// the actor will be treated as destroyed, in that undo an add will
 		// actually work
 		MarkAsGarbage();
diff --git a/Engine/Source/Runtime/Engine/Private/Animation/AnimNotifyState_Trail.cpp b/Engine/Source/Runtime/Engine/Private/Animation/AnimNotifyState_Trail.cpp
--- a/Engine/Source/Runtime/Engine/Private/Animation/AnimNotifyState_Trail.cpp
+++ b/Engine/Source/Runtime/Engine/Private/Animation/AnimNotifyState_Trail.cpp
@@ -29,7 +29,7 @@ static void GetCandidateSystems(USkeletalMeshComponent& MeshComp, ParticleSystem
 			{
 				Components.Add(ChildPSC);
 			}
-		}, false, RF_NoFlags, EInternalObjectFlags::PendingKill);
+		}, false, RF_NoFlags, EInternalObjectFlags::Garbage);
 	}
 }
 
diff --git a/Engine/Source/Runtime/Engine/Private/Components/ActorComponent.cpp b/Engine/Source/Runtime/Engine/Private/Components/ActorComponent.cpp
--- a/Engine/Source/Runtime/Engine/Private/Components/ActorComponent.cpp
+++ b/Engine/Source/Runtime/Engine/Private/Components/ActorComponent.cpp
@@ -1333,7 +1333,7 @@ void UActorComponent::RegisterComponentWithWorld(UWorld* InWorld, FRegisterCompo
 	if (IsCreatedByConstructionScript())
 	{
 		TArray<UObject*> Children;
-		GetObjectsWithOuter(this, Children, true, RF_NoFlags, EInternalObjectFlags::PendingKill);
+		GetObjectsWithOuter(this, Children, true, RF_NoFlags, EInternalObjectFlags::Garbage);
 
 		for (UObject* Child : Children)
 		{
diff --git a/Engine/Source/Runtime/Engine/Private/Components/LightComponent.cpp b/Engine/Source/Runtime/Engine/Private/Components/LightComponent.cpp
--- a/Engine/Source/Runtime/Engine/Private/Components/LightComponent.cpp
+++ b/Engine/Source/Runtime/Engine/Private/Components/LightComponent.cpp
@@ -1532,7 +1532,7 @@ void ULightComponent::ReassignStationaryLightChannels(UWorld* TargetWorld, bool
 	TMap<FLightAndChannel*, TArray<FLightAndChannel*> > LightToOverlapMap;
 
 	// Build an array of all static shadowing lights that need to be assigned
-	for (TObjectIterator<ULightComponent> LightIt(RF_ClassDefaultObject, /** bIncludeDerivedClasses */ true, /** InternalExcludeFlags */ EInternalObjectFlags::PendingKill); LightIt; ++LightIt)
+	for (TObjectIterator<ULightComponent> LightIt(RF_ClassDefaultObject, /** bIncludeDerivedClasses */ true, /** InternalExcludeFlags */ EInternalObjectFlags::Garbage); LightIt; ++LightIt)
 	{
 		ULightComponent* const LightComponent = *LightIt;
 		AActor* LightOwner = LightComponent->GetOwner();
diff --git a/Engine/Source/Runtime/Engine/Private/HierarchicalInstancedStaticMesh.cpp b/Engine/Source/Runtime/Engine/Private/HierarchicalInstancedStaticMesh.cpp
--- a/Engine/Source/Runtime/Engine/Private/HierarchicalInstancedStaticMesh.cpp
+++ b/Engine/Source/Runtime/Engine/Private/HierarchicalInstancedStaticMesh.cpp
@@ -163,7 +163,7 @@ static void FoliageCVarSinkFunction()
 		CachedFoliageDensityScale = FoliageDensityScale;
 		FoliageDensityScale = FMath::Clamp(FoliageDensityScale, 0.0f, 1.0f);
 
-		for (auto* Component : TObjectRange<UHierarchicalInstancedStaticMeshComponent>(RF_ClassDefaultObject | RF_ArchetypeObject, true, EInternalObjectFlags::PendingKill))
+		for (auto* Component : TObjectRange<UHierarchicalInstancedStaticMeshComponent>(RF_ClassDefaultObject | RF_ArchetypeObject, true, EInternalObjectFlags::Garbage))
 		{
 #if WITH_EDITOR
 			if (Component->bCanEnableDensityScaling)
diff --git a/Engine/Source/Runtime/Engine/Private/ISMPartition/ISMPartitionActor.cpp b/Engine/Source/Runtime/Engine/Private/ISMPartition/ISMPartitionActor.cpp
--- a/Engine/Source/Runtime/Engine/Private/ISMPartition/ISMPartitionActor.cpp
+++ b/Engine/Source/Runtime/Engine/Private/ISMPartition/ISMPartitionActor.cpp
@@ -43,7 +43,7 @@ static FAutoConsoleCommand DumpISMPartitionActors(
 				// Also output total ism count
 				int32 TotalCount = 0;
 				int32 ComponentCount = 0;
-				for (TObjectIterator<UInstancedStaticMeshComponent> It(RF_ClassDefaultObject | RF_ArchetypeObject, true, EInternalObjectFlags::PendingKill); It; ++It)
+				for (TObjectIterator<UInstancedStaticMeshComponent> It(RF_ClassDefaultObject | RF_ArchetypeObject, true, EInternalObjectFlags::Garbage); It; ++It)
 				{
 					UInstancedStaticMeshComponent* ISMComponent = *It;
 					if (ISMComponent)
diff --git a/Engine/Source/Runtime/Engine/Private/LODActor.cpp b/Engine/Source/Runtime/Engine/Private/LODActor.cpp
--- a/Engine/Source/Runtime/Engine/Private/LODActor.cpp
+++ b/Engine/Source/Runtime/Engine/Private/LODActor.cpp
@@ -1314,7 +1314,7 @@ void ALODActor::OnCVarsChanged()
 	{
 		CachedMaximumAllowedHLODLevel = MaximumAllowedHLODLevel;
 
-		for (ALODActor* Actor : TObjectRange<ALODActor>(RF_ClassDefaultObject | RF_ArchetypeObject, true, EInternalObjectFlags::PendingKill))
+		for (ALODActor* Actor : TObjectRange<ALODActor>(RF_ClassDefaultObject | RF_ArchetypeObject, true, EInternalObjectFlags::Garbage))
 		{
 			Actor->UpdateRegistrationToMatchMaximumLODLevel();
 		}
@@ -1349,7 +1349,7 @@ void ALODActor::OnCVarsChanged()
 	{
 		CachedDistances = HLODDistances;
 		const int32 NumDistances = CachedDistances.Num();
-		for (ALODActor* Actor : TObjectRange<ALODActor>(RF_ClassDefaultObject | RF_ArchetypeObject, true, EInternalObjectFlags::PendingKill))
+		for (ALODActor* Actor : TObjectRange<ALODActor>(RF_ClassDefaultObject | RF_ArchetypeObject, true, EInternalObjectFlags::Garbage))
 		{
 			Actor->UpdateOverrideTransitionDistance();
 		}
diff --git a/Engine/Source/Runtime/Engine/Private/Level.cpp b/Engine/Source/Runtime/Engine/Private/Level.cpp
--- a/Engine/Source/Runtime/Engine/Private/Level.cpp
+++ b/Engine/Source/Runtime/Engine/Private/Level.cpp
@@ -2040,7 +2040,7 @@ void ULevel::PostEditUndo()
 		{
 			Actors.Add(InnerActor);
 		}
-	}, /*bIncludeNestedObjects*/ false, /*ExclusionFlags*/ RF_NoFlags, /* InternalExclusionFlags */ EInternalObjectFlags::PendingKill);
+	}, /*bIncludeNestedObjects*/ false, /*ExclusionFlags*/ RF_NoFlags, /* InternalExclusionFlags */ EInternalObjectFlags::Garbage);
 
 	MarkLevelBoundsDirty();
 }
@@ -3341,7 +3341,7 @@ TArray<UBlueprint*> ULevel::GetLevelBlueprints() const
 		{
 			LevelBlueprints.Add(LevelChildBP);
 		}
-	}, false, RF_NoFlags, EInternalObjectFlags::PendingKill);
+	}, false, RF_NoFlags, EInternalObjectFlags::Garbage);
 
 	return LevelBlueprints;
 }
diff --git a/Engine/Source/Runtime/Engine/Private/LevelScriptActor.cpp b/Engine/Source/Runtime/Engine/Private/LevelScriptActor.cpp
--- a/Engine/Source/Runtime/Engine/Private/LevelScriptActor.cpp
+++ b/Engine/Source/Runtime/Engine/Private/LevelScriptActor.cpp
@@ -125,7 +125,7 @@ TArray<ALevelScriptActor*> ALevelScriptActor::FindSiblingLevelScriptActors() con
 	if (!ThisOuter->IsA<UPackage>())
 	{
 		TArray<UObject*> AllSiblingObjects;
-		GetObjectsWithOuter(ThisOuter, AllSiblingObjects, false, RF_NoFlags, EInternalObjectFlags::PendingKill);
+		GetObjectsWithOuter(ThisOuter, AllSiblingObjects, false, RF_NoFlags, EInternalObjectFlags::Garbage);
 
 		for (UObject* Sibling : AllSiblingObjects)
 		{
diff --git a/Engine/Source/Runtime/Engine/Private/LevelStreaming.cpp b/Engine/Source/Runtime/Engine/Private/LevelStreaming.cpp
--- a/Engine/Source/Runtime/Engine/Private/LevelStreaming.cpp
+++ b/Engine/Source/Runtime/Engine/Private/LevelStreaming.cpp
@@ -1084,7 +1084,7 @@ bool ULevelStreaming::RequestLevel(UWorld* PersistentWorld, bool bAllowLevelLoad
 
 	// Try to find the [to be] loaded package.
 	UWorld* World = nullptr;
-	UPackage* LevelPackage = (UPackage*)StaticFindObjectFast(UPackage::StaticClass(), nullptr, DesiredPackageName, 0, 0, RF_NoFlags, EInternalObjectFlags::PendingKill);
+	UPackage* LevelPackage = (UPackage*)StaticFindObjectFast(UPackage::StaticClass(), nullptr, DesiredPackageName, 0, 0, RF_NoFlags, EInternalObjectFlags::Garbage);
 	
 	if (LevelPackage)
 	{
diff --git a/Engine/Source/Runtime/Engine/Private/MapBuildData.cpp b/Engine/Source/Runtime/Engine/Private/MapBuildData.cpp
--- a/Engine/Source/Runtime/Engine/Private/MapBuildData.cpp
+++ b/Engine/Source/Runtime/Engine/Private/MapBuildData.cpp
@@ -91,7 +91,7 @@ void UWorld::PropagateLightingScenarioChange()
 		}
 	}
 
-	for (USceneComponent* Component : TObjectRange<USceneComponent>(RF_ClassDefaultObject | RF_ArchetypeObject, true, EInternalObjectFlags::PendingKill))
+	for (USceneComponent* Component : TObjectRange<USceneComponent>(RF_ClassDefaultObject | RF_ArchetypeObject, true, EInternalObjectFlags::Garbage))
 	{
 		if (Component->GetWorld() == this)
 		{
@@ -1071,7 +1071,7 @@ void UMapBuildDataRegistry::EmptyLevelData(const TSet<FGuid>* ResourcesToKeep)
 
 void UMapBuildDataRegistry::CleanupTransientOverrideMapBuildData()
 {
-	for (UHierarchicalInstancedStaticMeshComponent* Component : TObjectRange<UHierarchicalInstancedStaticMeshComponent>(RF_ClassDefaultObject | RF_ArchetypeObject, true, EInternalObjectFlags::PendingKill))
+	for (UHierarchicalInstancedStaticMeshComponent* Component : TObjectRange<UHierarchicalInstancedStaticMeshComponent>(RF_ClassDefaultObject | RF_ArchetypeObject, true, EInternalObjectFlags::Garbage))
 	{
 		for (auto& LOD : Component->LODData)
 		{
diff --git a/Engine/Source/Runtime/Engine/Private/PackedLevelActor/PackedLevelActorBuilder.cpp b/Engine/Source/Runtime/Engine/Private/PackedLevelActor/PackedLevelActorBuilder.cpp
--- a/Engine/Source/Runtime/Engine/Private/PackedLevelActor/PackedLevelActorBuilder.cpp
+++ b/Engine/Source/Runtime/Engine/Private/PackedLevelActor/PackedLevelActorBuilder.cpp
@@ -467,7 +467,7 @@ bool FPackedLevelActorBuilder::CreateOrUpdateBlueprintFromPacked(APackedLevelAct
 	if (BP->GeneratedClass)
 	{
 		TArray<UObject*> ObjectsOfClass;
-		GetObjectsOfClass(BP->GeneratedClass, ObjectsOfClass, true, RF_ClassDefaultObject | RF_ArchetypeObject, EInternalObjectFlags::PendingKill);
+		GetObjectsOfClass(BP->GeneratedClass, ObjectsOfClass, true, RF_ClassDefaultObject | RF_ArchetypeObject, EInternalObjectFlags::Garbage);
 		for (UObject* ObjectOfClass : ObjectsOfClass)
 		{
 			APackedLevelActor* PackedLevelActor = CastChecked<APackedLevelActor>(ObjectOfClass);
diff --git a/Engine/Source/Runtime/Engine/Private/Rendering/StaticLightingSystemInterface.cpp b/Engine/Source/Runtime/Engine/Private/Rendering/StaticLightingSystemInterface.cpp
--- a/Engine/Source/Runtime/Engine/Private/Rendering/StaticLightingSystemInterface.cpp
+++ b/Engine/Source/Runtime/Engine/Private/Rendering/StaticLightingSystemInterface.cpp
@@ -181,7 +181,7 @@ ENGINE_API void ToggleLightmapPreview_GameThread(UWorld* InWorld)
 						int32 NumPrimitiveComponents = 0;
 						int32 NumLightComponents = 0;
 
-						for (UPrimitiveComponent* Component : TObjectRange<UPrimitiveComponent>(RF_ClassDefaultObject | RF_ArchetypeObject, true, EInternalObjectFlags::PendingKill))
+						for (UPrimitiveComponent* Component : TObjectRange<UPrimitiveComponent>(RF_ClassDefaultObject | RF_ArchetypeObject, true, EInternalObjectFlags::Garbage))
 						{
 							if (Component->HasValidSettingsForStaticLighting(false))
 							{
@@ -189,7 +189,7 @@ ENGINE_API void ToggleLightmapPreview_GameThread(UWorld* InWorld)
 							}
 						}
 
-						for (ULightComponentBase* Component : TObjectRange<ULightComponentBase>(RF_ClassDefaultObject | RF_ArchetypeObject, true, EInternalObjectFlags::PendingKill))
+						for (ULightComponentBase* Component : TObjectRange<ULightComponentBase>(RF_ClassDefaultObject | RF_ArchetypeObject, true, EInternalObjectFlags::Garbage))
 						{
 							if (Component->bAffectsWorld && Component->HasStaticShadowing())
 							{
@@ -200,7 +200,7 @@ ENGINE_API void ToggleLightmapPreview_GameThread(UWorld* InWorld)
 						FScopedSlowTask SubSlowTask(NumPrimitiveComponents + NumLightComponents, LOCTEXT("RegisteringComponentsWithStaticLightingSystem", "Registering components with static lighting system"));
 						SubSlowTask.MakeDialog();
 
-						for (UPrimitiveComponent* Component : TObjectRange<UPrimitiveComponent>(RF_ClassDefaultObject | RF_ArchetypeObject, true, EInternalObjectFlags::PendingKill))
+						for (UPrimitiveComponent* Component : TObjectRange<UPrimitiveComponent>(RF_ClassDefaultObject | RF_ArchetypeObject, true, EInternalObjectFlags::Garbage))
 						{
 							if (Component->HasValidSettingsForStaticLighting(false))
 							{
@@ -210,7 +210,7 @@ ENGINE_API void ToggleLightmapPreview_GameThread(UWorld* InWorld)
 							}
 						}
 
-						for (ULightComponentBase* Component : TObjectRange<ULightComponentBase>(RF_ClassDefaultObject | RF_ArchetypeObject, true, EInternalObjectFlags::PendingKill))
+						for (ULightComponentBase* Component : TObjectRange<ULightComponentBase>(RF_ClassDefaultObject | RF_ArchetypeObject, true, EInternalObjectFlags::Garbage))
 						{
 							if (Component->bAffectsWorld && Component->HasStaticShadowing())
 							{
@@ -242,12 +242,12 @@ ENGINE_API void ToggleLightmapPreview_GameThread(UWorld* InWorld)
 					int32 NumPrimitiveComponents = 0;
 					int32 NumLightComponents = 0;
 
-					for (UPrimitiveComponent* Component : TObjectRange<UPrimitiveComponent>(RF_ClassDefaultObject | RF_ArchetypeObject, true, EInternalObjectFlags::PendingKill))
+					for (UPrimitiveComponent* Component : TObjectRange<UPrimitiveComponent>(RF_ClassDefaultObject | RF_ArchetypeObject, true, EInternalObjectFlags::Garbage))
 					{
 						NumPrimitiveComponents++;
 					}
 
-					for (ULightComponentBase* Component : TObjectRange<ULightComponentBase>(RF_ClassDefaultObject | RF_ArchetypeObject, true, EInternalObjectFlags::PendingKill))
+					for (ULightComponentBase* Component : TObjectRange<ULightComponentBase>(RF_ClassDefaultObject | RF_ArchetypeObject, true, EInternalObjectFlags::Garbage))
 					{
 						NumLightComponents++;
 					}
@@ -255,19 +255,19 @@ ENGINE_API void ToggleLightmapPreview_GameThread(UWorld* InWorld)
 					FScopedSlowTask SubSlowTask(NumPrimitiveComponents + NumLightComponents, LOCTEXT("UnregisteringComponentsWithStaticLightingSystem", "Unregistering components with static lighting system"));
 
 					// Unregister all landscapes first to prevent grass picking up landscape lightmaps
-					for (ULandscapeComponent* Component : TObjectRange<ULandscapeComponent>(RF_ClassDefaultObject | RF_ArchetypeObject, true, EInternalObjectFlags::PendingKill))
+					for (ULandscapeComponent* Component : TObjectRange<ULandscapeComponent>(RF_ClassDefaultObject | RF_ArchetypeObject, true, EInternalObjectFlags::Garbage))
 					{
 						FStaticLightingSystemInterface::OnPrimitiveComponentUnregistered.Broadcast(Component);
 					}
 
-					for (UPrimitiveComponent* Component : TObjectRange<UPrimitiveComponent>(RF_ClassDefaultObject | RF_ArchetypeObject, true, EInternalObjectFlags::PendingKill))
+					for (UPrimitiveComponent* Component : TObjectRange<UPrimitiveComponent>(RF_ClassDefaultObject | RF_ArchetypeObject, true, EInternalObjectFlags::Garbage))
 					{
 						FStaticLightingSystemInterface::OnPrimitiveComponentUnregistered.Broadcast(Component);
 
 						SubSlowTask.EnterProgressFrame(1, LOCTEXT("UnregisteringComponentsWithStaticLightingSystem", "Unregistering components with static lighting system"));
 					}
 
-					for (ULightComponentBase* Component : TObjectRange<ULightComponentBase>(RF_ClassDefaultObject | RF_ArchetypeObject, true, EInternalObjectFlags::PendingKill))
+					for (ULightComponentBase* Component : TObjectRange<ULightComponentBase>(RF_ClassDefaultObject | RF_ArchetypeObject, true, EInternalObjectFlags::Garbage))
 					{
 						FStaticLightingSystemInterface::OnLightComponentUnregistered.Broadcast(Component);
 
diff --git a/Engine/Source/Runtime/Engine/Private/StaticMeshResources.cpp b/Engine/Source/Runtime/Engine/Private/StaticMeshResources.cpp
--- a/Engine/Source/Runtime/Engine/Private/StaticMeshResources.cpp
+++ b/Engine/Source/Runtime/Engine/Private/StaticMeshResources.cpp
@@ -31,7 +31,7 @@ FStaticMeshComponentRecreateRenderStateContext::FStaticMeshComponentRecreateRend
 
 		const EObjectFlags AdditionalExclusionFlags = RF_ClassDefaultObject;
 		const bool bIncludeDerivedClasses = true;
-		const EInternalObjectFlags InternalExclusionFlags = EInternalObjectFlags::PendingKill;
+		const EInternalObjectFlags InternalExclusionFlags = EInternalObjectFlags::Garbage;
 
 		for (TObjectIterator<UStaticMeshComponent> It(AdditionalExclusionFlags, bIncludeDerivedClasses, InternalExclusionFlags); It; ++It)
 		{
diff --git a/Engine/Source/Runtime/Engine/Private/UnrealEngine.cpp b/Engine/Source/Runtime/Engine/Private/UnrealEngine.cpp
--- a/Engine/Source/Runtime/Engine/Private/UnrealEngine.cpp
+++ b/Engine/Source/Runtime/Engine/Private/UnrealEngine.cpp
@@ -15801,7 +15801,7 @@ void UEngine::CopyPropertiesForUnrelatedObjects(UObject* OldObject, UObject* New
 		{
 			ReferenceReplacementMap.Add(ObjectInOuter, nullptr);
 		}
-	}, true, RF_NoFlags, EInternalObjectFlags::PendingKill);
+	}, true, RF_NoFlags, EInternalObjectFlags::Garbage);
 
 	if(Params.bClearReferences)
 	{
diff --git a/Engine/Source/Runtime/Engine/Private/UnrealExporter.cpp b/Engine/Source/Runtime/Engine/Private/UnrealExporter.cpp
--- a/Engine/Source/Runtime/Engine/Private/UnrealExporter.cpp
+++ b/Engine/Source/Runtime/Engine/Private/UnrealExporter.cpp
@@ -514,7 +514,7 @@ void UExporter::EmitEndObject( FOutputDevice& Ar )
 FExportObjectInnerContext::FExportObjectInnerContext()
 {
 	// For each object . . .
-	for (UObject* InnerObj : TObjectRange<UObject>(RF_ClassDefaultObject, true, EInternalObjectFlags::PendingKill))
+	for (UObject* InnerObj : TObjectRange<UObject>(RF_ClassDefaultObject, true, EInternalObjectFlags::Garbage))
 	{
 		UObject* OuterObj = InnerObj->GetOuter();
 		if ( OuterObj )
@@ -539,7 +539,7 @@ FExportObjectInnerContext::FExportObjectInnerContext()
 FExportObjectInnerContext::FExportObjectInnerContext(const TArray<UObject*>& ObjsToIgnore)
 {
 	// For each object . . .
-	for (UObject* InnerObj : TObjectRange<UObject>(RF_ClassDefaultObject, true, EInternalObjectFlags::PendingKill))
+	for (UObject* InnerObj : TObjectRange<UObject>(RF_ClassDefaultObject, true, EInternalObjectFlags::Garbage))
 	{
 		if (!ObjsToIgnore.Contains(InnerObj))
 		{
@@ -592,7 +592,7 @@ void UExporter::ExportObjectInner(const FExportObjectInnerContext* Context, UObj
 	else
 	{
 		// NOTE: We ignore inner objects that have been tagged for death
-		GetObjectsWithOuter(Object, TempInners, false, RF_NoFlags, EInternalObjectFlags::PendingKill);
+		GetObjectsWithOuter(Object, TempInners, false, RF_NoFlags, EInternalObjectFlags::Garbage);
 	}
 	FExportObjectInnerContext::InnerList const& UnsortedObjectInners = ContextInners ? *ContextInners : TempInners;
 
@@ -974,6 +974,6 @@ void FSelectedActorExportObjectInnerContext::AddSelectedActor(const AActor* InAc
 			InnerList& InnersForOuterObject = ObjectToInnerMap.Add(OuterObj, InnerList());
 			InnersForOuterObject.Add(InnerObj);
 		}
-	}, /** bIncludeNestedObjects */ true, RF_NoFlags, EInternalObjectFlags::PendingKill);
+	}, /** bIncludeNestedObjects */ true, RF_NoFlags, EInternalObjectFlags::Garbage);
 }
 #endif
diff --git a/Engine/Source/Runtime/Engine/Private/World.cpp b/Engine/Source/Runtime/Engine/Private/World.cpp
--- a/Engine/Source/Runtime/Engine/Private/World.cpp
+++ b/Engine/Source/Runtime/Engine/Private/World.cpp
@@ -2044,7 +2044,7 @@ void UWorld::MarkObjectsPendingKill()
 	{
 		Object->MarkAsGarbage();
 	};
-	ForEachObjectWithOuter(this, MarkObjectPendingKill, true, RF_NoFlags, EInternalObjectFlags::PendingKill);
+	ForEachObjectWithOuter(this, MarkObjectPendingKill, true, RF_NoFlags, EInternalObjectFlags::Garbage);
 
 	MarkAsGarbage();
 	bMarkedObjectsPendingKill = true;
@@ -3275,9 +3275,9 @@ void FLevelStreamingGCHelper::PrepareStreamedOutLevelsForGC()
 				Packages.Add(Package, &bIsAlreadyInSet);
 				if (!bIsAlreadyInSet)
 				{
-					ForEachObjectWithPackage(Package, [](UObject* PackageObject) { PackageObject->MarkAsGarbage(); return true; }, true, RF_NoFlags, EInternalObjectFlags::PendingKill);
+					ForEachObjectWithPackage(Package, [](UObject* PackageObject) { PackageObject->MarkAsGarbage(); return true; }, true, RF_NoFlags, EInternalObjectFlags::Garbage);
 				}
-			}, true, RF_NoFlags, EInternalObjectFlags::PendingKill);
+			}, true, RF_NoFlags, EInternalObjectFlags::Garbage);
 		}
 	}
 
diff --git a/Engine/Source/Runtime/Engine/Private/WorldPartition/WorldPartitionLevelStreamingDynamic.cpp b/Engine/Source/Runtime/Engine/Private/WorldPartition/WorldPartitionLevelStreamingDynamic.cpp
--- a/Engine/Source/Runtime/Engine/Private/WorldPartition/WorldPartitionLevelStreamingDynamic.cpp
+++ b/Engine/Source/Runtime/Engine/Private/WorldPartition/WorldPartitionLevelStreamingDynamic.cpp
@@ -196,7 +196,7 @@ bool UWorldPartitionLevelStreamingDynamic::RequestLevel(UWorld* InPersistentWorl
 
 	// Try to find the package to load
 	const FName DesiredPackageName = GetWorldAssetPackageFName();
-	UPackage* LevelPackage = (UPackage*)StaticFindObjectFast(UPackage::StaticClass(), nullptr, DesiredPackageName, 0, 0, RF_NoFlags, EInternalObjectFlags::PendingKill);
+	UPackage* LevelPackage = (UPackage*)StaticFindObjectFast(UPackage::StaticClass(), nullptr, DesiredPackageName, 0, 0, RF_NoFlags, EInternalObjectFlags::Garbage);
 	UWorld* FoundWorld = LevelPackage ? UWorld::FindWorldInPackage(LevelPackage) : nullptr;
 	check(!FoundWorld || IsValidChecked(FoundWorld));
 	check(!FoundWorld || FoundWorld->PersistentLevel);
diff --git a/Engine/Source/Runtime/Engine/Public/EngineUtils.h b/Engine/Source/Runtime/Engine/Public/EngineUtils.h
--- a/Engine/Source/Runtime/Engine/Public/EngineUtils.h
+++ b/Engine/Source/Runtime/Engine/Public/EngineUtils.h
@@ -227,7 +227,7 @@ public:
 #endif // WITH_EDITOR
 		{
 			constexpr EObjectFlags ExcludeFlags = RF_ClassDefaultObject;
-			GetObjectsOfClass(InClass, ObjectArray, true, ExcludeFlags, EInternalObjectFlags::PendingKill);
+			GetObjectsOfClass(InClass, ObjectArray, true, ExcludeFlags, EInternalObjectFlags::Garbage);
 		}
 
 		const auto ActorSpawnedDelegate = FOnActorSpawned::FDelegate::CreateRaw(this, &FActorIteratorState::OnActorSpawned);
diff --git a/Engine/Source/Runtime/Foliage/Private/InstancedFoliage.cpp b/Engine/Source/Runtime/Foliage/Private/InstancedFoliage.cpp
--- a/Engine/Source/Runtime/Foliage/Private/InstancedFoliage.cpp
+++ b/Engine/Source/Runtime/Foliage/Private/InstancedFoliage.cpp
@@ -942,7 +942,7 @@ void UFoliageType::PostEditChangeProperty(struct FPropertyChangedEvent& Property
 	// Notify any currently-loaded InstancedFoliageActors
 	if (IsFoliageReallocationRequiredForPropertyChange(PropertyChangedEvent.Property))
 	{
-		for (TObjectIterator<AInstancedFoliageActor> It(RF_ClassDefaultObject, /** bIncludeDerivedClasses */ true, /** InternalExcludeFalgs */ EInternalObjectFlags::PendingKill); It; ++It)
+		for (TObjectIterator<AInstancedFoliageActor> It(RF_ClassDefaultObject, /** bIncludeDerivedClasses */ true, /** InternalExcludeFalgs */ EInternalObjectFlags::Garbage); It; ++It)
 		{
 			if (It->GetWorld() != nullptr)
 			{
@@ -958,7 +958,7 @@ void UFoliageType::PreEditChange(FProperty* PropertyAboutToChange)
 
 	if (IsSourcePropertyChange(PropertyAboutToChange))
 	{
-		for (TObjectIterator<AInstancedFoliageActor> It(RF_ClassDefaultObject, /** bIncludeDerivedClasses */ true, /** InternalExcludeFalgs */ EInternalObjectFlags::PendingKill); It; ++It)
+		for (TObjectIterator<AInstancedFoliageActor> It(RF_ClassDefaultObject, /** bIncludeDerivedClasses */ true, /** InternalExcludeFalgs */ EInternalObjectFlags::Garbage); It; ++It)
 		{
 			It->NotifyFoliageTypeWillChange(this);
 		}
diff --git a/Engine/Source/Runtime/Landscape/Private/Landscape.cpp b/Engine/Source/Runtime/Landscape/Private/Landscape.cpp
--- a/Engine/Source/Runtime/Landscape/Private/Landscape.cpp
+++ b/Engine/Source/Runtime/Landscape/Private/Landscape.cpp
@@ -1454,7 +1454,7 @@ static void OnStaticMeshLODDistanceScaleChanged()
 	{
 		LastValue = CVarStaticMeshLODDistanceScale.GetValueOnAnyThread();
 
-		for (auto* LandscapeComponent : TObjectRange<ULandscapeComponent>(RF_ClassDefaultObject | RF_ArchetypeObject, true, EInternalObjectFlags::PendingKill))
+		for (auto* LandscapeComponent : TObjectRange<ULandscapeComponent>(RF_ClassDefaultObject | RF_ArchetypeObject, true, EInternalObjectFlags::Garbage))
 		{
 			LandscapeComponent->MarkRenderStateDirty();
 		}
diff --git a/Engine/Source/Runtime/Landscape/Private/LandscapeGrass.cpp b/Engine/Source/Runtime/Landscape/Private/LandscapeGrass.cpp
--- a/Engine/Source/Runtime/Landscape/Private/LandscapeGrass.cpp
+++ b/Engine/Source/Runtime/Landscape/Private/LandscapeGrass.cpp
@@ -223,7 +223,7 @@ static void GrassCVarSinkFunction()
 		CachedGrassCullDistanceScale = GrassCullDistanceScale;
 		CachedDetailMode = DetailMode;
 
-		for (auto* Landscape : TObjectRange<ALandscapeProxy>(RF_ClassDefaultObject | RF_ArchetypeObject, true, EInternalObjectFlags::PendingKill))
+		for (auto* Landscape : TObjectRange<ALandscapeProxy>(RF_ClassDefaultObject | RF_ArchetypeObject, true, EInternalObjectFlags::Garbage))
 		{
 			Landscape->FlushGrassComponents(nullptr, false);
 		}
@@ -3352,7 +3352,7 @@ FAsyncGrassTask::~FAsyncGrassTask()
 
 static void FlushGrass(const TArray<FString>& Args)
 {
-	for (ALandscapeProxy* Landscape : TObjectRange<ALandscapeProxy>(RF_ClassDefaultObject | RF_ArchetypeObject, true, EInternalObjectFlags::PendingKill))
+	for (ALandscapeProxy* Landscape : TObjectRange<ALandscapeProxy>(RF_ClassDefaultObject | RF_ArchetypeObject, true, EInternalObjectFlags::Garbage))
 	{
 		Landscape->FlushGrassComponents();
 	}
@@ -3360,7 +3360,7 @@ static void FlushGrass(const TArray<FString>& Args)
 
 static void FlushGrassPIE(const TArray<FString>& Args)
 {
-	for (ALandscapeProxy* Landscape : TObjectRange<ALandscapeProxy>(RF_ClassDefaultObject | RF_ArchetypeObject, true, EInternalObjectFlags::PendingKill))
+	for (ALandscapeProxy* Landscape : TObjectRange<ALandscapeProxy>(RF_ClassDefaultObject | RF_ArchetypeObject, true, EInternalObjectFlags::Garbage))
 	{
 		Landscape->FlushGrassComponents(nullptr, false);
 	}
diff --git a/Engine/Source/Runtime/Landscape/Private/LandscapeRender.cpp b/Engine/Source/Runtime/Landscape/Private/LandscapeRender.cpp
--- a/Engine/Source/Runtime/Landscape/Private/LandscapeRender.cpp
+++ b/Engine/Source/Runtime/Landscape/Private/LandscapeRender.cpp
@@ -48,7 +48,7 @@ IMPLEMENT_TYPE_LAYOUT(FLandscapeVertexFactoryPixelShaderParameters);
 #if !UE_BUILD_SHIPPING
 static void OnLODDistributionScaleChanged(IConsoleVariable* CVar)
 {
-	for (auto* LandscapeComponent : TObjectRange<ULandscapeComponent>(RF_ClassDefaultObject | RF_ArchetypeObject, true, EInternalObjectFlags::PendingKill))
+	for (auto* LandscapeComponent : TObjectRange<ULandscapeComponent>(RF_ClassDefaultObject | RF_ArchetypeObject, true, EInternalObjectFlags::Garbage))
 	{
 		LandscapeComponent->MarkRenderStateDirty();
 	}
