Levels will now have the ability to clean up references that prevent them from being GC'd

diff --git a/Engine/Source/Runtime/Engine/Classes/Engine/Level.h b/Engine/Source/Runtime/Engine/Classes/Engine/Level.h
--- a/Engine/Source/Runtime/Engine/Classes/Engine/Level.h
+++ b/Engine/Source/Runtime/Engine/Classes/Engine/Level.h
@@ -910,6 +910,11 @@ public:
 	 */
 	ENGINE_API void CleanupLevel(bool bCleanupResources);
 
+	/**
+	 * Cleans all references from and to this level that may be preventing it from being Garbage Collected
+	 */
+	ENGINE_API void CleanupReferences();
+
 	/**
 	 * Clears all components of actors associated with this level (aka in Actors array) and 
 	 * also the BSP model components.
diff --git a/Engine/Source/Runtime/Engine/Private/Level.cpp b/Engine/Source/Runtime/Engine/Private/Level.cpp
--- a/Engine/Source/Runtime/Engine/Private/Level.cpp
+++ b/Engine/Source/Runtime/Engine/Private/Level.cpp
@@ -446,6 +446,15 @@ void ULevel::CleanupLevel(bool bCleanupResources)
 	}
 }
 
+void ULevel::CleanupReferences()
+{
+	// Make sure MetaData is not marked as standalone as this will prevent the level from being GC'd
+	if (GetOutermost()->HasMetaData())
+	{
+		GetOutermost()->GetMetaData()->ClearFlags(RF_Standalone);
+	}
+}
+
 void ULevel::PostInitProperties()
 {
 	Super::PostInitProperties();
diff --git a/Engine/Source/Runtime/Engine/Private/World.cpp b/Engine/Source/Runtime/Engine/Private/World.cpp
--- a/Engine/Source/Runtime/Engine/Private/World.cpp
+++ b/Engine/Source/Runtime/Engine/Private/World.cpp
@@ -2043,6 +2043,11 @@ void UWorld::MarkObjectsPendingKill()
 	auto MarkObjectPendingKill = [](UObject* Object)
 	{
 		Object->MarkAsGarbage();
+
+		if (ULevel* Level = Cast<ULevel>(Object))
+		{
+			Level->CleanupReferences();
+		}
 	};
 	ForEachObjectWithOuter(this, MarkObjectPendingKill, true, RF_NoFlags, EInternalObjectFlags::Garbage);
 
@@ -3278,6 +3283,8 @@ void FLevelStreamingGCHelper::PrepareStreamedOutLevelsForGC()
 					ForEachObjectWithPackage(Package, [](UObject* PackageObject) { PackageObject->MarkAsGarbage(); return true; }, true, RF_NoFlags, EInternalObjectFlags::Garbage);
 				}
 			}, true, RF_NoFlags, EInternalObjectFlags::Garbage);
+
+			Level->CleanupReferences();
 		}
 	}
 
@@ -5009,6 +5016,7 @@ void UWorld::CleanupWorldInternal(bool bSessionEnded, bool bCleanupResources, UW
 		if (PersistentLevel)
 		{
 			PersistentLevel->CleanupLevel(bCleanupResources);
+			PersistentLevel->CleanupReferences();
 		}
 
 		if (GetNumLevels() > 1)
@@ -5016,7 +5024,9 @@ void UWorld::CleanupWorldInternal(bool bSessionEnded, bool bCleanupResources, UW
 			check(GetLevel(0) == PersistentLevel);
 			for (int32 LevelIndex = 1; LevelIndex < GetNumLevels(); ++LevelIndex)
 			{
-				GetLevel(LevelIndex)->CleanupLevel(bCleanupResources);
+				ULevel* Level = GetLevel(LevelIndex);
+				Level->CleanupLevel(bCleanupResources);
+				Level->CleanupReferences();
 			}
 		}
 	}
