Compilation fixes for FVector2D/FVector2f

diff --git a/Engine/Plugins/Experimental/SkeletalReduction/Source/Private/SkeletalSimplifierVertex.h b/Engine/Plugins/Experimental/SkeletalReduction/Source/Private/SkeletalSimplifierVertex.h
--- a/Engine/Plugins/Experimental/SkeletalReduction/Source/Private/SkeletalSimplifierVertex.h
+++ b/Engine/Plugins/Experimental/SkeletalReduction/Source/Private/SkeletalSimplifierVertex.h
@@ -157,7 +157,7 @@ namespace SkeletalSimplifier
 			{
 				for (int32 i = 0; i < NumTexCoords; ++i)
 				{
-					TexCoords[i] = FVector2D(ForceInitToZero);
+					TexCoords[i] = FVector2f(ForceInitToZero);
 				}
 			}
 
diff --git a/Engine/Source/Runtime/Engine/Private/PrimitiveDrawingUtils.cpp b/Engine/Source/Runtime/Engine/Private/PrimitiveDrawingUtils.cpp
--- a/Engine/Source/Runtime/Engine/Private/PrimitiveDrawingUtils.cpp
+++ b/Engine/Source/Runtime/Engine/Private/PrimitiveDrawingUtils.cpp
@@ -757,7 +757,7 @@ void DrawTorus(FPrimitiveDrawInterface* PDI, const FMatrix& Transform, const FVe
 	FVector ZAxis = XAxis ^ YAxis;
 	ZAxis.Normalize();
 
-	FVector2D TC = FVector2D(0.0f, 0.0f);
+	FVector2f TC = FVector2f(0.0f, 0.0f);
 	float TCStepX = 1.0f / InnerSegments;
 	float TCStepY = 1.0f / OuterSegments;
 
@@ -1060,12 +1060,12 @@ void DrawRectangleMesh(FPrimitiveDrawInterface* PDI, const FVector& Center, cons
 		Center - XOffset + YOffset
 	};
 
-	FVector2D UVs[4] =
+	FVector2f UVs[4] =
 	{
-		FVector2D(0.0f,0.0f),
-		FVector2D(0.0f,1.0f),
-		FVector2D(1.0f,1.0f),
-		FVector2D(1.0f,0.0f),
+		FVector2f(0.0f,0.0f),
+		FVector2f(0.0f,1.0f),
+		FVector2f(1.0f,1.0f),
+		FVector2f(1.0f,0.0f),
 	};
 
 	FDynamicMeshBuilder MeshBuilder(PDI->View->GetFeatureLevel());
