Don't strip source morph data when cooking for mobile platforms as CPU morphing needs it

diff --git a/Engine/Source/Runtime/Engine/Classes/Animation/MorphTarget.h b/Engine/Source/Runtime/Engine/Classes/Animation/MorphTarget.h
--- a/Engine/Source/Runtime/Engine/Classes/Animation/MorphTarget.h
+++ b/Engine/Source/Runtime/Engine/Classes/Animation/MorphTarget.h
@@ -74,54 +74,7 @@ struct FMorphTargetLODModel
 	{ }
 
 	/** pipe operator */
-	friend FArchive& operator<<(FArchive& Ar, FMorphTargetLODModel& M)
-	{
-		Ar.UsingCustomVersion(FEditorObjectVersion::GUID);
-		Ar.UsingCustomVersion(FFortniteMainBranchObjectVersion::GUID);
-		Ar.UsingCustomVersion(FUE5PrivateFrostyStreamObjectVersion::GUID);
-
-		if (!Ar.IsObjectReferenceCollector())
-		{
-			if (Ar.IsLoading() && Ar.CustomVer(FEditorObjectVersion::GUID) < FEditorObjectVersion::AddedMorphTargetSectionIndices)
-			{
-				Ar << M.Vertices << M.NumBaseMeshVerts;
-				M.bGeneratedByEngine = false;
-			}
-			else if (Ar.IsLoading() && Ar.CustomVer(FFortniteMainBranchObjectVersion::GUID) < FFortniteMainBranchObjectVersion::SaveGeneratedMorphTargetByEngine)
-			{
-				Ar << M.Vertices << M.NumBaseMeshVerts << M.SectionIndices;
-				M.bGeneratedByEngine = false;
-			}
-			else
-			{
-				bool bIsCooked = false;
-				if (Ar.IsPersistent() && (Ar.CustomVer(FUE5PrivateFrostyStreamObjectVersion::GUID) >= FUE5PrivateFrostyStreamObjectVersion::StripMorphTargetSourceDataForCookedBuilds))
-				{
-					bIsCooked = Ar.IsCooking();
-					Ar << bIsCooked;
-				}
-
-				if (bIsCooked)
-				{
-					M.NumVertices = M.Vertices.Num();
-					Ar << M.NumVertices;
-				}
-				else
-				{
-					Ar << M.Vertices;
-
-					if (Ar.IsLoading())
-					{
-						M.NumVertices = M.Vertices.Num();
-					}
-				}
-
-				Ar << M.NumBaseMeshVerts << M.SectionIndices << M.bGeneratedByEngine;
-			}
-		}
-
-		return Ar;
-	}
+	friend FArchive& operator<<(FArchive& Ar, FMorphTargetLODModel& M);
 
 	void Reset()
 	{
@@ -202,6 +155,7 @@ public:
 	ENGINE_API virtual void EmptyMorphLODModels();
 
 	/** Discard CPU Buffers after render resources have been created. */
+	UE_DEPRECATED(5.0, "No longer in use, will be deleted. Whether to discard vertex data is now determined during cooking instead of loading.")
 	ENGINE_API virtual void DiscardVertexData();
 
 	/** Return true if this morph target uses engine built-in compression */
diff --git a/Engine/Source/Runtime/Engine/Private/MorphTools.cpp b/Engine/Source/Runtime/Engine/Private/MorphTools.cpp
--- a/Engine/Source/Runtime/Engine/Private/MorphTools.cpp
+++ b/Engine/Source/Runtime/Engine/Private/MorphTools.cpp
@@ -12,6 +12,57 @@
 #include "Rendering/SkeletalMeshLODModel.h"
 #include "UObject/GarbageCollection.h"
 #include "Algo/AnyOf.h"
+#include "Interfaces/ITargetPlatform.h"
+
+FArchive& operator<<(FArchive& Ar, FMorphTargetLODModel& M)
+{
+	Ar.UsingCustomVersion(FEditorObjectVersion::GUID);
+	Ar.UsingCustomVersion(FFortniteMainBranchObjectVersion::GUID);
+	Ar.UsingCustomVersion(FUE5PrivateFrostyStreamObjectVersion::GUID);
+
+	if (!Ar.IsObjectReferenceCollector())
+	{
+		if (Ar.IsLoading() && Ar.CustomVer(FEditorObjectVersion::GUID) < FEditorObjectVersion::AddedMorphTargetSectionIndices)
+		{
+			Ar << M.Vertices << M.NumBaseMeshVerts;
+			M.bGeneratedByEngine = false;
+		}
+		else if (Ar.IsLoading() && Ar.CustomVer(FFortniteMainBranchObjectVersion::GUID) < FFortniteMainBranchObjectVersion::SaveGeneratedMorphTargetByEngine)
+		{
+			Ar << M.Vertices << M.NumBaseMeshVerts << M.SectionIndices;
+			M.bGeneratedByEngine = false;
+		}
+		else
+		{
+			bool bVerticesAreStrippedForCookedBuilds = false;
+			if (Ar.IsPersistent() && (Ar.CustomVer(FUE5PrivateFrostyStreamObjectVersion::GUID) >= FUE5PrivateFrostyStreamObjectVersion::StripMorphTargetSourceDataForCookedBuilds))
+			{
+				// Strip source morph data for cooked build if targets don't include mobile. Mobile uses CPU morphing which needs the source morph data.
+				bVerticesAreStrippedForCookedBuilds = Ar.IsCooking() && (!Ar.CookingTarget()->SupportsFeature(ETargetPlatformFeatures::MobileRendering));
+				Ar << bVerticesAreStrippedForCookedBuilds;
+			}
+
+			if (bVerticesAreStrippedForCookedBuilds)
+			{
+				M.NumVertices = M.Vertices.Num();
+				Ar << M.NumVertices;
+			}
+			else
+			{
+				Ar << M.Vertices;
+
+				if (Ar.IsLoading())
+				{
+					M.NumVertices = M.Vertices.Num();
+				}
+			}
+
+			Ar << M.NumBaseMeshVerts << M.SectionIndices << M.bGeneratedByEngine;
+		}
+	}
+
+	return Ar;
+}
 
 /** compare based on base mesh source vertex indices */
 struct FCompareMorphTargetDeltas
diff --git a/Engine/Source/Runtime/Engine/Private/SkeletalMeshRenderData.cpp b/Engine/Source/Runtime/Engine/Private/SkeletalMeshRenderData.cpp
--- a/Engine/Source/Runtime/Engine/Private/SkeletalMeshRenderData.cpp
+++ b/Engine/Source/Runtime/Engine/Private/SkeletalMeshRenderData.cpp
@@ -739,16 +739,6 @@ void FSkeletalMeshRenderData::InitResources(bool bNeedsVertexColors, TArray<UMor
 			}
 		}
 
-#if !WITH_EDITOR
-		for (UMorphTarget* MorphTarget : InMorphTargets)
-		{
-			// A CPU copy of the morph deltas has been made so it is safe to 
-			// discard the original data. Keep CPU buffers when in the editor.
-			// Note: this data is stripped at cook time so it's probably empty.
-			MorphTarget->DiscardVertexData();
-		}
-#endif
-
 		ENQUEUE_RENDER_COMMAND(CmdSetSkeletalMeshReadyForStreaming)(
 			[this, Owner](FRHICommandListImmediate&)
 		{
