PR #8766: TextRenderComponents now allow Decal Materials (Contributed by [REDACTED_PERSON])

diff --git a/Engine/Source/Runtime/Engine/Private/Components/TextRenderComponent.cpp b/Engine/Source/Runtime/Engine/Private/Components/TextRenderComponent.cpp
--- a/Engine/Source/Runtime/Engine/Private/Components/TextRenderComponent.cpp
+++ b/Engine/Source/Runtime/Engine/Private/Components/TextRenderComponent.cpp
@@ -660,7 +660,7 @@ FTextRenderSceneProxy::FTextRenderSceneProxy( UTextRenderComponent* Component) :
 	{
 		const UMaterial* BaseMaterial = Component->TextMaterial->GetMaterial_Concurrent();
 
-		if(BaseMaterial->MaterialDomain == MD_Surface)
+		if(BaseMaterial->MaterialDomain == MD_Surface || BaseMaterial->MaterialDomain == MD_DeferredDecal)
 		{
 			EffectiveMaterial = Component->TextMaterial;
 		}
