Clearing an actors label without broadcasting OnActorLabelChanged could lead to desynced cached actor labels in the scene outliner.

diff --git a/Engine/Source/Runtime/Engine/Private/ActorEditor.cpp b/Engine/Source/Runtime/Engine/Private/ActorEditor.cpp
--- a/Engine/Source/Runtime/Engine/Private/ActorEditor.cpp
+++ b/Engine/Source/Runtime/Engine/Private/ActorEditor.cpp
@@ -973,6 +973,7 @@ bool AActor::IsActorLabelEditable() const
 void AActor::ClearActorLabel()
 {
 	ActorLabel.Reset();
+	FCoreDelegates::OnActorLabelChanged.Broadcast(this);
 }
 
 FFolder AActor::GetFolder() const
