Garbage Collector History: GC can now store information about traversed references which can then later be searched through with FReferenceChainSearch (obj refs)

diff --git a/Engine/Source/Runtime/CoreUObject/Private/UObject/GarbageCollection.cpp b/Engine/Source/Runtime/CoreUObject/Private/UObject/GarbageCollection.cpp
--- a/Engine/Source/Runtime/CoreUObject/Private/UObject/GarbageCollection.cpp
+++ b/Engine/Source/Runtime/CoreUObject/Private/UObject/GarbageCollection.cpp
@@ -32,6 +32,7 @@
 #include "HAL/Runnable.h"
 #include "HAL/RunnableThread.h"
 #include "UObject/FieldPathProperty.h"
+#include "UObject/GarbageCollectionHistory.h"
 
 /*-----------------------------------------------------------------------------
    Garbage collection.
@@ -613,6 +614,9 @@ void ShutdownGarbageCollection()
 	FGCArrayPool::Get().Cleanup();
 	delete GAsyncPurge;
 	GAsyncPurge = nullptr;
+#if ENABLE_GC_HISTORY
+	FGCHistory::Get().Cleanup();
+#endif
 }
 
 /**
@@ -625,7 +629,10 @@ class FGCReferenceProcessor
 {
 	// Local copies of globals to make sure they're cached
 	const int32 MinDesiredObjectsPerSubTask;
-	const bool GarbageReferenceTrackingEnabled;
+	const bool bGarbageReferenceTrackingEnabled;
+#if ENABLE_GC_HISTORY
+	const bool bWithHistory;
+#endif // ENABLE_GC_HISTORY
 
 public:
 
@@ -644,7 +651,10 @@ public:
 
 	FGCReferenceProcessor()
 		: MinDesiredObjectsPerSubTask(GMinDesiredObjectsPerSubTask)
-		, GarbageReferenceTrackingEnabled(!!GGarbageReferenceTrackingEnabled)
+		, bGarbageReferenceTrackingEnabled(!!GGarbageReferenceTrackingEnabled)
+#if ENABLE_GC_HISTORY
+		, bWithHistory(FGCHistory::Get().IsActive())
+#endif // ENABLE_GC_HISTORY
 	{
 	}
 
@@ -866,6 +876,31 @@ public:
 		}
 	}
 
+#if ENABLE_GC_HISTORY
+	FORCEINLINE void HandleHistoryReference(FGCArrayStruct& ObjectsToSerializeStruct, const UObject* const ReferencingObject, UObject*& Object, const int32 TokenIndex)
+	{
+		UObject* Referencer = ReferencingObject ? const_cast<UObject*>(ReferencingObject) : ObjectsToSerializeStruct.GetReferencingObject();		
+		FGCDirectReference Ref(Object);
+
+		if (TokenIndex >= 0)
+		{
+			UClass* Class = Referencer->GetClass();	
+			Ref.ReferencerName = Class->ReferenceTokenStream.GetTokenInfo(TokenIndex).Name;
+		}
+		else if (FGCObject::GGCObjectReferencer == Referencer)
+		{
+			Ref.ReferencerName = *FGCObject::GGCObjectReferencer->GetCurrentlySerializingObject()->GetReferencerName();
+		}
+
+		TArray<FGCDirectReference>*& DirectReferences = ObjectsToSerializeStruct.History.FindOrAdd(Referencer);
+		if (!DirectReferences)
+		{
+			DirectReferences = new TArray<FGCDirectReference>();
+		}
+		DirectReferences->Add(Ref);
+	}
+#endif // ENABLE_GC_HISTORY
+
 	/**
 	 * Handles object reference, potentially NULL'ing
 	 *
@@ -908,12 +943,12 @@ public:
 					return;
 				}
 
-				if (GarbageReferenceTrackingEnabled && !ObjectItem->HasAnyFlags(EInternalObjectFlags::PersistentGarbage))
+				if (bGarbageReferenceTrackingEnabled && !ObjectItem->HasAnyFlags(EInternalObjectFlags::PersistentGarbage))
 				{
 					HandleGarbageReference(ObjectsToSerializeStruct, ReferencingObject, Object, TokenIndex);
 				}
 			}
-			else if (GarbageReferenceTrackingEnabled)
+			else if (bGarbageReferenceTrackingEnabled)
 			{
 				// This object is being referenced by a persistent reference which means it wouldn't have been GC'd anyway
 				// so no need to track this object
@@ -937,6 +972,13 @@ public:
 				// Mark it as reachable.
 				if (ObjectItem->ThisThreadAtomicallyClearedRFUnreachable())
 				{
+#if ENABLE_GC_HISTORY
+					if (bWithHistory)
+					{
+						HandleHistoryReference(ObjectsToSerializeStruct, ReferencingObject, Object, TokenIndex);
+					}
+#endif // ENABLE_GC_HISTORY
+
 					// Objects that are part of a GC cluster should never have the unreachable flag set!
 					checkSlow(ObjectItem->GetOwnerIndex() <= 0);
 
@@ -968,6 +1010,13 @@ public:
 				}
 #endif
 
+#if ENABLE_GC_HISTORY
+				if (bWithHistory)
+				{
+					HandleHistoryReference(ObjectsToSerializeStruct, ReferencingObject, Object, TokenIndex);
+				}
+#endif // ENABLE_GC_HISTORY
+
 				// Mark it as reachable.
 				ObjectItem->ClearUnreachable();
 
@@ -2247,6 +2296,12 @@ void CollectGarbageInternal(EObjectFlags KeepFlags, bool bPerformFullPurge)
 			UE_LOG(LogGarbage, Log, TEXT("%f ms for GC"), (FPlatformTime::Seconds() - StartTime) * 1000);
 		}
 
+		FGCArrayPool& ArrayPool = FGCArrayPool::Get();
+		TArray<FGCArrayStruct*> AllArrays;
+		ArrayPool.GetAllArrayStructsFromPool(AllArrays);
+		// This needs to happen before clusters get dissolved otherwisise cluster information will be missing from history
+		ArrayPool.UpdateGCHistory(AllArrays);
+
 		// Reconstruct clusters if needed
 		if (GUObjectClusters.ClustersNeedDissolving())
 		{
@@ -2262,9 +2317,6 @@ void CollectGarbageInternal(EObjectFlags KeepFlags, bool bPerformFullPurge)
 		}
 
 		{
-			FGCArrayPool& ArrayPool = FGCArrayPool::Get();
-			TArray<FGCArrayStruct*> AllArrays;
-			ArrayPool.GetAllArrayStructsFromPool(AllArrays);
 			ArrayPool.DumpGarbageReferencers(AllArrays);
 		
 			GatherUnreachableObjects(!(Options & EFastReferenceCollectorOptions::Parallel));
@@ -2288,6 +2340,9 @@ void CollectGarbageInternal(EObjectFlags KeepFlags, bool bPerformFullPurge)
 				}
 			}
 
+			// Make sure nothing will be using potentially freed arrays
+			AllArrays.Empty();
+
 			if (bPerformFullPurge || !GIncrementalBeginDestroyEnabled)
 			{
 				UnhashUnreachableObjects(/**bUseTimeLimit = */ false);
@@ -2378,6 +2433,22 @@ void FGCArrayPool::DumpGarbageReferencers(TArray<FGCArrayStruct*>& AllArrays)
 	}
 }
 
+void FGCArrayPool::UpdateGCHistory(TArray<FGCArrayStruct*>& AllArrays)
+{
+#if ENABLE_GC_HISTORY
+	FGCHistory::Get().Update(AllArrays);
+
+	for (FGCArrayStruct* ArrayStruct : AllArrays)
+	{
+		for (TPair<UObject*, TArray<FGCDirectReference>*>& DirectReferenceInfos : ArrayStruct->History)
+		{
+			delete DirectReferenceInfos.Value;
+		}
+		ArrayStruct->History.Reset();
+	}
+#endif // ENABLE_GC_HISTORY
+}
+
 double GetLastGCTime()
 {
 	return GLastGCTime;
diff --git a/Engine/Source/Runtime/CoreUObject/Private/UObject/GarbageCollectionHistory.cpp b/Engine/Source/Runtime/CoreUObject/Private/UObject/GarbageCollectionHistory.cpp
new file mode 100644
--- /dev/null
+++ b/Engine/Source/Runtime/CoreUObject/Private/UObject/GarbageCollectionHistory.cpp
@@ -0,0 +1,191 @@
+// Copyright Epic Games, Inc. All Rights Reserved.
+
+/*=============================================================================
+	GarbageCollectionHistory.cpp: Unreal object garbage collection history code.
+=============================================================================*/
+
+#include "UObject/GarbageCollectionHistory.h"
+#include "HAL/ThreadSafeBool.h"
+#include "Misc/TimeGuard.h"
+#include "HAL/IConsoleManager.h"
+#include "Misc/App.h"
+#include "UObject/UObjectAllocator.h"
+#include "UObject/UObjectBase.h"
+#include "UObject/Object.h"
+#include "UObject/UnrealType.h"
+#include "UObject/GCObject.h"
+#include "UObject/GCScopeLock.h"
+#include "HAL/ExceptionHandling.h"
+
+
+/*-----------------------------------------------------------------------------
+   Garbage collection history code.
+-----------------------------------------------------------------------------*/
+
+#if ENABLE_GC_HISTORY
+
+static FAutoConsoleCommand CmdCalculateGCHistorySize(
+	TEXT("gc.CalculateHistorySize"),
+	TEXT(""),
+	FConsoleCommandWithArgsDelegate::CreateLambda([](const TArray<FString>& Args)
+		{
+			int64 TotalSize = FGCHistory::Get().GetAllocatedSize();
+			UE_LOG(LogGarbage, Display, TEXT("Memory allocated for GC history: %lldb"), TotalSize);
+		})
+);
+
+static FAutoConsoleCommand CmdSetHistorySize(
+	TEXT("gc.HistorySize"),
+	TEXT(""),
+	FConsoleCommandWithArgsDelegate::CreateLambda([](const TArray<FString>& Args)
+		{
+			if (Args.Num())
+			{
+				int32 HistorySize = FMath::Max(0, FCString::Atoi(*Args[0]));
+				FGCHistory::Get().SetHistorySize(HistorySize);
+				UE_LOG(LogGarbage, Display, TEXT("GC history size changed to: %d"), HistorySize);
+			}
+			else
+			{
+				UE_LOG(LogGarbage, Display, TEXT("GC history size currently set to: %d"), FGCHistory::Get().GetHistorySize());
+			}			
+		})
+);
+
+FGCHistory::~FGCHistory()
+{
+	Cleanup();
+}
+
+FGCHistory& FGCHistory::Get()
+{
+	static FGCHistory Singleton;
+	return Singleton;
+}
+
+void FGCHistory::Cleanup(FGCSnapshot& InSnapshot)
+{
+	for (TPair<FGCObjectInfo*, TArray<FGCDirectReferenceInfo>*>& DirectReferenceInfos : InSnapshot.DirectReferences)
+	{
+		delete DirectReferenceInfos.Value;
+	}
+	InSnapshot.DirectReferences.Reset();
+	for (TPair<UObject*, FGCObjectInfo*>& ObjectToInfoPair : InSnapshot.ObjectToInfoMap)
+	{
+		delete ObjectToInfoPair.Value;
+	}
+	InSnapshot.ObjectToInfoMap.Reset();
+}
+
+void FGCHistory::Cleanup()
+{
+	for (FGCSnapshot& Snapshot : Snapshots)
+	{
+		Cleanup(Snapshot);
+	}
+	Snapshots.Empty();
+	MostRecentSnapshotIndex = -1;	
+}
+
+void FGCHistory::SetHistorySize(int32 HistorySize)
+{
+	if (HistorySize != Snapshots.Num())
+	{
+		Cleanup();
+		
+		if (HistorySize > 0)
+		{
+			Snapshots.AddDefaulted(HistorySize);
+		}		
+	}
+}
+
+void FGCHistory::MergeArrayStructHistory(TMap<UObject*, TArray<FGCDirectReference>*>& History, FGCSnapshot& Snapshot)
+{
+	for (TPair<UObject*, TArray<FGCDirectReference>*>& ReferencePair : History)
+	{
+		FGCObjectInfo* ReferencerInfo = FGCObjectInfo::FindOrAddInfoHelper(ReferencePair.Key, Snapshot.ObjectToInfoMap);
+		TArray<FGCDirectReferenceInfo>*& DirectReferencesInfoArray = Snapshot.DirectReferences.FindOrAdd(ReferencerInfo);
+		if (!DirectReferencesInfoArray)
+		{
+			DirectReferencesInfoArray = new TArray<FGCDirectReferenceInfo>();
+			DirectReferencesInfoArray->Reserve(ReferencePair.Value->Num());
+		}
+		for (FGCDirectReference& DirectReference : *ReferencePair.Value)
+		{
+			FGCObjectInfo* ReferencedObjectInfo = FGCObjectInfo::FindOrAddInfoHelper(DirectReference.ReferencedObject, Snapshot.ObjectToInfoMap);
+			DirectReferencesInfoArray->Add(FGCDirectReferenceInfo(DirectReference.ReferencerName, ReferencedObjectInfo));
+		}
+	}
+}
+
+void FGCHistory::Update(TArray<FGCArrayStruct*>& AllArrays)
+{
+	if (IsActive())
+	{
+		// Advance to the next snapshot in the ring buffer
+		MostRecentSnapshotIndex = (MostRecentSnapshotIndex + 1) % Snapshots.Num();
+		FGCSnapshot& CurrentSnapshot = Snapshots[MostRecentSnapshotIndex];
+		if (CurrentSnapshot.DirectReferences.Num())
+		{
+			// Already taken so cleanup to be able to overwrite
+			Cleanup(CurrentSnapshot);
+		}
+
+		// Copy information gather in all of the GC struct arrays
+		for (FGCArrayStruct* ArrayStruct : AllArrays)
+		{
+			MergeArrayStructHistory(ArrayStruct->History, CurrentSnapshot);
+		}
+	}
+}
+
+FGCSnapshot* FGCHistory::GetSnapshot(int32 HistoryLevel)
+{
+	// Negative values make sense as input and are also accepted
+	HistoryLevel = FMath::Abs(HistoryLevel);
+	if (HistoryLevel >= Snapshots.Num())
+	{
+		// Clamp to the max available history size
+		HistoryLevel = (Snapshots.Num() - 1);
+	}
+	int32 SnapshotIndex = MostRecentSnapshotIndex - HistoryLevel;
+	if (SnapshotIndex < 0)
+	{
+		// Wrap around the ring buffer
+		SnapshotIndex += Snapshots.Num();
+	}
+	if (SnapshotIndex >= 0 && SnapshotIndex < Snapshots.Num() && Snapshots[SnapshotIndex].ObjectToInfoMap.Num())
+	{
+		return &Snapshots[SnapshotIndex];
+	}
+	else
+	{
+		// This means that the requested history level hasn't been filled yet or history recording is disabled
+		return nullptr;
+	}
+}
+
+int64 FGCHistory::GetAllocatedSize() const
+{
+	int64 TotalSize = Snapshots.GetAllocatedSize();
+	for (const FGCSnapshot& Snapshot : Snapshots)
+	{
+		TotalSize += Snapshot.GetAllocatedSize();
+	}
+	return TotalSize;
+}
+
+int64 FGCSnapshot::GetAllocatedSize() const
+{
+	int64 TotalSize = ObjectToInfoMap.GetAllocatedSize() + ObjectToInfoMap.Num() * sizeof(FGCObjectInfo);
+
+	TotalSize += DirectReferences.GetAllocatedSize();
+	for (const TPair<FGCObjectInfo*, TArray<FGCDirectReferenceInfo>*>& ReferencePair : DirectReferences)
+	{
+		TotalSize += ReferencePair.Value->GetAllocatedSize();
+	}
+	return TotalSize;
+}
+
+#endif // ENABLE_GC_HISTORY
diff --git a/Engine/Source/Runtime/CoreUObject/Private/UObject/Obj.cpp b/Engine/Source/Runtime/CoreUObject/Private/UObject/Obj.cpp
--- a/Engine/Source/Runtime/CoreUObject/Private/UObject/Obj.cpp
+++ b/Engine/Source/Runtime/CoreUObject/Private/UObject/Obj.cpp
@@ -22,7 +22,7 @@
 #include "UObject/UObjectHash.h"
 #include "UObject/Object.h"
 #include "Serialization/ArchiveUObject.h"
-#include "UObject/GarbageCollection.h"
+#include "UObject/GarbageCollectionHistory.h"
 #include "UObject/Class.h"
 #include "UObject/EnumProperty.h"
 #include "UObject/UObjectIterator.h"
@@ -4158,6 +4158,7 @@ bool StaticExec( UWorld* InWorld, const TCHAR* Cmd, FOutputDevice& Ar )
 			{
 				
 				EReferenceChainSearchMode SearchModeFlags = EReferenceChainSearchMode::PrintResults;
+				int32 HistoryLevel = 0;
 
 				FString Tok;
 				while(FParse::Token(Str, Tok, false))
@@ -4194,9 +4195,46 @@ bool StaticExec( UWorld* InWorld, const TCHAR* Cmd, FOutputDevice& Ar )
 					{
 						SearchModeFlags |= EReferenceChainSearchMode::FullChain;
 					}
+#if ENABLE_GC_HISTORY
+					else if (FParse::Value(Str, TEXT("history="), HistoryLevel))
+					{
+						if (FGCHistory::Get().IsActive())
+						{
+							int32 MaxHistoryLevel = FGCHistory::Get().GetHistorySize() - 1;
+							if (FMath::Abs(HistoryLevel) >= MaxHistoryLevel)
+							{
+								UE_LOG(LogObj, Log, TEXT("GC History level %d will be clamped to the current max %d"), HistoryLevel, MaxHistoryLevel);
+								HistoryLevel = MaxHistoryLevel;
+							}
+						}
+						else if (HistoryLevel != 0)
+						{
+							UE_LOG(LogObj, Log, TEXT("GC History is not currently enabled. Use 'gc.HistorySize 1' console command to enable it."));
+							HistoryLevel = 0;
+						}
+					}
+#endif // ENABLE_GC_HISTORY
 				}
 				
-				FReferenceChainSearch RefChainSearch(Object, SearchModeFlags);
+				if (HistoryLevel == 0)
+				{
+					FReferenceChainSearch RefChainSearch(Object, SearchModeFlags);
+				}
+#if ENABLE_GC_HISTORY
+				else
+				{
+					FReferenceChainSearch HistorySearch(SearchModeFlags);
+					FGCSnapshot* GCSnapshot = FGCHistory::Get().GetSnapshot(HistoryLevel);
+					if (GCSnapshot)
+					{
+						HistorySearch.PerformSearchFromGCSnapshot(Object, *GCSnapshot);
+					}
+					else
+					{
+						UE_LOG(LogObj, Log, TEXT("There's not been that many (%d) GC runs to be able to search for object refs in GC history."), (HistoryLevel + 1));
+					}
+				}
+#endif // ENABLE_GC_HISTORY
 			}
 			else
 			{
diff --git a/Engine/Source/Runtime/CoreUObject/Private/UObject/ReferenceChainSearch.cpp b/Engine/Source/Runtime/CoreUObject/Private/UObject/ReferenceChainSearch.cpp
--- a/Engine/Source/Runtime/CoreUObject/Private/UObject/ReferenceChainSearch.cpp
+++ b/Engine/Source/Runtime/CoreUObject/Private/UObject/ReferenceChainSearch.cpp
@@ -565,6 +565,10 @@ FReferenceChainSearch::FReferenceChainSearch(UObject* InObjectToFindReferencesTo
 	}
 }
 
+FReferenceChainSearch::FReferenceChainSearch(EReferenceChainSearchMode Mode)
+	: SearchMode(Mode)
+{
+}
 
 FReferenceChainSearch::~FReferenceChainSearch()
 {
@@ -587,6 +591,54 @@ void FReferenceChainSearch::PerformSearch()
 	}
 }
 
+#if ENABLE_GC_HISTORY
+void FReferenceChainSearch::PerformSearchFromGCSnapshot(UObject* InObjectToFindReferencesTo, FGCSnapshot& InSnapshot)
+{
+	FSlowHeartBeatScope DisableHangDetection; // This function can be very slow
+
+	Cleanup();
+
+	// Temporarily move the generated object info structs. We don't want to copy everything to minimize memory usage and save a few ms
+	ObjectToInfoMap = MoveTemp(InSnapshot.ObjectToInfoMap);
+
+	ObjectToFindReferencesTo = InObjectToFindReferencesTo;
+	ObjectInfoToFindReferencesTo = FGCObjectInfo::FindOrAddInfoHelper(ObjectToFindReferencesTo, ObjectToInfoMap);
+	FGCObjectInfo* GCObjectReferencerInfo = FGCObjectInfo::FindOrAddInfoHelper(FGCObject::GGCObjectReferencer, ObjectToInfoMap);
+
+	// We can avoid copying object infos but we need to regenerate direct reference infos
+	for (TPair<FGCObjectInfo*, TArray<FGCDirectReferenceInfo>*>& DirectReferencesInfo : InSnapshot.DirectReferences)
+	{
+		FGraphNode* ObjectNode = FindOrAddNode(DirectReferencesInfo.Key);
+		for (FGCDirectReferenceInfo& ReferenceInfo : *DirectReferencesInfo.Value)
+		{
+			FGraphNode* ReferencedObjectNode = FindOrAddNode(ReferenceInfo.ReferencedObjectInfo);
+			EReferenceType ReferenceType = EReferenceType::Unknown;
+			if (GCObjectReferencerInfo == DirectReferencesInfo.Key || ReferenceInfo.ReferencerName == NAME_None)
+			{
+				ReferenceType = EReferenceType::AddReferencedObjects;
+			}
+			else
+			{
+				ReferenceType = EReferenceType::Property;
+			}
+			ObjectNode->ReferencedObjects.Add(FNodeReferenceInfo(ReferencedObjectNode, ReferenceType, ReferenceInfo.ReferencerName, nullptr, 0));
+			ReferencedObjectNode->ReferencedByObjects.Add(ObjectNode);
+		}
+	}
+
+	// Second pass creates all reference chains
+	PerformSearch();
+
+	if (!!(SearchMode & (EReferenceChainSearchMode::PrintResults | EReferenceChainSearchMode::PrintAllResults)))
+	{
+		PrintResults(!!(SearchMode & EReferenceChainSearchMode::PrintAllResults));
+	}
+
+	// Return the object info struct back to the snapshot
+	InSnapshot.ObjectToInfoMap = MoveTemp(ObjectToInfoMap);
+}
+#endif // ENABLE_GC_HISTORY
+
 void FReferenceChainSearch::FindDirectReferencesForObjects()
 {
 	TSet<FObjectReferenceInfo> ReferencedObjects;
diff --git a/Engine/Source/Runtime/CoreUObject/Public/UObject/FastReferenceCollector.h b/Engine/Source/Runtime/CoreUObject/Public/UObject/FastReferenceCollector.h
--- a/Engine/Source/Runtime/CoreUObject/Public/UObject/FastReferenceCollector.h
+++ b/Engine/Source/Runtime/CoreUObject/Public/UObject/FastReferenceCollector.h
@@ -238,6 +238,12 @@ public:
 	 */
 	void DumpGarbageReferencers(TArray<FGCArrayStruct*>& AllArrays);
 
+	/**
+	 * Updates GC history after the last GC run.
+	 * @param AllArrays Arrays to process
+	 */
+	void UpdateGCHistory(TArray<FGCArrayStruct*>& AllArrays);
+
 	/**
 	 * Grabs all arrays from the pool 
 	 */
diff --git a/Engine/Source/Runtime/CoreUObject/Public/UObject/GarbageCollection.h b/Engine/Source/Runtime/CoreUObject/Public/UObject/GarbageCollection.h
--- a/Engine/Source/Runtime/CoreUObject/Public/UObject/GarbageCollection.h
+++ b/Engine/Source/Runtime/CoreUObject/Public/UObject/GarbageCollection.h
@@ -29,6 +29,8 @@
 /** Token debug info (token names) enabled in non-shipping builds */
 #define ENABLE_GC_TOKEN_DEBUG_INFO (!UE_BUILD_SHIPPING)
 
+#define ENABLE_GC_HISTORY (!UE_BUILD_SHIPPING)
+
 COREUOBJECT_API DECLARE_LOG_CATEGORY_EXTERN(LogGarbage, Warning, All);
 DECLARE_STATS_GROUP(TEXT("Garbage Collection"), STATGROUP_GC, STATCAT_Advanced);
 
@@ -567,6 +569,17 @@ struct FGarbageReferenceInfo
 	FString GetReferencingObjectInfo() const;
 };
 
+struct FGCDirectReference
+{
+	FGCDirectReference() = default;
+	explicit FGCDirectReference(UObject* Obj)
+		: ReferencedObject(Obj)
+	{}
+	/** Property or FGCObject name referencing this object */
+	FName ReferencerName;
+	UObject* ReferencedObject = nullptr;
+};
+
 /** Struct to hold the objects to serialize array and the list of weak references. This is allocated by ArrayPool */
 struct FGCArrayStruct
 {
@@ -577,6 +590,9 @@ struct FGCArrayStruct
 	TArray<UObject*> ObjectsToSerialize;
 	TArray<UObject**> WeakReferences;
 	TArray<FGarbageReferenceInfo> GarbageReferences;
+#if ENABLE_GC_HISTORY
+	TMap<UObject*, TArray<FGCDirectReference>*> History;
+#endif
 
 	FORCEINLINE UObject* GetReferencingObject()
 	{
diff --git a/Engine/Source/Runtime/CoreUObject/Public/UObject/GarbageCollectionHistory.h b/Engine/Source/Runtime/CoreUObject/Public/UObject/GarbageCollectionHistory.h
new file mode 100644
--- /dev/null
+++ b/Engine/Source/Runtime/CoreUObject/Public/UObject/GarbageCollectionHistory.h
@@ -0,0 +1,104 @@
+// Copyright Epic Games, Inc. All Rights Reserved.
+
+/*=============================================================================
+	GarbageCollectionHistory.h: Unreal realtime garbage collection history
+=============================================================================*/
+
+#pragma once
+
+#include "UObject/GCObjectInfo.h"
+
+#if ENABLE_GC_HISTORY
+
+/**
+ * Structure that holds information about a direct reference from an object
+ **/
+struct FGCDirectReferenceInfo
+{
+	FGCDirectReferenceInfo() = default;
+	explicit FGCDirectReferenceInfo(FName InReferencerName, FGCObjectInfo* Obj)
+		: ReferencerName(InReferencerName)
+		, ReferencedObjectInfo(Obj)
+	{}
+	/** Property or FGCObject name referencing this object */
+	FName ReferencerName;
+	/** Info struct of the object being referenced */
+	FGCObjectInfo* ReferencedObjectInfo = nullptr;
+};
+
+/**
+ * Structure that holds all direct references traversed in a GC run as well as FGCObjectInfo structs created for all participating objects
+ **/
+struct FGCSnapshot
+{
+	/** FGCObjectInfo structs generated for each of the objects encountered during GC */
+	TMap<UObject*, FGCObjectInfo*> ObjectToInfoMap;
+	/** Lis of direct references for all objects */
+	TMap<FGCObjectInfo*, TArray<FGCDirectReferenceInfo>*> DirectReferences;
+
+	/** Returns the number of bytes allocated by a single snapshot */
+	int64 GetAllocatedSize() const;
+};
+
+/**
+ * Garbage Collector History. Holds snapshots of a number of previous GC runs.
+ **/
+class COREUOBJECT_API FGCHistory
+{
+	/** Captured snapshots of previous GC runs. This is essentially a preallocated ring buffer. */
+	TArray<FGCSnapshot> Snapshots;
+	/** Index of the last recorded snapshot in the Snapshots array */
+	int32 MostRecentSnapshotIndex = -1;
+
+	~FGCHistory();
+
+	void Cleanup(FGCSnapshot& InShapshot);	
+	void MergeArrayStructHistory(TMap<UObject*, TArray<FGCDirectReference>*>& History, FGCSnapshot& Snapshot);
+
+public:
+
+	/** Gets the FGCHistory singleton. */
+	static FGCHistory& Get();
+
+	/** Copies information generated by Garbage Collector into the next snapshot struct. */
+	void Update(TArray<FGCArrayStruct*>& AllArrays);
+
+	/** Frees memory allocated by GC history. */
+	void Cleanup();
+
+	/** 
+	 * Sets how many GC runs the history can record. 
+	 * @param HistorySize The number of GC runs the history can record. If 0 GC history will be disabled.
+	 */
+	void SetHistorySize(int32 HistorySize);
+
+	FORCEINLINE int32 GetHistorySize() const
+	{
+		return Snapshots.Num();
+	}
+
+	FORCEINLINE bool IsActive() const
+	{
+		return GetHistorySize() > 0;
+	}
+
+	/**
+	 * Gets a snapshot of a previous GC run.
+	 * @param HistoryLevel Index of a previous run: 0 - the last one, 1 - the one before and up to (GetHistorySize() - 1). Negative values are also accepted.
+	 * @returns Pointer to a prevous GC run snapshot or null if there's not been that many recorded GC runs yet
+	 */
+	FGCSnapshot* GetSnapshot(int32 HistoryLevel);
+
+	/**
+	 * Returns the snapshot of the last GC run. See GetSnapshot(int32)
+	 **/
+	FGCSnapshot* GetLastSnapshot()
+	{
+		return GetSnapshot(0);
+	}
+
+	/** Returns the number of bytes allocated by GC history */
+	int64 GetAllocatedSize() const;
+};
+
+#endif // ENABLE_GC_HISTORY
\ No newline at end of file
diff --git a/Engine/Source/Runtime/CoreUObject/Public/UObject/ReferenceChainSearch.h b/Engine/Source/Runtime/CoreUObject/Public/UObject/ReferenceChainSearch.h
--- a/Engine/Source/Runtime/CoreUObject/Public/UObject/ReferenceChainSearch.h
+++ b/Engine/Source/Runtime/CoreUObject/Public/UObject/ReferenceChainSearch.h
@@ -4,7 +4,7 @@
 
 #include "CoreMinimal.h"
 #include "UObject/UObjectGlobals.h"
-#include "UObject/GCObjectInfo.h"
+#include "UObject/GarbageCollectionHistory.h"
 #include "HAL/ThreadHeartBeat.h"
 
 /** Search mode flags */
@@ -233,11 +233,19 @@ public:
 	};
 
 	/** Constructs a new search engine and finds references to the specified object */
-	COREUOBJECT_API FReferenceChainSearch(UObject* InObjectToFindReferencesTo, EReferenceChainSearchMode Mode = EReferenceChainSearchMode::PrintResults);
+	COREUOBJECT_API explicit FReferenceChainSearch(UObject* InObjectToFindReferencesTo, EReferenceChainSearchMode Mode = EReferenceChainSearchMode::PrintResults);
+
+	/** Constructs a new search engine but does not find references to any objects until one of the PerformSearch*() functions is called */
+	COREUOBJECT_API explicit FReferenceChainSearch(EReferenceChainSearchMode Mode);
 
 	/** Destructor */
 	COREUOBJECT_API ~FReferenceChainSearch();
 
+#if ENABLE_GC_HISTORY
+	/** Searches for references in a previous GC run snapshot temporarily acquiring its object info */
+	COREUOBJECT_API void PerformSearchFromGCSnapshot(UObject* InObjectToFindReferencesTo, FGCSnapshot& InSnapshot);
+#endif //ENABLE_GC_HISTORY
+
 	/**
 	 * Dumps results to log
 	 * @param bDumpAllChains - if set to false, the output will be trimmed to the first 100 reference chains
diff --git a/Engine/Source/Runtime/Engine/Private/UnrealEngine.cpp b/Engine/Source/Runtime/Engine/Private/UnrealEngine.cpp
--- a/Engine/Source/Runtime/Engine/Private/UnrealEngine.cpp
+++ b/Engine/Source/Runtime/Engine/Private/UnrealEngine.cpp
@@ -14924,59 +14924,118 @@ void UEngine::VerifyLoadMapWorldCleanup()
 	}
 }
 
+static void PrintStaleReferenceChainsAndFindReferencingObjects(UObject* ObjectToFindReferencesTo, FReferenceChainSearch& RefChainSearch, FGCObjectInfo*& OutGarbageObject, FGCObjectInfo*& OutReferencingObject)
+{
+	RefChainSearch.PrintResults([&ObjectToFindReferencesTo, &OutGarbageObject, &OutReferencingObject](FReferenceChainSearch::FCallbackParams& Params)
+		{
+			check(Params.Object);
+			if (!Params.Object->IsValid())
+			{
+				// We may find many chains that lead to the leak but for brevity we only report the first one in the fatal error below
+				if (!OutGarbageObject)
+				{
+					OutGarbageObject = Params.Object;
+					OutReferencingObject = Params.Referencer;
+				}
+				// We only really care about a reference to the first garbage / pending kill object in the chain
+				// as the reference that holds it is the one that prevents the world from being GC'd
+				return false;
+			}
+			else
+			{
+				UObject* ResolvedObject = Params.Object->TryResolveObject();
+				if (ResolvedObject != ObjectToFindReferencesTo)
+				{
+					// Keep track of the last referencing object, we may need it if there's no garbage referencers
+					OutReferencingObject = Params.Object;
+				}
+				return true;
+			}
+		}, false);
+}
+
+static FString GetPathToStaleObjectReferencer(UObject* ObjectToFindReferencesTo, FReferenceChainSearch& RefChainSearch)
+{
+	bool bFirstGarbageObjectFound = false;
+	FString PathToCulprit = RefChainSearch.GetRootPath([&bFirstGarbageObjectFound, ObjectToFindReferencesTo](FReferenceChainSearch::FCallbackParams& Params)
+		{
+			check(Params.Object);
+			// Mark the first garbage reference or the target object reference as the culprit
+			if ((!Params.Object->IsValid() || Params.Object->TryResolveObject() == ObjectToFindReferencesTo) && !bFirstGarbageObjectFound)
+			{
+				bFirstGarbageObjectFound = true;
+				check(Params.Out);
+				Params.Out->Logf(ELogVerbosity::Log, TEXT("%s    ^^^^^^^^^^^^^^^^^^^^^^^^^^^^^^^^^^^^^^^^^^^^^^^^^^^^^^^^^^^^^^"), FCString::Spc(Params.Indent));
+				Params.Out->Logf(ELogVerbosity::Log, TEXT("%s    ^ This reference is preventing the old %s from being GC'd ^"), FCString::Spc(Params.Indent), *ObjectToFindReferencesTo->GetClass()->GetName());
+			}
+			return true;
+		});
+	return PathToCulprit;
+}
+
 void UEngine::FindAndPrintStaleReferencesToObject(UObject* ObjectToFindReferencesTo, ELogVerbosity::Type Verbosity)
 {
 	check(ObjectToFindReferencesTo);
 
+	UE_LOG(LogLoad, Log, TEXT("Looking for existing references to %s..."), *ObjectToFindReferencesTo->GetFullName());
+
+	bool bReferenceChainFound = false;
 	FGCObjectInfo* OutGarbageObject = nullptr;
 	FGCObjectInfo* OutReferencingObject = nullptr;
 	FReferenceChainSearch RefChainSearch(ObjectToFindReferencesTo, EReferenceChainSearchMode::Default);
-	RefChainSearch.PrintResults([&OutGarbageObject, &OutReferencingObject](FReferenceChainSearch::FCallbackParams& Params)
+#if ENABLE_GC_HISTORY
+	FReferenceChainSearch HistorySearch(EReferenceChainSearchMode::Default); // we may not need it but it has to live in this scope so that FGCObjectInfos are being kept alive
+#endif
+
+	PrintStaleReferenceChainsAndFindReferencingObjects(ObjectToFindReferencesTo, RefChainSearch, OutGarbageObject, OutReferencingObject);
+
+	FString PathToCulprit;
+	FString GarbageErrorMessage;
+	if (RefChainSearch.GetReferenceChains().Num())
 	{
-		check(Params.Object);
-		if (!Params.Object->IsValid())
+		bReferenceChainFound = true;
+		PathToCulprit = GetPathToStaleObjectReferencer(ObjectToFindReferencesTo, RefChainSearch);
+		FString GarbageObjectName = OutGarbageObject ? OutGarbageObject->GetFullName() : ObjectToFindReferencesTo->GetFullName();
+		checkf(OutReferencingObject, TEXT("No object referencing %s found even though we have a valid reference chain"), *ObjectToFindReferencesTo->GetPathName());
+		GarbageErrorMessage = FString::Printf(TEXT("Object %s is being referenced by %s"), *GarbageObjectName, *OutReferencingObject->GetFullName());
+	}
+#if ENABLE_GC_HISTORY
+	else
+	{
+		FGCSnapshot* LastGCSnapshot = FGCHistory::Get().GetLastSnapshot();
+		if (LastGCSnapshot)
 		{
-			// We may find many chains that lead to the leak but for brevity we only report the first one in the fatal error below
-			if (!OutGarbageObject)
+			UE_LOG(LogLoad, Log, TEXT("Looking for references to %s in the last GC run history..."), *ObjectToFindReferencesTo->GetFullName());
+
+			HistorySearch.PerformSearchFromGCSnapshot(ObjectToFindReferencesTo, *LastGCSnapshot);
+
+			OutGarbageObject = nullptr;
+			OutReferencingObject = nullptr;
+			PrintStaleReferenceChainsAndFindReferencingObjects(ObjectToFindReferencesTo, HistorySearch, OutGarbageObject, OutReferencingObject);
+			if (HistorySearch.GetReferenceChains().Num())
 			{
-				OutGarbageObject = Params.Object;
-				OutReferencingObject = Params.Referencer;
+				bReferenceChainFound = true;
+				PathToCulprit = GetPathToStaleObjectReferencer(ObjectToFindReferencesTo, RefChainSearch);
+				FString GarbageObjectName = OutGarbageObject ? OutGarbageObject->GetFullName() : ObjectToFindReferencesTo->GetFullName();
+				checkf(OutReferencingObject, TEXT("No object referencing %s found even though we have a valid reference chain"), *ObjectToFindReferencesTo->GetPathName());
+				GarbageErrorMessage = FString::Printf(TEXT("Garbage object %s was previously being referenced by %s"), *GarbageObjectName, *OutReferencingObject->GetFullName());
 			}
-			// We only really care about a reference to the first garbage / pending kill object in the chain
-			// as the reference that holds it is the one that prevents the world from being GC'd
-			return false;
 		}
-		else
+	}
+#endif // ENABLE_GC_HISTORY
+
+	if (!bReferenceChainFound)
+	{
+		if (OutGarbageObject)
 		{
-			return true;
+			GarbageErrorMessage = FString::Printf(TEXT("Garbage object %s is not referenced so it may have a flag set that's preventing it from being destroyed (see log for details)"), *OutGarbageObject->GetFullName());
 		}
-	}, false);
-	bool bFirstGarbageObjectFound = false;
-	FString PathToCulprit = RefChainSearch.GetRootPath([&bFirstGarbageObjectFound, ObjectToFindReferencesTo](FReferenceChainSearch::FCallbackParams& Params)
-	{
-		check(Params.Object);
-		if (!Params.Object->IsValid() && !bFirstGarbageObjectFound)
+		else
 		{
-			bFirstGarbageObjectFound = true;
-			check(Params.Out);
-			Params.Out->Logf(ELogVerbosity::Log, TEXT("%s    ^^^^^^^^^^^^^^^^^^^^^^^^^^^^^^^^^^^^^^^^^^^^^^^^^^^^^^^^^^^^^^"), FCString::Spc(Params.Indent));
-			Params.Out->Logf(ELogVerbosity::Log, TEXT("%s    ^ This reference is preventing the old %s from being GC'd ^"), FCString::Spc(Params.Indent), *ObjectToFindReferencesTo->GetClass()->GetName());
+			GarbageErrorMessage = TEXT("However it's not referenced by any object. It may have a flag set that's preventing it from being destroyed (see log for details)");
 		}
-		return true;
-	});
-	FString GarbageErrorMessage;
-	if (OutGarbageObject && OutReferencingObject)
-	{
-		GarbageErrorMessage = FString::Printf(TEXT("Garbage object %s is being referenced by %s"), *OutGarbageObject->GetFullName(), *OutReferencingObject->GetFullName());
-	}
-	else if (OutGarbageObject)
-	{
-		GarbageErrorMessage = FString::Printf(TEXT("Garbage object %s is not referenced so it may have a flag set that's preventing it from being destroyed (see log for details)"), *OutGarbageObject->GetFullName());
-	}
-	else
-	{
-		GarbageErrorMessage = TEXT("However it's not referenced by any garbage object. It may have a flag set that's preventing it from being destroyed (see log for details)");
 	}
+
 	if (Verbosity == ELogVerbosity::Fatal)
 	{
 		UE_LOG(LogLoad, Fatal, TEXT("Old %s not cleaned up by GC! %s:") LINE_TERMINATOR TEXT("%s"),
