ISM + WPO improvements.

Add r.RayTracing.Geometry.InstancedStaticMeshes.SimulationClusterRadius to bucket and sort ISM + WPO simulated instances based on distance to view
Refactor cvar code to fix/enable r.RayTracing.Geometry.StaticMeshes.WPO = 2
Fix bugs with current ISM + WPO if the local transform uses a non-identity matrix or if the simulated instance transform is non-identity

#PR 8807  [REDACTED_HANDLE] at Nvidia
#8807

diff --git a/Engine/Shaders/Private/RayTracing/RayTracingDynamicMesh.usf b/Engine/Shaders/Private/RayTracing/RayTracingDynamicMesh.usf
--- a/Engine/Shaders/Private/RayTracing/RayTracingDynamicMesh.usf
+++ b/Engine/Shaders/Private/RayTracing/RayTracingDynamicMesh.usf
@@ -52,8 +52,9 @@ uint NumVertices;
 uint MinVertexIndex;
 uint PrimitiveId;
 uint OutputVertexBaseIndex;
-float4x4 InstanceTransform;
-float4x4 InverseTransform;
+int InstanceId;
+float4x4 WorldToInstance;
+
 RWBuffer<float> RWVertexPositions;
 int bApplyWorldPositionOffset;
 
@@ -72,14 +73,14 @@ void RayTracingDynamicGeometryConverterCS(uint3 DispatchThreadId : SV_DispatchTh
 	if (DispatchThreadId.x < NumActualVertices)
 	{
 		FVertexFactoryInput Input = LoadVertexFactoryInputForDynamicUpdate(VertexIndex / 3, VertexIndex % 3, PrimitiveId);
-		FVertexFactoryIntermediates VFIntermediates = GetVertexFactoryIntermediates(Input);
+
 #if USE_INSTANCING
-		// Instance transform supplied explicitly
-		VFIntermediates.InstanceOrigin = InstanceTransform[3];
-		VFIntermediates.InstanceTransform1 = InstanceTransform[0];
-		VFIntermediates.InstanceTransform2 = InstanceTransform[1];
-		VFIntermediates.InstanceTransform3 = InstanceTransform[2];
+		Input.InstanceId = InstanceId;
+#elif USE_INSTANCE_CULLING
+		Input.DrawInstanceId = InstanceId;
 #endif
+		FVertexFactoryIntermediates VFIntermediates = GetVertexFactoryIntermediates(Input);
+
 		float4 WorldPosition = VertexFactoryGetWorldPosition(Input, VFIntermediates);
 		float3x3 TangentToLocal = VertexFactoryGetTangentToLocal(Input, VFIntermediates);
 		FMaterialVertexParameters VertexParameters = GetMaterialVertexParameters(Input, VFIntermediates, WorldPosition.xyz, TangentToLocal);
@@ -93,9 +94,9 @@ void RayTracingDynamicGeometryConverterCS(uint3 DispatchThreadId : SV_DispatchTh
 			WorldPosition.xyz += MakeFinite(GetMaterialWorldPositionOffset(VertexParameters));
 		}
 
-#if USE_INSTANCING
+#if USE_INSTANCING || USE_INSTANCE_CULLING
 		//Reverse transform to neutral space
-		WorldPosition = mul(WorldPosition, InverseTransform);
+		WorldPosition = mul(WorldPosition, WorldToInstance);
 #elif RAY_TRACING_DYNAMIC_MESH_IN_LOCAL_SPACE
 		// Move the point back into local space because the RT instance will be placed there
 		// This the recommended default behavior so that the transform applied by the VertexFactory and the RT transform are always the same
diff --git a/Engine/Source/Runtime/Engine/Private/InstancedStaticMesh.cpp b/Engine/Source/Runtime/Engine/Private/InstancedStaticMesh.cpp
--- a/Engine/Source/Runtime/Engine/Private/InstancedStaticMesh.cpp
+++ b/Engine/Source/Runtime/Engine/Private/InstancedStaticMesh.cpp
@@ -60,7 +60,7 @@ IMPLEMENT_TYPE_LAYOUT(FInstancedStaticMeshVertexFactoryShaderParameters);
 CSV_DEFINE_CATEGORY_MODULE(ENGINE_API, InstancedStaticMeshComponent, true);
 
 const int32 InstancedStaticMeshMaxTexCoord = 8;
-static const int32 MaxSimulatedInstances = 256;
+static const uint32 MaxSimulatedInstances = 256;
 
 IMPLEMENT_HIT_PROXY(HInstancedStaticMeshInstance, HHitProxy);
 
@@ -118,6 +118,11 @@ static TAutoConsoleVariable<int32> CVarRayTracingInstancesEvaluateWPO(
 	TEXT(" -1 - on only for meshes with evaluate WPO enabled")
 );
 
+static TAutoConsoleVariable<float> CVarRayTracingInstancesSimulationClusterRadius(
+	TEXT("r.RayTracing.Geometry.InstancedStaticMeshes.SimulationClusterRadius"),
+	500.0f, // 5 m
+	TEXT("Bucket instances based on distance to camera for simulating WPO (default = 500 (5m), disable if <= 0)"));
+
 static TAutoConsoleVariable<int32> CVarRayTracingSimulatedInstanceCount(
 	TEXT("r.RayTracing.Geometry.InstancedStaticMeshes.SimulationCount"),
 	1,
@@ -1144,7 +1149,7 @@ void FPerInstanceRenderData::UpdateBoundsTransforms_Concurrent()
 
 void FPerInstanceRenderData::UpdateBoundsTransforms()
 {
-	const int32 InstanceCount = InstanceBuffer.GetNumInstances();
+	const uint32 InstanceCount = InstanceBuffer.GetNumInstances();
 	PerInstanceTransforms.Empty(InstanceCount);
 
 	if (bTrackBounds)
@@ -1152,7 +1157,7 @@ void FPerInstanceRenderData::UpdateBoundsTransforms()
 		FBoxSphereBounds LocalBounds = FBoxSphereBounds(InstanceLocalBounds);
 		PerInstanceBounds.Empty(InstanceCount);
 
-		for (int32 InstanceIndex = 0; InstanceIndex < InstanceCount; ++InstanceIndex)
+		for (uint32 InstanceIndex = 0; InstanceIndex < InstanceCount; ++InstanceIndex)
 		{
 			if (!InstanceBuffer.GetInstanceData() || !InstanceBuffer.GetInstanceData()->IsValidIndex(InstanceIndex))
 		    {
@@ -1169,7 +1174,7 @@ void FPerInstanceRenderData::UpdateBoundsTransforms()
 	}
 	else
 	{
-		for (int32 InstanceIndex = 0; InstanceIndex < InstanceCount; ++InstanceIndex)
+		for (uint32 InstanceIndex = 0; InstanceIndex < InstanceCount; ++InstanceIndex)
 		{
 			if (!InstanceBuffer.GetInstanceData() || !InstanceBuffer.GetInstanceData()->IsValidIndex(InstanceIndex))
 		    {
@@ -1721,7 +1726,7 @@ void FInstancedStaticMeshSceneProxy::GetDynamicRayTracingInstances(struct FRayTr
 		return;
 	}
 
-	const int32 InstanceCount = InstancedRenderData.PerInstanceRenderData->InstanceBuffer.GetNumInstances();
+	const uint32 InstanceCount = InstancedRenderData.PerInstanceRenderData->InstanceBuffer.GetNumInstances();
 
 	if (InstanceCount == 0)
 	{
@@ -1735,6 +1740,12 @@ void FInstancedStaticMeshSceneProxy::GetDynamicRayTracingInstances(struct FRayTr
 		return;
 	}
 
+	struct SVisibleInstance
+	{
+		uint32 InstanceIndex = 0;
+		float DistanceToView = 0;
+	};
+
 	//setup a 'template' for the instance first, so we aren't duplicating work
 	//#dxr_todo: when multiple LODs are used, template needs to be an array of templates, probably best initialized on-demand via a lamda
 	FRayTracingInstance RayTracingInstanceTemplate;
@@ -1743,10 +1754,13 @@ void FInstancedStaticMeshSceneProxy::GetDynamicRayTracingInstances(struct FRayTr
 	RayTracingInstanceTemplate.Geometry = &RenderData->LODResources[LOD].RayTracingGeometry;
 
 	// Which index holds the reference to the particular simulated instance
-	TArray<int32> ActiveInstances;
+	TArray<uint32> ActiveInstances;
+
+	// Visible instances
+	TArray<SVisibleInstance> VisibleInstances;
 
-	const int32 RequestedSimulatedInstances = CVarRayTracingSimulatedInstanceCount.GetValueOnRenderThread();
-	const int32 SimulatedInstances = FMath::Min(RequestedSimulatedInstances == -1 ? InstanceCount : FMath::Clamp(RequestedSimulatedInstances, 1, InstanceCount), MaxSimulatedInstances);
+	const uint32 RequestedSimulatedInstances = CVarRayTracingSimulatedInstanceCount.GetValueOnRenderThread();
+	const uint32 SimulatedInstances = FMath::Min(RequestedSimulatedInstances == -1 ? InstanceCount : FMath::Clamp(RequestedSimulatedInstances, 1u, InstanceCount), MaxSimulatedInstances);
 
 	const int32 WPOEvalMode = CVarRayTracingInstancesEvaluateWPO.GetValueOnRenderThread();
 	const bool bWantsWPOEvaluation = WPOEvalMode < 0 ? bDynamicRayTracingGeometry : WPOEvalMode != 0;
@@ -1773,8 +1787,6 @@ void FInstancedStaticMeshSceneProxy::GetDynamicRayTracingInstances(struct FRayTr
 		}
 		RayTracingWPOInstanceTemplate.BuildInstanceMaskAndFlags(GetScene().GetFeatureLevel());
 	
-		
-
 		if (RayTracingDynamicData.Num() != SimulatedInstances || LOD != CachedRayTracingLOD)
 		{
 			SetupRayTracingDynamicInstances(SimulatedInstances, LOD);
@@ -1787,18 +1799,33 @@ void FInstancedStaticMeshSceneProxy::GetDynamicRayTracingInstances(struct FRayTr
 		}
 	}
 
-	//preallocate the worst-case to prevent an explosion of reallocs
-	//#dxr_todo: possibly track used instances and reserve based on previous behavior
-	RayTracingInstanceTemplate.InstanceTransforms.Reserve(InstanceCount);
+	VisibleInstances.Reserve(InstanceCount);
+
+	FVector ScaleVector = GetLocalToWorld().GetScaleVector();
+	FMatrix WorldToLocal = GetLocalToWorld().InverseFast();
+	float LocalToWorldScale = FMath::Max3(ScaleVector.X, ScaleVector.Y, ScaleVector.Z);
+	FVector LocalViewPosition = WorldToLocal.TransformPosition(Context.ReferenceView->ViewLocation);
 
-	FMatrix ToWorld = InstancedRenderData.Component->GetComponentTransform().ToMatrixWithScale();
+	auto GetDistanceToInstance =
+		[&LocalViewPosition, &LocalToWorldScale](const FVector4f& InstanceSphere, float& InstanceRadius, float& DistanceToInstanceCenter, float& DistanceToInstanceStart)
+	{
+		FVector4f InstanceLocation = InstanceSphere;
+		FVector4f VToInstanceCenter = (FVector4f) LocalViewPosition - InstanceLocation;
+
+		DistanceToInstanceCenter = VToInstanceCenter.Size();
+		InstanceRadius = InstanceSphere.W;
 
-	// whether to use angular culling instead of distance, angle is halved as it is compared against the projection of the radius rather than the diameter
-	const float CullAngle = FMath::Min(CVarRayTracingInstancesCullAngle.GetValueOnRenderThread(), 179.9f) * 0.5f;
+		// Scale accounts for possibly scaling in LocalToWorld, since measurements are all in local
+		DistanceToInstanceStart = (DistanceToInstanceCenter - InstanceRadius) * LocalToWorldScale;
+	};
 
+	const TArray<FRenderTransform>& PerInstanceTransforms = InstancedRenderData.PerInstanceRenderData->GetPerInstanceTransforms();
 	const TArray<FVector4f>& PerInstanceBounds = InstancedRenderData.PerInstanceRenderData->GetPerInstanceBounds();
 	if (CVarRayTracingRenderInstancesCulling.GetValueOnRenderThread() > 0 && PerInstanceBounds.Num())
 	{
+		// whether to use angular culling instead of distance, angle is halved as it is compared against the projection of the radius rather than the diameter
+		const float CullAngle = FMath::Min(CVarRayTracingInstancesCullAngle.GetValueOnRenderThread(), 179.9f) * 0.5f;
+
 		if (CullAngle < 0.0f)
 		{
 			//
@@ -1811,20 +1838,10 @@ void FInstancedStaticMeshSceneProxy::GetDynamicRayTracingInstances(struct FRayTr
 			const bool ApplyGeneralCulling = BVHCullRadius > 0.0f;
 			const bool ApplyLowScaleCulling = BVHLowScaleThreshold > 0.0f && BVHLowScaleRadius > 0.0f;
 
-			FVector ScaleVector = GetLocalToWorld().GetScaleVector();
-			FMatrix WorldToLocal = GetLocalToWorld().InverseFast();
-			float Scale = FMath::Max3(ScaleVector.X, ScaleVector.Y, ScaleVector.Z);
-			FVector LocalViewPosition = WorldToLocal.TransformPosition(Context.ReferenceView->ViewLocation);
-
-			const TArray<FRenderTransform>& PerInstanceTransforms = InstancedRenderData.PerInstanceRenderData->GetPerInstanceTransforms();
-			for (int32 InstanceIndex = 0; InstanceIndex < InstanceCount; InstanceIndex++)
+			for (uint32 InstanceIndex = 0; InstanceIndex < InstanceCount; InstanceIndex++)
 			{
-				FVector4f InstanceSphere = PerInstanceBounds[InstanceIndex];
-				FVector InstanceLocation = FVector(InstanceSphere.X, InstanceSphere.Y, InstanceSphere.Z); // LWC_TODO: Precision loss, but bounds are already float by this point
-				FVector VToInstanceCenter = LocalViewPosition - InstanceLocation;
-				float DistanceToInstanceCenter = VToInstanceCenter.Size(); // LWC_TODO: Precision loss
-				float InstanceRadius = InstanceSphere.W;
-				float DistanceToInstanceStart = (DistanceToInstanceCenter - InstanceRadius) * Scale; //scale accounts for possibly scaling in LocalToWorld, since measurements are all in local
+				float InstanceRadius, DistanceToInstanceCenter, DistanceToInstanceStart;
+				GetDistanceToInstance(PerInstanceBounds[InstanceIndex], InstanceRadius, DistanceToInstanceCenter, DistanceToInstanceStart);
 
 				// Cull instance based on distance
 				if (DistanceToInstanceStart > BVHCullRadius && ApplyGeneralCulling)
@@ -1837,9 +1854,7 @@ void FInstancedStaticMeshSceneProxy::GetDynamicRayTracingInstances(struct FRayTr
 						continue;
 				}
 
-				const FMatrix InstanceTransform = PerInstanceTransforms[InstanceIndex].ToMatrix() * GetLocalToWorld();
-				RayTracingInstanceTemplate.InstanceTransforms.Add(InstanceTransform);
-
+				VisibleInstances.Add(SVisibleInstance{ InstanceIndex, DistanceToInstanceStart });
 			}
 		}
 		else
@@ -1852,78 +1867,14 @@ void FInstancedStaticMeshSceneProxy::GetDynamicRayTracingInstances(struct FRayTr
 			//
 			float Ratio = FMath::Tan(CullAngle / 360.0f * 2.0f * PI);
 
-			FVector ScaleVector = GetLocalToWorld().GetScaleVector();
-			FMatrix WorldToLocal = GetLocalToWorld().InverseFast();
-			float Scale = FMath::Max3(ScaleVector.X, ScaleVector.Y, ScaleVector.Z);
-			FVector LocalViewPosition = WorldToLocal.TransformPosition(Context.ReferenceView->ViewLocation);
-
-			const TArray<FRenderTransform>& PerInstanceTransforms = InstancedRenderData.PerInstanceRenderData->GetPerInstanceTransforms();
-			for (int32 InstanceIndex = 0; InstanceIndex < InstanceCount; InstanceIndex++)
+			for (uint32 InstanceIndex = 0; InstanceIndex < InstanceCount; InstanceIndex++)
 			{
-				FVector4f InstanceSphere = PerInstanceBounds[InstanceIndex];
-				FVector InstanceLocation = FVector(InstanceSphere.X, InstanceSphere.Y, InstanceSphere.Z); // LWC_TODO: Precision loss, but bounds are already float by this point
-				FVector VToInstanceCenter = LocalViewPosition - InstanceLocation;
-				float DistanceToInstanceCenter = VToInstanceCenter.Size();  // LWC_TODO: Precision loss
+				float InstanceRadius, DistanceToInstanceCenter, DistanceToInstanceStart;
+				GetDistanceToInstance(PerInstanceBounds[InstanceIndex], InstanceRadius, DistanceToInstanceCenter, DistanceToInstanceStart);
 
-				if (DistanceToInstanceCenter * Ratio <= InstanceSphere.W * Scale)
+				if (DistanceToInstanceCenter * Ratio <= InstanceRadius * LocalToWorldScale)
 				{
-					FMatrix InstanceTransform = PerInstanceTransforms[InstanceIndex].ToMatrix() * GetLocalToWorld();
-					const int32 DynamicInstanceIdx = InstanceIndex % SimulatedInstances;
-
-					if (bHasWorldPositionOffset && InstancedRenderData.VertexFactories[LOD].GetType()->SupportsRayTracingDynamicGeometry())
-					{
-						FRayTracingInstance *DynamicInstance = nullptr;
-
-						if (ActiveInstances[DynamicInstanceIdx] == INDEX_NONE)
-						{
-							// First case of this dynamic instance, setup the material and add it
-							float InstanceRandom;
-							InstancedRenderData.PerInstanceRenderData->InstanceBuffer.GetInstanceRandomID(InstanceIndex, InstanceRandom);
-
-							const FStaticMeshLODResources& LODModel = RenderData->LODResources[LOD];
-
-							FRayTracingDynamicData &DynamicData = RayTracingDynamicData[DynamicInstanceIdx];
-
-							ActiveInstances[DynamicInstanceIdx] = OutRayTracingInstances.Num();
-							FRayTracingInstance &RayTracingInstance = OutRayTracingInstances.Add_GetRef(RayTracingWPOInstanceTemplate);
-							RayTracingInstance.Geometry = &DynamicData.DynamicGeometry;
-							RayTracingInstance.InstanceTransforms.Reserve(InstanceCount);
-
-							DynamicInstance = &RayTracingInstance;
-
-							FRayTracingInstance SimulationInstance = RayTracingWPODynamicTemplate;
-
-							// ToDo - deeper dive into ensuring better instance simulation matching
-							FMatrix Passthrough = FMatrix::Identity;
-							Passthrough.M[3][3] = InstanceRandom;
-
-							Context.DynamicRayTracingGeometriesToUpdate.Add(
-								FRayTracingDynamicGeometryUpdateParams
-								{
-									SimulationInstance.Materials,
-									false,
-									(uint32)LODModel.GetNumVertices(),
-									uint32((SIZE_T)LODModel.GetNumVertices() * sizeof(FVector3f)),
-									DynamicData.DynamicGeometry.Initializer.TotalPrimitiveCount,
-									&DynamicData.DynamicGeometry,
-									nullptr,
-									true,
-									Passthrough
-								}
-							);
-						}
-						else
-						{
-							DynamicInstance = &OutRayTracingInstances[ActiveInstances[DynamicInstanceIdx]];
-						}
-
-						DynamicInstance->InstanceTransforms.Add(InstanceTransform);
-
-					}
-					else
-					{
-						RayTracingInstanceTemplate.InstanceTransforms.Emplace(InstanceTransform);
-					}
+					VisibleInstances.Add(SVisibleInstance{ InstanceIndex, DistanceToInstanceStart });
 				}
 			}
 		}
@@ -1931,72 +1882,122 @@ void FInstancedStaticMeshSceneProxy::GetDynamicRayTracingInstances(struct FRayTr
 	else
 	{
 		// No culling
-		const TArray<FRenderTransform>& PerInstanceTransforms = InstancedRenderData.PerInstanceRenderData->GetPerInstanceTransforms();
-		for (int32 InstanceIndex = 0; InstanceIndex < InstanceCount; ++InstanceIndex)
+		for (uint32 InstanceIndex = 0; InstanceIndex < InstanceCount; ++InstanceIndex)
 		{
-			FMatrix InstanceTransform = PerInstanceTransforms[InstanceIndex].ToMatrix() * GetLocalToWorld();
+			VisibleInstances.Add(SVisibleInstance{ InstanceIndex, -1.0f });
+		}
+	}
 
-			if (bHasWorldPositionOffset && InstancedRenderData.VertexFactories[LOD].GetType()->SupportsRayTracingDynamicGeometry())
+	// Bucket which visible instances to simulate based on camera distance
+	if (bHasWorldPositionOffset && VisibleInstances.Num() > 0)
+	{
+		float SimulationClusterRadius = CVarRayTracingInstancesSimulationClusterRadius.GetValueOnRenderThread();
+
+		if (SimulationClusterRadius > 0.0)
+		{
+			if (SimulatedInstances < (uint32)VisibleInstances.Num())
 			{
-				FRayTracingInstance* DynamicInstance = nullptr;
+				const FMatrix& InvProjMatrix = Context.ReferenceView->ViewMatrices.GetInvProjectionMatrix();
 
-				const int32 DynamicInstanceIdx = InstanceIndex % SimulatedInstances;
+				// In no culling case, we are missing distance to view, so fill it in now
+				if (CVarRayTracingRenderInstancesCulling.GetValueOnRenderThread() == 0)
+				{
+					for (uint32 InstanceIndex = 0; InstanceIndex < (uint32)VisibleInstances.Num(); InstanceIndex++)
+					{
+						float InstanceRadius, DistanceToInstanceCenter, DistanceToInstanceStart;
+						GetDistanceToInstance(PerInstanceBounds[InstanceIndex], InstanceRadius, DistanceToInstanceCenter, DistanceToInstanceStart);
 
-				if (ActiveInstances[DynamicInstanceIdx] == INDEX_NONE)
+						VisibleInstances[InstanceIndex].DistanceToView = DistanceToInstanceStart;
+					}
+				}
+
+				uint32 PartitionIndex = 0;
+
+				// Let's partition the array so the closest elements are the ones we want to simulate
+				for (uint32 InstanceIndex = (uint32)VisibleInstances.Num() - 1; PartitionIndex < InstanceIndex; )
 				{
-					// First case of this dynamic instance, setup the material and add it
-					float InstanceRandom;
-					InstancedRenderData.PerInstanceRenderData->InstanceBuffer.GetInstanceRandomID(InstanceIndex, InstanceRandom);
+					if (VisibleInstances[InstanceIndex].DistanceToView > SimulationClusterRadius)
+					{
+						InstanceIndex--;
+					}
+					else
+					{
+						Swap(VisibleInstances[PartitionIndex], VisibleInstances[InstanceIndex]);
 
-					const FStaticMeshLODResources& LODModel = RenderData->LODResources[LOD];
+						PartitionIndex++;
+					}
+				}
 
-					FRayTracingDynamicData &DynamicData = RayTracingDynamicData[DynamicInstanceIdx];
+				// And partial sort those that fall within this bucket
+				Sort(&VisibleInstances[0], PartitionIndex, [](const SVisibleInstance& Instance0, const SVisibleInstance& Instance1)
+				{
+					return Instance0.DistanceToView < Instance1.DistanceToView;
+				}
+				);
+			}
+		}
+	}
 
-					ActiveInstances[DynamicInstanceIdx] = OutRayTracingInstances.Num();
-					FRayTracingInstance &RayTracingInstance = OutRayTracingInstances.Add_GetRef(RayTracingWPOInstanceTemplate);
-					RayTracingInstance.Geometry = &DynamicData.DynamicGeometry;
-					RayTracingInstance.InstanceTransforms.Reserve(InstanceCount);
+	//preallocate the worst-case to prevent an explosion of reallocs
+	//#dxr_todo: possibly track used instances and reserve based on previous behavior
+	RayTracingInstanceTemplate.InstanceTransforms.Reserve(InstanceCount);
 
-					DynamicInstance = &RayTracingInstance;
+	// Add all visible instances
+	for (SVisibleInstance VisibleInstance : VisibleInstances)
+	{
+		const uint32 InstanceIndex = VisibleInstance.InstanceIndex;
 
-					FRayTracingInstance SimulationInstance = RayTracingWPODynamicTemplate;
+		FMatrix InstanceToWorld = PerInstanceTransforms[InstanceIndex].ToMatrix() * GetLocalToWorld();
+		const uint32 DynamicInstanceIdx = InstanceIndex % SimulatedInstances;
 
-					// ToDo - deeper dive into ensuring better instance simulation matching
-					FMatrix Passthrough = FMatrix::Identity;
-					Passthrough.M[3][3] = InstanceRandom;
+		if (bHasWorldPositionOffset && InstancedRenderData.VertexFactories[LOD].GetType()->SupportsRayTracingDynamicGeometry())
+		{
+			FRayTracingInstance* DynamicInstance = nullptr;
 
-					Context.DynamicRayTracingGeometriesToUpdate.Add(
-						FRayTracingDynamicGeometryUpdateParams
-						{
-							SimulationInstance.Materials,
-							false,
-							(uint32)LODModel.GetNumVertices(),
-							uint32((SIZE_T)LODModel.GetNumVertices() * sizeof(FVector3f)),
-							DynamicData.DynamicGeometry.Initializer.TotalPrimitiveCount,
-							&DynamicData.DynamicGeometry,
-							nullptr,
-							true,
-							Passthrough
-						}
-					);
+			if (ActiveInstances[DynamicInstanceIdx] == -1)
+			{
+				// first case of this dynamic instance, setup the material and add it
+				const FStaticMeshLODResources& LODModel = RenderData->LODResources[LOD];
 
-				}
-				else
-				{
-					DynamicInstance = &OutRayTracingInstances[ActiveInstances[DynamicInstanceIdx]];
-				}
+				FRayTracingDynamicData& DynamicData = RayTracingDynamicData[DynamicInstanceIdx];
 
-				DynamicInstance->InstanceTransforms.Add(InstanceTransform);
+				ActiveInstances[DynamicInstanceIdx] = OutRayTracingInstances.Num();
+				FRayTracingInstance& RayTracingInstance = OutRayTracingInstances.Add_GetRef(RayTracingWPOInstanceTemplate);
+				RayTracingInstance.Geometry = &DynamicData.DynamicGeometry;
+				RayTracingInstance.InstanceTransforms.Reserve(InstanceCount);
 
+				DynamicInstance = &RayTracingInstance;
+
+				Context.DynamicRayTracingGeometriesToUpdate.Add(
+					FRayTracingDynamicGeometryUpdateParams
+					{
+						RayTracingWPODynamicTemplate.Materials,
+						false,
+						(uint32)LODModel.GetNumVertices(),
+						uint32((SIZE_T)LODModel.GetNumVertices() * sizeof(FVector)),
+						DynamicData.DynamicGeometry.Initializer.TotalPrimitiveCount,
+						&DynamicData.DynamicGeometry,
+						nullptr,
+						true,
+						InstanceIndex,
+						(FMatrix44f) InstanceToWorld.Inverse()
+					}
+				);
 			}
 			else
 			{
-				RayTracingInstanceTemplate.InstanceTransforms.Add(InstanceTransform);
+				DynamicInstance = &OutRayTracingInstances[ActiveInstances[DynamicInstanceIdx]];
 			}
+
+			DynamicInstance->InstanceTransforms.Add(InstanceToWorld);
+
+		}
+		else
+		{
+			RayTracingInstanceTemplate.InstanceTransforms.Emplace(InstanceToWorld);
 		}
 	}
 
-
 	if (RayTracingInstanceTemplate.InstanceTransforms.Num() > 0)
 	{
 		int32 SectionCount = InstancedRenderData.LODModels[LOD].Sections.Num();
diff --git a/Engine/Source/Runtime/Engine/Private/StaticMeshRender.cpp b/Engine/Source/Runtime/Engine/Private/StaticMeshRender.cpp
--- a/Engine/Source/Runtime/Engine/Private/StaticMeshRender.cpp
+++ b/Engine/Source/Runtime/Engine/Private/StaticMeshRender.cpp
@@ -156,6 +156,28 @@ static TAutoConsoleVariable<int32> CVarRayTracingStaticMeshesWPO(
 	TEXT(" 2: static meshes with world position offset visible in ray tracing, WPO evaluation disabled")
 );
 
+FAutoConsoleVariableSink CVarRayTracingStaticMeshesWPOSink(FConsoleCommandDelegate::CreateLambda([]()
+{
+	static int32 CachedRayTracingStaticMeshesWPO = CVarRayTracingStaticMeshesWPO.GetValueOnGameThread();
+
+	int32 RayTracingStaticMeshesWPO = CVarRayTracingStaticMeshesWPO.GetValueOnGameThread();
+
+	if (CachedRayTracingStaticMeshesWPO != RayTracingStaticMeshesWPO)
+	{
+		CachedRayTracingStaticMeshesWPO = RayTracingStaticMeshesWPO;
+
+		// NV-JIRA UE-668: Do this as a task on the game thread to break up a possible
+		// reentry call to USkeletalMeshComponent::PostAnimEvaluation if BP toggles this CVar
+		FFunctionGraphTask::CreateAndDispatchWhenReady([]()
+		{
+			// Easiest way to update all static scene proxies is to recreate them
+			FlushRenderingCommands();
+			FGlobalComponentReregisterContext ReregisterContext;
+			
+		}, TStatId(), nullptr, ENamedThreads::GameThread);
+	}
+}));
+
 static TAutoConsoleVariable<int32> CVarRayTracingStaticMeshesWPOCulling(
 	TEXT("r.RayTracing.Geometry.StaticMeshes.WPO.Culling"),
 	1,
@@ -265,10 +287,20 @@ FStaticMeshSceneProxy::FStaticMeshSceneProxy(UStaticMeshComponent* InComponent,
 
 #if RHI_RAYTRACING
 	bSupportRayTracing = InComponent->GetStaticMesh()->bSupportRayTracing;
-	bDynamicRayTracingGeometry = bSupportRayTracing && InComponent->bEvaluateWorldPositionOffset && MaterialRelevance.bUsesWorldPositionOffset;
+	bDynamicRayTracingGeometry = false;
 	
 	if (IsRayTracingEnabled() && bSupportRayTracing)
 	{
+		if (CVarRayTracingStaticMeshesWPO.GetValueOnAnyThread() > 0)
+		{
+			bDynamicRayTracingGeometry = MaterialRelevance.bUsesWorldPositionOffset;
+
+			if (CVarRayTracingStaticMeshesWPO.GetValueOnAnyThread() == 1)
+			{
+				bDynamicRayTracingGeometry &= InComponent->bEvaluateWorldPositionOffset;
+			}
+		}
+
 		RayTracingGeometries.AddDefaulted(RenderData->LODResources.Num());
 		for (int32 LODIndex = 0; LODIndex < RenderData->LODResources.Num(); LODIndex++)
 		{
@@ -371,6 +403,10 @@ FStaticMeshSceneProxy::FStaticMeshSceneProxy(UStaticMeshComponent* InComponent,
 void FStaticMeshSceneProxy::SetEvaluateWorldPositionOffsetInRayTracing(bool NewValue)
 {
 #if RHI_RAYTRACING
+	// Skip this code if the cvar settings are forcing WPO evaluation on/off
+	if (CVarRayTracingStaticMeshesWPO.GetValueOnAnyThread() != 1)
+		return;
+
 	NewValue &= MaterialRelevance.bUsesWorldPositionOffset;
 	if (NewValue && !bDynamicRayTracingGeometry)
 	{
@@ -1699,7 +1735,7 @@ bool FStaticMeshSceneProxy::HasRayTracingRepresentation() const
 
 void FStaticMeshSceneProxy::GetDynamicRayTracingInstances(FRayTracingMaterialGatheringContext& Context, TArray<FRayTracingInstance>& OutRayTracingInstances )
 {
-	if (DynamicRayTracingGeometries.Num() <= 0 || CVarRayTracingStaticMeshes.GetValueOnRenderThread() == 0 || CVarRayTracingStaticMeshesWPO.GetValueOnRenderThread() == 0)
+	if (DynamicRayTracingGeometries.Num() <= 0 || CVarRayTracingStaticMeshes.GetValueOnRenderThread() == 0)
 	{
 		return;
 	}
@@ -1708,7 +1744,7 @@ void FStaticMeshSceneProxy::GetDynamicRayTracingInstances(FRayTracingMaterialGat
 	const uint32 LODIndex = FMath::Max(GetLOD(Context.ReferenceView), (int32)GetCurrentFirstLODIdx_RenderThread());
 	const FStaticMeshLODResources& LODModel = RenderData->LODResources[LODIndex];
 
-	bool bEvaluateWPO = CVarRayTracingStaticMeshesWPO.GetValueOnRenderThread() == 1;
+	bool bEvaluateWPO = bDynamicRayTracingGeometry;
 
 	if (bEvaluateWPO && CVarRayTracingStaticMeshesWPOCulling.GetValueOnRenderThread() > 0)
 	{
diff --git a/Engine/Source/Runtime/Engine/Public/SceneManagement.h b/Engine/Source/Runtime/Engine/Public/SceneManagement.h
--- a/Engine/Source/Runtime/Engine/Public/SceneManagement.h
+++ b/Engine/Source/Runtime/Engine/Public/SceneManagement.h
@@ -2602,8 +2602,8 @@ struct FRayTracingDynamicGeometryUpdateParams
 
 	bool bApplyWorldPositionOffset = true;
 
-	// contains Instance random in [3][3] by convention
-	FMatrix InstanceTransform = FMatrix::Identity;
+	uint32 InstanceId = 0;
+	FMatrix44f WorldToInstance = FMatrix44f::Identity;
 };
 
 struct FRayTracingMaterialGatheringContext
diff --git a/Engine/Source/Runtime/Renderer/Private/RayTracing/RayTracingDynamicGeometry.cpp b/Engine/Source/Runtime/Renderer/Private/RayTracing/RayTracingDynamicGeometry.cpp
--- a/Engine/Source/Runtime/Renderer/Private/RayTracing/RayTracingDynamicGeometry.cpp
+++ b/Engine/Source/Runtime/Renderer/Private/RayTracing/RayTracingDynamicGeometry.cpp
@@ -45,8 +45,8 @@ public:
 		PrimitiveId.Bind(Initializer.ParameterMap, TEXT("PrimitiveId"));
 		OutputVertexBaseIndex.Bind(Initializer.ParameterMap, TEXT("OutputVertexBaseIndex"));
 		bApplyWorldPositionOffset.Bind(Initializer.ParameterMap, TEXT("bApplyWorldPositionOffset"));
-		InstanceTransform.Bind(Initializer.ParameterMap, TEXT("InstanceTransform"));
-		InverseTransform.Bind(Initializer.ParameterMap, TEXT("InverseTransform"));
+		InstanceId.Bind(Initializer.ParameterMap, TEXT("InstanceId"));
+		WorldToInstance.Bind(Initializer.ParameterMap, TEXT("WorldToInstance"));
 	}
 
 	FRayTracingDynamicGeometryConverterCS() = default;
@@ -98,8 +98,8 @@ public:
 	LAYOUT_FIELD(FShaderParameter, PrimitiveId);
 	LAYOUT_FIELD(FShaderParameter, bApplyWorldPositionOffset);
 	LAYOUT_FIELD(FShaderParameter, OutputVertexBaseIndex);
-	LAYOUT_FIELD(FShaderParameter, InstanceTransform);
-	LAYOUT_FIELD(FShaderParameter, InverseTransform);
+	LAYOUT_FIELD(FShaderParameter, InstanceId);
+	LAYOUT_FIELD(FShaderParameter, WorldToInstance);
 };
 
 IMPLEMENT_MATERIAL_SHADER_TYPE(, FRayTracingDynamicGeometryConverterCS, TEXT("/Engine/Private/RayTracing/RayTracingDynamicMesh.usf"), TEXT("RayTracingDynamicGeometryConverterCS"), SF_Compute);
@@ -204,11 +204,6 @@ void FRayTracingDynamicGeometryCollection::AddDynamicMeshBatchForGeometryUpdate(
 		RWBuffer = &VertexPositionBuffer->RWBuffer;
 	}
 
-	FMatrix44f InstanceTransform = FMatrix44f(UpdateParams.InstanceTransform);		// LWC_TODO: Precision loss
-	FMatrix44f InverseTransform = InstanceTransform;
-	InverseTransform.M[3][3] = 1.0f;
-	InverseTransform = InverseTransform.InverseFast();
-
 	for (const FMeshBatch& MeshBatch : UpdateParams.MeshBatches)
 	{
 		if (!ensureMsgf(MeshBatch.VertexFactory->GetType()->SupportsRayTracingDynamicGeometry(),
@@ -276,8 +271,8 @@ void FRayTracingDynamicGeometryCollection::AddDynamicMeshBatchForGeometryUpdate(
 		SingleShaderBindings.Add(Shader->PrimitiveId, PrimitiveId);
 		SingleShaderBindings.Add(Shader->OutputVertexBaseIndex, OutputVertexBaseIndex);
 		SingleShaderBindings.Add(Shader->bApplyWorldPositionOffset, UpdateParams.bApplyWorldPositionOffset ? 1 : 0);
-		SingleShaderBindings.Add(Shader->InstanceTransform, InstanceTransform);
-		SingleShaderBindings.Add(Shader->InverseTransform, InverseTransform);
+		SingleShaderBindings.Add(Shader->InstanceId, UpdateParams.InstanceId);
+		SingleShaderBindings.Add(Shader->WorldToInstance, UpdateParams.WorldToInstance);
 
 #if MESH_DRAW_COMMAND_DEBUG_DATA
 		FMeshProcessorShaders ShadersForDebug = Shaders.GetUntypedShaders();
