Double default range for movable directional light shadows (from 200m to 400m) and turn on DF shadows beyond that by default + some changed VSM config.
- add engine relese version & handling in serialization to retain default values for content saved with previous versions.
- set VSM to use distance culling for non nanite by default (r.Shadow.Virtual.UseFarShadowCulling).

diff --git a/Engine/Source/Runtime/Core/Public/UObject/UE5ReleaseStreamObjectVersion.h b/Engine/Source/Runtime/Core/Public/UObject/UE5ReleaseStreamObjectVersion.h
--- a/Engine/Source/Runtime/Core/Public/UObject/UE5ReleaseStreamObjectVersion.h
+++ b/Engine/Source/Runtime/Core/Public/UObject/UE5ReleaseStreamObjectVersion.h
@@ -101,6 +101,9 @@ struct CORE_API FUE5ReleaseStreamObjectVersion
 		BlueprintPinsUseRealNumbers,
 #endif
 		
+		// Changed shadow defaults for directional light components, version needed to not affect old things
+		UpdatedDirectionalLightShadowDefaults,
+
 		// -----<new versions can be added above this line>-------------------------------------------------
 		VersionPlusOne,
 		LatestVersion = VersionPlusOne - 1
diff --git a/Engine/Source/Runtime/Engine/Private/Components/DirectionalLightComponent.cpp b/Engine/Source/Runtime/Engine/Private/Components/DirectionalLightComponent.cpp
--- a/Engine/Source/Runtime/Engine/Private/Components/DirectionalLightComponent.cpp
+++ b/Engine/Source/Runtime/Engine/Private/Components/DirectionalLightComponent.cpp
@@ -974,18 +974,23 @@ UDirectionalLightComponent::UDirectionalLightComponent(const FObjectInitializer&
 
 	ShadowCascadeBiasDistribution = 1;
 	WholeSceneDynamicShadowRadius_DEPRECATED = 20000.0f;
-	DynamicShadowDistanceMovableLight = 20000.0f;
 	DynamicShadowDistanceStationaryLight = 0.f;
 
 	ForwardShadingPriority = 0;
-	DistanceFieldShadowDistance = 30000.0f;
+	// For 5.0 double dynamic shadow distance by default & enable DF shadows.
+	// Note, we reset this to old defaults in Serialize for lights created with older versions of the engine.
+	DynamicShadowDistanceMovableLight = 40000.0f;
+	DistanceFieldShadowDistance = 51200.0f;
+	DynamicShadowCascades = 4;
+	bUseRayTracedDistanceFieldShadows = true;
+
 	TraceDistance = 10000.0f;
 	FarShadowDistance = 300000.0f;
 	LightSourceAngle = 0.5357f;		// Angle of earth's sun
 	LightSourceSoftAngle = 0.0f;
 	ShadowSourceAngleFactor = 1.0f;
 
-	DynamicShadowCascades = 3;
+
 	CascadeDistributionExponent = 3.0f;
 	CascadeTransitionFraction = 0.1f;
 	ShadowDistanceFadeoutFraction = 0.1f;
@@ -1290,10 +1295,23 @@ void UDirectionalLightComponent::SetAtmosphereSunLightIndex(int32 NewValue)
 
 void UDirectionalLightComponent::Serialize(FArchive& Ar)
 {
-	Super::Serialize(Ar);
-
+	Ar.UsingCustomVersion(FUE5ReleaseStreamObjectVersion::GUID);
 	Ar.UsingCustomVersion(FUE5MainStreamObjectVersion::GUID);
 
+	if (Ar.IsLoading())
+	{
+		// Reset to older defaults if the light was created with a previous version of the engine.
+		if (Ar.CustomVer(FUE5ReleaseStreamObjectVersion::GUID) < FUE5ReleaseStreamObjectVersion::UpdatedDirectionalLightShadowDefaults)
+		{
+			DynamicShadowDistanceMovableLight = 20000.0f;
+			DistanceFieldShadowDistance = 30000.0f;
+			DynamicShadowCascades = 3;
+			bUseRayTracedDistanceFieldShadows = false;
+		}
+	}
+
+	Super::Serialize(Ar);
+
 	if(Ar.UEVer() < VER_UE4_REMOVE_LIGHT_MOBILITY_CLASSES)
 	{
 		// If outer is a DirectionalLight, we use the ADirectionalLight::LoadedFromAnotherClass path
diff --git a/Engine/Source/Runtime/Renderer/Private/ShadowSetup.cpp b/Engine/Source/Runtime/Renderer/Private/ShadowSetup.cpp
--- a/Engine/Source/Runtime/Renderer/Private/ShadowSetup.cpp
+++ b/Engine/Source/Runtime/Renderer/Private/ShadowSetup.cpp
@@ -272,7 +272,7 @@ static TAutoConsoleVariable<int32> CVarAlwaysAllocateMaxResolutionAtlases(
 
 static TAutoConsoleVariable<int32> CVarVsmUseFarShadowRules(
 	TEXT("r.Shadow.Virtual.UseFarShadowCulling"),
-	0,
+	1,
 	TEXT("Switch between implementing the far shadow culling logic for VSMs."),
 	ECVF_RenderThreadSafe
 );
