Cleaning up codec definitions on soundwaves.

- Adding new enumeration to select PCM, ADPCM, Bink, and "platform specific"
- Updated platform codecs and various APIs to select correctly which codec to use
- Changed "is seekable streaming" semantics to "Is Seekable" since that is more correct in a post-stream-caching world
- Tried to clean up platform backends. Unable to move Bink decoder to multiplatform code because of module dependencies, will require a future refactor of our decoder module setup.

diff --git a/Engine/Plugins/Runtime/Metasound/Source/MetasoundEngine/Private/MetasoundWaveProxyReader.cpp b/Engine/Plugins/Runtime/Metasound/Source/MetasoundEngine/Private/MetasoundWaveProxyReader.cpp
--- a/Engine/Plugins/Runtime/Metasound/Source/MetasoundEngine/Private/MetasoundWaveProxyReader.cpp
+++ b/Engine/Plugins/Runtime/Metasound/Source/MetasoundEngine/Private/MetasoundWaveProxyReader.cpp
@@ -423,7 +423,7 @@ namespace Metasound
 
 		// For non-seekable streaming waves, use a fallback method to seek
 		// to the start time. 
-		const bool bUseFallbackSeekMethod = (GMetaSoundWaveProxyReaderSimulateSeekOnNonSeekable != 0) && (InStartTimeInSeconds != 0.f) && (!WaveProxy->IsSeekableStreaming());
+		const bool bUseFallbackSeekMethod = (GMetaSoundWaveProxyReaderSimulateSeekOnNonSeekable != 0) && (InStartTimeInSeconds != 0.f) && (!WaveProxy->IsSeekable());
 		if (bUseFallbackSeekMethod)
 		{
 			if (!bFallbackSeekMethodWarningLogged)
diff --git a/Engine/Source/Developer/Android/AndroidTargetPlatform/Private/AndroidTargetPlatform.cpp b/Engine/Source/Developer/Android/AndroidTargetPlatform/Private/AndroidTargetPlatform.cpp
--- a/Engine/Source/Developer/Android/AndroidTargetPlatform/Private/AndroidTargetPlatform.cpp
+++ b/Engine/Source/Developer/Android/AndroidTargetPlatform/Private/AndroidTargetPlatform.cpp
@@ -616,12 +616,9 @@ const UTextureLODSettings& FAndroidTargetPlatform::GetTextureLODSettings() const
 }
 
 
-FName FAndroidTargetPlatform::GetWaveFormat( const class USoundWave* Wave ) const
+FName FAndroidTargetPlatform::GetWaveFormat(const class USoundWave* Wave) const
 {
-	static const FName NAME_ADPCM(TEXT("ADPCM"));
-	static const FName NAME_OGG(TEXT("OGG"));
-	static const FName NAME_BINKA(TEXT("BINKA"));
-
+	// Platform specific codec name
 	static bool bFormatRead = false;
 	static FName NAME_FORMAT;
 	if (!bFormatRead)
@@ -638,12 +635,12 @@ FName FAndroidTargetPlatform::GetWaveFormat( const class USoundWave* Wave ) cons
 		}
 
 #if WITH_OGGVORBIS
-		if (AudioSetting == NAME_OGG || AudioSetting == NAME_None)
+		if (AudioSetting == Audio::NAME_OGG || AudioSetting == NAME_None)
 		{
-			NAME_FORMAT = NAME_OGG;
+			NAME_FORMAT = Audio::NAME_OGG;
 		}
 #else
-		if (AudioSetting == NAME_OGG)
+		if (AudioSetting == Audio::NAME_OGG)
 		{
 			UE_LOG(LogAudio, Error, TEXT("Attemped to select Ogg Vorbis encoding when the cooker is built without Ogg Vorbis support."));
 		}
@@ -651,33 +648,26 @@ FName FAndroidTargetPlatform::GetWaveFormat( const class USoundWave* Wave ) cons
 		else
 		{
 			// Otherwise return ADPCM as it'll either be option '2' or 'default' depending on WITH_OGGVORBIS config
-			NAME_FORMAT = NAME_ADPCM;
+			NAME_FORMAT = Audio::NAME_ADPCM;
 		}
 	}
 
-	if (Wave->bUseBinkAudio)
+	FName FormatName = Audio::ToName(Wave->GetSoundAssetCompressionType());
+	if (FormatName == Audio::NAME_PLATFORM_SPECIFIC)
 	{
-		return NAME_BINKA;
+		FormatName = NAME_FORMAT;
 	}
-
-	if (Wave->IsSeekableStreaming())
-	{
-		return NAME_ADPCM;
-	}
-
-	return NAME_FORMAT;
+	return FormatName;
 }
 
 
 void FAndroidTargetPlatform::GetAllWaveFormats(TArray<FName>& OutFormats) const
 {
-	static const FName NAME_OGG(TEXT("OGG"));
-	static const FName NAME_ADPCM(TEXT("ADPCM"));
-	static const FName NAME_BINKA(TEXT("BINKA"));
 
-	OutFormats.Add(NAME_BINKA);
-	OutFormats.Add(NAME_OGG);
-	OutFormats.Add(NAME_ADPCM);
+	OutFormats.Add(Audio::NAME_BINKA);
+	OutFormats.Add(Audio::NAME_OGG);
+	OutFormats.Add(Audio::NAME_PCM);
+	OutFormats.Add(Audio::NAME_ADPCM);
 }
 
 #endif //WITH_ENGINE
diff --git a/Engine/Source/Developer/AudioFormatADPCM/Private/AudioFormatADPCM.cpp b/Engine/Source/Developer/AudioFormatADPCM/Private/AudioFormatADPCM.cpp
--- a/Engine/Source/Developer/AudioFormatADPCM/Private/AudioFormatADPCM.cpp
+++ b/Engine/Source/Developer/AudioFormatADPCM/Private/AudioFormatADPCM.cpp
@@ -20,8 +20,6 @@ DEFINE_LOG_CATEGORY_STATIC(LogAudioFormatADPCM, Log, All);
 	((uint32)(uint8)(ch0) | ((uint32)(uint8)(ch1) << 8) |\
 	((uint32)(uint8)(ch2) << 16) | ((uint32)(uint8)(ch3) << 24 ))
 
-static FName NAME_ADPCM(TEXT("ADPCM"));
-
 namespace
 {
 	struct RiffDataChunk
@@ -452,19 +450,20 @@ public:
 
 	virtual uint16 GetVersion(FName Format) const override
 	{
-		check(Format == NAME_ADPCM);
+		check(Format == Audio::NAME_ADPCM || Format == Audio::NAME_PCM);
 		return UE_AUDIO_ADPCM_VER;
 	}
 
 
 	virtual void GetSupportedFormats(TArray<FName>& OutFormats) const
 	{
-		OutFormats.Add(NAME_ADPCM);
+		OutFormats.Add(Audio::NAME_ADPCM);
+		OutFormats.Add(Audio::NAME_PCM);
 	}
 
 	virtual bool Cook(FName Format, const TArray<uint8>& SrcBuffer, FSoundQualityInfo& QualityInfo, TArray<uint8>& CompressedDataStore) const
 	{
-		check(Format == NAME_ADPCM);
+		check(Format == Audio::NAME_ADPCM || Format == Audio::NAME_PCM);
 
 		if (QualityInfo.Quality == 100)
 		{
@@ -481,7 +480,7 @@ public:
 	virtual bool CookSurround(FName Format, const TArray<TArray<uint8> >& SrcBuffers, FSoundQualityInfo& QualityInfo, TArray<uint8>& CompressedDataStore) const
 	{
 		// Ensure the right format
-		check(Format == NAME_ADPCM);
+		check(Format == Audio::NAME_ADPCM || Format == Audio::NAME_PCM);
 		// Ensure at least two channel
 		check(SrcBuffers.Num() > 1);
 		// Ensure one buffer per channel
@@ -506,7 +505,7 @@ public:
 
 	virtual int32 Recompress(FName Format, const TArray<uint8>& SrcBuffer, FSoundQualityInfo& QualityInfo, TArray<uint8>& OutBuffer) const
 	{
-		check(Format == NAME_ADPCM);
+		check(Format == Audio::NAME_ADPCM || Format == Audio::NAME_PCM);
 
 		// Recompress is only necessary during editor previews
 		return 0;
diff --git a/Engine/Source/Developer/IOS/IOSTargetPlatform/Private/IOSTargetPlatform.cpp b/Engine/Source/Developer/IOS/IOSTargetPlatform/Private/IOSTargetPlatform.cpp
--- a/Engine/Source/Developer/IOS/IOSTargetPlatform/Private/IOSTargetPlatform.cpp
+++ b/Engine/Source/Developer/IOS/IOSTargetPlatform/Private/IOSTargetPlatform.cpp
@@ -675,25 +675,21 @@ const UTextureLODSettings& FIOSTargetPlatform::GetTextureLODSettings() const
 	return *TextureLODSettings;
 }
 
-
 FName FIOSTargetPlatform::GetWaveFormat( const USoundWave* Wave ) const
 {
-	if (Wave->bUseBinkAudio)
+	FName FormatName = Audio::ToName(Wave->GetSoundAssetCompressionType());
+	if (FormatName == Audio::NAME_PLATFORM_SPECIFIC)
 	{
-		static FName NAME_BINKA(TEXT("BINKA"));
-		return NAME_BINKA;
+			FormatName = Audio::NAME_ADPCM;
 	}
-	static FName NAME_ADPCM(TEXT("ADPCM"));
-	return NAME_ADPCM;
+	return FormatName;
 }
 
-
 void FIOSTargetPlatform::GetAllWaveFormats(TArray<FName>& OutFormat) const
 {
-	static FName NAME_BINKA(TEXT("BINKA"));
-	static FName NAME_ADPCM(TEXT("ADPCM"));
-	OutFormat.Add(NAME_ADPCM);
-	OutFormat.Add(NAME_BINKA);
+	OutFormat.Add(Audio::NAME_ADPCM);
+	OutFormat.Add(Audio::NAME_PCM);
+	OutFormat.Add(Audio::NAME_BINKA);
 }
 
 #endif // WITH_ENGINE
diff --git a/Engine/Source/Developer/Linux/LinuxTargetPlatform/Private/LinuxTargetPlatform.h b/Engine/Source/Developer/Linux/LinuxTargetPlatform/Private/LinuxTargetPlatform.h
--- a/Engine/Source/Developer/Linux/LinuxTargetPlatform/Private/LinuxTargetPlatform.h
+++ b/Engine/Source/Developer/Linux/LinuxTargetPlatform/Private/LinuxTargetPlatform.h
@@ -313,40 +313,30 @@ public:
 
 	virtual FName GetWaveFormat( const class USoundWave* Wave ) const override
 	{
-		static const FName NAME_ADPCM(TEXT("ADPCM"));
-		static const FName NAME_OGG(TEXT("OGG"));
-		static const FName NAME_OPUS(TEXT("OPUS"));
-		static const FName NAME_BINKA(TEXT("BINKA"));
+		FName FormatName = Audio::ToName(Wave->GetSoundAssetCompressionType());
 
-		if (Wave->bUseBinkAudio)
+		if (FormatName == Audio::NAME_PLATFORM_SPECIFIC)
 		{
-			return NAME_BINKA;
-		}
+			if (Wave->IsStreaming(*this->IniPlatformName()))
+			{
+				return Audio::NAME_OPUS;
+			}
 
-		if (Wave->IsSeekableStreaming())
-		{
-			return NAME_ADPCM;
+			return Audio::NAME_OGG;
 		}
-
-		if (Wave->IsStreaming(*this->IniPlatformName()))
+		else
 		{
-			return NAME_OPUS;
+			return FormatName;
 		}
-
-		return NAME_OGG;
 	}
 
 	virtual void GetAllWaveFormats(TArray<FName>& OutFormats) const override
 	{
-		static const FName NAME_ADPCM(TEXT("ADPCM"));
-		static const FName NAME_OGG(TEXT("OGG"));
-		static const FName NAME_OPUS(TEXT("OPUS"));
-		static const FName NAME_BINKA(TEXT("BINKA"));
-
-		OutFormats.Add(NAME_BINKA);
-		OutFormats.Add(NAME_ADPCM);
-		OutFormats.Add(NAME_OGG);
-		OutFormats.Add(NAME_OPUS);
+		OutFormats.Add(Audio::NAME_BINKA);
+		OutFormats.Add(Audio::NAME_ADPCM);
+		OutFormats.Add(Audio::NAME_PCM);
+		OutFormats.Add(Audio::NAME_OGG);
+		OutFormats.Add(Audio::NAME_OPUS);
 	}
 
 #endif //WITH_ENGINE
diff --git a/Engine/Source/Developer/Mac/MacTargetPlatform/Private/GenericMacTargetPlatform.h b/Engine/Source/Developer/Mac/MacTargetPlatform/Private/GenericMacTargetPlatform.h
--- a/Engine/Source/Developer/Mac/MacTargetPlatform/Private/GenericMacTargetPlatform.h
+++ b/Engine/Source/Developer/Mac/MacTargetPlatform/Private/GenericMacTargetPlatform.h
@@ -323,42 +323,31 @@ public:
 		FTargetPlatformBase::AddDependencySCArrayHelper(OutDependencies, TEXT("Binaries/ThirdParty/ShaderConductor/Mac/libShaderConductor.dylib"));
 	}
 
-	virtual FName GetWaveFormat( const class USoundWave* Wave ) const override
+	virtual FName GetWaveFormat(const class USoundWave* Wave) const override
 	{
-		static const FName NAME_ADPCM(TEXT("ADPCM"));
-		static const FName NAME_OGG(TEXT("OGG"));
-		static const FName NAME_OPUS(TEXT("OPUS"));
-		static const FName NAME_BINKA(TEXT("BINKA"));
-
-		if (Wave->bUseBinkAudio)
+		FName FormatName = Audio::ToName(Wave->GetSoundAssetCompressionType());
+		if (FormatName == Audio::NAME_PLATFORM_SPECIFIC)
 		{
-			return NAME_BINKA;
-		}
+			if (Wave->IsStreaming(*this->IniPlatformName()))
+			{
+				return Audio::NAME_OPUS;
+			}
 
-		if (Wave->IsSeekableStreaming())
-		{
-			return NAME_ADPCM;
+			return Audio::NAME_OGG;
 		}
-
-		if (Wave->IsStreaming(*this->IniPlatformName()))
+		else
 		{
-			return NAME_OPUS;
+			return FormatName;
 		}
-
-		return NAME_OGG;
 	}
 
 	virtual void GetAllWaveFormats(TArray<FName>& OutFormats) const override
 	{
-		static const FName NAME_ADPCM(TEXT("ADPCM"));
-		static const FName NAME_OGG(TEXT("OGG"));
-		static const FName NAME_OPUS(TEXT("OPUS"));
-		static const FName NAME_BINKA(TEXT("BINKA"));
-
-		OutFormats.Add(NAME_ADPCM);
-		OutFormats.Add(NAME_OGG);
-		OutFormats.Add(NAME_OPUS);
-		OutFormats.Add(NAME_BINKA);
+		OutFormats.Add(Audio::NAME_ADPCM);
+		OutFormats.Add(Audio::NAME_PCM);
+		OutFormats.Add(Audio::NAME_OGG);
+		OutFormats.Add(Audio::NAME_OPUS);
+		OutFormats.Add(Audio::NAME_BINKA);
 	}
 
 #endif //WITH_ENGINE
diff --git a/Engine/Source/Developer/Windows/WindowsTargetPlatform/Private/GenericWindowsTargetPlatform.h b/Engine/Source/Developer/Windows/WindowsTargetPlatform/Private/GenericWindowsTargetPlatform.h
--- a/Engine/Source/Developer/Windows/WindowsTargetPlatform/Private/GenericWindowsTargetPlatform.h
+++ b/Engine/Source/Developer/Windows/WindowsTargetPlatform/Private/GenericWindowsTargetPlatform.h
@@ -420,45 +420,34 @@ public:
 		TextureLODSettings = InTextureLODSettings;
 	}
 
-	virtual FName GetWaveFormat( const class USoundWave* Wave ) const override
+	virtual FName GetWaveFormat(const class USoundWave* Wave) const override
 	{
-		static const FName NAME_ADPCM(TEXT("ADPCM"));
-		static const FName NAME_OGG(TEXT("OGG"));
-		static const FName NAME_OPUS(TEXT("OPUS"));
-		static const FName NAME_BINKA(TEXT("BINKA"));
+		FName FormatName = Audio::ToName(Wave->GetSoundAssetCompressionType());
 
-		if (Wave->bUseBinkAudio)
+		if (FormatName == Audio::NAME_PLATFORM_SPECIFIC)
 		{
-			return NAME_BINKA;
-		}
+#if !USE_VORBIS_FOR_STREAMING
+			if (Wave->IsStreaming())
+			{
+				return Audio::NAME_OPUS;
+			}
+#endif
 
-		// Seekable streams need to pick a codec which allows fixed-sized frames so we can compute stream chunk index to load
-		if (Wave->IsSeekableStreaming())
-		{
-			return NAME_ADPCM;
+			return Audio::NAME_OGG;
 		}
-
-#if !USE_VORBIS_FOR_STREAMING
-		if (Wave->IsStreaming())
+		else
 		{
-			return NAME_OPUS;
+			return FormatName;
 		}
-#endif
-
-		return NAME_OGG;
 	}
 
 	virtual void GetAllWaveFormats(TArray<FName>& OutFormats) const override
 	{
-		static const FName NAME_ADPCM(TEXT("ADPCM"));
-		static const FName NAME_OGG(TEXT("OGG"));
-		static const FName NAME_OPUS(TEXT("OPUS"));
-		static const FName NAME_BINKA(TEXT("BINKA"));
-
-		OutFormats.Add(NAME_BINKA);
-		OutFormats.Add(NAME_ADPCM);
-		OutFormats.Add(NAME_OGG);
-		OutFormats.Add(NAME_OPUS);
+		OutFormats.Add(Audio::NAME_BINKA);
+		OutFormats.Add(Audio::NAME_ADPCM);
+		OutFormats.Add(Audio::NAME_PCM);
+		OutFormats.Add(Audio::NAME_OGG);
+		OutFormats.Add(Audio::NAME_OPUS);
 	}
 
 	virtual void GetWaveFormatModuleHints(TArray<FName>& OutModuleNames) const override
diff --git a/Engine/Source/Runtime/Android/AudioMixerAndroid/Private/AudioMixerPlatformAndroid.cpp b/Engine/Source/Runtime/Android/AudioMixerAndroid/Private/AudioMixerPlatformAndroid.cpp
--- a/Engine/Source/Runtime/Android/AudioMixerAndroid/Private/AudioMixerPlatformAndroid.cpp
+++ b/Engine/Source/Runtime/Android/AudioMixerAndroid/Private/AudioMixerPlatformAndroid.cpp
@@ -6,10 +6,10 @@
 #include "CoreGlobals.h"
 #include "Misc/ConfigCacheIni.h"
 #include "Misc/ScopeLock.h"
+#include "AudioDevice.h"
 
 #if WITH_ENGINE
 #include "VorbisAudioInfo.h"
-#include "ADPCMAudioInfo.h"
 #include "BinkAudioInfo.h"
 #include "AudioPluginUtilities.h"
 #endif
@@ -401,11 +401,6 @@ namespace Audio
 		}
 	}
 
-	bool FMixerPlatformAndroid::SupportsRealtimeDecompression() const
-	{
-		return true;
-	}
-
 	void FMixerPlatformAndroid::SubmitBuffer(const uint8* Buffer)
 	{
 		check(DeviceBuffer.Num() == NumSamplesPerDeviceCallback);
@@ -424,81 +419,50 @@ namespace Audio
 		}
 	}
 
-	FName FMixerPlatformAndroid::GetRuntimeFormat(USoundWave* InSoundWave)
+	FName FMixerPlatformAndroid::GetRuntimeFormat(const USoundWave* InSoundWave) const
 	{
-#if WITH_ENGINE
-		static FName NAME_ADPCM(TEXT("ADPCM"));
-		static const FName NAME_BINKA(TEXT("BINKA"));
+		FName RuntimeFormat = Audio::ToName(InSoundWave->GetSoundAssetCompressionType());
 
-		if (InSoundWave->bUseBinkAudio)
+		if (RuntimeFormat == Audio::NAME_PLATFORM_SPECIFIC)
 		{
-			return NAME_BINKA;
-		}
-
-		if (InSoundWave->IsSeekableStreaming())
-		{
-			return NAME_ADPCM;
-		}
-
 #if WITH_OGGVORBIS
-		static FName NAME_OGG(TEXT("OGG"));
-		if (InSoundWave->HasCompressedData(NAME_OGG))
-		{
-			return NAME_OGG;
+			if (InSoundWave->HasCompressedData(Audio::NAME_OGG))
+			{
+				RuntimeFormat = Audio::NAME_OGG;
+			}
+			else
+#endif // #if WITH_OGGVORBIS
+			{
+				RuntimeFormat = Audio::NAME_ADPCM;
+			}
 		}
-#endif // WITH_OGGVORBIS
-		return NAME_ADPCM;
-#else
-		return FName(TEXT("None"));
-#endif //WITH_ENGINE
-	}
 
-	bool FMixerPlatformAndroid::HasCompressedAudioInfoClass(USoundWave* InSoundWave)
-	{
-		return true;
+		return RuntimeFormat;
 	}
 
-	ICompressedAudioInfo* FMixerPlatformAndroid::CreateCompressedAudioInfo(USoundWave* InSoundWave)
+	ICompressedAudioInfo* FMixerPlatformAndroid::CreateCompressedAudioInfo(const FName& InRuntimeFormat) const
 	{
-#if WITH_ENGINE
-		if (InSoundWave->bUseBinkAudio)
-		{
-			return new FBinkAudioInfo();
-		}
-
-		if (InSoundWave->IsSeekableStreaming())
-		{
-			return new FADPCMAudioInfo();
-		}
-
-		return new FVorbisAudioInfo();
-#else
-		return nullptr;
-#endif // WITH_ENGINE
-	}
+		ICompressedAudioInfo* Decoder = nullptr;
 
-	ICompressedAudioInfo* FMixerPlatformAndroid::CreateCompressedAudioInfo(const FSoundWaveProxyPtr& InSoundWave)
-	{
-		if (!ensure(InSoundWave.IsValid()))
+#if WITH_OGGVORBIS
+		if (InRuntimeFormat == Audio::NAME_OGG)
 		{
-			return nullptr;
+			Decoder = new FVorbisAudioInfo();
 		}
-
-#if WITH_ENGINE
-		if (InSoundWave->UseBinkAudio())
+		else
+#endif // #if WITH_OGGVORBIS
+#if WITH_BINK_AUDIO
+		if (InRuntimeFormat == Audio::NAME_BINKA)
 		{
-			return new FBinkAudioInfo();
-		}
-
-		if (InSoundWave->IsSeekableStreaming())
+			Decoder = new FBinkAudioInfo();
+		} 
+		else
+#endif // WITH_BINK_AUDIO			
 		{
-			return new FADPCMAudioInfo();
+			Decoder = Audio::CreateSoundAssetDecoder(InRuntimeFormat);
 		}
-
-		return new FVorbisAudioInfo();
-#endif // WITH_ENGINE
-
-		return nullptr;
+		ensureMsgf(Decoder != nullptr, TEXT("Failed to create a sound asset decoder for compression type: %s"), *InRuntimeFormat.ToString());
+		return Decoder;
 	}
 
 	FString FMixerPlatformAndroid::GetDefaultDeviceName()
diff --git a/Engine/Source/Runtime/Android/AudioMixerAndroid/Private/AudioMixerPlatformAndroid.h b/Engine/Source/Runtime/Android/AudioMixerAndroid/Private/AudioMixerPlatformAndroid.h
--- a/Engine/Source/Runtime/Android/AudioMixerAndroid/Private/AudioMixerPlatformAndroid.h
+++ b/Engine/Source/Runtime/Android/AudioMixerAndroid/Private/AudioMixerPlatformAndroid.h
@@ -33,15 +33,12 @@ namespace Audio
 		virtual bool StopAudioStream() override;
 		virtual FAudioPlatformDeviceInfo GetPlatformDeviceInfo() const override;
 		virtual void SubmitBuffer(const uint8* Buffer) override;
-		virtual FName GetRuntimeFormat(USoundWave* InSoundWave) override;
-		virtual bool HasCompressedAudioInfoClass(USoundWave* InSoundWave) override;
-		virtual ICompressedAudioInfo* CreateCompressedAudioInfo(USoundWave* InSoundWave) override;
-		virtual ICompressedAudioInfo* CreateCompressedAudioInfo(const FSoundWaveProxyPtr& InSoundWave) override;
+		virtual FName GetRuntimeFormat(const USoundWave* InSoundWave) const override;
+		virtual ICompressedAudioInfo* CreateCompressedAudioInfo(const FName& InRuntimeFormat) const override;
 		virtual FString GetDefaultDeviceName() override;
 		virtual FAudioPlatformSettings GetPlatformSettings() const override;
 		virtual void SuspendContext() override;
 		virtual void ResumeContext() override;
-		virtual bool SupportsRealtimeDecompression() const override;
 
 		//~ End IAudioMixerPlatformInterface
 		
diff --git a/Engine/Source/Runtime/Apple/AudioMixerAudioUnit/Private/AudioMixerPlatformAudioUnit.cpp b/Engine/Source/Runtime/Apple/AudioMixerAudioUnit/Private/AudioMixerPlatformAudioUnit.cpp
--- a/Engine/Source/Runtime/Apple/AudioMixerAudioUnit/Private/AudioMixerPlatformAudioUnit.cpp
+++ b/Engine/Source/Runtime/Apple/AudioMixerAudioUnit/Private/AudioMixerPlatformAudioUnit.cpp
@@ -4,12 +4,12 @@
 #include "AudioMixerPlatformAudioUnitUtils.h"
 #include "Modules/ModuleManager.h"
 #include "AudioMixer.h"
+#include "AudioDevice.h"
 #include "AudioMixerDevice.h"
 #include "CoreGlobals.h"
 #include "CoreMinimal.h"
 #include "Misc/ConfigCacheIni.h"
 #include "Misc/CoreDelegates.h"
-#include "ADPCMAudioInfo.h"
 #include "BinkAudioInfo.h"
 
 /*
@@ -387,47 +387,44 @@ namespace Audio
 		int32 PushResult = CircularOutputBuffer.Push((const int8*)Buffer, BytesToSubmitToAudioMixer);
 		check(PushResult == BytesToSubmitToAudioMixer);
 	}
-	
-	FName FMixerPlatformAudioUnit::GetRuntimeFormat(USoundWave* InSoundWave)
+
+	FName FMixerPlatformAudioUnit::GetRuntimeFormat(const USoundWave* InSoundWave) const
 	{
-		if (InSoundWave->bUseBinkAudio)
+		FName RuntimeFormat = Audio::ToName(InSoundWave->GetSoundAssetCompressionType());
+
+		// On IOS the default platform specific codec is ADPCM
+		if (RuntimeFormat == Audio::NAME_PLATFORM_SPECIFIC)
 		{
-			static const FName NAME_BINKA(TEXT("BINKA"));
-			return NAME_BINKA;
+			RuntimeFormat = Audio::NAME_ADPCM;
 		}
-
-		static FName NAME_ADPCM(TEXT("ADPCM"));
-		return NAME_ADPCM;
-	}
-	
-	bool FMixerPlatformAudioUnit::HasCompressedAudioInfoClass(USoundWave* InSoundWave)
-	{
-		return true;
-	}
-	
-	ICompressedAudioInfo* FMixerPlatformAudioUnit::CreateCompressedAudioInfo(USoundWave* InSoundWave)
-	{
-		if (InSoundWave->bUseBinkAudio)
+		else if (InRuntimeFormat == Audio::NAME_BINKA)
 		{
 			return new FBinkAudioInfo();
 		}
-		return new FADPCMAudioInfo();
+
+		return RuntimeFormat;
 	}
 
-	ICompressedAudioInfo* FMixerPlatformAudioUnit::CreateCompressedAudioInfo(const FSoundWaveProxyPtr& InSoundWave)
+	ICompressedAudioInfo* FMixerPlatformCoreAudio::CreateCompressedAudioInfo(const FName& InRuntimeFormat) const
 	{
-		if (!ensure(InSoundWave.IsValid()))
+		ICompressedAudioInfo* Decoder = nullptr;
+
+		if (RuntimeFormat == Audio::NAME_PLATFORM_SPECIFIC)
 		{
-			return nullptr;
+			Decoder = new FADPCMAudioInfo();
+		}	
+		else if (InRuntimeFormat == Audio::NAME_BINKA)
+		{
+			Decoder = new FBinkAudioInfo();
 		}
-		else if (InSoundWave->UseBinkAudio())
+		else
 		{
-			return new FBinkAudioInfo();
+			Decoder = Audio::CreateSoundAssetDecoder(InRuntimeFormat);
 		}
-
-		return new FADPCMAudioInfo();
+		ensureMsgf(Decoder != nullptr, TEXT("Failed to create a sound asset decoder for compression type: %s"), *InRuntimeFormat.ToString());
+		return Decoder;
 	}
-	
+
 	FString FMixerPlatformAudioUnit::GetDefaultDeviceName()
 	{
 		return FString();
diff --git a/Engine/Source/Runtime/Apple/AudioMixerAudioUnit/Private/AudioMixerPlatformAudioUnit.h b/Engine/Source/Runtime/Apple/AudioMixerAudioUnit/Private/AudioMixerPlatformAudioUnit.h
--- a/Engine/Source/Runtime/Apple/AudioMixerAudioUnit/Private/AudioMixerPlatformAudioUnit.h
+++ b/Engine/Source/Runtime/Apple/AudioMixerAudioUnit/Private/AudioMixerPlatformAudioUnit.h
@@ -35,19 +35,13 @@ namespace Audio
 		virtual bool MoveAudioStreamToNewAudioDevice(const FString& InNewDeviceId) override;
 		virtual FAudioPlatformDeviceInfo GetPlatformDeviceInfo() const override;
 		virtual void SubmitBuffer(const uint8* Buffer) override;
-		virtual FName GetRuntimeFormat(USoundWave* InSoundWave) override;
-		virtual bool HasCompressedAudioInfoClass(USoundWave* InSoundWave) override;
-		virtual ICompressedAudioInfo* CreateCompressedAudioInfo(USoundWave* InSoundWave) override;
-		virtual ICompressedAudioInfo* CreateCompressedAudioInfo(const FSoundWaveProxyPtr& InSoundWave) override;
+		virtual FName GetRuntimeFormat(const USoundWave* InSoundWave) const override;
 		virtual FString GetDefaultDeviceName() override;
 		virtual FAudioPlatformSettings GetPlatformSettings() const override;
 		virtual int32 GetNumFrames(const int32 InNumReqestedFrames) override;
 		virtual void ResumeContext() override;
 		virtual void SuspendContext() override;
-		
-		/** Whether or not the platform supports realtime decompression. */
-		virtual bool SupportsRealtimeDecompression() const override { return true; }
-		
+			
 		/** Whether or not the platform disables caching of decompressed PCM data (i.e. to save memory on fixed memory platforms) */
 		virtual bool DisablePCMAudioCaching() const override { return true; }
 		
diff --git a/Engine/Source/Runtime/Apple/AudioMixerCoreAudio/Private/AudioMixerPlatformCoreAudio.cpp b/Engine/Source/Runtime/Apple/AudioMixerCoreAudio/Private/AudioMixerPlatformCoreAudio.cpp
--- a/Engine/Source/Runtime/Apple/AudioMixerCoreAudio/Private/AudioMixerPlatformCoreAudio.cpp
+++ b/Engine/Source/Runtime/Apple/AudioMixerCoreAudio/Private/AudioMixerPlatformCoreAudio.cpp
@@ -10,7 +10,7 @@
 #if WITH_ENGINE
 #include "VorbisAudioInfo.h"
 #include "OpusAudioInfo.h"
-#include "ADPCMAudioInfo.h"
+#include "AudioDevice.h"
 #include "BinkAudioInfo.h"
 #endif
 /**
@@ -360,111 +360,49 @@ namespace Audio
 		RemainingBytesInCurrentSubmittedBuffer = BytesPerSubmittedBuffer;
 	}
 
-	FName FMixerPlatformCoreAudio::GetRuntimeFormat(USoundWave* InSoundWave)
+	FName FMixerPlatformCoreAudio::GetRuntimeFormat(const USoundWave* InSoundWave) const
 	{
-#if WITH_ENGINE
-		static const FName NAME_ADPCM(TEXT("ADPCM"));
-		static const FName NAME_OGG(TEXT("OGG"));
-		static const FName NAME_OPUS(TEXT("OPUS"));
-		static const FName NAME_BINKA(TEXT("BINKA"));
-
-		if (InSoundWave->bUseBinkAudio)
-		{
-			return NAME_BINKA;
-		}
-
-		if (InSoundWave->IsSeekableStreaming())
-		{
-			return NAME_ADPCM;
-		}
-
-		if (InSoundWave->IsStreaming(nullptr))
-		{
-			return NAME_OPUS;
-		}
-
-		return NAME_OGG;
-#else
-		return FName();
-#endif
-	}
-
-	bool FMixerPlatformCoreAudio::HasCompressedAudioInfoClass(USoundWave* InSoundWave)
-	{
-		return true;
-	}
-
-	ICompressedAudioInfo* FMixerPlatformCoreAudio::CreateCompressedAudioInfo(USoundWave* InSoundWave)
-	{
-#if WITH_ENGINE
-		check(InSoundWave);
-
-		if (InSoundWave->bUseBinkAudio)
-		{
-			return new FBinkAudioInfo();
-		}
-
-		if (InSoundWave->IsSeekableStreaming())
-		{
-			return new FADPCMAudioInfo();
-		}
+		FName RuntimeFormat = Audio::ToName(InSoundWave->GetSoundAssetCompressionType());
 
-		if (InSoundWave->IsStreaming())
+		if (RuntimeFormat == Audio::NAME_PLATFORM_SPECIFIC)
 		{
-			return new FOpusAudioInfo();
-		}
-
-		static const FName NAME_OGG(TEXT("OGG"));
-		if (FPlatformProperties::RequiresCookedData() ? InSoundWave->HasCompressedData(NAME_OGG) : (InSoundWave->GetCompressedData(NAME_OGG) != nullptr))
-		{
-			ICompressedAudioInfo* CompressedInfo = new FVorbisAudioInfo();
-			if (!CompressedInfo)
+			if (InSoundWave->IsStreaming(nullptr))
+			{
+				RuntimeFormat = Audio::NAME_OPUS;
+			}
+			else
 			{
-				UE_LOG(LogAudio, Error, TEXT("Failed to create new FVorbisAudioInfo for SoundWave %s: out of memory."), *InSoundWave->GetName());
-				return nullptr;
+				RuntimeFormat = Audio::NAME_OGG;
 			}
-			return CompressedInfo;
 		}
 
-#endif // WITH_ENGINE
-		return nullptr;
+		return RuntimeFormat;
 	}
 
-	ICompressedAudioInfo* FMixerPlatformCoreAudio::CreateCompressedAudioInfo(const FSoundWaveProxyPtr& InSoundWave)
+	ICompressedAudioInfo* FMixerPlatformCoreAudio::CreateCompressedAudioInfo(const FName& InRuntimeFormat) const
 	{
-#if WITH_ENGINE
-		static const FName NAME_OGG(TEXT("OGG"));
+		ICompressedAudioInfo* Decoder = nullptr;
 
-		if (!ensure(InSoundWave.IsValid()))
+		if (InRuntimeFormat == Audio::NAME_OGG)
 		{
-			return nullptr;
+			Decoder = new FVorbisAudioInfo();
 		}
-		else if (InSoundWave->UseBinkAudio())
+		else if (InRuntimeFormat == Audio::NAME_OPUS)
 		{
-			return new FBinkAudioInfo();
+			Decoder = new FOpusAudioInfo();
 		}
-		else if (InSoundWave->IsSeekableStreaming())
+#if WITH_BINK_AUDIO
+		else if (InRuntimeFormat == Audio::NAME_BINKA)
 		{
-			return new FADPCMAudioInfo();
-		}
-		else if (InSoundWave->IsStreaming())
-		{
-			return new FOpusAudioInfo();
+			return new FBinkAudioInfo();
 		}
-		else if (FPlatformProperties::RequiresCookedData() ? InSoundWave->HasCompressedData(NAME_OGG) : (InSoundWave->GetCompressedData(NAME_OGG) != nullptr))
+#endif // WITH_BINK_AUDIO	
+		else
 		{
-			ICompressedAudioInfo* CompressedInfo = new FVorbisAudioInfo();
-			if (!CompressedInfo)
-			{
-				UE_LOG(LogAudio, Error, TEXT("Failed to create new FVorbisAudioInfo for SoundWave %s: out of memory."), *InSoundWave->GetFName().ToString());
-				return nullptr;
-			}
-			return CompressedInfo;
+			Decoder = Audio::CreateSoundAssetDecoder(InRuntimeFormat);
 		}
-
-#endif // WITH_ENGINE
-
-		return nullptr;
+		ensureMsgf(Decoder != nullptr, TEXT("Failed to create a sound asset decoder for compression type: %s"), *InRuntimeFormat.ToString());
+		return Decoder;
 	}
 
 	FString FMixerPlatformCoreAudio::GetDefaultDeviceName()
diff --git a/Engine/Source/Runtime/Apple/AudioMixerCoreAudio/Private/AudioMixerPlatformCoreAudio.h b/Engine/Source/Runtime/Apple/AudioMixerCoreAudio/Private/AudioMixerPlatformCoreAudio.h
--- a/Engine/Source/Runtime/Apple/AudioMixerCoreAudio/Private/AudioMixerPlatformCoreAudio.h
+++ b/Engine/Source/Runtime/Apple/AudioMixerCoreAudio/Private/AudioMixerPlatformCoreAudio.h
@@ -35,10 +35,8 @@ namespace Audio
 		virtual bool MoveAudioStreamToNewAudioDevice(const FString& InNewDeviceId) override;
 		virtual FAudioPlatformDeviceInfo GetPlatformDeviceInfo() const override;
 		virtual void SubmitBuffer(const uint8* Buffer) override;
-		virtual FName GetRuntimeFormat(USoundWave* InSoundWave) override;
-		virtual bool HasCompressedAudioInfoClass(USoundWave* InSoundWave) override;
-		virtual ICompressedAudioInfo* CreateCompressedAudioInfo(USoundWave* InSoundWave) override;
-		virtual ICompressedAudioInfo* CreateCompressedAudioInfo(const FSoundWaveProxyPtr& InSoundWave) override;
+		virtual FName GetRuntimeFormat(const USoundWave* InSoundWave) const override;
+		virtual ICompressedAudioInfo* CreateCompressedAudioInfo(const FName& InRuntimeFormat) const override;
 		virtual FString GetDefaultDeviceName() override;
 		virtual FAudioPlatformSettings GetPlatformSettings() const override;
         virtual int32 GetNumFrames(const int32 InNumReqestedFrames) override;
diff --git a/Engine/Source/Runtime/AudioMixer/Private/AudioMixerBuffer.cpp b/Engine/Source/Runtime/AudioMixer/Private/AudioMixerBuffer.cpp
--- a/Engine/Source/Runtime/AudioMixer/Private/AudioMixerBuffer.cpp
+++ b/Engine/Source/Runtime/AudioMixer/Private/AudioMixerBuffer.cpp
@@ -357,7 +357,7 @@ namespace Audio
 		{
 			// When set to seekable streaming, missing the first chunk is possible and
 			// does not signify any issue with the asset itself, so don't mark it as invalid.
-			if (InWave && !InWave->IsSeekableStreaming())
+			if (InWave && !InWave->IsSeekable())
 			{
 				UE_LOG(LogAudioMixer, Warning,
 					TEXT("FMixerBuffer::CreateStreamingBuffer failed to StreamCompressedInfo on SoundWave '%s'.  Invalidating wave resource data (asset now requires re-cook)."),
diff --git a/Engine/Source/Runtime/AudioMixer/Private/AudioMixerDevice.cpp b/Engine/Source/Runtime/AudioMixer/Private/AudioMixerDevice.cpp
--- a/Engine/Source/Runtime/AudioMixer/Private/AudioMixerDevice.cpp
+++ b/Engine/Source/Runtime/AudioMixer/Private/AudioMixerDevice.cpp
@@ -646,22 +646,35 @@ namespace Audio
 		return new FMixerSource(this);
 	}
 
-	FName FMixerDevice::GetRuntimeFormat(USoundWave* InSoundWave)
+	FName FMixerDevice::GetRuntimeFormat(const USoundWave* InSoundWave) const
 	{
-		check(AudioMixerPlatform);
-		return AudioMixerPlatform->GetRuntimeFormat(InSoundWave);
+		FName RuntimeFormat = Audio::ToName(InSoundWave->GetSoundAssetCompressionType());
+
+		// If not specified, the default platform codec is BINK.
+		if (RuntimeFormat == Audio::NAME_PLATFORM_SPECIFIC)
+		{
+			RuntimeFormat = AudioMixerPlatform->GetRuntimeFormat(InSoundWave);
+			if (RuntimeFormat.IsNone())
+			{
+				// TODO: make this an override per platform for "default" platform-specific codec
+				RuntimeFormat = Audio::NAME_BINKA;
+			}
+		}
+		return RuntimeFormat;
 	}
 
 	bool FMixerDevice::HasCompressedAudioInfoClass(USoundWave* InSoundWave)
 	{
 		check(InSoundWave);
 		check(AudioMixerPlatform);
-		return AudioMixerPlatform->HasCompressedAudioInfoClass(InSoundWave);
+		// Every platform has compressed audio. 
+		return true;
 	}
 
 	bool FMixerDevice::SupportsRealtimeDecompression() const
 	{
-		return AudioMixerPlatform->SupportsRealtimeDecompression();
+		// Every platform supports realtime decompression.
+		return true;
 	}
 
 	bool FMixerDevice::DisablePCMAudioCaching() const
@@ -669,17 +682,14 @@ namespace Audio
 		return AudioMixerPlatform->DisablePCMAudioCaching();
 	}
 
-	class ICompressedAudioInfo* FMixerDevice::CreateCompressedAudioInfo(USoundWave* InSoundWave)
+	ICompressedAudioInfo* FMixerDevice::CreateCompressedAudioInfo(const USoundWave* InSoundWave) const
 	{
-		check(InSoundWave);
-		check(AudioMixerPlatform);
-		return AudioMixerPlatform->CreateCompressedAudioInfo(InSoundWave);
+		return AudioMixerPlatform->CreateCompressedAudioInfo(GetRuntimeFormat(InSoundWave));
 	}
 
-	class ICompressedAudioInfo* FMixerDevice::CreateCompressedAudioInfo(const FSoundWaveProxyPtr& InSoundWave)
+	class ICompressedAudioInfo* FMixerDevice::CreateCompressedAudioInfo(const FSoundWaveProxyPtr& InSoundWaveProxy) const
 	{
-		check(AudioMixerPlatform);
-		return AudioMixerPlatform->CreateCompressedAudioInfo(InSoundWave);
+		return AudioMixerPlatform->CreateCompressedAudioInfo(InSoundWaveProxy->GetRuntimeFormat());
 	}
 
 	bool FMixerDevice::ValidateAPICall(const TCHAR* Function, uint32 ErrorCode)
diff --git a/Engine/Source/Runtime/AudioMixer/Public/AudioMixerDevice.h b/Engine/Source/Runtime/AudioMixer/Public/AudioMixerDevice.h
--- a/Engine/Source/Runtime/AudioMixer/Public/AudioMixerDevice.h
+++ b/Engine/Source/Runtime/AudioMixer/Public/AudioMixerDevice.h
@@ -82,6 +82,7 @@ namespace Audio
 		FCriticalSection MutationLock;
 	};
 
+
 	class AUDIOMIXER_API FMixerDevice :	public FAudioDevice,
 										public IAudioMixer,
 										public FGCObject
@@ -103,12 +104,12 @@ namespace Audio
 		virtual double GetAudioTime() const override;
 		virtual FAudioEffectsManager* CreateEffectsManager() override;
 		virtual FSoundSource* CreateSoundSource() override;
-		virtual FName GetRuntimeFormat(USoundWave* SoundWave) override;
+		virtual FName GetRuntimeFormat(const USoundWave* SoundWave) const override;
 		virtual bool HasCompressedAudioInfoClass(USoundWave* SoundWave) override;
 		virtual bool SupportsRealtimeDecompression() const override;
 		virtual bool DisablePCMAudioCaching() const override;
-		virtual class ICompressedAudioInfo* CreateCompressedAudioInfo(USoundWave* SoundWave) override;
-		virtual class ICompressedAudioInfo* CreateCompressedAudioInfo(const FSoundWaveProxyPtr& SoundWave) override;
+		virtual class ICompressedAudioInfo* CreateCompressedAudioInfo(const USoundWave* SoundWave) const override;
+		virtual class ICompressedAudioInfo* CreateCompressedAudioInfo(const FSoundWaveProxyPtr& SoundWave) const override;
 		virtual bool ValidateAPICall(const TCHAR* Function, uint32 ErrorCode) override;
 		virtual bool Exec(UWorld* InWorld, const TCHAR* Cmd, FOutputDevice& Ar) override;
 		virtual void CountBytes(class FArchive& Ar) override;
diff --git a/Engine/Source/Runtime/AudioMixerCore/Public/AudioMixer.h b/Engine/Source/Runtime/AudioMixerCore/Public/AudioMixer.h
--- a/Engine/Source/Runtime/AudioMixerCore/Public/AudioMixer.h
+++ b/Engine/Source/Runtime/AudioMixerCore/Public/AudioMixer.h
@@ -122,7 +122,6 @@ using FSoundWavePtr = TSharedPtr<FSoundWaveData, ESPMode::ThreadSafe>;
 
 namespace Audio
 {
-
    	/** Structure to hold platform device information **/
 	struct FAudioPlatformDeviceInfo
 	{
@@ -484,18 +483,9 @@ namespace Audio
 		/** Submit the given buffer to the platform's output audio device. */
 		virtual void SubmitBuffer(const uint8* Buffer) {};
 
-		/** Returns the name of the format of the input sound wave. */
-		virtual FName GetRuntimeFormat(USoundWave* InSoundWave) = 0;
-
 		/** Allows platforms to filter the requested number of frames to render. Some platforms only support specific frame counts. */
 		virtual int32 GetNumFrames(const int32 InNumReqestedFrames) { return InNumReqestedFrames; }
 
-		/** Checks if the platform has a compressed audio format for sound waves. */
-		virtual bool HasCompressedAudioInfoClass(USoundWave* InSoundWave) = 0;
-
-		/** Whether or not the platform supports realtime decompression. */
-		virtual bool SupportsRealtimeDecompression() const { return false; }
-
 		/** Whether or not the platform disables caching of decompressed PCM data (i.e. to save memory on fixed memory platforms) */
 		virtual bool DisablePCMAudioCaching() const { return false; }
 
@@ -505,9 +495,11 @@ namespace Audio
 		/** Whether this is an interface for a non-realtime renderer. If true, synch events will behave differently to avoid deadlocks. */
 		virtual bool IsNonRealtime() const { return false; }
 
-		/** Creates a Compressed audio info class suitable for decompressing this SoundWave. */
-		virtual ICompressedAudioInfo* CreateCompressedAudioInfo(USoundWave* SoundWave) = 0;
-		virtual ICompressedAudioInfo* CreateCompressedAudioInfo(const FSoundWaveProxyPtr& SoundWave) { return nullptr; }
+		/** Returns the FName version of a given sound wave's format. Override to provide a runtime format (codec) that is platform specific. */
+		virtual FName GetRuntimeFormat(const USoundWave* InSoundWave) const { return FName(); }
+
+		/** Creates a Compressed audio info class suitable for decompressing this SoundWave. OVerride to create platform-specific decoders. */
+		virtual ICompressedAudioInfo* CreateCompressedAudioInfo(const FName& InRuntimeFormat) const { return nullptr; };
 
 		/** Return any optional device name defined in platform configuratio. */
 		virtual FString GetDefaultDeviceName() = 0;
diff --git a/Engine/Source/Runtime/AudioMixerCore/Public/AudioMixerTypes.h b/Engine/Source/Runtime/AudioMixerCore/Public/AudioMixerTypes.h
--- a/Engine/Source/Runtime/AudioMixerCore/Public/AudioMixerTypes.h
+++ b/Engine/Source/Runtime/AudioMixerCore/Public/AudioMixerTypes.h
@@ -58,6 +58,17 @@ namespace Audio {
 		};
 	}
 
+	// Indicates a platform-specific format
+	static FName NAME_PLATFORM_SPECIFIC(TEXT("PLATFORM_SPECIFIC"));
+
+	// Supported on all platforms:
+	static FName NAME_BINKA(TEXT("BINKA"));
+	static FName NAME_ADPCM(TEXT("ADPCM"));
+	static FName NAME_PCM(TEXT("PCM"));
+
+	// Not yet supported on all platforms as a selectable option so is included under "platform specific" enumeration for now. 
+	static FName NAME_OGG(TEXT("OGG"));
+	static FName NAME_OPUS(TEXT("OPUS"));
 }
 
 struct AUDIOMIXERCORE_API FAudioPlatformSettings
diff --git a/Engine/Source/Runtime/Engine/Classes/Sound/SoundWave.h b/Engine/Source/Runtime/Engine/Classes/Sound/SoundWave.h
--- a/Engine/Source/Runtime/Engine/Classes/Sound/SoundWave.h
+++ b/Engine/Source/Runtime/Engine/Classes/Sound/SoundWave.h
@@ -283,6 +283,42 @@ enum class ESoundWaveFFTSize : uint8
 	VeryLarge_2048,
 };
 
+// Sound Asset Compression Type
+UENUM(BlueprintType)
+enum class ESoundAssetCompressionType : uint8
+{
+	// Perceptual-based codec which supports all features across all platforms.
+	BinkAudio,
+
+	// Will encode the asset using ADPCM, a time-domain codec that has fixed-sized quality and about ~4x compression ratio, but is relatively cheap to decode.
+	ADPCM,
+
+	// Uncompressed audio. Large memory usage (streamed chunks contain less audio per chunk) but extremely cheap to decode and supports all features. 
+	PCM,
+
+	// Encodes the asset to a platform specific format and will be different depending on the platform. It does not currently support seeking.
+	PlatformSpecific,
+};
+
+
+namespace Audio
+{
+	static FName ToName(ESoundAssetCompressionType InDecoderType)
+	{
+		switch (InDecoderType)
+		{
+		case ESoundAssetCompressionType::PlatformSpecific:		return NAME_PLATFORM_SPECIFIC;
+
+		case ESoundAssetCompressionType::BinkAudio:				return NAME_BINKA;
+		case ESoundAssetCompressionType::ADPCM:					return NAME_ADPCM;
+		case ESoundAssetCompressionType::PCM:					return NAME_PCM;
+		default:
+			ensure(false);
+			return TEXT("UNKNOWN");
+		}
+	}
+}
+
 // Struct defining a cue point in a sound wave asset
 USTRUCT()
 struct FSoundWaveCuePoint
@@ -316,18 +352,19 @@ struct ISoundWaveClient
 	virtual bool OnIsReadyForFinishDestroy(class USoundWave* Wave) const = 0;
 	virtual void OnFinishDestroy(class USoundWave* Wave) = 0;
 };
-
-
 UCLASS(hidecategories=Object, editinlinenew, BlueprintType, meta= (LoadBehavior = "LazyOnDemand"))
 class ENGINE_API USoundWave : public USoundBase, public IAudioProxyDataFactory
 {
 	GENERATED_UCLASS_BODY()
 
-public:
-	/** Platform agnostic compression quality. 1..100 with 1 being best compression and 100 being best quality. */
-	UPROPERTY(EditAnywhere, Category = "Format|Quality", meta = (DisplayName = "Compression", ClampMin = "1", ClampMax = "100"), AssetRegistrySearchable)
+private:
+
+	/** Platform agnostic compression quality. 1..100 with 1 being best compression and 100 being best quality. ADPCM and PCM sound asset compression types ignore this parameter. */
+	UPROPERTY(EditAnywhere, Category = "Format|Quality", meta = (DisplayName = "Compression", ClampMin = "1", ClampMax = "100", EditCondition = "SoundAssetCompressionType != ESoundAssetCompressionType::PCM || SoundAssetCompressionType != ESoundAssetCompressionType::ADPCM"), AssetRegistrySearchable)
 	int32 CompressionQuality;
 
+public:
+
 	/** Priority of this sound when streaming (lower priority streams may not always play) */
 	UPROPERTY(EditAnywhere, Category = "Playback|Streaming", meta = (ClampMin = 0))
 	int32 StreamingPriority;
@@ -346,18 +383,25 @@ public:
 	UPROPERTY(EditAnywhere, Category = Sound, AssetRegistrySearchable)
 	uint8 bLooping : 1;
 
-	/** Whether this sound can be streamed to avoid increased memory usage. If using Stream Caching, use Loading Behavior instead to control memory usage. */
-	UPROPERTY(EditAnywhere, Category = "Playback|Streaming", meta = (DisplayName = "Force Streaming"))
+	/** Here for legacy code. */
+	UPROPERTY()
 	uint8 bStreaming : 1;
 
-	/** Whether this sound supports seeking. This requires recooking with a codec which supports seekability and streaming. */
-	UPROPERTY(EditAnywhere, Category = "Playback|Streaming", meta = (DisplayName = "Seekable", EditCondition = "bStreaming"))
-	uint8 bSeekableStreaming : 1;
+private:
 
-	/** If true, the sound will compress to the Bink Audio format whenever available. */
+	/** The compression type to use for the sound wave asset. */
 	UPROPERTY(EditAnywhere, Category = "Format")
+	ESoundAssetCompressionType SoundAssetCompressionType = ESoundAssetCompressionType::PlatformSpecific;
+
+	// Deprecated compression type properties
+	UPROPERTY(meta = (DeprecatedProperty, DeprecationMessage = "5.0 - Property is deprecated. bSeekableStreaming now means ADPCM codec in SoundAssetCompressionType."))
+	uint8 bSeekableStreaming : 1;
+
+	UPROPERTY(meta = (DeprecatedProperty, DeprecationMessage = "5.0 - Property is deprecated. bUseBinkAudio now means Bink codec in SoundAssetCompressionType."))
 	uint8 bUseBinkAudio : 1;
 
+public:
+
 	/** the number of sounds currently playing this sound wave. */
 	FThreadSafeCounter NumSourcesPlaying;
 
@@ -372,7 +416,8 @@ public:
 		NumSourcesPlaying.Decrement();
 	}
 
-
+	/** Returns the sound's asset compression type. */
+	ESoundAssetCompressionType GetSoundAssetCompressionType() const;
 
 private:
 	// cached proxy
@@ -633,6 +678,9 @@ protected:
 
 public:
 
+	// Returns the compression quality of the sound asset.
+	int32 GetCompressionQuality() const;
+
 	/** Resource index to cross reference with buffers */
 	int32 ResourceID;
 
@@ -952,8 +1000,9 @@ public:
 	bool IsStreaming(const TCHAR* PlatformName = nullptr) const;
 	bool IsStreaming(const FPlatformAudioCookOverrides& Overrides) const;
 
-	/** Checks whether sound has seekable streaming enabled. */
-	bool IsSeekableStreaming() const;
+	/** Returns whether the sound is seekable. */
+	bool IsSeekable() const;
+
 	/**
 	 * Checks whether we should use the load on demand cache.
 	 */
@@ -1106,11 +1155,11 @@ public:
 
 	void ReleaseCompressedAudio();
 
-	bool UseBinkAudio() const { return bUseBinkAudio; }
 	bool IsStreaming() const { return bIsStreaming; }
+	ESoundAssetCompressionType GetSoundCompressionType() const { return SoundAssetCompressionType; }
 	bool IsLooping() const { return bIsLooping; }
 	bool ShouldUseStreamCaching() const { return bShouldUseStreamCaching; }
-	bool IsSeekableStreaming() const {  return bIsStreaming && bSeekableStreaming; }
+	bool IsSeekable() const {  return bIsSeekable; }
 	bool IsRetainingAudio() const {  return FirstChunk.IsValid(); }
 	bool WasLoadingBehaviorOverridden() const { return bLoadingBehaviorOverridden; }
 	ESoundWaveLoadingBehavior GetLoadingBehavior() const { return LoadingBehavior; }
@@ -1168,6 +1217,7 @@ private:
 	FName RuntimeFormat{ "FSoundWaveProxy_InvalidFormat" };
 	FObjectKey SoundWaveKeyCached;
 	TArray<FSoundWaveCuePoint> CuePoints;
+	ESoundAssetCompressionType SoundAssetCompressionType;
 
 	float SampleRate;
 	float Duration;
@@ -1180,8 +1230,7 @@ private:
 	uint8 bIsLooping : 1;
 	uint8 bIsTemplate : 1;
 	uint8 bIsStreaming : 1;
-	uint8 bUseBinkAudio : 1;
-	uint8 bSeekableStreaming : 1;
+	uint8 bIsSeekable : 1;
 	uint8 bShouldUseStreamCaching : 1;
 	uint8 bLoadingBehaviorOverridden : 1;
 
@@ -1234,7 +1283,7 @@ public:
 	bool UseBinkAudio() const;
 	bool IsRetainingAudio() const;
 	bool ShouldUseStreamCaching() const;
-	bool IsSeekableStreaming() const;
+	bool IsSeekable() const;
 	bool IsZerothChunkDataLoaded() const;
 	bool WasLoadingBehaviorOverridden() const;
 	bool HasCompressedData(FName Format, ITargetPlatform* TargetPlatform = USoundWave::GetRunningPlatform()) const;
diff --git a/Engine/Source/Runtime/Engine/Private/Audio.cpp b/Engine/Source/Runtime/Engine/Private/Audio.cpp
--- a/Engine/Source/Runtime/Engine/Private/Audio.cpp
+++ b/Engine/Source/Runtime/Engine/Private/Audio.cpp
@@ -1080,17 +1080,7 @@ bool FWaveInstance::IsSeekable() const
 		return false;
 	}
 
-	if (WaveData->bIsSourceBus || WaveData->bProcedural)
-	{
-		return false;
-	}
-
-	if (IsStreaming() && !WaveData->IsSeekableStreaming())
-	{
-		return false;
-	}
-
-	return true;
+	return WaveData->IsSeekable();
 }
 
 bool FWaveInstance::IsStreaming() const
diff --git a/Engine/Source/Runtime/Engine/Private/AudioDerivedData.cpp b/Engine/Source/Runtime/Engine/Private/AudioDerivedData.cpp
--- a/Engine/Source/Runtime/Engine/Private/AudioDerivedData.cpp
+++ b/Engine/Source/Runtime/Engine/Private/AudioDerivedData.cpp
@@ -1143,20 +1143,20 @@ static void CookSimpleWave(USoundWave* SoundWave, FName FormatName, const IAudio
 		FSoundQualityInfo QualityInfo = { 0 };
 		if (CompressionOverrides)
 		{
-			float ModifiedCompressionQuality = (float)SoundWave->CompressionQuality * CompressionOverrides->CompressionQualityModifier;
+			float ModifiedCompressionQuality = (float)SoundWave->GetCompressionQuality() * CompressionOverrides->CompressionQualityModifier;
 			if (ModifiedCompressionQuality >= 1.0f)
 			{
 				QualityInfo.Quality = FMath::FloorToInt(ModifiedCompressionQuality);
-				UE_LOG(LogAudioDerivedData, Display, TEXT("Compression Quality for %s will be modified from %d to %d."), *SoundWave->GetFullName(), SoundWave->CompressionQuality, QualityInfo.Quality);
+				UE_LOG(LogAudioDerivedData, Display, TEXT("Compression Quality for %s will be modified from %d to %d."), *SoundWave->GetFullName(), SoundWave->GetCompressionQuality(), QualityInfo.Quality);
 			}
 			else
 			{
-				QualityInfo.Quality = SoundWave->CompressionQuality;
+				QualityInfo.Quality = SoundWave->GetCompressionQuality();
 			}
 		}
 		else
 		{
-			QualityInfo.Quality = SoundWave->CompressionQuality;
+			QualityInfo.Quality = SoundWave->GetCompressionQuality();
 		}
 
 		QualityInfo.NumChannels = *WaveInfo.pChannels;
@@ -1366,11 +1366,11 @@ static void CookSurroundWave( USoundWave* SoundWave, FName FormatName, const IAu
 				CompressionQualityModifier = CompressionOverrides->CompressionQualityModifier;
 			}
 
-			float ModifiedCompressionQuality = (float)SoundWave->CompressionQuality;
+			float ModifiedCompressionQuality = (float)SoundWave->GetCompressionQuality();
 
 			if (!FMath::IsNearlyEqual(CompressionQualityModifier, 1.0f))
 			{
-				ModifiedCompressionQuality = (float)SoundWave->CompressionQuality * CompressionQualityModifier;
+				ModifiedCompressionQuality = (float)SoundWave->GetCompressionQuality() * CompressionQualityModifier;
 			}
 			
 			if (ModifiedCompressionQuality >= 1.0f)
@@ -1379,7 +1379,7 @@ static void CookSurroundWave( USoundWave* SoundWave, FName FormatName, const IAu
 			}
 			else
 			{
-				QualityInfo.Quality = SoundWave->CompressionQuality;
+				QualityInfo.Quality = SoundWave->GetCompressionQuality();
 			}
 
 			QualityInfo.NumChannels = ChannelCount;
diff --git a/Engine/Source/Runtime/Engine/Private/AudioDevice.cpp b/Engine/Source/Runtime/Engine/Private/AudioDevice.cpp
--- a/Engine/Source/Runtime/Engine/Private/AudioDevice.cpp
+++ b/Engine/Source/Runtime/Engine/Private/AudioDevice.cpp
@@ -3,6 +3,7 @@
 
 #include "ActiveSound.h"
 #include "ActiveSoundUpdateInterface.h"
+#include "ADPCMAudioInfo.h"
 #include "Audio.h"
 #include "AudioCompressionSettingsUtils.h"
 #include "AudioDecompress.h"
@@ -10,6 +11,7 @@
 #include "AudioEffect.h"
 #include "AudioPluginUtilities.h"
 #include "Audio/AudioDebug.h"
+#include "AudioMixerDevice.h"
 #include "Components/AudioComponent.h"
 #include "ContentStreaming.h"
 #include "DrawDebugHelpers.h"
@@ -43,7 +45,6 @@
 #include "UObject/UObjectHash.h"
 #include "UObject/UObjectIterator.h"
 #include "UObject/Package.h"
-#include "AudioMixerDevice.h"
 
 #if WITH_EDITOR
 #include "AudioEditorModule.h"
@@ -176,6 +177,26 @@ FAutoConsoleVariableRef CVarFlushAudioRenderThreadOnGC(
 	TEXT("When set to 1, every time the GC runs, we flush all pending audio render thread commands.\n"),
 	ECVF_Default);
 
+namespace Audio
+{
+	ICompressedAudioInfo* CreateSoundAssetDecoder(const FName& InRuntimeFormat)
+	{
+		// The decoder will need to be instantiated for the specific platform
+		if (InRuntimeFormat == Audio::NAME_PLATFORM_SPECIFIC)
+		{
+			return nullptr;
+		}
+		// Multi-platform decoders -- can be instantiated here for all platforms.
+		else if (InRuntimeFormat == Audio::NAME_ADPCM || InRuntimeFormat == Audio::NAME_PCM)
+		{
+			// Our PCM is currently really our "ADPCM" encoder at 100 quality
+			return new FADPCMAudioInfo();
+		}
+		ensureMsgf(false, TEXT("Unknown runtime sound asset format type."));
+		return nullptr;
+	}
+}
+
 namespace AudioDeviceUtils
 {
 	using FVirtualLoopPair = TPair<FActiveSound*, FAudioVirtualLoop>;
diff --git a/Engine/Source/Runtime/Engine/Private/SoundWave.cpp b/Engine/Source/Runtime/Engine/Private/SoundWave.cpp
--- a/Engine/Source/Runtime/Engine/Private/SoundWave.cpp
+++ b/Engine/Source/Runtime/Engine/Private/SoundWave.cpp
@@ -184,8 +184,7 @@ void FSoundWaveData::InitializeDataFromSoundWave(USoundWave& InWave)
 	bIsLooping = InWave.IsLooping();
 	bIsTemplate = InWave.IsTemplate();
 	bIsStreaming = InWave.IsStreaming(nullptr);
-	bUseBinkAudio = InWave.bUseBinkAudio;
-	bSeekableStreaming = InWave.bSeekableStreaming;
+	SoundAssetCompressionType = InWave.GetSoundAssetCompressionType();
 	bShouldUseStreamCaching = InWave.ShouldUseStreamCaching();
 }
 
@@ -732,12 +731,36 @@ void USoundWave::Serialize( FArchive& Ar )
 	bool bShouldStreamSound = false;
 
 #if WITH_EDITORONLY_DATA
-		bLoadedFromCookedData = Ar.IsLoading() && bCooked;
-		if (bVirtualizeWhenSilent_DEPRECATED)
+	bLoadedFromCookedData = Ar.IsLoading() && bCooked;
+	if (bVirtualizeWhenSilent_DEPRECATED)
+	{
+		bVirtualizeWhenSilent_DEPRECATED = 0;
+		VirtualizationMode = EVirtualizationMode::PlayWhenSilent;
+	}
+
+	// If we have deprecated properties with data, migrate the data to the new enum
+	if (bUseBinkAudio || bSeekableStreaming)
+	{
+		// Bink audio takes precedence over bSeekableStreaming
+		if (bUseBinkAudio)
 		{
-			bVirtualizeWhenSilent_DEPRECATED = 0;
-			VirtualizationMode = EVirtualizationMode::PlayWhenSilent;
+			SoundAssetCompressionType = ESoundAssetCompressionType::BinkAudio;
 		}
+		else if (bSeekableStreaming)
+		{
+			// PCM was CompressionQuality 100
+			if (CompressionQuality == 100)
+			{
+				SoundAssetCompressionType = ESoundAssetCompressionType::PCM;
+			}
+			else
+			{
+				SoundAssetCompressionType = ESoundAssetCompressionType::ADPCM;
+			}
+		}
+		bUseBinkAudio = 0;
+		bSeekableStreaming = 0;
+	}
 #endif // WITH_EDITORONLY_DATA
 
 	if (Ar.IsSaving() || Ar.IsCooking())
@@ -895,6 +918,11 @@ void USoundWave::Serialize( FArchive& Ar )
 	}
 }
 
+ESoundAssetCompressionType USoundWave::GetSoundAssetCompressionType() const
+{
+	return SoundAssetCompressionType;
+}
+
 float USoundWave::GetSubtitlePriority() const
 {
 	return SubtitlePriority;
@@ -2182,13 +2210,12 @@ void USoundWave::PostEditChangeProperty(FPropertyChangedEvent& PropertyChangedEv
 	static const FName CompressionQualityFName = GET_MEMBER_NAME_CHECKED(USoundWave, CompressionQuality);
 	static const FName SampleRateFName = GET_MEMBER_NAME_CHECKED(USoundWave,SampleRateQuality);
 	static const FName StreamingFName = GET_MEMBER_NAME_CHECKED(USoundWave, bStreaming);
-	static const FName SeekableStreamingFName = GET_MEMBER_NAME_CHECKED(USoundWave, bSeekableStreaming);
-	static const FName UseBinkAudioFName = GET_MEMBER_NAME_CHECKED(USoundWave, bUseBinkAudio);
+	static const FName SoundAssetCompressionTypeFName = GET_MEMBER_NAME_CHECKED(USoundWave, SoundAssetCompressionType);
 	static const FName LoadingBehaviorFName = GET_MEMBER_NAME_CHECKED(USoundWave, LoadingBehavior);
 
 	// force proxy flags to be up to date
-	SoundWaveDataPtr->bUseBinkAudio = bUseBinkAudio;
-	SoundWaveDataPtr->bSeekableStreaming = bSeekableStreaming;
+	SoundWaveDataPtr->bIsSeekable = IsSeekable();
+	SoundWaveDataPtr->SoundAssetCompressionType = SoundAssetCompressionType;
 	SoundWaveDataPtr->bIsStreaming = IsStreaming(nullptr);
 	SoundWaveDataPtr->bShouldUseStreamCaching = ShouldUseStreamCaching();
 
@@ -2222,7 +2249,7 @@ void USoundWave::PostEditChangeProperty(FPropertyChangedEvent& PropertyChangedEv
 				}
 			}
 
-			if (Name == CompressionQualityFName || Name == SampleRateFName || Name == StreamingFName || Name == SeekableStreamingFName || Name == UseBinkAudioFName || Name == LoadingBehaviorFName)
+			if (Name == CompressionQualityFName || Name == SampleRateFName || Name == StreamingFName || Name == SoundAssetCompressionTypeFName || Name == LoadingBehaviorFName)
 			{
 				InvalidateCompressedData();
 				FreeResources();
@@ -2934,9 +2961,13 @@ TArrayView<const uint8> USoundWave::GetZerothChunk(bool bForImmediatePlayback)
 	}
 }
 
-bool USoundWave::IsSeekableStreaming() const
+bool USoundWave::IsSeekable() const
 {
-	return bStreaming && bSeekableStreaming;
+	if (bIsSourceBus || bProcedural || SoundAssetCompressionType == ESoundAssetCompressionType::PlatformSpecific)
+	{
+		return false;
+	}
+	return true;
 }
 
 bool USoundWave::GetSoundWavesWithCookedAnalysisData(TArray<USoundWave*>& OutSoundWaves)
@@ -2959,6 +2990,20 @@ bool USoundWave::HasCookedAmplitudeEnvelopeData() const
 	return CookedEnvelopeTimeData.Num() > 0;
 }
 
+int32 USoundWave::GetCompressionQuality() const
+{
+	if (SoundAssetCompressionType == ESoundAssetCompressionType::PCM)
+	{
+		// PCM is treated as ADPCM at 100-quality
+		return 100;
+	}
+	else
+	{
+		return CompressionQuality;
+	}
+}
+
+
 FSoundWaveProxyPtr USoundWave::CreateSoundWaveProxy()
 {
 #if WITH_EDITORONLY_DATA
@@ -3829,12 +3874,6 @@ bool FSoundWaveProxy::IsStreaming() const
 	return SoundWaveDataPtr->IsStreaming();
 }
 
-bool FSoundWaveProxy::UseBinkAudio() const
-{
-	check(SoundWaveDataPtr);
-	return SoundWaveDataPtr->UseBinkAudio();
-}
-
 bool FSoundWaveProxy::IsRetainingAudio() const
 {
 	check(SoundWaveDataPtr);
@@ -3847,12 +3886,13 @@ bool FSoundWaveProxy::ShouldUseStreamCaching() const
 	return SoundWaveDataPtr->ShouldUseStreamCaching();
 }
 
-bool FSoundWaveProxy::IsSeekableStreaming() const
+bool FSoundWaveProxy::IsSeekable() const
 {
 	check(SoundWaveDataPtr);
-	return SoundWaveDataPtr->IsSeekableStreaming();
+	return SoundWaveDataPtr->IsSeekable();
 }
 
+
 bool FSoundWaveProxy::IsZerothChunkDataLoaded() const
 {
 	check(SoundWaveDataPtr);
diff --git a/Engine/Source/Runtime/Engine/Public/AudioDevice.h b/Engine/Source/Runtime/Engine/Public/AudioDevice.h
--- a/Engine/Source/Runtime/Engine/Public/AudioDevice.h
+++ b/Engine/Source/Runtime/Engine/Public/AudioDevice.h
@@ -65,6 +65,9 @@ struct FWaveInstance;
 namespace Audio
 {
 	class FAudioDebugger;
+
+	/** Returns a decoder for the sound assets's given compression type. Will return nullptr if the compression type is platform-dependent. */
+	ENGINE_API ICompressedAudioInfo* CreateSoundAssetDecoder(const FName& InRuntimeFormat);
 }
 
 /**
@@ -1099,7 +1102,7 @@ public:
 	 */
 	void DeactivateReverbEffect(FName TagName);
 
-	virtual FName GetRuntimeFormat(USoundWave* SoundWave) = 0;
+	virtual FName GetRuntimeFormat(const USoundWave* SoundWave) const = 0;
 
 	/** Whether this SoundWave has an associated info class to decompress it */
 	virtual bool HasCompressedAudioInfoClass(USoundWave* SoundWave) { return false; }
@@ -1117,8 +1120,8 @@ public:
 	}
 
 	/** Creates a Compressed audio info class suitable for decompressing this SoundWave */
-	virtual ICompressedAudioInfo* CreateCompressedAudioInfo(USoundWave* SoundWave) { return nullptr; }
-	virtual ICompressedAudioInfo* CreateCompressedAudioInfo(const FSoundWaveProxyPtr& SoundWave) { return nullptr; }
+	virtual ICompressedAudioInfo* CreateCompressedAudioInfo(const USoundWave* SoundWave) const { return nullptr; }
+	virtual ICompressedAudioInfo* CreateCompressedAudioInfo(const FSoundWaveProxyPtr& SoundWave) const { return nullptr; }
 
 	/**
 	 * Check for errors and output a human readable string
diff --git a/Engine/Source/Runtime/Linux/AudioMixerSDL/Private/AudioMixerPlatformSDL.cpp b/Engine/Source/Runtime/Linux/AudioMixerSDL/Private/AudioMixerPlatformSDL.cpp
--- a/Engine/Source/Runtime/Linux/AudioMixerSDL/Private/AudioMixerPlatformSDL.cpp
+++ b/Engine/Source/Runtime/Linux/AudioMixerSDL/Private/AudioMixerPlatformSDL.cpp
@@ -6,6 +6,7 @@
 #include "AudioMixerDevice.h"
 #include "CoreGlobals.h"
 #include "Misc/ConfigCacheIni.h"
+#include "AudioDevice.h"
 
 #if WITH_ENGINE
 #include "AudioPluginUtilities.h"
@@ -330,112 +331,49 @@ namespace Audio
 		ReadNextBuffer();
 	}
 
-	FName FMixerPlatformSDL::GetRuntimeFormat(USoundWave* InSoundWave)
+	FName FMixerPlatformSDL::GetRuntimeFormat(const USoundWave* InSoundWave) const
 	{
-#if WITH_ENGINE
-		static FName NAME_OGG(TEXT("OGG"));
-		static FName NAME_OPUS(TEXT("OPUS"));
-		static FName NAME_ADPCM(TEXT("ADPCM"));
-		
-#if WITH_BINK_AUDIO
-		static FName NAME_BINKA(TEXT("BINKA"));
-		if (InSoundWave->bUseBinkAudio)
-		{
-			return NAME_BINKA;
-		}
-#endif // WITH_BINK_AUDIO
-		if (InSoundWave->IsStreaming(nullptr))
+		FName RuntimeFormat = Audio::ToName(InSoundWave->GetSoundAssetCompressionType());
+
+		if (RuntimeFormat == Audio::NAME_PLATFORM_SPECIFIC)
 		{
-			if (InSoundWave->IsSeekableStreaming())
+			if (InSoundWave->HasCompressedData(Audio::NAME_OGG))
 			{
-				return NAME_ADPCM;
+				RuntimeFormat = Audio::NAME_OGG;
 			}
-
-			return NAME_OPUS;
-		}
-		return NAME_OGG;
-#else
-		checkNoEntry();
-		return FName();
-#endif // WITH_ENGINE
-	}
-
-	bool FMixerPlatformSDL::HasCompressedAudioInfoClass(USoundWave* InSoundWave)
-	{
-		return true;
-	}
-
-	ICompressedAudioInfo* FMixerPlatformSDL::CreateCompressedAudioInfo(USoundWave* InSoundWave)
-	{
-#if WITH_ENGINE
-		check(InSoundWave);
-
-#if WITH_BINK_AUDIO
-		if (InSoundWave->bUseBinkAudio)
-		{
-			return new FBinkAudioInfo();
-		}
-#endif // WITH_BINK_AUDIO
-
-		if (InSoundWave->IsStreaming())
-		{
-			if (InSoundWave->IsSeekableStreaming())
+			else
 			{
-				return new FADPCMAudioInfo();
+				RuntimeFormat = Audio::NAME_OPUS;
 			}
-
-			return new FOpusAudioInfo();
-		}
-
-		static FName NAME_OGG(TEXT("OGG"));
-		if (InSoundWave->HasCompressedData(NAME_OGG))
-		{
-			return new FVorbisAudioInfo();
 		}
 
-		return new FADPCMAudioInfo();
-#else
-		checkNoEntry();
-		return nullptr;
-#endif // WITH_ENGINE
+		return RuntimeFormat;
 	}
 
-	ICompressedAudioInfo* FMixerPlatformSDL::CreateCompressedAudioInfo(const FSoundWaveProxyPtr& InSoundWave)
+	ICompressedAudioInfo* FMixerPlatformSDL::CreateCompressedAudioInfo(const FName& InRuntimeFormat) const
 	{
-#if WITH_ENGINE
-		if (!InSoundWave.IsValid())
+		ICompressedAudioInfo* Decoder = nullptr;
+
+		if (InRuntimeFormat == Audio::NAME_OGG)
 		{
-			return nullptr;
+			Decoder = new FVorbisAudioInfo();
 		}
-
-#if WITH_BINK_AUDIO
-		if (InSoundWave->UseBinkAudio())
+		else if (InRuntimeFormat == Audio::NAME_OPUS)
 		{
-			return new FBinkAudioInfo();
+			Decoder = new FOpusAudioInfo();
 		}
-#endif // WITH_BINK_AUDIO
-
-		if (InSoundWave->IsStreaming())
+#if WITH_BINK_AUDIO
+		else if (InRuntimeFormat == Audio::NAME_BINKA)
 		{
-			if (InSoundWave->IsSeekableStreaming())
-			{
-				return new FADPCMAudioInfo();
-			}
-
-			return new FOpusAudioInfo();
+			Decoder = new FBinkAudioInfo();
 		}
-
-		static FName NAME_OGG(TEXT("OGG"));
-		if (InSoundWave->HasCompressedData(NAME_OGG))
+#endif // WITH_BINK_AUDIO	
+		else
 		{
-			return new FVorbisAudioInfo();
+			Decoder = Audio::CreateSoundAssetDecoder(InRuntimeFormat);
 		}
-
-		return new FADPCMAudioInfo();
-#else
-		checkNoEntry();
-		return nullptr;
-#endif // WITH_ENGINE
+		ensureMsgf(Decoder != nullptr, TEXT("Failed to create a sound asset decoder for compression type: %s"), *InRuntimeFormat.ToString());
+		return Decoder;
 	}
 
 	FString FMixerPlatformSDL::GetDefaultDeviceName()
diff --git a/Engine/Source/Runtime/Linux/AudioMixerSDL/Private/AudioMixerPlatformSDL.h b/Engine/Source/Runtime/Linux/AudioMixerSDL/Private/AudioMixerPlatformSDL.h
--- a/Engine/Source/Runtime/Linux/AudioMixerSDL/Private/AudioMixerPlatformSDL.h
+++ b/Engine/Source/Runtime/Linux/AudioMixerSDL/Private/AudioMixerPlatformSDL.h
@@ -34,10 +34,8 @@ namespace Audio
 		bool StopAudioStream() override;
 		FAudioPlatformDeviceInfo GetPlatformDeviceInfo() const override;
 		void SubmitBuffer(const uint8* Buffer) override;
-		FName GetRuntimeFormat(USoundWave* InSoundWave) override;
-		bool HasCompressedAudioInfoClass(USoundWave* InSoundWave) override;
-		ICompressedAudioInfo* CreateCompressedAudioInfo(USoundWave* InSoundWave) override;
-		ICompressedAudioInfo* CreateCompressedAudioInfo(const FSoundWaveProxyPtr& InSoundWave) override;
+		FName GetRuntimeFormat(const USoundWave* InSoundWave) const override;
+		ICompressedAudioInfo* CreateCompressedAudioInfo(const FName& InRuntimeFormat) const override;
 		FString GetDefaultDeviceName() override;
 		FAudioPlatformSettings GetPlatformSettings() const override;
 		void ResumeContext() override;
diff --git a/Engine/Source/Runtime/NonRealtimeAudioRenderer/NonRealtimeAudioRenderer.Build.cs b/Engine/Source/Runtime/NonRealtimeAudioRenderer/NonRealtimeAudioRenderer.Build.cs
--- a/Engine/Source/Runtime/NonRealtimeAudioRenderer/NonRealtimeAudioRenderer.Build.cs
+++ b/Engine/Source/Runtime/NonRealtimeAudioRenderer/NonRealtimeAudioRenderer.Build.cs
@@ -22,6 +22,7 @@ namespace UnrealBuildTool.Rules
                     "Core",
                     "CoreUObject",
                     "Engine",
+                    "BinkAudioDecoder",
                     "AudioMixerCore",
                     "SignalProcessing",
                     "AudioMixer",
diff --git a/Engine/Source/Runtime/NonRealtimeAudioRenderer/Private/AudioMixerPlatformNonRealtime.cpp b/Engine/Source/Runtime/NonRealtimeAudioRenderer/Private/AudioMixerPlatformNonRealtime.cpp
--- a/Engine/Source/Runtime/NonRealtimeAudioRenderer/Private/AudioMixerPlatformNonRealtime.cpp
+++ b/Engine/Source/Runtime/NonRealtimeAudioRenderer/Private/AudioMixerPlatformNonRealtime.cpp
@@ -18,10 +18,14 @@
 #if WITH_XMA2
 #include "XMAAudioInfo.h"
 #endif  //#if WITH_XMA2
+#if WITH_BINK_AUDIO
+#include "BinkAudioInfo.h"
+#endif // WITH_BINK_AUDIO
 #include "OpusAudioInfo.h"
 #include "VorbisAudioInfo.h"
 #include "ADPCMAudioInfo.h"
 #include "Interfaces/IAudioFormat.h"
+#include "AudioDevice.h"
 
 #include "CoreGlobals.h"
 #include "Misc/ConfigCacheIni.h"
@@ -304,155 +308,61 @@ namespace Audio
 		}
 	}
 
-	FName FMixerPlatformNonRealtime::GetRuntimeFormat(USoundWave* InSoundWave)
-	{
-		//TODO: Set this up to propogate to 
-
-		static FName NAME_OGG(TEXT("OGG"));
-		static FName NAME_OPUS(TEXT("OPUS"));
-		static FName NAME_XMA(TEXT("XMA"));
-		static FName NAME_ADPCM(TEXT("ADPCM"));
-
-		if (InSoundWave->IsSeekableStreaming())
-		{
-			return NAME_ADPCM;
-		}
-
-#if WITH_XMA2 && USE_XMA2_FOR_STREAMING
-		if (InSoundWave->IsStreaming(nullptr) && InSoundWave->NumChannels <= 2)
-		{
-			return NAME_XMA;
-		}
-#endif
-
-		if (InSoundWave->IsStreaming(nullptr))
-		{
-#if USE_VORBIS_FOR_STREAMING
-			return NAME_OGG;
-#else
-			return NAME_OPUS;
-#endif
-		}
-
 #if WITH_XMA2
-		if (InSoundWave->NumChannels <= 2)
-		{
-			return NAME_XMA;
-		}
-#endif //#if WITH_XMA2
-
-		return NAME_OGG;
-	}
-
-	bool FMixerPlatformNonRealtime::HasCompressedAudioInfoClass(USoundWave* InSoundWave)
-	{
-#if PLATFORM_WINDOWS || HAS_COMPRESSED_AUDIO_INFO_CLASS
-		return true;
-#else
-		return false;
-#endif // PLATFORM_WINDOWS || HAS_COMPRESSED_AUDIO_INFO_CLASS
-	}
+	static FName NAME_XMA(TEXT("XMA"));
+#endif
 
-	ICompressedAudioInfo* FMixerPlatformNonRealtime::CreateCompressedAudioInfo(USoundWave* InSoundWave)
+	FName FMixerPlatformNonRealtime::GetRuntimeFormat(const USoundWave* InSoundWave) const
 	{
-		// TODO: Currently this is a copy paste of the XAudio2 platform interface. Ultimately, this function needs to propogate to the current platform's correct CrateCompressedAudioInfo call.
-#if PLATFORM_WINDOWS || HAS_COMPRESSED_AUDIO_INFO_CLASS
-		check(InSoundWave);
-
-#if WITH_XMA2 && USE_XMA2_FOR_STREAMING
-		if (InSoundWave->IsStreaming() && InSoundWave->NumChannels <= 2 )
-		{
-			if (InSoundWave->IsSeekableStreaming())
-			{
-				return new FADPCMAudioInfo();
-			}
-			return XMA2_INFO_NEW();
-		}
-#endif
+		FName RuntimeFormat = Audio::ToName(InSoundWave->GetSoundAssetCompressionType());
 
-		if (InSoundWave->IsStreaming(nullptr))
+		if (RuntimeFormat == Audio::NAME_PLATFORM_SPECIFIC)
 		{
-			if (InSoundWave->IsSeekableStreaming())
+#if WITH_XMA2 && USE_XMA2_FOR_STREAMING
+			if (InSoundWave->NumChannels <= 2)
 			{
-				return new FADPCMAudioInfo();
+				RuntimeFormat = Audio::NAME_XMA;
 			}
+#endif // WITH_XMA2 && USE_XMA2_FOR_STREAMING
 
 #if USE_VORBIS_FOR_STREAMING
-			return new FVorbisAudioInfo();
+			RuntimeFormat = Audio::NAME_OGG;
 #else
-			return new FOpusAudioInfo();
-#endif
-		}
-
-		static const FName NAME_OGG(TEXT("OGG"));
-		if (FPlatformProperties::RequiresCookedData() ? InSoundWave->HasCompressedData(NAME_OGG) : (InSoundWave->GetCompressedData(NAME_OGG) != nullptr))
-		{
-			return new FVorbisAudioInfo();
-		}
-
-#if WITH_XMA2
-		static const FName NAME_XMA(TEXT("XMA"));
-		if (FPlatformProperties::RequiresCookedData() ? InSoundWave->HasCompressedData(NAME_XMA) : (InSoundWave->GetCompressedData(NAME_XMA) != nullptr))
-		{
-			return XMA2_INFO_NEW();
+			RuntimeFormat = Audio::NAME_OPUS;
+#endif // USE_VORBIS_FOR_STREAMING
 		}
-#endif // WITH_XMA2
-#endif // PLATFORM_WINDOWS || HAS_COMPRESSED_AUDIO_INFO_CLASS
 
-		return nullptr;
+		return RuntimeFormat;
 	}
 
-	ICompressedAudioInfo* FMixerPlatformNonRealtime::CreateCompressedAudioInfo(const FSoundWaveProxyPtr& InSoundWave)
+	ICompressedAudioInfo* FMixerPlatformNonRealtime::CreateCompressedAudioInfo(const FName& InRuntimeFormat) const
 	{
-		if (!InSoundWave.IsValid())
-		{
-			return nullptr;
-		}
-
-		// TODO: Currently this is a copy paste of the XAudio2 platform interface. Ultimately, this function needs to propogate to the current platform's correct CrateCompressedAudioInfo call.
-#if PLATFORM_WINDOWS || HAS_COMPRESSED_AUDIO_INFO_CLASS
-
+		// Need to create a platform-specific codec
 #if WITH_XMA2 && USE_XMA2_FOR_STREAMING
-		if (InSoundWave->IsStreaming() && InSoundWave->GetNumChannels() <= 2)
+		if (InRuntimeFormat == NAME_XMA)
 		{
-			if (InSoundWave->IsSeekableStreaming())
-			{
-				return new FADPCMAudioInfo();
-			}
 			return XMA2_INFO_NEW();
 		}
-#endif
-
-		if (InSoundWave->IsStreaming())
-		{
-			if (InSoundWave->IsSeekableStreaming())
-			{
-				return new FADPCMAudioInfo();
-			}
-
+#endif // WITH_XMA2 && USE_XMA2_FOR_STREAMING			
 #if USE_VORBIS_FOR_STREAMING
+		if (InRuntimeFormat == Audio::NAME_OGG)
+		{
 			return new FVorbisAudioInfo();
-#else
-			return new FOpusAudioInfo();
-#endif
 		}
-
-		static const FName NAME_OGG(TEXT("OGG"));
-		if (FPlatformProperties::RequiresCookedData() ? InSoundWave->HasCompressedData(NAME_OGG) : (InSoundWave->GetCompressedData(NAME_OGG) != nullptr))
+#else // USE_VORBIS_FOR_STREAMING
+		if (InRuntimeFormat == Audio::NAME_OPUS)
 		{
-			return new FVorbisAudioInfo();
+			return new FOpusAudioInfo();
 		}
-
-#if WITH_XMA2
-		static const FName NAME_XMA(TEXT("XMA"));
-		if (FPlatformProperties::RequiresCookedData() ? InSoundWave->HasCompressedData(NAME_XMA) : (InSoundWave->GetCompressedData(NAME_XMA) != nullptr))
+#endif // #if USE_VORBIS_FOR_STREAMING
+#if WITH_BINK_AUDIO
+		if (InRuntimeFormat == Audio::NAME_BINKA)
 		{
-			return XMA2_INFO_NEW();
+			return new FBinkAudioInfo();
 		}
-#endif // WITH_XMA2
-#endif // PLATFORM_WINDOWS || HAS_COMPRESSED_AUDIO_INFO_CLASS
+#endif
 
-		return nullptr;
+		return Audio::CreateSoundAssetDecoder(InRuntimeFormat);
 	}
 
 	FString FMixerPlatformNonRealtime::GetDefaultDeviceName()
diff --git a/Engine/Source/Runtime/NonRealtimeAudioRenderer/Public/AudioMixerPlatformNonRealtime.h b/Engine/Source/Runtime/NonRealtimeAudioRenderer/Public/AudioMixerPlatformNonRealtime.h
--- a/Engine/Source/Runtime/NonRealtimeAudioRenderer/Public/AudioMixerPlatformNonRealtime.h
+++ b/Engine/Source/Runtime/NonRealtimeAudioRenderer/Public/AudioMixerPlatformNonRealtime.h
@@ -37,12 +37,9 @@ namespace Audio
 		virtual void ResumePlaybackOnNewDevice() override;
 		virtual FAudioPlatformDeviceInfo GetPlatformDeviceInfo() const override;
 		virtual void SubmitBuffer(const uint8* Buffer) override;
-		virtual FName GetRuntimeFormat(USoundWave* InSoundWave) override;
-		virtual bool HasCompressedAudioInfoClass(USoundWave* InSoundWave) override;
-		virtual bool SupportsRealtimeDecompression() const override { return true; }
 		virtual bool DisablePCMAudioCaching() const override;
-		virtual ICompressedAudioInfo* CreateCompressedAudioInfo(USoundWave* InSoundWave) override;
-		virtual ICompressedAudioInfo* CreateCompressedAudioInfo(const FSoundWaveProxyPtr& InSoundWave) override;
+		virtual FName GetRuntimeFormat(const USoundWave* InSoundWave) const override;
+		virtual ICompressedAudioInfo* CreateCompressedAudioInfo(const FName& InRuntimeFormat) const override;
 		virtual FString GetDefaultDeviceName() override;
 		virtual FAudioPlatformSettings GetPlatformSettings() const override;
 		virtual void OnHardwareUpdate() override;
diff --git a/Engine/Source/Runtime/Windows/AudioMixerXAudio2/Private/AudioMixerPlatformXAudio2.cpp b/Engine/Source/Runtime/Windows/AudioMixerXAudio2/Private/AudioMixerPlatformXAudio2.cpp
--- a/Engine/Source/Runtime/Windows/AudioMixerXAudio2/Private/AudioMixerPlatformXAudio2.cpp
+++ b/Engine/Source/Runtime/Windows/AudioMixerXAudio2/Private/AudioMixerPlatformXAudio2.cpp
@@ -8,6 +8,7 @@
 
 #include "AudioMixerPlatformXAudio2.h"
 #include "AudioMixer.h"
+#include "AudioDevice.h"
 #include "HAL/PlatformAffinity.h"
 #include "HAL/PlatformTime.h"
 #include "HAL/PlatformProcess.h"
@@ -23,7 +24,6 @@
 #endif  //#if WITH_XMA2
 #include "OpusAudioInfo.h"
 #include "VorbisAudioInfo.h"
-#include "ADPCMAudioInfo.h"
 #include "AudioPluginUtilities.h"
 #if WITH_BINK_AUDIO
 #include "BinkAudioInfo.h"
@@ -1845,168 +1845,74 @@ namespace Audio
 		}
 	}
 
-	FName FMixerPlatformXAudio2::GetRuntimeFormat(USoundWave* InSoundWave)
-	{
-		static FName NAME_OGG(TEXT("OGG"));
-		static FName NAME_OPUS(TEXT("OPUS"));
-		static FName NAME_XMA(TEXT("XMA"));
-		static FName NAME_ADPCM(TEXT("ADPCM"));
-		static FName NAME_BINKA(TEXT("BINKA"));
-
-#if WITH_ENGINE
+#if WITH_XMA2
+	static FName NAME_XMA(TEXT("XMA"));
+#endif
 
-#if WITH_BINK_AUDIO
-		if (InSoundWave->bUseBinkAudio)
-		{
-			return NAME_BINKA;
-		}
-#endif	// WITH_BINK_AUDIO	
+	FName FMixerPlatformXAudio2::GetRuntimeFormat(const USoundWave* InSoundWave) const
+	{
+		FName RuntimeFormat = Audio::ToName(InSoundWave->GetSoundAssetCompressionType());
 
-		if (InSoundWave->IsStreaming(nullptr))
+		if (RuntimeFormat == Audio::NAME_PLATFORM_SPECIFIC)
 		{
-			if (InSoundWave->IsSeekableStreaming())
-			{
-				return NAME_ADPCM;
-			}
-
 #if WITH_XMA2 && USE_XMA2_FOR_STREAMING
 			if (InSoundWave->NumChannels <= 2)
 			{
-				return NAME_XMA;
+				RuntimeFormat = Audio::NAME_XMA;
 			}
 #endif // WITH_XMA2 && USE_XMA2_FOR_STREAMING
 
 #if USE_VORBIS_FOR_STREAMING
-			return NAME_OGG;
+			RuntimeFormat = Audio::NAME_OGG;
+#else
+			RuntimeFormat = Audio::NAME_OPUS;
 #endif // USE_VORBIS_FOR_STREAMING
 		}
 
-#if WITH_XMA2
-		if (InSoundWave->NumChannels <= 2)
-		{
-			return NAME_XMA;
-		}
-#endif // WITH_XMA2
-#endif // WITH_ENGINE
-		return NAME_OGG;
-	}
-
-	bool FMixerPlatformXAudio2::HasCompressedAudioInfoClass(USoundWave* InSoundWave)
-	{
-		return true;
+		return RuntimeFormat;
 	}
 
-	ICompressedAudioInfo* FMixerPlatformXAudio2::CreateCompressedAudioInfo(USoundWave* InSoundWave)
+	ICompressedAudioInfo* FMixerPlatformXAudio2::CreateCompressedAudioInfo(const FName& InRuntimeFormat) const
 	{
-#if WITH_ENGINE
-		check(InSoundWave);
-
-#if WITH_BINK_AUDIO
-		if (InSoundWave->bUseBinkAudio)
-		{
-			return new FBinkAudioInfo();
-		}
-#endif // WITH_BINK_AUDIO
-
-		if (InSoundWave->IsStreaming(nullptr))
-		{
-			if (InSoundWave->IsSeekableStreaming())
-			{
-				return new FADPCMAudioInfo();
-			}
-		}
+		ICompressedAudioInfo* Decoder = nullptr;
 
+		// Need to create a platform-specific codec
 #if WITH_XMA2 && USE_XMA2_FOR_STREAMING
-		if (InSoundWave->IsStreaming() && InSoundWave->NumChannels <= 2 )
+		if (InRuntimeFormat == Audio::NAME_XMA)
 		{
-			return XMA2_INFO_NEW();
+			Decoder = XMA2_INFO_NEW();
 		}
-#endif
-
-		if (InSoundWave->IsStreaming(nullptr))
+		else
+#endif // WITH_XMA2 && USE_XMA2_FOR_STREAMING			
+		if (InRuntimeFormat == Audio::NAME_OGG)
 		{
 #if USE_VORBIS_FOR_STREAMING
-			return new FVorbisAudioInfo();
-#else
-			return new FOpusAudioInfo();
-#endif // USE_VORBIS_FOR_STREAMING
+			Decoder = new FVorbisAudioInfo();
+#endif // #if USE_VORBIS_FOR_STREAMING
 		}
-
-		static const FName NAME_OGG(TEXT("OGG"));
-		if (FPlatformProperties::RequiresCookedData() ? InSoundWave->HasCompressedData(NAME_OGG) : (InSoundWave->GetCompressedData(NAME_OGG) != nullptr))
-		{
-			return new FVorbisAudioInfo();
-		}
-
-#if WITH_XMA2
-		static const FName NAME_XMA(TEXT("XMA"));
-		if (FPlatformProperties::RequiresCookedData() ? InSoundWave->HasCompressedData(NAME_XMA) : (InSoundWave->GetCompressedData(NAME_XMA) != nullptr))
+		else if (InRuntimeFormat == Audio::NAME_OPUS)
 		{
-			return XMA2_INFO_NEW();
+#if !USE_VORBIS_FOR_STREAMING
+			Decoder = new FOpusAudioInfo();
+#endif // #if !USE_VORBIS_FOR_STREAMING
 		}
-#endif // WITH_XMA2
-#endif // WITH_ENGINE
-		return nullptr;
-	}
-
-
-	ICompressedAudioInfo* FMixerPlatformXAudio2::CreateCompressedAudioInfo(const FSoundWaveProxyPtr& InSoundWave)
-	{
-#if WITH_ENGINE
-		check(InSoundWave);
-
 #if WITH_BINK_AUDIO
-		if (InSoundWave->UseBinkAudio())
+		else if (InRuntimeFormat == Audio::NAME_BINKA)
 		{
 			return new FBinkAudioInfo();
 		}
-#endif // WITH_BINK_AUDIO
-
-		// STREAMING:
-		if (InSoundWave->IsStreaming())
-		{
-			if (InSoundWave->IsSeekableStreaming())
-			{
-				return new FADPCMAudioInfo();
-			}
-	#if WITH_XMA2 && USE_XMA2_FOR_STREAMING
-			else if (InSoundWave->GetNumChannels() <= 2)
-			{
-				return XMA2_INFO_NEW();
-			}
-	#endif
-			else
-			{
-#if USE_VORBIS_FOR_STREAMING
-				return new FVorbisAudioInfo();
-#else
-				return new FOpusAudioInfo();
-#endif // USE_VORBIS_FOR_STREAMING
-			}
-		}
-
-		// NOT STREAMING:
-		static const FName NAME_OGG(TEXT("OGG"));
-		if (FPlatformProperties::RequiresCookedData() ? InSoundWave->HasCompressedData(NAME_OGG) : (InSoundWave->GetCompressedData(NAME_OGG) != nullptr))
-		{
-			return new FVorbisAudioInfo();
-		}
-
-#if WITH_XMA2
-		static const FName NAME_XMA(TEXT("XMA"));
-		if (FPlatformProperties::RequiresCookedData() ? InSoundWave->HasCompressedData(NAME_XMA) : (InSoundWave->GetCompressedData(NAME_XMA) != nullptr))
+#endif // #if WITH_BINK_AUDIO
+		else
 		{
-			return XMA2_INFO_NEW();
+			// Fallback to the multiplatform decoder
+			Decoder = Audio::CreateSoundAssetDecoder(InRuntimeFormat);
 		}
-#endif // WITH_XMA2
-#endif // WITH_ENGINE
-
-		return nullptr;
+		ensureMsgf(Decoder != nullptr, TEXT("Failed to create a sound asset decoder for compression type: %s"), *InRuntimeFormat.ToString());
+		return Decoder;
 	}
 
 	FString FMixerPlatformXAudio2::GetDefaultDeviceName()
 	{
-		//GConfig->GetString(TEXT("/Script/WindowsTargetPlatform.WindowsTargetSettings"), TEXT("AudioDevice"), WindowsAudioDeviceName, GEngineIni);
 		return FString();
 	}
 
diff --git a/Engine/Source/Runtime/Windows/AudioMixerXAudio2/Private/AudioMixerPlatformXAudio2.h b/Engine/Source/Runtime/Windows/AudioMixerXAudio2/Private/AudioMixerPlatformXAudio2.h
--- a/Engine/Source/Runtime/Windows/AudioMixerXAudio2/Private/AudioMixerPlatformXAudio2.h
+++ b/Engine/Source/Runtime/Windows/AudioMixerXAudio2/Private/AudioMixerPlatformXAudio2.h
@@ -81,12 +81,9 @@ namespace Audio
 		virtual void ResumePlaybackOnNewDevice() override;
 		virtual FAudioPlatformDeviceInfo GetPlatformDeviceInfo() const override;
 		virtual void SubmitBuffer(const uint8* Buffer) override;
-		virtual FName GetRuntimeFormat(USoundWave* InSoundWave) override;
-		virtual bool HasCompressedAudioInfoClass(USoundWave* InSoundWave) override;
-		virtual bool SupportsRealtimeDecompression() const override { return true; }
 		virtual bool DisablePCMAudioCaching() const override;
-		virtual ICompressedAudioInfo* CreateCompressedAudioInfo(USoundWave* InSoundWave) override;
-		virtual ICompressedAudioInfo* CreateCompressedAudioInfo(const FSoundWaveProxyPtr& InSoundWave) override;
+		virtual FName GetRuntimeFormat(const USoundWave* InSoundWave) const override;
+		virtual ICompressedAudioInfo* CreateCompressedAudioInfo(const FName& InRuntimeFormat) const override;
 		virtual FString GetDefaultDeviceName() override;
 		virtual FAudioPlatformSettings GetPlatformSettings() const override;
 		virtual void OnHardwareUpdate() override;
