Undo changelist 18658885
Removes reference to DataLayerInstance and DataLayerAssets in code

diff --git a/Engine/Source/Editor/DataLayerEditor/Private/DataLayer/DataLayerFactory.cpp b/Engine/Source/Editor/DataLayerEditor/Private/DataLayer/DataLayerFactory.cpp
deleted file mode 100644
diff --git a/Engine/Source/Editor/DataLayerEditor/Public/DataLayer/DataLayerFactory.h b/Engine/Source/Editor/DataLayerEditor/Public/DataLayer/DataLayerFactory.h
deleted file mode 100644
diff --git a/Engine/Source/Editor/EditorStyle/Private/StarshipStyle.cpp b/Engine/Source/Editor/EditorStyle/Private/StarshipStyle.cpp
--- a/Engine/Source/Editor/EditorStyle/Private/StarshipStyle.cpp
+++ b/Engine/Source/Editor/EditorStyle/Private/StarshipStyle.cpp
@@ -6085,7 +6085,6 @@ void FStarshipEditorStyle::FStyle::SetupClassIconsAndThumbnails()
 			{TEXT("CurveBase")},
 			{TEXT("DataAsset")},
 			{TEXT("DataTable")},
-			{TEXT("DataLayerAsset")},
 			{TEXT("DecalActor")},
 			{TEXT("DefaultPawn")},
 			{TEXT("DialogueVoice")},
diff --git a/Engine/Source/Editor/UnrealEd/Classes/Commandlets/WorldPartitionDataLayerToAssetCommandLet.h b/Engine/Source/Editor/UnrealEd/Classes/Commandlets/WorldPartitionDataLayerToAssetCommandLet.h
deleted file mode 100644
diff --git a/Engine/Source/Editor/UnrealEd/Private/Commandlets/WorldPartitionCommandletHelpers.cpp b/Engine/Source/Editor/UnrealEd/Private/Commandlets/WorldPartitionCommandletHelpers.cpp
deleted file mode 100644
diff --git a/Engine/Source/Editor/UnrealEd/Private/Commandlets/WorldPartitionDataLayerToAssetCommandlet.cpp b/Engine/Source/Editor/UnrealEd/Private/Commandlets/WorldPartitionDataLayerToAssetCommandlet.cpp
deleted file mode 100644
diff --git a/Engine/Source/Editor/UnrealEd/Public/Commandlets/WorldPartitionCommandletHelpers.h b/Engine/Source/Editor/UnrealEd/Public/Commandlets/WorldPartitionCommandletHelpers.h
deleted file mode 100644
diff --git a/Engine/Source/Editor/UnrealEd/UnrealEd.Build.cs b/Engine/Source/Editor/UnrealEd/UnrealEd.Build.cs
--- a/Engine/Source/Editor/UnrealEd/UnrealEd.Build.cs
+++ b/Engine/Source/Editor/UnrealEd/UnrealEd.Build.cs
@@ -136,7 +136,6 @@ public class UnrealEd : ModuleRules
 				"CinematicCamera",
 				"CurveEditor",
 				"ContentBrowserData",
-				"DataLayerEditor",
 				"DerivedDataCache",
 				"Zen",
 				"DesktopPlatform",
diff --git a/Engine/Source/Runtime/Engine/Classes/GameFramework/Actor.h b/Engine/Source/Runtime/Engine/Classes/GameFramework/Actor.h
--- a/Engine/Source/Runtime/Engine/Classes/GameFramework/Actor.h
+++ b/Engine/Source/Runtime/Engine/Classes/GameFramework/Actor.h
@@ -39,8 +39,6 @@ struct FAttachedActorInfo;
 struct FNetViewer;
 struct FNetworkObjectInfo;
 class UDataLayer;
-class UDataLayerAsset;
-class UDataLayerInstance;
 class AWorldDataLayers;
 class UActorFolder;
 
@@ -832,10 +830,6 @@ protected:
 	UPROPERTY(EditAnywhere, AdvancedDisplay, Category = DataLayers)
 	TArray<FActorDataLayer> DataLayers;
 
-	/** DataLayers the actor belongs to.*/
-	UPROPERTY()
-	TArray<TObjectPtr<const UDataLayerAsset>> DataLayerAssets;
-
 	TArray<FActorDataLayer> PreEditChangeDataLayers;
 
 public:
@@ -1169,13 +1163,6 @@ public:
 	bool IsValidForDataLayer() const;
 	void FixupDataLayers(bool bRevertChangesOnLockedDataLayer = false);
 	static const FName GetDataLayersPropertyName() { return GET_MEMBER_NAME_CHECKED(AActor, DataLayers); }
-
-	bool AddDataLayer(const UDataLayerAsset* DataLayerAsset);
-	bool RemoveDataLayer(const UDataLayerAsset* DataLayerAsset);
-	bool ContainsDataLayer(const UDataLayerAsset* DataLayerAsset) const;
-	const TArray<TObjectPtr<const UDataLayerAsset>>& GetDataLayerAssets() const { return DataLayerAssets; }
-
-	TArray<FActorDataLayer> const& GetActorDataLayers() const { return DataLayers; };
 #endif
 
 	//~=============================================================================
diff --git a/Engine/Source/Runtime/Engine/Private/ActorEditor.cpp b/Engine/Source/Runtime/Engine/Private/ActorEditor.cpp
--- a/Engine/Source/Runtime/Engine/Private/ActorEditor.cpp
+++ b/Engine/Source/Runtime/Engine/Private/ActorEditor.cpp
@@ -15,7 +15,6 @@
 #include "WorldPartition/WorldPartitionSubsystem.h"
 #include "WorldPartition/DataLayer/WorldDataLayers.h"
 #include "WorldPartition/DataLayer/DataLayer.h"
-#include "WorldPartition/DataLayer/DataLayerAsset.h"
 #include "EditorSupportDelegates.h"
 #include "Logging/TokenizedMessage.h"
 #include "Logging/MessageLog.h"
@@ -1321,34 +1320,10 @@ bool AActor::AddDataLayer(const UDataLayer* DataLayer)
 	return bActorWasModified;
 }
 
-bool AActor::AddDataLayer(const UDataLayerAsset* DataLayerAsset)
-{
-	bool bActorWasModified = false;
-	if (SupportsDataLayer() && DataLayerAsset && !ContainsDataLayer(DataLayerAsset))
-	{
-		Modify();
-		bActorWasModified = true;
-		DataLayerAssets.Add(DataLayerAsset);
-	}
-	return bActorWasModified;
-}
-
 bool AActor::RemoveDataLayer(const UDataLayer* DataLayer)
 {
 	bool bActorWasModified = false;
 	if (ContainsDataLayer(DataLayer))
-	{
-		Modify();
-		bActorWasModified = true;
-		DataLayers.Remove(FActorDataLayer(DataLayer->GetFName()));
-	}
-	return bActorWasModified;
-}
-
-bool AActor::RemoveDataLayer(const UDataLayerAsset* DataLayerAsset)
-{
-	bool bActorWasModified = false;
-	if (ContainsDataLayer(DataLayerAsset))
 	{
 		if (!bActorWasModified)
 		{
@@ -1356,7 +1331,7 @@ bool AActor::RemoveDataLayer(const UDataLayerAsset* DataLayerAsset)
 			bActorWasModified = true;
 		}
 
-		DataLayerAssets.Remove(DataLayerAsset);
+		DataLayers.Remove(FActorDataLayer(DataLayer->GetFName()));
 	}
 	return bActorWasModified;
 }
@@ -1377,11 +1352,6 @@ bool AActor::ContainsDataLayer(const UDataLayer* DataLayer) const
 	return DataLayer && DataLayers.Contains(FActorDataLayer(DataLayer->GetFName()));
 }
 
-bool AActor::ContainsDataLayer(const UDataLayerAsset* DataLayerAsset) const
-{
-	return DataLayerAsset && DataLayerAssets.Contains(DataLayerAsset);
-}
-
 bool AActor::HasDataLayers() const
 {
 	return DataLayers.Num() > 0;
diff --git a/Engine/Source/Runtime/Engine/Private/WorldPartition/DataLayer/DataLayerAsset.cpp b/Engine/Source/Runtime/Engine/Private/WorldPartition/DataLayer/DataLayerAsset.cpp
deleted file mode 100644
diff --git a/Engine/Source/Runtime/Engine/Private/WorldPartition/DataLayer/DataLayerInstance.cpp b/Engine/Source/Runtime/Engine/Private/WorldPartition/DataLayer/DataLayerInstance.cpp
deleted file mode 100644
diff --git a/Engine/Source/Runtime/Engine/Private/WorldPartition/DataLayer/WorldDataLayers.cpp b/Engine/Source/Runtime/Engine/Private/WorldPartition/DataLayer/WorldDataLayers.cpp
--- a/Engine/Source/Runtime/Engine/Private/WorldPartition/DataLayer/WorldDataLayers.cpp
+++ b/Engine/Source/Runtime/Engine/Private/WorldPartition/DataLayer/WorldDataLayers.cpp
@@ -6,7 +6,6 @@
 
 #include "WorldPartition/DataLayer/WorldDataLayers.h"
 #include "WorldPartition/DataLayer/DataLayer.h"
-#include "WorldPartition/DataLayer/DataLayerInstance.h"
 #include "WorldPartition/DataLayer/DataLayerSubsystem.h"
 #include "EngineUtils.h"
 #include "Engine/CoreSettings.h"
@@ -606,18 +605,6 @@ UDataLayer* AWorldDataLayers::CreateDataLayer(FName InName, EObjectFlags InObjec
 	return NewDataLayer;
 }
 
-UDataLayerInstance* AWorldDataLayers::CreateDataLayer(const UDataLayerAsset* InDataLayerAsset)
-{
-	Modify();
-
-	const FName DataLayerUniqueName = *FString::Format(TEXT("{0}_{1}"), { InDataLayerAsset->GetName(), FGuid::NewGuid().ToString() });
-	UDataLayerInstance* NewDataLayer = NewObject<UDataLayerInstance>(this, DataLayerUniqueName, RF_Transactional | RF_NoFlags);
-	check(NewDataLayer != NULL);
-	NewDataLayer->OnCreated(InDataLayerAsset);
-	DataLayerInstances.Add(NewDataLayer);
-	return NewDataLayer;
-}
-
 bool AWorldDataLayers::RemoveDataLayers(const TArray<UDataLayer*>& InDataLayers)
 {
 	bool bIsModified = false;
@@ -645,17 +632,6 @@ bool AWorldDataLayers::RemoveDataLayer(UDataLayer* InDataLayer)
 	return false;
 }
 
-bool AWorldDataLayers::RemoveDataLayer(UDataLayerInstance* InDataLayerInstance)
-{
-	if (ContainsDataLayer(InDataLayerInstance))
-	{
-		Modify();
-		DataLayerInstances.Remove(InDataLayerInstance);
-		return true;
-	}
-	return false;
-}
-
 void AWorldDataLayers::SetAllowRuntimeDataLayerEditing(bool bInAllowRuntimeDataLayerEditing)
 {
 	if (bAllowRuntimeDataLayerEditing != bInAllowRuntimeDataLayerEditing)
@@ -671,11 +647,6 @@ bool AWorldDataLayers::ContainsDataLayer(const UDataLayer* InDataLayer) const
 	return WorldDataLayers.Contains(InDataLayer);
 }
 
-bool AWorldDataLayers::ContainsDataLayer(const UDataLayerInstance* InDataLayerInstance) const 
-{
-	return DataLayerInstances.Contains(InDataLayerInstance);
-}
-
 const UDataLayer* AWorldDataLayers::GetDataLayerFromName(const FName& InDataLayerName) const
 {
 #if WITH_EDITOR	
diff --git a/Engine/Source/Runtime/Engine/Public/WorldPartition/DataLayer/DataLayerAsset.h b/Engine/Source/Runtime/Engine/Public/WorldPartition/DataLayer/DataLayerAsset.h
deleted file mode 100644
diff --git a/Engine/Source/Runtime/Engine/Public/WorldPartition/DataLayer/DataLayerInstance.h b/Engine/Source/Runtime/Engine/Public/WorldPartition/DataLayer/DataLayerInstance.h
deleted file mode 100644
diff --git a/Engine/Source/Runtime/Engine/Public/WorldPartition/DataLayer/WorldDataLayers.h b/Engine/Source/Runtime/Engine/Public/WorldPartition/DataLayer/WorldDataLayers.h
--- a/Engine/Source/Runtime/Engine/Public/WorldPartition/DataLayer/WorldDataLayers.h
+++ b/Engine/Source/Runtime/Engine/Public/WorldPartition/DataLayer/WorldDataLayers.h
@@ -10,8 +10,6 @@
 #include "WorldDataLayers.generated.h"
 
 class UDataLayer;
-class UDataLayerInstance;
-class UDataLayerAsset;
 
 /**
  * Actor containing all data layers for a world
@@ -21,8 +19,6 @@ class ENGINE_API AWorldDataLayers : public AInfo
 {
 	GENERATED_UCLASS_BODY()
 
-	friend class UDataLayerToAssetCommandlet;
-
 public:
 	virtual void PostLoad() override;
 	virtual void RewindForReplay() override;
@@ -36,15 +32,12 @@ public:
 
 	static AWorldDataLayers* Create(UWorld* World);
 	UDataLayer* CreateDataLayer(FName InName = TEXT("DataLayer"), EObjectFlags InObjectFlags = RF_NoFlags);
-	UDataLayerInstance* CreateDataLayer(const UDataLayerAsset* InDataLayerAsset);
 	bool RemoveDataLayer(UDataLayer* InDataLayer);
 	bool RemoveDataLayers(const TArray<UDataLayer*>& InDataLayers);
 	FName GenerateUniqueDataLayerLabel(const FName& InDataLayerLabel) const;
 	void SetAllowRuntimeDataLayerEditing(bool bInAllowRuntimeDataLayerEditing);
 	bool GetAllowRuntimeDataLayerEditing() const { return bAllowRuntimeDataLayerEditing; }
 
-	bool RemoveDataLayer(UDataLayerInstance* InDataLayerInstance);
-
 	//~ Begin Helper Functions
 	TArray<const UDataLayer*> GetDataLayerObjects(const TArray<FActorDataLayer>& DataLayers) const;
 	TArray<const UDataLayer*> GetDataLayerObjects(const TArray<FName>& InDataLayerNames) const;
@@ -65,8 +58,6 @@ public:
 	void ForEachDataLayer(TFunctionRef<bool(UDataLayer*)> Func);
 	void ForEachDataLayer(TFunctionRef<bool(UDataLayer*)> Func) const;
 
-	bool ContainsDataLayer(const UDataLayerInstance* InDataLayerInstance) const;
-
 	// DataLayer Runtime State
 	void SetDataLayerRuntimeState(FActorDataLayer InDataLayer, EDataLayerRuntimeState InState, bool bIsRecursive = false);
 	EDataLayerRuntimeState GetDataLayerRuntimeStateByName(FName InDataLayerName) const;
@@ -127,12 +118,6 @@ private:
 	bool bAllowRuntimeDataLayerEditing;
 #endif
 
-	UPROPERTY()
-	TSet<TObjectPtr<UDataLayerInstance>> DataLayerInstances;
-
-	UPROPERTY()
-	TMap<FName, TWeakObjectPtr<UDataLayerInstance>> DeprecatedDataLayerNameToDataLayerInstance;
-
 	UPROPERTY()
 	TSet<TObjectPtr<UDataLayer>> WorldDataLayers;
 
