Optimize ULevel::GetLoadedExternalObjectPackages()

Don't check all loaded packages unless the external actor folders actually exist on disk.

#rnx

diff --git a/Engine/Source/Runtime/Engine/Private/Level.cpp b/Engine/Source/Runtime/Engine/Private/Level.cpp
--- a/Engine/Source/Runtime/Engine/Private/Level.cpp
+++ b/Engine/Source/Runtime/Engine/Private/Level.cpp
@@ -3278,19 +3278,24 @@ TArray<UPackage*> ULevel::GetLoadedExternalObjectPackages() const
 	TArray<FString> ExternalObjectsPaths;
 	ExternalObjectsPaths.Add(SanitizeExternalPath(ULevel::GetExternalActorsPath(World->GetPackage(), (World->OriginalWorldName == NAME_None) ? World->GetName() : World->OriginalWorldName.ToString())));
 	ExternalObjectsPaths.Add(SanitizeExternalPath(FExternalPackageHelper::GetExternalObjectsPath(World->GetPackage(), (World->OriginalWorldName == NAME_None) ? World->GetName() : World->OriginalWorldName.ToString())));
-	ExternalObjectsPaths = ExternalObjectsPaths.FilterByPredicate([](const FString& Path) { return !Path.IsEmpty(); });
-	
-	for (TObjectIterator<UPackage> It; It; ++It)
+	ExternalObjectsPaths = ExternalObjectsPaths.FilterByPredicate([](const FString& Path) {
+			return !Path.IsEmpty() && FPaths::DirectoryExists(FPackageName::LongPackageNameToFilename(Path));
+		});
+
+	if (!ExternalObjectsPaths.IsEmpty())
 	{
-		TStringBuilder<256> PackageName;
-		It->GetLoadedPath().AppendPackageName(PackageName);
-		FStringView PackageNameStringView(PackageName);
-		for (const FString& ExternalObjectsPath : ExternalObjectsPaths)
+		for (TObjectIterator<UPackage> It; It; ++It)
 		{
-			if (PackageNameStringView.Contains(ExternalObjectsPath))
+			TStringBuilder<256> PackageName;
+			It->GetLoadedPath().AppendPackageName(PackageName);
+			FStringView PackageNameStringView(PackageName);
+			for (const FString& ExternalObjectsPath : ExternalObjectsPaths)
 			{
-				ExternalObjectPackages.Add(*It);
-				break;
+				if (PackageNameStringView.Contains(ExternalObjectsPath))
+				{
+					ExternalObjectPackages.Add(*It);
+					break;
+				}
 			}
 		}
 	}
