Add missing Super::BeginPlay call to WorldPartitionReplay

diff --git a/Engine/Source/Runtime/Engine/Private/WorldPartition/WorldPartitionReplay.cpp b/Engine/Source/Runtime/Engine/Private/WorldPartition/WorldPartitionReplay.cpp
--- a/Engine/Source/Runtime/Engine/Private/WorldPartition/WorldPartitionReplay.cpp
+++ b/Engine/Source/Runtime/Engine/Private/WorldPartition/WorldPartitionReplay.cpp
@@ -109,6 +109,8 @@ bool AWorldPartitionReplay::IsEnabled(UWorld* World)
 
 void AWorldPartitionReplay::BeginPlay()
 {
+	Super::BeginPlay();
+	
 	UWorld* World = GetWorld();
 	check(World->IsGameWorld());
 	check(World->IsPartitionedWorld());
