GetBlendProfile is now exposed to blueprint (to use with MontagePlay/Stop_WithBlendSettings)

diff --git a/Engine/Source/Runtime/Engine/Classes/Animation/Skeleton.h b/Engine/Source/Runtime/Engine/Classes/Animation/Skeleton.h
--- a/Engine/Source/Runtime/Engine/Classes/Animation/Skeleton.h
+++ b/Engine/Source/Runtime/Engine/Classes/Animation/Skeleton.h
@@ -595,6 +595,7 @@ public:
 	TArray<TObjectPtr<UBlendProfile>> BlendProfiles;
 
 	/** Get the specified blend profile by name */
+	UFUNCTION(BlueprintPure, Category = Skeleton)
 	ENGINE_API UBlendProfile* GetBlendProfile(const FName& InProfileName);
 
 	/** Create a new blend profile with the specified name */
