Post sync fixups.
This goes with CL 18817999 and contains files that needed to be updated to match concurrent code changes.

diff --git a/Engine/Plugins/FX/Niagara/Source/Niagara/Private/NiagaraComponent.cpp b/Engine/Plugins/FX/Niagara/Source/Niagara/Private/NiagaraComponent.cpp
--- a/Engine/Plugins/FX/Niagara/Source/Niagara/Private/NiagaraComponent.cpp
+++ b/Engine/Plugins/FX/Niagara/Source/Niagara/Private/NiagaraComponent.cpp
@@ -3241,7 +3241,7 @@ void UNiagaraComponent::SetOverrideParameterStoreValue(const FNiagaraVariableBas
 		{
 			FVector3f WorldPos;
 			FMemory::Memcpy(&WorldPos, InValue.GetBytes(), InValue.GetNumBytes());
-			OverrideParameters.SetPositionParameterValue(WorldPos, InKey.GetName(), true);
+			OverrideParameters.SetPositionParameterValue((FVector)WorldPos, InKey.GetName(), true);
 		}
 		else
 		{
diff --git a/Engine/Plugins/FX/Niagara/Source/Niagara/Private/NiagaraDataInterfaceSkeletalMeshConnectivity.cpp b/Engine/Plugins/FX/Niagara/Source/Niagara/Private/NiagaraDataInterfaceSkeletalMeshConnectivity.cpp
--- a/Engine/Plugins/FX/Niagara/Source/Niagara/Private/NiagaraDataInterfaceSkeletalMeshConnectivity.cpp
+++ b/Engine/Plugins/FX/Niagara/Source/Niagara/Private/NiagaraDataInterfaceSkeletalMeshConnectivity.cpp
@@ -293,9 +293,9 @@ static bool BuildAdjacencyBuffer(const FSkeletalMeshLODRenderData& LodRenderData
 
 		const FVector P[3] =
 		{
-			VertexBuffer.VertexPosition(V[0]),
-			VertexBuffer.VertexPosition(V[1]),
-			VertexBuffer.VertexPosition(V[2])
+			VertexBuffer.VertexPosition((FVector)V[0]),
+			VertexBuffer.VertexPosition((FVector)V[1]),
+			VertexBuffer.VertexPosition((FVector)V[2])
 		};
 
 		if (SortBySize)
diff --git a/Engine/Source/Runtime/Engine/Private/PhysicsEngine/BodySetup.cpp b/Engine/Source/Runtime/Engine/Private/PhysicsEngine/BodySetup.cpp
--- a/Engine/Source/Runtime/Engine/Private/PhysicsEngine/BodySetup.cpp
+++ b/Engine/Source/Runtime/Engine/Private/PhysicsEngine/BodySetup.cpp
@@ -159,7 +159,7 @@ UBodySetup::UBodySetup(const FObjectInitializer& ObjectInitializer)
 #if WITH_EDITORONLY_DATA
 	BuildScale_DEPRECATED = 1.0f;
 #endif
-	BuildScale3D = FVector3f::OneVector;
+	BuildScale3D = FVector::OneVector;
 	SetFlags(RF_Transactional);
 	bSharedCookedData = false;
 	CookedFormatDataOverride = nullptr;
@@ -1006,7 +1006,7 @@ void UBodySetup::RescaleSimpleCollision( FVector BuildScale )
 			SphylElem->Length *= ScaleMultiplier;
 		}
 
-		BuildScale3D = FVector3f(BuildScale);	//LWC_TODO: Precision loss
+		BuildScale3D = BuildScale;
 	}
 }
 
@@ -1231,7 +1231,7 @@ void UBodySetup::PostLoad()
 #if WITH_EDITORONLY_DATA
 	if ( GetLinkerUEVersion() < VER_UE4_BUILD_SCALE_VECTOR )
 	{
-		BuildScale3D = FVector3f( BuildScale_DEPRECATED );
+		BuildScale3D = BuildScale_DEPRECATED;
 	}
 #endif
 
