Fox CIS

diff --git a/Engine/Source/Runtime/Engine/Private/PhysicsEngine/BodySetup.cpp b/Engine/Source/Runtime/Engine/Private/PhysicsEngine/BodySetup.cpp
--- a/Engine/Source/Runtime/Engine/Private/PhysicsEngine/BodySetup.cpp
+++ b/Engine/Source/Runtime/Engine/Private/PhysicsEngine/BodySetup.cpp
@@ -1231,7 +1231,7 @@ void UBodySetup::PostLoad()
 #if WITH_EDITORONLY_DATA
 	if ( GetLinkerUEVersion() < VER_UE4_BUILD_SCALE_VECTOR )
 	{
-		BuildScale3D = BuildScale_DEPRECATED;
+		BuildScale3D = FVector(BuildScale_DEPRECATED);
 	}
 #endif
 
