Fixed some issues found by the static analysis.

diff --git a/Engine/Source/Runtime/Engine/Private/Elements/Framework/TypedElementCommonActions.cpp b/Engine/Source/Runtime/Engine/Private/Elements/Framework/TypedElementCommonActions.cpp
--- a/Engine/Source/Runtime/Engine/Private/Elements/Framework/TypedElementCommonActions.cpp
+++ b/Engine/Source/Runtime/Engine/Private/Elements/Framework/TypedElementCommonActions.cpp
@@ -24,6 +24,13 @@ bool UTypedElementCommonActions::DeleteSelectedElements(UTypedElementSelectionSe
 	if (!SelectionSet)
 	{
 		FFrame::KismetExecutionMessage(TEXT("SelectionSet is null."), ELogVerbosity::Error);
+		return false;
+	}
+
+	if (!World)
+	{
+		FFrame::KismetExecutionMessage(TEXT("World is null."), ELogVerbosity::Error);
+		return false;
 	}
 
 	FTypedElementListRef NormalizedElements = SelectionSet->GetNormalizedSelection(FTypedElementSelectionNormalizationOptions());
@@ -96,6 +103,12 @@ bool UTypedElementCommonActions::DeleteNormalizedElements(const FScriptTypedElem
 		return false;
 	}
 
+	if (!World)
+	{
+		FFrame::KismetExecutionMessage(TEXT("World is null."), ELogVerbosity::Error);
+		return false;
+	}
+
 	return DeleteNormalizedElements(NativeList, World, InSelectionSet, DeletionOptions);
 }
 
@@ -104,6 +117,13 @@ TArray<FScriptTypedElementHandle> UTypedElementCommonActions::K2_DuplicateSelect
 	if (!SelectionSet)
 	{
 		FFrame::KismetExecutionMessage(TEXT("SelectionSet is null."), ELogVerbosity::Error);
+		return {};
+	}
+
+	if (!World)
+	{
+		FFrame::KismetExecutionMessage(TEXT("World is null."), ELogVerbosity::Error);
+		return {};
 	}
 
 	return TypedElementUtil::ConvertToScriptElementArray(DuplicateSelectedElements(SelectionSet, World, LocationOffset), SelectionSet->GetElementList()->GetRegistry());
@@ -118,6 +138,12 @@ TArray<FScriptTypedElementHandle> UTypedElementCommonActions::DuplicateNormalize
 		return {};
 	}
 
+	if (!World)
+	{
+		FFrame::KismetExecutionMessage(TEXT("World is null."), ELogVerbosity::Error);
+		return {};
+	}
+
 	return TypedElementUtil::ConvertToScriptElementArray(DuplicateNormalizedElements(NativeList, World, LocationOffset), NativeList->GetRegistry());
 }
 
