Landscape rendering flat at 0 height

Reverted default value of URuntimeVirtualTextureComponent::bHidePrimitives to false

diff --git a/Engine/Source/Runtime/Engine/Classes/Components/RuntimeVirtualTextureComponent.h b/Engine/Source/Runtime/Engine/Classes/Components/RuntimeVirtualTextureComponent.h
--- a/Engine/Source/Runtime/Engine/Classes/Components/RuntimeVirtualTextureComponent.h
+++ b/Engine/Source/Runtime/Engine/Classes/Components/RuntimeVirtualTextureComponent.h
@@ -43,7 +43,7 @@ protected:
 
 	/** Hide primitives in the main pass. Hidden primitives will be those that draw to this virtual texture with 'Draw in Main Pass' set to 'From Virtual Texture'. */
 	UPROPERTY(EditAnywhere, AdvancedDisplay, Category = VirtualTexture)
-	bool bHidePrimitives = true;
+	bool bHidePrimitives = false;
 
 	/** Texture object containing streamed low mips. This can reduce rendering update cost. */
 	UPROPERTY(EditAnywhere, BlueprintReadOnly, NonPIEDuplicateTransient, Category = VirtualTextureBuild)
