fixed deprecating functions in camera manager

diff --git a/Engine/Plugins/MovieScene/MovieRenderPipeline/Source/MovieRenderPipelineRenderPasses/Private/MoviePipelineDeferredPasses.cpp b/Engine/Plugins/MovieScene/MovieRenderPipeline/Source/MovieRenderPipelineRenderPasses/Private/MoviePipelineDeferredPasses.cpp
--- a/Engine/Plugins/MovieScene/MovieRenderPipeline/Source/MovieRenderPipelineRenderPasses/Private/MoviePipelineDeferredPasses.cpp
+++ b/Engine/Plugins/MovieScene/MovieRenderPipeline/Source/MovieRenderPipelineRenderPasses/Private/MoviePipelineDeferredPasses.cpp
@@ -618,9 +618,9 @@ void UMoviePipelineDeferredPassBase::PostRendererSubmission(const FMoviePipeline
 	
 	// Draw letterboxing
 	APlayerCameraManager* PlayerCameraManager = GetPipeline()->GetWorld()->GetFirstPlayerController()->PlayerCameraManager;
-	if(PlayerCameraManager && PlayerCameraManager->GetCameraCachePOV().bConstrainAspectRatio)
+	if(PlayerCameraManager && PlayerCameraManager->GetCameraCacheView().bConstrainAspectRatio)
 	{
-		const FMinimalViewInfo CameraCache = PlayerCameraManager->GetCameraCachePOV();
+		const FMinimalViewInfo CameraCache = PlayerCameraManager->GetCameraCacheView();
 		UMoviePipelineOutputSetting* OutputSettings = GetPipeline()->GetPipelineMasterConfig()->FindSetting<UMoviePipelineOutputSetting>();
 		check(OutputSettings);
 		
diff --git a/Engine/Plugins/MovieScene/MovieRenderPipeline/Source/MovieRenderPipelineRenderPasses/Private/MoviePipelineImagePassBase.cpp b/Engine/Plugins/MovieScene/MovieRenderPipeline/Source/MovieRenderPipelineRenderPasses/Private/MoviePipelineImagePassBase.cpp
--- a/Engine/Plugins/MovieScene/MovieRenderPipeline/Source/MovieRenderPipelineRenderPasses/Private/MoviePipelineImagePassBase.cpp
+++ b/Engine/Plugins/MovieScene/MovieRenderPipeline/Source/MovieRenderPipelineRenderPasses/Private/MoviePipelineImagePassBase.cpp
@@ -306,9 +306,9 @@ FSceneView* UMoviePipelineImagePassBase::GetSceneViewForSampleState(FSceneViewFa
 		APlayerCameraManager* PlayerCameraManager = GetPipeline()->GetWorld()->GetFirstPlayerController()->PlayerCameraManager;
 		
 		// Stretch the fovs if the view is constrained to the camera's aspect ratio
-		if (PlayerCameraManager && PlayerCameraManager->GetCameraCachePOV().bConstrainAspectRatio)
+		if (PlayerCameraManager && PlayerCameraManager->GetCameraCacheView().bConstrainAspectRatio)
 		{
-			const FMinimalViewInfo CameraCache = PlayerCameraManager->GetCameraCachePOV();
+			const FMinimalViewInfo CameraCache = PlayerCameraManager->GetCameraCacheView();
 			const float DestAspectRatio = ViewInitOptions.GetViewRect().Width() / (float)ViewInitOptions.GetViewRect().Height();
 
 			// If the camera's aspect ratio has a thinner width, then stretch the horizontal fov more than usual to 
@@ -460,7 +460,7 @@ void UMoviePipelineImagePassBase::BlendPostProcessSettings(FSceneView* InView)
 			InView->ColorScale = FLinearColor(ColorScale.X, ColorScale.Y, ColorScale.Z);
 		}
 
-		FMinimalViewInfo ViewInfo = LocalPlayerController->PlayerCameraManager->GetCameraCachePOV();
+		FMinimalViewInfo ViewInfo = LocalPlayerController->PlayerCameraManager->GetCameraCacheView();
 		for (int32 PPIdx = 0; PPIdx < CameraAnimPPBlendWeights->Num(); ++PPIdx)
 		{
 			InView->OverridePostProcessSettings((*CameraAnimPPSettings)[PPIdx], (*CameraAnimPPBlendWeights)[PPIdx]);
diff --git a/Engine/Plugins/MovieScene/TemplateSequence/Source/TemplateSequence/Private/Tests/SequenceCameraShakeTestUtil.cpp b/Engine/Plugins/MovieScene/TemplateSequence/Source/TemplateSequence/Private/Tests/SequenceCameraShakeTestUtil.cpp
--- a/Engine/Plugins/MovieScene/TemplateSequence/Source/TemplateSequence/Private/Tests/SequenceCameraShakeTestUtil.cpp
+++ b/Engine/Plugins/MovieScene/TemplateSequence/Source/TemplateSequence/Private/Tests/SequenceCameraShakeTestUtil.cpp
@@ -7,12 +7,12 @@
 
 FMinimalViewInfo USequenceCameraShakeTestUtil::GetCameraCachePOV(APlayerController* PlayerController)
 {
-	return PlayerController->PlayerCameraManager->GetCameraCachePOV();
+	return PlayerController->PlayerCameraManager->GetCameraCacheView();
 }
 
 FMinimalViewInfo USequenceCameraShakeTestUtil::GetLastFrameCameraCachePOV(APlayerController* PlayerController)
 {
-	return PlayerController->PlayerCameraManager->GetLastFrameCameraCachePOV();
+	return PlayerController->PlayerCameraManager->GetLastFrameCameraCacheView();
 }
 
 bool USequenceCameraShakeTestUtil::GetPostProcessBlendCache(APlayerController* PlayerController, int32 PPIndex, FPostProcessSettings& OutPPSettings, float& OutPPBlendWeight)
diff --git a/Engine/Source/Runtime/Engine/Classes/Camera/PlayerCameraManager.h b/Engine/Source/Runtime/Engine/Classes/Camera/PlayerCameraManager.h
--- a/Engine/Source/Runtime/Engine/Classes/Camera/PlayerCameraManager.h
+++ b/Engine/Source/Runtime/Engine/Classes/Camera/PlayerCameraManager.h
@@ -322,10 +322,18 @@ public:
 	virtual void SetLastFrameCameraCachePOV(const FMinimalViewInfo& InPOV);
 
 	/** Gets value of CameraCachePrivate.POV */
-	virtual const FMinimalViewInfo& GetCameraCachePOV() const;
+	virtual const FMinimalViewInfo& GetCameraCacheView() const;
 
 	/** Gets value of LastFrameCameraCachePrivate.POV */
-	virtual const FMinimalViewInfo& GetLastFrameCameraCachePOV() const;
+	virtual const FMinimalViewInfo& GetLastFrameCameraCacheView() const;
+
+	/** Gets value of CameraCachePrivate.POV */
+	UE_DEPRECATED(5.0, "APlayerCameraManager::GetCameraCacheView is favored now, and this function forwards to that one.")
+	virtual FMinimalViewInfo GetCameraCachePOV() const;
+
+	/** Gets value of LastFrameCameraCachePrivate.POV */
+	UE_DEPRECATED(5.0, "APlayerCameraManager::GetLastFrameCameraCacheView is favored now, and this function forwards to that one.")
+	virtual FMinimalViewInfo GetLastFrameCameraCachePOV() const;
 
 	/** Get value of CameraCachePrivate.Time  */
 	float GetCameraCacheTime() const { return CameraCachePrivate.TimeStamp; }
diff --git a/Engine/Source/Runtime/Engine/Private/LocalPlayer.cpp b/Engine/Source/Runtime/Engine/Private/LocalPlayer.cpp
--- a/Engine/Source/Runtime/Engine/Private/LocalPlayer.cpp
+++ b/Engine/Source/Runtime/Engine/Private/LocalPlayer.cpp
@@ -707,7 +707,7 @@ void ULocalPlayer::GetViewPoint(FMinimalViewInfo& OutViewInfo) const
 	{
 		if (PlayerController->PlayerCameraManager != NULL)
 		{
-			OutViewInfo = PlayerController->PlayerCameraManager->GetCameraCachePOV();
+			OutViewInfo = PlayerController->PlayerCameraManager->GetCameraCacheView();
 			OutViewInfo.FOV = PlayerController->PlayerCameraManager->GetFOVAngle();
 			PlayerController->GetPlayerViewPoint(/*out*/ OutViewInfo.Location, /*out*/ OutViewInfo.Rotation);
 		}
diff --git a/Engine/Source/Runtime/Engine/Private/PlayerCameraManager.cpp b/Engine/Source/Runtime/Engine/Private/PlayerCameraManager.cpp
--- a/Engine/Source/Runtime/Engine/Private/PlayerCameraManager.cpp
+++ b/Engine/Source/Runtime/Engine/Private/PlayerCameraManager.cpp
@@ -150,7 +150,7 @@ void APlayerCameraManager::SetViewTarget(class AActor* NewTarget, struct FViewTa
 			}
 
 			// use last frame's POV
-			ViewTarget.POV = GetLastFrameCameraCachePOV();
+			ViewTarget.POV = GetLastFrameCameraCacheView();
 			BlendParams = TransitionParams;
 			BlendTimeToGo = TransitionParams.BlendTime;
 
@@ -836,7 +836,7 @@ void APlayerCameraManager::EndPlay(const EEndPlayReason::Type EndPlayReason)
 
 void APlayerCameraManager::InitializeFor(APlayerController* PC)
 {
-	FMinimalViewInfo DefaultFOVCache = GetCameraCachePOV();
+	FMinimalViewInfo DefaultFOVCache = GetCameraCacheView();
 	DefaultFOVCache.FOV = DefaultFOV;
 	SetCameraCachePOV(DefaultFOVCache);
 
@@ -855,7 +855,7 @@ void APlayerCameraManager::InitializeFor(APlayerController* PC)
 
 float APlayerCameraManager::GetFOVAngle() const
 {
-	return (LockedFOV > 0.f) ? LockedFOV : GetCameraCachePOV().FOV;
+	return (LockedFOV > 0.f) ? LockedFOV : GetCameraCacheView().FOV;
 }
 
 void APlayerCameraManager::SetFOV(float NewFOV)
@@ -890,19 +890,19 @@ void APlayerCameraManager::UnlockOrthoWidth()
 
 void APlayerCameraManager::GetCameraViewPoint(FVector& OutCamLoc, FRotator& OutCamRot) const
 {
-	const FMinimalViewInfo& CurrentPOV = GetCameraCachePOV();
+	const FMinimalViewInfo& CurrentPOV = GetCameraCacheView();
 	OutCamLoc = CurrentPOV.Location;
 	OutCamRot = CurrentPOV.Rotation;
 }
 
 FRotator APlayerCameraManager::GetCameraRotation() const
 {
-	return GetCameraCachePOV().Rotation;
+	return GetCameraCacheView().Rotation;
 }
 
 FVector APlayerCameraManager::GetCameraLocation() const
 {
-	return GetCameraCachePOV().Location;
+	return GetCameraCacheView().Location;
 }
 
 void APlayerCameraManager::SetDesiredColorScale(FVector NewColorScale, float InterpTime)
@@ -953,8 +953,8 @@ void APlayerCameraManager::UpdateCamera(float DeltaTime)
 			const float ClientNetCamUpdateDeltaTime = GameNetworkManager->ClientNetCamUpdateDeltaTime;
 			const float ClientNetCamUpdatePositionLimit = GameNetworkManager->ClientNetCamUpdatePositionLimit;
 
-			const FMinimalViewInfo& CurrentPOV = GetCameraCachePOV();
-			const FMinimalViewInfo& LastPOV = GetLastFrameCameraCachePOV();
+			const FMinimalViewInfo& CurrentPOV = GetCameraCacheView();
+			const FMinimalViewInfo& LastPOV = GetLastFrameCameraCacheView();
 
 			FVector ClientCameraPosition = FRepMovement::RebaseOntoZeroOrigin(CurrentPOV.Location, this);
 			FVector PrevClientCameraPosition = FRepMovement::RebaseOntoZeroOrigin(LastPOV.Location, this);
@@ -1173,7 +1173,7 @@ void APlayerCameraManager::FillCameraCache(const FMinimalViewInfo& NewInfo)
 	const float CurrentGameTime = GetWorld()->TimeSeconds;
 	if (CurrentCacheTime != CurrentGameTime)
 	{
-		SetLastFrameCameraCachePOV(GetCameraCachePOV());
+		SetLastFrameCameraCachePOV(GetCameraCacheView());
 		SetLastFrameCameraCacheTime(CurrentCacheTime);
 	}
 
@@ -1240,7 +1240,7 @@ void APlayerCameraManager::LimitViewYaw(FRotator& ViewRotation, float InViewYawM
 
 void APlayerCameraManager::DisplayDebug(class UCanvas* Canvas, const FDebugDisplayInfo& DebugDisplay, float& YL, float& YPos)
 {
-	const FMinimalViewInfo& CurrentPOV = GetCameraCachePOV();
+	const FMinimalViewInfo& CurrentPOV = GetCameraCacheView();
 
 	FDisplayDebugManager& DisplayDebugManager = Canvas->DisplayDebugManager;
 	DisplayDebugManager.SetDrawColor(FColor(255, 255, 255));
@@ -1266,11 +1266,11 @@ void APlayerCameraManager::ApplyWorldOffset(const FVector& InOffset, bool bWorld
 {
 	Super::ApplyWorldOffset(InOffset, bWorldShift);
 
-	FMinimalViewInfo CurrentPOV = GetCameraCachePOV();
+	FMinimalViewInfo CurrentPOV = GetCameraCacheView();
 	CurrentPOV.Location += InOffset;
 	SetCameraCachePOV(CurrentPOV);
 
-	FMinimalViewInfo LastFramePOV = GetLastFrameCameraCachePOV();
+	FMinimalViewInfo LastFramePOV = GetLastFrameCameraCacheView();
 	LastFramePOV.Location += InOffset;
 	SetLastFrameCameraCachePOV(LastFramePOV);
 
@@ -1568,16 +1568,26 @@ void APlayerCameraManager::SetLastFrameCameraCachePOV(const FMinimalViewInfo& In
 	LastFrameCameraCachePrivate.POV = InPOV;
 }
 
-const FMinimalViewInfo& APlayerCameraManager::GetCameraCachePOV() const
+const FMinimalViewInfo& APlayerCameraManager::GetCameraCacheView() const
 {
 	return CameraCachePrivate.POV;
 }
 
-const FMinimalViewInfo& APlayerCameraManager::GetLastFrameCameraCachePOV() const
+const FMinimalViewInfo& APlayerCameraManager::GetLastFrameCameraCacheView() const
 {
 	return LastFrameCameraCachePrivate.POV;
 }
 
+FMinimalViewInfo APlayerCameraManager::GetCameraCachePOV() const
+{
+	return GetCameraCacheView();
+}
+
+FMinimalViewInfo APlayerCameraManager::GetLastFrameCameraCachePOV() const
+{
+	return GetLastFrameCameraCacheView();
+}
+
 PRAGMA_DISABLE_DEPRECATION_WARNINGS
 APlayerCameraManager::~APlayerCameraManager()
 {
diff --git a/Engine/Source/Runtime/Engine/Private/PlayerController.cpp b/Engine/Source/Runtime/Engine/Private/PlayerController.cpp
--- a/Engine/Source/Runtime/Engine/Private/PlayerController.cpp
+++ b/Engine/Source/Runtime/Engine/Private/PlayerController.cpp
@@ -1772,7 +1772,7 @@ void APlayerController::ServerUpdateCamera_Implementation(FVector_NetQuantize Ca
 #endif
 	{
 		//@TODO: CAMERA: Fat pipe
-		FMinimalViewInfo NewInfo = PlayerCameraManager->GetCameraCachePOV();
+		FMinimalViewInfo NewInfo = PlayerCameraManager->GetCameraCacheView();
 		NewInfo.Location = NewPOV.Location;
 		NewInfo.Rotation = NewPOV.Rotation;
 		PlayerCameraManager->FillCameraCache(NewInfo);
