Adding logging to the buoyancy component.  Adds 1D case to ComputeSprungMasses in SimpleSuspension to handle one axis being zero.

diff --git a/Engine/Plugins/Experimental/Water/Source/Runtime/Private/BuoyancyComponent.cpp b/Engine/Plugins/Experimental/Water/Source/Runtime/Private/BuoyancyComponent.cpp
--- a/Engine/Plugins/Experimental/Water/Source/Runtime/Private/BuoyancyComponent.cpp
+++ b/Engine/Plugins/Experimental/Water/Source/Runtime/Private/BuoyancyComponent.cpp
@@ -384,7 +384,8 @@ void UBuoyancyComponent::ComputePontoonCoefficients()
 		{
 			const FVector& LocalCOM = BodyInstance->GetMassSpaceLocal().GetLocation();
 			//Distribute a mass of 1 to each pontoon so that we get a scaling factor based on position relative to CoM
-			FSimpleSuspensionHelpers::ComputeSprungMasses(LocalPontoonLocations, LocalCOM, 1.f, PontoonCoefficients);
+			bool ComputeSuccess = FSimpleSuspensionHelpers::ComputeSprungMasses(LocalPontoonLocations, LocalCOM, 1.f, PontoonCoefficients);
+			ensureMsgf(ComputeSuccess, TEXT("Failed to compute sprung masses for: %s"), *GetOwner()->GetName());
 		}
 	}
 
diff --git a/Engine/Source/Runtime/Engine/Private/PhysicsEngine/SimpleSuspension.cpp b/Engine/Source/Runtime/Engine/Private/PhysicsEngine/SimpleSuspension.cpp
--- a/Engine/Source/Runtime/Engine/Private/PhysicsEngine/SimpleSuspension.cpp
+++ b/Engine/Source/Runtime/Engine/Private/PhysicsEngine/SimpleSuspension.cpp
@@ -37,6 +37,35 @@ void FSimpleSuspension::Update(const FTransform& LocalToWorld, const FVector& Li
 	FSimpleSuspensionHelpers::ComputeSuspensionForces(LinearVelocity, AngularVelocityRad, SuspensionState, SuspensionParams.SpringParams, SuspensionState);
 }
 
+bool FSimpleSuspensionHelpers::ComputeSingleAxisLambda(const Chaos::FReal AxisDot, const Chaos::FReal SumAxis, const uint32 Count, TArray<Chaos::FReal, TFixedAllocator<2>>& Lambdas)
+{
+	using Chaos::FReal;
+
+	//compute determinant
+	const FReal DetLL = AxisDot * Count - SumAxis * SumAxis;
+	if (!ensureMsgf(DetLL > SMALL_NUMBER || DetLL < -SMALL_NUMBER,
+		TEXT("Spring configuration is invalid! Please make sure no two springs are at the same location.")))
+	{
+		return false;
+	}
+
+	const FReal LambdaB0 = Lambdas[0];
+	const FReal LambdaB1 = Lambdas[1];
+
+	//Compute the inverse matrix
+	const FReal DetLLInv = 1.f / DetLL;
+	const FReal InvLL00 = Count * DetLLInv;
+	const FReal InvLL01 = -SumAxis * DetLLInv;
+	const FReal InvLL10 = -SumAxis * DetLLInv;
+	const FReal InvLL11 = AxisDot * DetLLInv;
+
+	//compute Lagrange Multipliers - The third lambda value will always be zero in the 1D case.
+	Lambdas[0] = InvLL00 * LambdaB0 + InvLL01 * LambdaB1;
+	Lambdas[1] = InvLL10 * LambdaB0 + InvLL11 * LambdaB1;
+
+	return true;
+}
+
 bool FSimpleSuspensionHelpers::ComputeSprungMasses(const TArray<FVector>& MassSpringPositions, const float TotalMass, TArray<float>& OutSprungMasses)
 {
 	using Chaos::FReal;
@@ -162,7 +191,7 @@ bool FSimpleSuspensionHelpers::ComputeSprungMasses(const TArray<FVector>& MassSp
 	N+3 equations and N+3 unknowns.
 
 	[  0   0   0   x0  x1  x2  ... ] [ lambda0 ] = [ m x_c ]
-	[  0   0   0   y0  y1  y2  ... ] [ lambda1 ]   [ m z_c ]
+	[  0   0   0   y0  y1  y2  ... ] [ lambda1 ]   [ m y_c ]
 	[  0   0   0   1   1   1   ... ] [ lambda2 ]   [ m     ]
 	[  x0  y0  1   2   0   0   ... ] [ m0      ]   [ 2 m_u ]
 	[  x1  y1  1   0   2   0       ] [ m1      ]   [ 2 m_u ]
@@ -180,7 +209,7 @@ bool FSimpleSuspensionHelpers::ComputeSprungMasses(const TArray<FVector>& MassSp
 	[  L   2 I ] [ m_vec      ]   [ v ]
 
 	where Lt = Transpose(L), lambda_vec = { lambda0, lambda1, lambda2 },
-	m_vec = { m0, m1, m2, ... }, u = m { x_c, z_c, 1 }, and v = 2 m_u { 1, 1, 1, ...}.
+	m_vec = { m0, m1, m2, ... }, u = m { x_c, y_c, 1 }, and v = 2 m_u { 1, 1, 1, ...}.
 
 	This system can be solved for lambda_vec and subsequently m_vec, which is our solution vector.
 
@@ -215,39 +244,75 @@ bool FSimpleSuspensionHelpers::ComputeSprungMasses(const TArray<FVector>& MassSp
 		XDotY += X * Y;
 	}
 
-	// Compute determinant of system matrix, in prep for inversion
-	const FReal DetLL
-		= (XDotX * YDotY * Count)
-		+ (2.f * XDotY * SumX * SumY)
-		- (YDotY * SumX * SumX)
-		- (XDotX * SumY * SumY)
-		- (XDotY * Count);
+	//calculate the lambdas - we approximate the center of mass as zero for each axis
+	FReal Lambda0 = 0.0f;
+	FReal Lambda1 = 0.0f;
+	FReal Lambda2 = 0.0f;
 
-	// Make sure the matrix is invertible!
-	if (!ensureMsgf(DetLL > SMALL_NUMBER || DetLL < -SMALL_NUMBER, TEXT("Spring configuration is invalid! Please make sure no two springs are at the same location.")))
-	{
-		return false;
-	}
-
-	// Compute the elements of the inverse matrix
-	const FReal DetLLInv = 1.f / DetLL;
-	const FReal InvLL00 = ((Count * YDotY) - (SumY * SumY)) * DetLLInv;
-	const FReal InvLL01 = ((SumX * SumY) - (Count * XDotY)) * DetLLInv;
-	const FReal InvLL02 = ((SumY * XDotY) - (SumX * YDotY)) * DetLLInv;
-	const FReal InvLL10 = InvLL01; // Symmetry!
-	const FReal InvLL11 = ((Count * XDotX) - (SumX * SumX)) * DetLLInv;
-	const FReal InvLL12 = ((SumX * XDotY) - (SumY * XDotX)) * DetLLInv; // = InvLL21. Symmetry!
-	const FReal InvLL20 = InvLL02; // Symmetry!
-	const FReal InvLL21 = InvLL12; // Symmetry!
-	const FReal InvLL22 = ((XDotX * YDotY) - (XDotY * XDotY)) * DetLLInv;
-
-	// Compute the Lagrange multipliers
 	const FReal LambdaB0 = 2.f * AverageMass * SumX;
 	const FReal LambdaB1 = 2.f * AverageMass * SumY;
 	const FReal LambdaB2 = (2.f * AverageMass * CountN) - (2.f * TotalMass);
-	const FReal Lambda0 = (InvLL00 * LambdaB0) + (InvLL01 * LambdaB1) + (InvLL02 * LambdaB2);
-	const FReal Lambda1 = (InvLL10 * LambdaB0) + (InvLL11 * LambdaB1) + (InvLL12 * LambdaB2);
-	const FReal Lambda2 = (InvLL20 * LambdaB0) + (InvLL21 * LambdaB1) + (InvLL22 * LambdaB2);
+
+	//if one axis is zero, we actually lose a constraint/equation and we need to adjust our calculation
+	if (YDotY < SMALL_NUMBER  && YDotY > -SMALL_NUMBER)
+	{
+		TArray<Chaos::FReal, TFixedAllocator<2>> Lambdas;
+		Lambdas.Add(LambdaB0);
+		Lambdas.Add(LambdaB2);
+
+		ComputeSingleAxisLambda(XDotX, SumX, Count, Lambdas);
+
+		Lambda0 = Lambdas[0];
+		Lambda2 = Lambdas[1];
+	}
+
+	else if (XDotX < SMALL_NUMBER && XDotX > -SMALL_NUMBER)
+	{
+		TArray<Chaos::FReal, TFixedAllocator<2>> Lambdas;
+		Lambdas.Add(LambdaB1);
+		Lambdas.Add(LambdaB2);
+
+		ComputeSingleAxisLambda(YDotY, SumY, Count, Lambdas);
+
+		Lambda1 = Lambdas[0];
+		Lambda2 = Lambdas[1];
+	}
+
+	//2D constraint calculation
+	else
+	{
+		// Compute determinant of system matrix, in prep for inversion
+		const FReal DetLL
+			= (XDotX * YDotY * Count)
+			+ (2.f * XDotY * SumX * SumY)
+			- (YDotY * SumX * SumX)
+			- (XDotX * SumY * SumY)
+			- (XDotY * XDotY * Count);
+
+		// Make sure the matrix is invertible!
+		if (!ensureMsgf(DetLL > SMALL_NUMBER || DetLL < -SMALL_NUMBER,
+			TEXT("Spring configuration is invalid! Please make sure no two springs are at the same location.")))
+		{
+			return false;
+		}
+
+		// Compute the elements of the inverse matrix
+		const FReal DetLLInv = 1.f / DetLL;
+		const FReal InvLL00 = ((Count * YDotY) - (SumY * SumY)) * DetLLInv;
+		const FReal InvLL01 = ((SumX * SumY) - (Count * XDotY)) * DetLLInv;
+		const FReal InvLL02 = ((SumY * XDotY) - (SumX * YDotY)) * DetLLInv;
+		const FReal InvLL10 = InvLL01; // Symmetry!
+		const FReal InvLL11 = ((Count * XDotX) - (SumX * SumX)) * DetLLInv;
+		const FReal InvLL12 = ((SumX * XDotY) - (SumY * XDotX)) * DetLLInv; // = InvLL21. Symmetry!
+		const FReal InvLL20 = InvLL02; // Symmetry!
+		const FReal InvLL21 = InvLL12; // Symmetry!
+		const FReal InvLL22 = ((XDotX * YDotY) - (XDotY * XDotY)) * DetLLInv;
+
+		// Compute the Lagrange multipliers
+		Lambda0 = (InvLL00 * LambdaB0) + (InvLL01 * LambdaB1) + (InvLL02 * LambdaB2);
+		Lambda1 = (InvLL10 * LambdaB0) + (InvLL11 * LambdaB1) + (InvLL12 * LambdaB2);
+		Lambda2 = (InvLL20 * LambdaB0) + (InvLL21 * LambdaB1) + (InvLL22 * LambdaB2);		
+	}
 
 	// Compute the masses
 	for (uint32 Index = 0; Index < Count; ++Index)
diff --git a/Engine/Source/Runtime/Engine/Public/Physics/SimpleSuspension.h b/Engine/Source/Runtime/Engine/Public/Physics/SimpleSuspension.h
--- a/Engine/Source/Runtime/Engine/Public/Physics/SimpleSuspension.h
+++ b/Engine/Source/Runtime/Engine/Public/Physics/SimpleSuspension.h
@@ -138,6 +138,9 @@ struct FSimpleSuspensionHelpers
 		of mass. */
 	ENGINE_API static bool ComputeSprungMasses(const TArray<FVector>& LocalSpringPositions, const FVector& LocalCenterOfMass, const float TotalMass, TArray<float>& OutSprungMasses);
 
+	/** Calculates the Lambdas in the single axis case*/
+	static bool ComputeSingleAxisLambda(const Chaos::FReal AxisDot, const Chaos::FReal SumAxis, const uint32 Count, TArray<Chaos::FReal, TFixedAllocator<2>>& Lambdas);
+
 	/** Given a sprung mass and a spring stiffness, compute the natural
 		frequency of the spring */
 	static void ComputeSpringNaturalFrequencyAndDampingRatio(const float SprungMass, const float SpringStiffness, const float SpringDamping, float& OutNaturalFrequency, float& OutDampingRatio);
