Fixed: Vulkan validation assert related to a mobile implementation of SSProfilePreIntegratedTexture

diff --git a/Engine/Source/Runtime/Engine/Private/Rendering/SubsurfaceProfile.cpp b/Engine/Source/Runtime/Engine/Private/Rendering/SubsurfaceProfile.cpp
--- a/Engine/Source/Runtime/Engine/Private/Rendering/SubsurfaceProfile.cpp
+++ b/Engine/Source/Runtime/Engine/Private/Rendering/SubsurfaceProfile.cpp
@@ -26,6 +26,16 @@ static TAutoConsoleVariable<int32> CVarSSProfilesPreIntegratedTextureForceUpdate
 	ECVF_Cheat | ECVF_RenderThreadSafe);
 #endif
 
+static bool ForceUpdateSSProfilesPreIntegratedTexture()
+{
+#if !(UE_BUILD_SHIPPING || UE_BUILD_TEST)
+	return (CVarSSProfilesPreIntegratedTextureForceUpdate.GetValueOnAnyThread() == 1);
+#else
+	return false;
+#endif
+}
+
+
 // lives on the render thread
 ENGINE_API TGlobalResource<FSubsurfaceProfileTexture> GSubsurfaceProfileTextureObject;
 
@@ -192,31 +202,28 @@ IPooledRenderTarget* FSubsurfaceProfileTexture::GetSSProfilesPreIntegratedTextur
 
 	// PreIntegrated SSS look up texture
 	int32 SSProfilesPreIntegratedTextureResolution = FMath::RoundUpToPowerOfTwo(FMath::Max(CVarSSProfilesPreIntegratedTextureResolution.GetValueOnAnyThread(), 32));
-
-	int32 SubsurfaceProfileEntriesNum = SubsurfaceProfileEntries.Num();
-
-	bool bSSProfilesPreIntegratedTextureNeedsUpdate = false;
-
-	if (!GSSProfilesPreIntegratedTexture || GSSProfilesPreIntegratedTexture->GetDesc().Extent != SSProfilesPreIntegratedTextureResolution)
+	
+	// Generate the new preintegrated texture if needed.
+	if (!GSSProfilesPreIntegratedTexture || 
+		GSSProfilesPreIntegratedTexture->GetDesc().Extent != SSProfilesPreIntegratedTextureResolution ||
+		ForceUpdateSSProfilesPreIntegratedTexture())
 	{
-		bSSProfilesPreIntegratedTextureNeedsUpdate = true;
+		GSSProfilesPreIntegratedTexture.SafeRelease();
 
+		int32 SubsurfaceProfileEntriesNum = SubsurfaceProfileEntries.Num();
 		// Use RGB10A2 since it could be compressed by mobile hardware according to ARM.
-		FPooledRenderTargetDesc Desc(FPooledRenderTargetDesc::Create2DArrayDesc(SSProfilesPreIntegratedTextureResolution, PF_A2B10G10R10, FClearValueBinding::Black, TexCreate_TargetArraySlicesIndependently, TexCreate_ShaderResource | TexCreate_RenderTargetable, false, SubsurfaceProfileEntriesNum));
+		FRDGTextureDesc ProfileTextureDesc = FRDGTextureDesc::Create2DArray(
+			SSProfilesPreIntegratedTextureResolution, 
+			PF_A2B10G10R10, 
+			FClearValueBinding::Black, 
+			TexCreate_TargetArraySlicesIndependently | TexCreate_ShaderResource | TexCreate_RenderTargetable,
+			SubsurfaceProfileEntriesNum);
+		FRDGTextureRef ProfileTexture = GraphBuilder.CreateTexture(ProfileTextureDesc, TEXT("SSProfilePreIntegratedTexture"));
 
-		GRenderTargetPool.FindFreeElement(GraphBuilder.RHICmdList, Desc, GSSProfilesPreIntegratedTexture, TEXT("SSProfilePreIntegratedTexture"));
-	}
-	// Generate the new preintegrated texture if needed.
-	if (bSSProfilesPreIntegratedTextureNeedsUpdate
-#if !(UE_BUILD_SHIPPING || UE_BUILD_TEST)
-		|| CVarSSProfilesPreIntegratedTextureForceUpdate.GetValueOnAnyThread() == 1
-#endif
-		)
-	{
 		for (int32 i = 0; i < SubsurfaceProfileEntriesNum; ++i)
 		{
 			FSSProfilePreIntegratedPS::FParameters* PassParameters = GraphBuilder.AllocParameters<FSSProfilePreIntegratedPS::FParameters>();
-			PassParameters->RenderTargets[0] = FRenderTargetBinding(GraphBuilder.RegisterExternalTexture(GSSProfilesPreIntegratedTexture), ERenderTargetLoadAction::EClear, 0, i);
+			PassParameters->RenderTargets[0] = FRenderTargetBinding(ProfileTexture, ERenderTargetLoadAction::EClear, 0, i);
 
 			PassParameters->SourceSSProfilesTexture = GetSubsurfaceProfileTextureWithFallback();
 			PassParameters->SourceSSProfilesSampler = TStaticSamplerState<SF_Point, AM_Clamp, AM_Clamp, AM_Clamp>::GetRHI();
@@ -224,9 +231,7 @@ IPooledRenderTarget* FSubsurfaceProfileTexture::GetSSProfilesPreIntegratedTextur
 			FIntVector SourceSSProfilesTextureSize = PassParameters->SourceSSProfilesTexture->GetSizeXYZ();
 
 			PassParameters->SourceSSProfilesTextureSizeAndInvSize = FVector4f(SourceSSProfilesTextureSize.X, SourceSSProfilesTextureSize.Y, 1.0f / SourceSSProfilesTextureSize.X, 1.0f / SourceSSProfilesTextureSize.Y);
-
 			PassParameters->TargetSSProfilesPreIntegratedTextureSizeAndInvSize = FVector4f(SSProfilesPreIntegratedTextureResolution, SSProfilesPreIntegratedTextureResolution, 1.0f / SSProfilesPreIntegratedTextureResolution, 1.0f / SSProfilesPreIntegratedTextureResolution);
-
 			PassParameters->SourceSubsurfaceProfileInt = i;
 
 			FIntRect ViewRect = FIntRect(FIntPoint(0, 0), SSProfilesPreIntegratedTextureResolution);
@@ -236,6 +241,8 @@ IPooledRenderTarget* FSubsurfaceProfileTexture::GetSSProfilesPreIntegratedTextur
 
 			FPixelShaderUtils::AddFullscreenPass(GraphBuilder, GlobalShaderMap, RDG_EVENT_NAME("SSS::SSProfilePreIntegrated"), PixelShader, PassParameters, ViewRect);
 		}
+
+		GraphBuilder.QueueTextureExtraction(ProfileTexture, &GSSProfilesPreIntegratedTexture);
 	}
 
 	return GSSProfilesPreIntegratedTexture;
