Make GameThread_GetParameterValue determine parameter value validity in the same way as RenderThread_GetParameterValue. This fixes a bug where some mesh components fail to register themselves as instances of some textures becuase the corresponding material parameters have been overridden to none in MIDs. Failing to register those components can lead to blurry textures because the streamer doesn't know that the instances are rendered right in front of the camera.

diff --git a/Engine/Source/Runtime/Engine/Classes/Materials/MaterialInstance.h b/Engine/Source/Runtime/Engine/Classes/Materials/MaterialInstance.h
--- a/Engine/Source/Runtime/Engine/Classes/Materials/MaterialInstance.h
+++ b/Engine/Source/Runtime/Engine/Classes/Materials/MaterialInstance.h
@@ -98,6 +98,8 @@ struct FScalarParameterValue
 
 	bool IsOverride() const { return true; }
 
+	bool IsValid() const { return true; }
+
 	bool operator==(const FScalarParameterValue& Other) const
 	{
 		return
@@ -152,6 +154,8 @@ struct FVectorParameterValue
 
 	bool IsOverride() const { return true; }
 
+	bool IsValid() const { return true; }
+
 	bool operator==(const FVectorParameterValue& Other) const
 	{
 		return
@@ -199,6 +203,8 @@ struct FDoubleVectorParameterValue
 
 	bool IsOverride() const { return true; }
 
+	bool IsValid() const { return true; }
+
 	bool operator==(const FDoubleVectorParameterValue& Other) const
 	{
 		return
@@ -250,6 +256,8 @@ struct FTextureParameterValue
 
 	bool IsOverride() const { return true; }
 
+	bool IsValid() const { return GetValue(*this) != nullptr; }
+
 	bool operator==(const FTextureParameterValue& Other) const
 	{
 		return
@@ -296,6 +304,8 @@ struct FRuntimeVirtualTextureParameterValue
 
 	bool IsOverride() const { return true; }
 
+	bool IsValid() const { return GetValue(*this) != nullptr; }
+
 	bool operator==(const FRuntimeVirtualTextureParameterValue& Other) const
 	{
 		return
@@ -350,6 +360,8 @@ struct FFontParameterValue
 
 	bool IsOverride() const { return true; }
 
+	bool IsValid() const { return GetValue(*this) != nullptr; }
+
 	bool operator==(const FFontParameterValue& Other) const
 	{
 		return
diff --git a/Engine/Source/Runtime/Engine/Private/Materials/MaterialInstanceSupport.h b/Engine/Source/Runtime/Engine/Private/Materials/MaterialInstanceSupport.h
--- a/Engine/Source/Runtime/Engine/Private/Materials/MaterialInstanceSupport.h
+++ b/Engine/Source/Runtime/Engine/Private/Materials/MaterialInstanceSupport.h
@@ -324,7 +324,7 @@ inline bool GameThread_GetParameterValue(const TArray<ParameterType>& Parameters
 	for (int32 ParameterIndex = 0; ParameterIndex < Parameters.Num(); ParameterIndex++)
 	{
 		const ParameterType* Parameter = &Parameters[ParameterIndex];
-		if (Parameter->IsOverride() && Parameter->ParameterInfo == ParameterInfo)
+		if (Parameter->IsOverride() && Parameter->IsValid() && Parameter->ParameterInfo == ParameterInfo)
 		{
 			Parameter->GetValue(OutResult);
 			return true;
diff --git a/Engine/Source/Runtime/Engine/Public/StaticParameterSet.h b/Engine/Source/Runtime/Engine/Public/StaticParameterSet.h
--- a/Engine/Source/Runtime/Engine/Public/StaticParameterSet.h
+++ b/Engine/Source/Runtime/Engine/Public/StaticParameterSet.h
@@ -143,6 +143,8 @@ struct FStaticSwitchParameter : public FStaticParameterBase
 		OutResult.ExpressionGuid = ExpressionGUID;
 #endif
 	}
+
+	bool IsValid() const { return true; }
 };
 
 /**
@@ -228,6 +230,8 @@ struct FStaticComponentMaskParameter : public FStaticParameterBase
 		OutResult.ExpressionGuid = ExpressionGUID;
 #endif
 	}
+
+	bool IsValid() const { return true; }
 };
 
 /**
