Default r.Decal.NormalReprojectionEnabled to false.
Fix crash when it is set to true in forward rendering.
Fix WPO not writing velocity for case of velocity in depth pass where it isn't a perf win and the behavior is confusing.
#rnx

diff --git a/Engine/Source/Runtime/Engine/Private/PrimitiveSceneProxy.cpp b/Engine/Source/Runtime/Engine/Private/PrimitiveSceneProxy.cpp
--- a/Engine/Source/Runtime/Engine/Private/PrimitiveSceneProxy.cpp
+++ b/Engine/Source/Runtime/Engine/Private/PrimitiveSceneProxy.cpp
@@ -38,6 +38,13 @@ static TAutoConsoleVariable<int32> CVarCacheWPOPrimitives(
 	ECVF_RenderThreadSafe | ECVF_Scalability
 	);
 
+static TAutoConsoleVariable<int32> CVarVertexDeformationOutputsVelocity(
+	TEXT("r.VertexDeformationOutputsVelocity"),
+	0,
+	TEXT("Enables materials with World Position Offset and/or World Displacement to output velocities during velocity pass even when the actor has not moved. \n")
+	TEXT("This only has an impact if r.VelocityOutputPass=2. \n")
+	TEXT("This will incur a performance cost that can be quite significant if many objects are using WPO, such as a forest of trees."));
+
 bool CacheShadowDepthsFromPrimitivesUsingWPO()
 {
 	return CVarCacheWPOPrimitives.GetValueOnAnyThread(true) != 0;
@@ -98,6 +105,17 @@ bool SupportsNaniteRendering(const FVertexFactory* RESTRICT VertexFactory, const
 	return false;
 }
 
+static bool VertexDeformationOutputsVelocity()
+{
+	// This is only optional when velocity pass is after base pass (when it could create many more draw calls).
+	// Note this output pass setting is read only.
+	static const auto CVarVelocityOutputPass = IConsoleManager::Get().FindConsoleVariable(TEXT("r.VelocityOutputPass"));
+	bool bVertexDeformationOutputsVelocityForPass = CVarVelocityOutputPass && CVarVelocityOutputPass->GetInt() != 2;
+	
+	bool bVertexDeformationOutputsVelocity = CVarVertexDeformationOutputsVelocity.GetValueOnAnyThread() != 0;
+	return bVertexDeformationOutputsVelocityForPass || bVertexDeformationOutputsVelocity;
+}
+
 FPrimitiveSceneProxy::FPrimitiveSceneProxy(const UPrimitiveComponent* InComponent, FName InResourceName)
 :
 #if !(UE_BUILD_SHIPPING || UE_BUILD_TEST)
@@ -315,9 +333,7 @@ FPrimitiveSceneProxy::FPrimitiveSceneProxy(const UPrimitiveComponent* InComponen
 	FObjectCacheEventSink::NotifyUsedMaterialsChanged_Concurrent(InComponent, UsedMaterialsForVerification);
 #endif
 
-	static const auto CVarVertexDeformationOutputsVelocity = IConsoleManager::Get().FindConsoleVariable(TEXT("r.VertexDeformationOutputsVelocity"));
-
-	if (!bAlwaysHasVelocity && IsMovable() && CVarVertexDeformationOutputsVelocity && CVarVertexDeformationOutputsVelocity->GetInt())
+	if (!bAlwaysHasVelocity && IsMovable() && VertexDeformationOutputsVelocity())
 	{
 		ERHIFeatureLevel::Type FeatureLevel = GetScene().GetFeatureLevel();
 
diff --git a/Engine/Source/Runtime/Renderer/Private/CompositionLighting/PostProcessDeferredDecals.cpp b/Engine/Source/Runtime/Renderer/Private/CompositionLighting/PostProcessDeferredDecals.cpp
--- a/Engine/Source/Runtime/Renderer/Private/CompositionLighting/PostProcessDeferredDecals.cpp
+++ b/Engine/Source/Runtime/Renderer/Private/CompositionLighting/PostProcessDeferredDecals.cpp
@@ -43,7 +43,7 @@ static TAutoConsoleVariable<float> CVarDBufferDecalNormalReprojectionThresholdHi
 
 static TAutoConsoleVariable<bool> CVarDBufferDecalNormalReprojectionEnabled(
 	TEXT("r.Decal.NormalReprojectionEnabled"),
-	true, 
+	false, 
 	TEXT("If true, normal reprojection from the previous frame is allowed in SceneTexture nodes on DBuffer decals, provided that motion ")
 	TEXT("in depth prepass is enabled as well (r.VelocityOutputPass=0). Otherwise the fallback is the normal extracted from the depth buffer."),
 	ECVF_RenderThreadSafe);
@@ -561,10 +561,8 @@ void AddDeferredDecalPass(
 
 void ExtractNormalsForNextFrameReprojection(FRDGBuilder& GraphBuilder, const FSceneTextures& SceneTextures, const TArray<FViewInfo>& Views)
 {
-	static auto CVarNormalReprojectionEnabled = IConsoleManager::Get().FindConsoleVariable(TEXT("r.Decal.NormalReprojectionEnabled"));
-
 	// save the previous frame if early motion vectors are enabled and normal reprojection is enabled, so there should be no cost if these options are off
-	bool bApplyReproject = FVelocityRendering::DepthPassCanOutputVelocity() && CVarNormalReprojectionEnabled && CVarNormalReprojectionEnabled->GetInt() > 0;
+	bool bApplyReproject = FVelocityRendering::DepthPassCanOutputVelocity() && CVarDBufferDecalNormalReprojectionEnabled.GetValueOnRenderThread();
 
 	if (bApplyReproject)
 	{
diff --git a/Engine/Source/Runtime/Renderer/Private/DeferredShadingRenderer.cpp b/Engine/Source/Runtime/Renderer/Private/DeferredShadingRenderer.cpp
--- a/Engine/Source/Runtime/Renderer/Private/DeferredShadingRenderer.cpp
+++ b/Engine/Source/Runtime/Renderer/Private/DeferredShadingRenderer.cpp
@@ -2648,8 +2648,11 @@ void FDeferredShadingSceneRenderer::Render(FRDGBuilder& GraphBuilder)
 		}
 	}
 
-	// mark GBufferA for saving for next frame if it's needed
-	ExtractNormalsForNextFrameReprojection(GraphBuilder, SceneTextures, Views);
+	if (bUseGBuffer)
+	{
+		// mark GBufferA for saving for next frame if it's needed
+		ExtractNormalsForNextFrameReprojection(GraphBuilder, SceneTextures, Views);
+	}
 
 	// Rebuild scene textures to include GBuffers.
 	SceneTextures.SetupMode |= ESceneTextureSetupMode::GBuffers;
diff --git a/Engine/Source/Runtime/Renderer/Private/VelocityRendering.cpp b/Engine/Source/Runtime/Renderer/Private/VelocityRendering.cpp
--- a/Engine/Source/Runtime/Renderer/Private/VelocityRendering.cpp
+++ b/Engine/Source/Runtime/Renderer/Private/VelocityRendering.cpp
@@ -47,13 +47,6 @@ static TAutoConsoleVariable<int32> CVarParallelVelocity(
 	ECVF_RenderThreadSafe
 	);
 
-static TAutoConsoleVariable<int32> CVarVertexDeformationOutputsVelocity(
-	TEXT("r.VertexDeformationOutputsVelocity"),
-	0,
-	TEXT("Enables materials with World Position Offset and/or World Displacement to output velocities during velocity pass even when the actor has not moved. \n")
-	TEXT("This only has an impact if r.VelocityOutputPass=2. \n")
-	TEXT("This will incur a performance cost that can be quite significant if many objects are using WPO, such as a forest of trees."));
-
 static TAutoConsoleVariable<int32> CVarRHICmdFlushRenderThreadTasksVelocityPass(
 	TEXT("r.RHICmdFlushRenderThreadTasksVelocityPass"),
 	0,
