Fix for ClientCameraPosition doesn't fit in FVector_NetQuantize for ServerUpdateCamera log spam.

The position was previously clamped to well below the world max bounds. Clamping is now removed as the serialization can handle any value, except for non-finite numbers which will cause an ensure.

#rnx

diff --git a/Engine/Source/Runtime/Engine/Private/PlayerCameraManager.cpp b/Engine/Source/Runtime/Engine/Private/PlayerCameraManager.cpp
--- a/Engine/Source/Runtime/Engine/Private/PlayerCameraManager.cpp
+++ b/Engine/Source/Runtime/Engine/Private/PlayerCameraManager.cpp
@@ -974,21 +974,6 @@ void APlayerCameraManager::UpdateCamera(float DeltaTime)
 
 				if ((CompressedRotation != PrevCompressedRotation) || !ClientCameraPosition.Equals(PrevClientCameraPosition) || (TimeSinceLastServerUpdateCamera > ServerUpdateCameraTimeout))
 				{
-#if !(UE_BUILD_SHIPPING || UE_BUILD_TEST)
-					if (ClientCameraPosition.X > 1048576.0f || ClientCameraPosition.X < -1048576.0f ||
-						ClientCameraPosition.Y > 1048576.0f || ClientCameraPosition.Y < -1048576.0f ||
-						ClientCameraPosition.Z > 1048576.0f || ClientCameraPosition.Z < -1048576.0f)
-					{
-						UE_LOG(LogPlayerCameraManager, Warning, TEXT("ClientCameraPosition %f %f %f doesn't fit in FVector_NetQuantize for ServerUpdateCamera, capping"), ClientCameraPosition.X, ClientCameraPosition.Y, ClientCameraPosition.Z);
-					}
-#endif //!(UE_BUILD_SHIPPING || UE_BUILD_TEST)
-					
-					const float MaxQuantize = 1048575.f;
-					const float MinQuantize = -1048575.f;
-					ClientCameraPosition.X = FMath::Clamp(ClientCameraPosition.X, MinQuantize, MaxQuantize);
-					ClientCameraPosition.Y = FMath::Clamp(ClientCameraPosition.Y, MinQuantize, MaxQuantize);
-					ClientCameraPosition.Z = FMath::Clamp(ClientCameraPosition.Z, MinQuantize, MaxQuantize);
-
 					PCOwner->ServerUpdateCamera(ClientCameraPosition, CompressedRotation);
 
 					TimeSinceLastServerUpdateCamera = 0.0f;
