Material Pow() expressions result type is driven by left-hand-side input only

diff --git a/Engine/Source/Runtime/Engine/Private/Materials/HLSLMaterialDerivativeAutogen.cpp b/Engine/Source/Runtime/Engine/Private/Materials/HLSLMaterialDerivativeAutogen.cpp
--- a/Engine/Source/Runtime/Engine/Private/Materials/HLSLMaterialDerivativeAutogen.cpp
+++ b/Engine/Source/Runtime/Engine/Private/Materials/HLSLMaterialDerivativeAutogen.cpp
@@ -412,6 +412,7 @@ FMaterialDerivativeAutogen::FOperationType2 FMaterialDerivativeAutogen::GetFunc2
 					return FOperationType2(NonLWCType, LWCType, NonLWCType);
 				case EFunc2::Min:
 				case EFunc2::Max:
+				case EFunc2::PowPositiveClamped:
 					// See note about Min/Max below
 					return LhsLWCType;
 				default:
@@ -435,9 +436,11 @@ FMaterialDerivativeAutogen::FOperationType2 FMaterialDerivativeAutogen::GetFunc2
 				return FOperationType2(EDerivativeType::Float1, NonLWCType, NonLWCType);
 			case EFunc2::Min:
 			case EFunc2::Max:
+			case EFunc2::PowPositiveClamped:
 				// Min/Max should probably use the same logic as other binary operations
 				// However, legacy HLSLMaterialTranslator::Min/Max simply coerce the Rhs value to the type of the Lhs value
 				// Need to preserve this behavior here (may want to wrap this behind a legacy flag/CVAR at some point)
+				// Legacy PowPositiveClamped behavior is implemented differently, but end result is the same
 				return LhsNonLWCType;
 			default:
 				return NonLWCType;
