Montages using inertial blends now support blend profiles

diff --git a/Engine/Source/Runtime/AnimGraphRuntime/Private/AnimNodes/AnimNode_Slot.cpp b/Engine/Source/Runtime/AnimGraphRuntime/Private/AnimNodes/AnimNode_Slot.cpp
--- a/Engine/Source/Runtime/AnimGraphRuntime/Private/AnimNodes/AnimNode_Slot.cpp
+++ b/Engine/Source/Runtime/AnimGraphRuntime/Private/AnimNodes/AnimNode_Slot.cpp
@@ -39,13 +39,13 @@ void FAnimNode_Slot::Update_AnyThread(const FAnimationUpdateContext& Context)
 	// Update cache in AnimInstance.
 	Context.AnimInstanceProxy->UpdateSlotNodeWeight(SlotName, WeightData.SlotNodeWeight, Context.GetFinalBlendWeight());
 
-	float InertializationDuration;
-	if (Context.AnimInstanceProxy->GetSlotInertializationRequest(SlotName, InertializationDuration))
+	UE::Anim::FSlotInertializationRequest InertializationRequest;
+	if (Context.AnimInstanceProxy->GetSlotInertializationRequest(SlotName, InertializationRequest))
 	{
 		UE::Anim::IInertializationRequester* InertializationRequester = Context.GetMessage<UE::Anim::IInertializationRequester>();
 		if (InertializationRequester)
 		{
-			InertializationRequester->RequestInertialization(InertializationDuration);
+			InertializationRequester->RequestInertialization(InertializationRequest.Get<0>(), InertializationRequest.Get<1>());
 		}
 		else
 		{
diff --git a/Engine/Source/Runtime/Engine/Classes/Animation/AnimInstance.h b/Engine/Source/Runtime/Engine/Classes/Animation/AnimInstance.h
--- a/Engine/Source/Runtime/Engine/Classes/Animation/AnimInstance.h
+++ b/Engine/Source/Runtime/Engine/Classes/Animation/AnimInstance.h
@@ -36,6 +36,11 @@ struct FAnimNode_LinkedAnimLayer;
 
 typedef TArray<FTransform> FTransformArrayA2;
 
+namespace UE::Anim
+{
+	using FSlotInertializationRequest = TPair<float, const UBlendProfile*>;
+}	// namespace UE::Anim
+
 struct FParallelEvaluationData
 {
 	FBlendedHeapCurve& OutCurve;
@@ -884,7 +889,7 @@ protected:
 	TMap<class UAnimMontage*, struct FAnimMontageInstance*> ActiveMontagesMap;
 
 	/**  Inertialization requests gathered this frame. Gets reset in UpdateMontageEvaluationData */
-	TMap<FName, float> SlotGroupInertializationRequestMap;
+	TMap<FName, UE::Anim::FSlotInertializationRequest> SlotGroupInertializationRequestMap;
 
 	/* StopAllMontagesByGroupName needs a BlendMode and BlendProfile to function properly if using non-default ones in your montages. If you want default BlendMode/BlendProfiles, you need to update the calling code to do so. */
 	UE_DEPRECATED(5.0, "Use StopAllMontagesByGroupName with other signature.")
@@ -903,12 +908,12 @@ protected:
 
 public:
 
-	/**  Builds an inertialization request from the montage's group and provided duration */
-	void RequestMontageInertialization(const UAnimMontage* Montage, float Duration);
+	/**  Builds an inertialization request from the montage's group, provided duration and optional blend profile*/
+	void RequestMontageInertialization(const UAnimMontage* Montage, float Duration, const UBlendProfile* BlendProfile = nullptr);
 
 	/**  Requests an inertial blend during the next anim graph update. Requires your anim graph to have a slot node belonging to the specified group name */
 	UFUNCTION(BlueprintCallable, Category = "Inertial Blend")
-	void RequestSlotGroupInertialization(FName InSlotGroupName, float Duration);
+	void RequestSlotGroupInertialization(FName InSlotGroupName, float Duration, const UBlendProfile* BlendProfile = nullptr);
 
 	/** Queue a Montage BlendingOut Event to be triggered. */
 	void QueueMontageBlendingOutEvent(const FQueuedMontageBlendingOutEvent& MontageBlendingOutEvent);
diff --git a/Engine/Source/Runtime/Engine/Private/Animation/AnimInstance.cpp b/Engine/Source/Runtime/Engine/Private/Animation/AnimInstance.cpp
--- a/Engine/Source/Runtime/Engine/Private/Animation/AnimInstance.cpp
+++ b/Engine/Source/Runtime/Engine/Private/Animation/AnimInstance.cpp
@@ -1766,20 +1766,20 @@ void UAnimInstance::Montage_Advance(float DeltaSeconds)
 	}
 }
 
-void UAnimInstance::RequestSlotGroupInertialization(FName InSlotGroupName, float Duration)
+void UAnimInstance::RequestSlotGroupInertialization(FName InSlotGroupName, float Duration, const UBlendProfile* BlendProfile)
 {
 	// Must add this on both the anim instance and proxy's map, as this could called after UAnimInstance::UpdateMontageEvaluationData.
-	SlotGroupInertializationRequestMap.FindOrAdd(InSlotGroupName) = Duration;
-	GetProxyOnAnyThread<FAnimInstanceProxy>().GetSlotGroupInertializationRequestMap().FindOrAdd(InSlotGroupName) = Duration;
+	SlotGroupInertializationRequestMap.FindOrAdd(InSlotGroupName) = UE::Anim::FSlotInertializationRequest(Duration, BlendProfile);
+	GetProxyOnAnyThread<FAnimInstanceProxy>().GetSlotGroupInertializationRequestMap().FindOrAdd(InSlotGroupName) = UE::Anim::FSlotInertializationRequest(Duration, BlendProfile);
 }
 
-void UAnimInstance::RequestMontageInertialization(const UAnimMontage* Montage, float Duration)
+void UAnimInstance::RequestMontageInertialization(const UAnimMontage* Montage, float Duration, const UBlendProfile* BlendProfile)
 {
 	if (Montage)
 	{
 		// Adds a new request or overwrites an existing one
 		// We always overwrite with the last request, instead of using the shortest one (differs from AnimNode_Inertialization), because we expect the last montage played/stopped to take precedence
-		SlotGroupInertializationRequestMap.FindOrAdd(Montage->GetGroupName()) = Duration;
+		SlotGroupInertializationRequestMap.FindOrAdd(Montage->GetGroupName()) = UE::Anim::FSlotInertializationRequest(Duration, BlendProfile);
 	}
 }
 
diff --git a/Engine/Source/Runtime/Engine/Private/Animation/AnimInstanceProxy.cpp b/Engine/Source/Runtime/Engine/Private/Animation/AnimInstanceProxy.cpp
--- a/Engine/Source/Runtime/Engine/Private/Animation/AnimInstanceProxy.cpp
+++ b/Engine/Source/Runtime/Engine/Private/Animation/AnimInstanceProxy.cpp
@@ -844,12 +844,12 @@ void FAnimInstanceProxy::UpdateSlotNodeWeight(const FName& SlotNodeName, float I
 	}
 }
 
-bool FAnimInstanceProxy::GetSlotInertializationRequest(const FName& SlotName, float& OutDuration)
+bool FAnimInstanceProxy::GetSlotInertializationRequest(const FName& SlotName, UE::Anim::FSlotInertializationRequest& OutRequest)
 {
 	const FName GroupName = Skeleton ? Skeleton->GetSlotGroupName(SlotName) : NAME_None;
-	if (const float* RequestDuration = GetSlotGroupInertializationRequestMap().Find(GroupName))
+	if (const UE::Anim::FSlotInertializationRequest* FoundRequest = GetSlotGroupInertializationRequestMap().Find(GroupName))
 	{
-		OutDuration = *RequestDuration;
+		OutRequest = *FoundRequest;
 		return true;
 	}
 
diff --git a/Engine/Source/Runtime/Engine/Private/Animation/AnimMontage.cpp b/Engine/Source/Runtime/Engine/Private/Animation/AnimMontage.cpp
--- a/Engine/Source/Runtime/Engine/Private/Animation/AnimMontage.cpp
+++ b/Engine/Source/Runtime/Engine/Private/Animation/AnimMontage.cpp
@@ -1388,7 +1388,7 @@ void FAnimMontageInstance::Play(float InPlayRate, const FMontageBlendSettings& B
 		const float InertialBlendDuration = BlendInArgs.BlendTime;
 		// Request new inertialization for new montage's group name
 		// If there is an existing inertialization request, we overwrite that here.
-		AnimInstance->RequestMontageInertialization(Montage, InertialBlendDuration);
+		AnimInstance->RequestMontageInertialization(Montage, InertialBlendDuration, BlendInSettings.BlendProfile);
 
 		// When using inertialization, we need to instantly blend in.
 		BlendInArgs.BlendTime = 0.0f;
@@ -1456,7 +1456,7 @@ void FAnimMontageInstance::Stop(const FMontageBlendSettings& InBlendOutSettings,
 				if (bShouldInertialize)
 				{
 					// Send the inertial blend request to the anim instance
-					Inst->RequestMontageInertialization(Montage, InBlendOutSettings.Blend.BlendTime);
+					Inst->RequestMontageInertialization(Montage, InBlendOutSettings.Blend.BlendTime, InBlendOutSettings.BlendProfile);
 				}
 			}
 		}
diff --git a/Engine/Source/Runtime/Engine/Public/Animation/AnimInstanceProxy.h b/Engine/Source/Runtime/Engine/Public/Animation/AnimInstanceProxy.h
--- a/Engine/Source/Runtime/Engine/Public/Animation/AnimInstanceProxy.h
+++ b/Engine/Source/Runtime/Engine/Public/Animation/AnimInstanceProxy.h
@@ -449,7 +449,7 @@ public:
 	// Allow slot nodes to store off their weight during ticking
 	void UpdateSlotNodeWeight(const FName& SlotNodeName, float InLocalMontageWeight, float InNodeGlobalWeight);
 
-	bool GetSlotInertializationRequest(const FName& SlotName, float& OutDuration);
+	bool GetSlotInertializationRequest(const FName& SlotName, UE::Anim::FSlotInertializationRequest& OutRequest);
 
 	/** Register a named slot */
 	void RegisterSlotNodeWithAnimInstance(const FName& SlotNodeName);
@@ -793,7 +793,7 @@ protected:
 		Note that there might be multiple Active at the same time. This will only return the first active one it finds. **/
 	const FMontageEvaluationState* GetActiveMontageEvaluationState() const;
 
-	TMap<FName, float>& GetSlotGroupInertializationRequestMap() { return SlotGroupInertializationRequestMap; }
+	TMap<FName, UE::Anim::FSlotInertializationRequest>& GetSlotGroupInertializationRequestMap() { return SlotGroupInertializationRequestMap; }
 
 	/** Access montage array data */
 	TArray<FMontageEvaluationState>& GetMontageEvaluationData();
@@ -1093,8 +1093,8 @@ private:
 	TArray<FMontageEvaluationState> MontageEvaluationData;
 	TArray<FMontageEvaluationState>* MainMontageEvaluationData;
 
-	// Inertialization request for each slot. Mapped with SlotNameToTrackerIndex. Initialized to a value of 1.0f, which indicates no request for this slot.
-	TMap<FName, float> SlotGroupInertializationRequestMap;
+	// Inertialization request for each slot.
+	TMap<FName, UE::Anim::FSlotInertializationRequest> SlotGroupInertializationRequestMap;
 
 	/** Delegate fired on the game thread before update occurs */
 	TArray<FAnimNode_Base*> GameThreadPreUpdateNodes;
