Remove check() for input type to material Cross(), allow casting any input type to float3, match legacy behavior

diff --git a/Engine/Source/Runtime/Engine/Private/Materials/HLSLMaterialDerivativeAutogen.cpp b/Engine/Source/Runtime/Engine/Private/Materials/HLSLMaterialDerivativeAutogen.cpp
--- a/Engine/Source/Runtime/Engine/Private/Materials/HLSLMaterialDerivativeAutogen.cpp
+++ b/Engine/Source/Runtime/Engine/Private/Materials/HLSLMaterialDerivativeAutogen.cpp
@@ -448,8 +448,6 @@ FMaterialDerivativeAutogen::FOperationType2 FMaterialDerivativeAutogen::GetFunc2
 		}
 		break;
 	case EFunc2::Cross:
-		check(LhsNonLWCType == EDerivativeType::Float3);
-		check(RhsNonLWCType == EDerivativeType::Float3);
 		return EDerivativeType::Float3;
 	default:
 		check(0);
