Deleting old implementation of RouteActorInitialization now that the new implementation seems stable (we still have a way of toggling granular processing of actor initialization).

diff --git a/Engine/Source/Runtime/Engine/Classes/Engine/Level.h b/Engine/Source/Runtime/Engine/Classes/Engine/Level.h
--- a/Engine/Source/Runtime/Engine/Classes/Engine/Level.h
+++ b/Engine/Source/Runtime/Engine/Classes/Engine/Level.h
@@ -1027,11 +1027,6 @@ public:
 	 */
 	void ResetRouteActorInitializationState();
 
-	/**
-	 * Old implementation for routing actor initialization in full.
-	 */
-	void RouteActorInitializeOld();
-
 	/**
 	 * Routes pre and post initialize to actors and also sets volumes.
 	 *
diff --git a/Engine/Source/Runtime/Engine/Private/Level.cpp b/Engine/Source/Runtime/Engine/Private/Level.cpp
--- a/Engine/Source/Runtime/Engine/Private/Level.cpp
+++ b/Engine/Source/Runtime/Engine/Private/Level.cpp
@@ -98,14 +98,6 @@ static FAutoConsoleVariableRef CVarActorClusteringEnabled(
 	ECVF_Default
 );
 
-bool GUseLegacyRouteActorInitialization = false;
-static FAutoConsoleVariableRef CVarUseLegacyRouteActorInitialization(
-	TEXT("s.UseLegacyRouteActorInitialization"),
-	GUseLegacyRouteActorInitialization,
-	TEXT("Toggle for whether to use the old non-granular implementation of route actor initialization."),
-	ECVF_Default
-);
-
 #if WITH_EDITOR
 void FLevelActorFoldersHelper::SetUseActorFolders(ULevel* InLevel, bool bInEnabled)
 {
@@ -2658,65 +2650,9 @@ void ULevel::ResetRouteActorInitializationState()
 	 RouteActorInitializationIndex = 0;
 }
 
-void ULevel::RouteActorInitializeOld()
-{
-	// Send PreInitializeComponents and collect volumes.
-	for (int32 Index = 0; Index < Actors.Num(); ++Index)
-	{
-		AActor* const Actor = Actors[Index];
-		if (Actor && !Actor->IsActorInitialized())
-		{
-			Actor->PreInitializeComponents();
-		}
-	}
-
-	const bool bCallBeginPlay = OwningWorld->HasBegunPlay();
-	TArray<AActor*> ActorsToBeginPlay;
-
-	// Send InitializeComponents on components and PostInitializeComponents.
-	for (int32 Index = 0; Index < Actors.Num(); ++Index)
-	{
-		AActor* const Actor = Actors[Index];
-		if (Actor)
-		{
-			if (!Actor->IsActorInitialized())
-			{
-				// Call Initialize on Components.
-				Actor->InitializeComponents();
-
-				Actor->PostInitializeComponents(); // should set Actor->bActorInitialized = true
-				if (!Actor->IsActorInitialized() && IsValidChecked(Actor))
-				{
-					UE_LOG(LogActor, Fatal, TEXT("%s failed to route PostInitializeComponents.  Please call Super::PostInitializeComponents() in your <className>::PostInitializeComponents() function. "), *Actor->GetFullName());
-				}
-
-				if (bCallBeginPlay && !Actor->IsChildActor())
-				{
-					ActorsToBeginPlay.Add(Actor);
-				}
-			}
-		}
-	}
-
-	// Do this in a second pass to make sure they're all initialized before begin play starts
-	for (int32 ActorIndex = 0; ActorIndex < ActorsToBeginPlay.Num(); ActorIndex++)
-	{
-		AActor* Actor = ActorsToBeginPlay[ActorIndex];
-		SCOPE_CYCLE_COUNTER(STAT_ActorBeginPlay);
-		Actor->DispatchBeginPlay(/*bFromLevelStreaming*/ true);
-	}
-
-	RouteActorInitializationState = ERouteActorInitializationState::Finished;
-}
-
 void ULevel::RouteActorInitialize(int32 NumActorsToProcess)
 {
 	TRACE_OBJECT_EVENT(this, RouteActorInitialize);
-	if (GUseLegacyRouteActorInitialization)
-	{
-		RouteActorInitializeOld();
-		return;
-	}
 
 	const bool bFullProcessing = (NumActorsToProcess <= 0);
 	switch (RouteActorInitializationState)
@@ -2757,7 +2693,7 @@ void ULevel::RouteActorInitialize(int32 NumActorsToProcess)
 						Actor->PostInitializeComponents();
 						if (!Actor->IsActorInitialized() && IsValidChecked(Actor))
 						{
-							UE_LOG(LogActor, Fatal, TEXT("%s failed to route PostInitializeComponents. Please call Super::PostInitializeComponents() in your <className>::PostInitializeComponents() function. "), *Actor->GetFullName());
+							UE_LOG(LogActor, Fatal, TEXT("%s failed to route PostInitializeComponents. Please call Super::PostInitializeComponents() in your <className>::PostInitializeComponents() function."), *Actor->GetFullName());
 						}
 					}
 				}
