Updating types for ComputeSingleAxisLambda function

diff --git a/Engine/Source/Runtime/Engine/Private/PhysicsEngine/SimpleSuspension.cpp b/Engine/Source/Runtime/Engine/Private/PhysicsEngine/SimpleSuspension.cpp
--- a/Engine/Source/Runtime/Engine/Private/PhysicsEngine/SimpleSuspension.cpp
+++ b/Engine/Source/Runtime/Engine/Private/PhysicsEngine/SimpleSuspension.cpp
@@ -37,7 +37,7 @@ void FSimpleSuspension::Update(const FTransform& LocalToWorld, const FVector& Li
 	FSimpleSuspensionHelpers::ComputeSuspensionForces(LinearVelocity, AngularVelocityRad, SuspensionState, SuspensionParams.SpringParams, SuspensionState);
 }
 
-bool FSimpleSuspensionHelpers::ComputeSingleAxisLambda(const Chaos::FReal AxisDot, const Chaos::FReal SumAxis, const uint32 Count, TArray<Chaos::FReal, TFixedAllocator<2>>& Lambdas)
+bool FSimpleSuspensionHelpers::ComputeSingleAxisLambda(const FVector::FReal AxisDot, const FVector::FReal SumAxis, const uint32 Count, TArray<FVector::FReal, TFixedAllocator<2>>& Lambdas)
 {
 	using Chaos::FReal;
 
diff --git a/Engine/Source/Runtime/Engine/Public/Physics/SimpleSuspension.h b/Engine/Source/Runtime/Engine/Public/Physics/SimpleSuspension.h
--- a/Engine/Source/Runtime/Engine/Public/Physics/SimpleSuspension.h
+++ b/Engine/Source/Runtime/Engine/Public/Physics/SimpleSuspension.h
@@ -139,7 +139,7 @@ struct FSimpleSuspensionHelpers
 	ENGINE_API static bool ComputeSprungMasses(const TArray<FVector>& LocalSpringPositions, const FVector& LocalCenterOfMass, const float TotalMass, TArray<float>& OutSprungMasses);
 
 	/** Calculates the Lambdas in the single axis case*/
-	static bool ComputeSingleAxisLambda(const Chaos::FReal AxisDot, const Chaos::FReal SumAxis, const uint32 Count, TArray<Chaos::FReal, TFixedAllocator<2>>& Lambdas);
+	static bool ComputeSingleAxisLambda(const FVector::FReal AxisDot, const FVector::FReal SumAxis, const uint32 Count, TArray<FVector::FReal, TFixedAllocator<2>>& Lambdas);
 
 	/** Given a sprung mass and a spring stiffness, compute the natural
 		frequency of the spring */
