Fixed replication _FAST macros which fail to pass on the parameter required for acceleration.

diff --git a/Engine/Source/Runtime/Engine/Public/Net/UnrealNetwork.h b/Engine/Source/Runtime/Engine/Public/Net/UnrealNetwork.h
--- a/Engine/Source/Runtime/Engine/Public/Net/UnrealNetwork.h
+++ b/Engine/Source/Runtime/Engine/Public/Net/UnrealNetwork.h
@@ -351,7 +351,7 @@ DisableAllReplicatedPropertiesOfClass(StaticClass(), c::StaticClass(), SuperClas
 	static const bool bIsValid_##c_##v = ValidateReplicatedClassInheritance(StaticClass(), c::StaticClass(), TEXT("DISABLE_ALL_CLASS_REPLICATED_PROPERTIES")); \
 	const TCHAR* DoRepPropertyName_##c(TEXT(#c)); \
 	const NetworkingPrivate::FRepClassDescriptor ClassDescriptor_##c(DoRepPropertyName_##c, (int32)c::ENetFields_Private::NETFIELD_REP_START, (int32)c::ENetFields_Private::NETFIELD_REP_END); \
-	DisableAllReplicatedPropertiesOfClass(StaticClass(), c::StaticClass(), SuperClassBehavior, OutLifetimeProps); \
+	DisableAllReplicatedPropertiesOfClass(ClassDescriptor_##c, SuperClassBehavior, OutLifetimeProps); \
 }	
 
 ENGINE_API void DisableReplicatedLifetimeProperty(const UClass* ThisClass, const UClass* PropertyClass, FName PropertyName, TArray< FLifetimeProperty >& OutLifetimeProps);
@@ -374,7 +374,7 @@ ENGINE_API void DisableAllReplicatedPropertiesOfClass(const NetworkingPrivate::F
 	static const bool bIsValid_##c_##v = ValidateReplicatedClassInheritance(StaticClass(), c::StaticClass(), TEXT(#v)); \
 	const TCHAR* DoRepPropertyName_##c_##v(TEXT(#v)); \
 	const NetworkingPrivate::FRepPropertyDescriptor PropertyDescriptor_##c_##v(DoRepPropertyName_##c_##v, (int32)c::ENetFields_Private::v, 1); \
-	ResetReplicatedLifetimeProperty(StaticClass(), c::StaticClass(), GET_MEMBER_NAME_CHECKED(c,v), cond, OutLifetimeProps); \
+	ResetReplicatedLifetimeProperty(PropertyDescriptor_##c_##v, cond, OutLifetimeProps); \
 }
 
 #define RESET_REPLIFETIME_CONDITION_FAST_STATIC_ARRAY(c,v,cond) \
@@ -382,7 +382,7 @@ ENGINE_API void DisableAllReplicatedPropertiesOfClass(const NetworkingPrivate::F
 	static const bool bIsValid_##c_##v = ValidateReplicatedClassInheritance(StaticClass(), c::StaticClass(), TEXT(#v)); \
 	const TCHAR* DoRepPropertyName_##c_##v(TEXT(#v)); \
 	const NetworkingPrivate::FRepPropertyDescriptor PropertyDescriptor_##c_##v(DoRepPropertyName_##c_##v, (int32)c::ENetFields_Private::v##_STATIC_ARRAY, (int32)c::EArrayDims_Private::v); \
-	ResetReplicatedLifetimeProperty(StaticClass(), c::StaticClass(), GET_MEMBER_NAME_CHECKED(c,v), cond, OutLifetimeProps); \
+	ResetReplicatedLifetimeProperty(PropertyDescriptor_##c_##v, cond, OutLifetimeProps); \
 }
 
 #define RESET_REPLIFETIME_FAST(c,v) RESET_REPLIFETIME_CONDITION_FAST(c, v, COND_None)
@@ -395,7 +395,7 @@ ENGINE_API void DisableAllReplicatedPropertiesOfClass(const NetworkingPrivate::F
 	static const bool bIsValid_##c_##v = ValidateReplicatedClassInheritance(StaticClass(), c::StaticClass(), TEXT(#v)); \
 	const TCHAR* DoRepPropertyName_##c_##v(TEXT(#v)); \
 	const NetworkingPrivate::FRepPropertyDescriptor PropertyDescriptor_##c_##v(DoRepPropertyName_##c_##v, (int32)c::ENetFields_Private::v, 1); \
-	ResetReplicatedLifetimeProperty(StaticClass(), c::StaticClass(), GET_MEMBER_NAME_CHECKED(c,v), params, OutLifetimeProps); \
+	ResetReplicatedLifetimeProperty(PropertyDescriptor_##c_##v, params, OutLifetimeProps); \
 }
 
 ENGINE_API void ResetReplicatedLifetimeProperty(
