- No longer create a sound proxy if the app can't render (play) audio
#p4v-preflight-copy 18774992

diff --git a/Engine/Source/Runtime/Engine/Private/SoundWave.cpp b/Engine/Source/Runtime/Engine/Private/SoundWave.cpp
--- a/Engine/Source/Runtime/Engine/Private/SoundWave.cpp
+++ b/Engine/Source/Runtime/Engine/Private/SoundWave.cpp
@@ -1444,7 +1444,8 @@ void USoundWave::PostLoad()
 	// If our loading behavior is defined by a sound class, we need to update whether this sound wave actually needs to retain its audio data or not.
 	ESoundWaveLoadingBehavior ActualLoadingBehavior = GetLoadingBehavior();
 
-	if (!InternalProxy.IsValid() && ActualLoadingBehavior != ESoundWaveLoadingBehavior::ForceInline)
+	// no need for a proxy if we can't render audio!
+	if (!InternalProxy.IsValid() && ActualLoadingBehavior != ESoundWaveLoadingBehavior::ForceInline && FApp::CanEverRenderAudio())
 	{
 		InternalProxy = CreateSoundWaveProxy();
 	}
@@ -2202,8 +2203,6 @@ void USoundWave::BakeEnvelopeAnalysis()
 
 void USoundWave::PostEditChangeProperty(FPropertyChangedEvent& PropertyChangedEvent)
 {
-	// maxtodo: test this, stale data? looping? channel count?
-
 	Super::PostEditChangeProperty(PropertyChangedEvent);
 	check(SoundWaveDataPtr);
 
