Make rerun construction scripts editor only

diff --git a/Engine/Plugins/Compositing/Composure/Source/Composure/Classes/CompositingElement.h b/Engine/Plugins/Compositing/Composure/Source/Composure/Classes/CompositingElement.h
--- a/Engine/Plugins/Compositing/Composure/Source/Composure/Classes/CompositingElement.h
+++ b/Engine/Plugins/Compositing/Composure/Source/Composure/Classes/CompositingElement.h
@@ -386,9 +386,9 @@ public:
 #endif
 	//~ End UObject interface
 	
-	//~ Begin AAcotr interface
-	virtual void RerunConstructionScripts() override;
+	//~ Begin AActor interface
 #if WITH_EDITOR
+	virtual void RerunConstructionScripts() override;
 	virtual void OnConstruction(const FTransform& Transform) override;
 #endif
 	//~ End AAcotr interface
diff --git a/Engine/Plugins/Compositing/Composure/Source/Composure/Classes/ComposurePipelineBaseActor.h b/Engine/Plugins/Compositing/Composure/Source/Composure/Classes/ComposurePipelineBaseActor.h
--- a/Engine/Plugins/Compositing/Composure/Source/Composure/Classes/ComposurePipelineBaseActor.h
+++ b/Engine/Plugins/Compositing/Composure/Source/Composure/Classes/ComposurePipelineBaseActor.h
@@ -59,8 +59,10 @@ public:
 	virtual void SetAutoRunChildrenAndSelf(bool bAutoRunChildAndSelf) {}
 
 public:
+#if WITH_EDITOR
 	//~ AActor interface
 	virtual void RerunConstructionScripts() override;
+#endif
 
 protected:
 	/**
diff --git a/Engine/Plugins/Runtime/nDisplay/Source/DisplayCluster/Public/DisplayClusterRootActor.h b/Engine/Plugins/Runtime/nDisplay/Source/DisplayCluster/Public/DisplayClusterRootActor.h
--- a/Engine/Plugins/Runtime/nDisplay/Source/DisplayCluster/Public/DisplayClusterRootActor.h
+++ b/Engine/Plugins/Runtime/nDisplay/Source/DisplayCluster/Public/DisplayClusterRootActor.h
@@ -104,7 +104,9 @@ protected:
 	virtual void PostLoad() override;
 	virtual void PostActorCreated() override;
 	virtual void BeginDestroy() override;
+#if WITH_EDITOR
 	virtual void RerunConstructionScripts() override;
+#endif
 
 	// Initializes the actor on spawn and load
 	void InitializeRootActor();
diff --git a/Engine/Source/Runtime/Engine/Classes/Engine/Level.h b/Engine/Source/Runtime/Engine/Classes/Engine/Level.h
--- a/Engine/Source/Runtime/Engine/Classes/Engine/Level.h
+++ b/Engine/Source/Runtime/Engine/Classes/Engine/Level.h
@@ -683,7 +683,9 @@ public:
 	{
 		Init,
 		RegisterInitialComponents,
+#if WITH_EDITOR
 		RunConstructionScripts,
+#endif
 		Finalize
 	};
 
@@ -1311,9 +1313,9 @@ public:
 
 private:
 	bool IncrementalRegisterComponents(bool bPreRegisterComponents, int32 NumComponentsToUpdate, FRegisterComponentContext* Context);
+#if WITH_EDITOR
 	bool IncrementalRunConstructionScripts(bool bProcessAllActors);
 
-#if WITH_EDITOR
 private:
 	/**
 	 * Potentially defer the running of an actor's construction script on load
diff --git a/Engine/Source/Runtime/Engine/Classes/GameFramework/Actor.h b/Engine/Source/Runtime/Engine/Classes/GameFramework/Actor.h
--- a/Engine/Source/Runtime/Engine/Classes/GameFramework/Actor.h
+++ b/Engine/Source/Runtime/Engine/Classes/GameFramework/Actor.h
@@ -2954,8 +2954,10 @@ public:
 	/** Returns true if the actor's class has a non trivial user construction script. */
 	bool HasNonTrivialUserConstructionScript() const;
 
+#if WITH_EDITOR
 	/** Rerun construction scripts, destroying all autogenerated components; will attempt to preserve the root component location. */
 	virtual void RerunConstructionScripts();
+#endif
 
 	/** 
 	 * Debug helper to show the component hierarchy of this actor.
diff --git a/Engine/Source/Runtime/Engine/Private/ActorConstruction.cpp b/Engine/Source/Runtime/Engine/Private/ActorConstruction.cpp
--- a/Engine/Source/Runtime/Engine/Private/ActorConstruction.cpp
+++ b/Engine/Source/Runtime/Engine/Private/ActorConstruction.cpp
@@ -231,6 +231,7 @@ bool AActor::HasNonTrivialUserConstructionScript() const
 	return false;
 }
 
+#if WITH_EDITOR
 void AActor::RerunConstructionScripts()
 {
 	checkf(!HasAnyFlags(RF_ClassDefaultObject), TEXT("RerunConstructionScripts should never be called on a CDO as it can mutate the transient data on the CDO which then propagates to instances!"));
@@ -242,7 +243,6 @@ void AActor::RerunConstructionScripts()
 	if (bAllowReconstruction)
 	{
 		TRACE_CPUPROFILER_EVENT_SCOPE(AActor::RerunConstructionScripts);
-#if WITH_EDITOR
 		if(GIsEditor)
 		{
 			// Don't allow reconstruction if we're still in the middle of construction.
@@ -266,7 +266,7 @@ void AActor::RerunConstructionScripts()
 				}
 			}
 		}
-#endif
+
 		// Child Actors can be customized in many ways by their parents construction scripts and rerunning directly on them would wipe
 		// that out. So instead we redirect up the hierarchy
 		if (IsChildActor())
@@ -328,7 +328,6 @@ void AActor::RerunConstructionScripts()
 			}
 		}
 
-#if WITH_EDITOR
 		if (!CurrentTransactionAnnotation.IsValid())
 		{
 			CurrentTransactionAnnotation = FActorTransactionAnnotation::Create(this, false);
@@ -367,9 +366,6 @@ void AActor::RerunConstructionScripts()
 
 			bUseRootComponentProperties = false;
 		}
-#else
-		InstanceDataCache = new FComponentInstanceDataCache(this);
-#endif
 
 		if (bUseRootComponentProperties)
 		{
@@ -430,8 +426,6 @@ void AActor::RerunConstructionScripts()
 			}
 		}
 
-#if WITH_EDITOR
-
 		// Return the component which was added by the construction script.
 		// It may be the same as the argument, or a parent component if the argument was a native subobject.
 		auto GetComponentAddedByConstructionScript = [](UActorComponent* Component) -> UActorComponent*
@@ -516,7 +510,6 @@ void AActor::RerunConstructionScripts()
 				ComponentData.UCSComponentIndex = Component->GetUCSSerializationIndex();
 			}
 		}
-#endif
 
 		// Destroy existing components
 		DestroyConstructedComponents();
@@ -569,7 +562,6 @@ void AActor::RerunConstructionScripts()
 		// Restore the undo buffer
 		GUndo = CurrentTransaction;
 
-#if WITH_EDITOR
 		// Create the mapping of old->new components and notify the editor of the replacements
 		TInlineComponentArray<UActorComponent*> NewComponents;
 		GetComponents(NewComponents);
@@ -688,12 +680,9 @@ void AActor::RerunConstructionScripts()
 		{
 			CurrentTransactionAnnotation = nullptr;
 		}
-#else
-		delete InstanceDataCache;
-#endif
-
 	}
 }
+#endif
 
 namespace
 {
diff --git a/Engine/Source/Runtime/Engine/Private/Level.cpp b/Engine/Source/Runtime/Engine/Private/Level.cpp
--- a/Engine/Source/Runtime/Engine/Private/Level.cpp
+++ b/Engine/Source/Runtime/Engine/Private/Level.cpp
@@ -1351,6 +1351,10 @@ static void SortActorsHierarchy(TArray<AActor*>& Actors, ULevel* Level)
 DECLARE_CYCLE_STAT(TEXT("Deferred Init Bodies"), STAT_DeferredUpdateBodies, STATGROUP_Physics);
 void ULevel::IncrementalUpdateComponents(int32 NumComponentsToUpdate, bool bRerunConstructionScripts, FRegisterComponentContext* Context)
 {
+#if !WITH_EDITOR
+	ensure(!bRerunConstructionScripts);
+#endif
+
 	TRACE_CPUPROFILER_EVENT_SCOPE(ULevel::IncrementalUpdateComponents);
 
 	// A value of 0 means that we want to update all components.
@@ -1380,17 +1384,23 @@ void ULevel::IncrementalUpdateComponents(int32 NumComponentsToUpdate, bool bReru
 		case EIncrementalComponentState::RegisterInitialComponents:
 			if (IncrementalRegisterComponents(true, NumComponentsToUpdate, Context))
 			{
-				bool ShouldRunConstructionScripts = !bHasRerunConstructionScripts && bRerunConstructionScripts && !IsTemplate() && !GIsUCCMakeStandaloneHeaderGenerator;
-				IncrementalComponentState = ShouldRunConstructionScripts ? EIncrementalComponentState::RunConstructionScripts : EIncrementalComponentState::Finalize;
+#if WITH_EDITOR
+				const bool bShouldRunConstructionScripts = !bHasRerunConstructionScripts && bRerunConstructionScripts && !IsTemplate() && !GIsUCCMakeStandaloneHeaderGenerator;
+				IncrementalComponentState = bShouldRunConstructionScripts ? EIncrementalComponentState::RunConstructionScripts : EIncrementalComponentState::Finalize;
+#else
+				IncrementalComponentState = EIncrementalComponentState::Finalize;
+#endif
 			}
 			break;
 
+#if WITH_EDITOR
 		case EIncrementalComponentState::RunConstructionScripts:
 			if (IncrementalRunConstructionScripts(bFullyUpdateComponents))
 			{
 				IncrementalComponentState = EIncrementalComponentState::Finalize;
 			}
 			break;
+#endif
 
 		case EIncrementalComponentState::Finalize:
 			IncrementalComponentState = EIncrementalComponentState::Init;
@@ -1471,6 +1481,7 @@ bool ULevel::IncrementalRegisterComponents(bool bPreRegisterComponents, int32 Nu
 	return false;
 }
 
+#if WITH_EDITOR
 bool ULevel::IncrementalRunConstructionScripts(bool bProcessAllActors)
 {
 	// Find next valid actor to process components registration
@@ -1530,6 +1541,7 @@ bool ULevel::IncrementalRunConstructionScripts(bool bProcessAllActors)
 	}
 	return false;
 }
+#endif
 
 bool ULevel::IncrementalUnregisterComponents(int32 NumComponentsToUnregister)
 {
diff --git a/Engine/Source/Runtime/Engine/Private/World.cpp b/Engine/Source/Runtime/Engine/Private/World.cpp
--- a/Engine/Source/Runtime/Engine/Private/World.cpp
+++ b/Engine/Source/Runtime/Engine/Private/World.cpp
@@ -1998,7 +1998,8 @@ void UWorld::InitializeNewWorld(const InitializationValues IVS, bool bInSkipInit
 		InitWorld(IVS);
 
 		// Update components.
-		UpdateWorldComponents(true, false);
+		const bool bRerunConstructionScripts = !FPlatformProperties::RequiresCookedData();
+		UpdateWorldComponents(bRerunConstructionScripts, false);
 	}
 }
 
@@ -2190,6 +2191,10 @@ void UWorld::UpdateWorldComponents(bool bRerunConstructionScripts, bool bCurrent
 {
 	TRACE_CPUPROFILER_EVENT_SCOPE(UWorld::UpdateWorldComponents);
 
+#if !WITH_EDITOR
+	ensure(!bRerunConstructionScripts);
+#endif
+
 	if ( !IsRunningDedicatedServer() )
 	{
 		if(!LineBatcher)
diff --git a/Engine/Source/Runtime/Foliage/Public/InstancedFoliageActor.h b/Engine/Source/Runtime/Foliage/Public/InstancedFoliageActor.h
--- a/Engine/Source/Runtime/Foliage/Public/InstancedFoliageActor.h
+++ b/Engine/Source/Runtime/Foliage/Public/InstancedFoliageActor.h
@@ -53,7 +53,9 @@ public:
 
 	//~ Begin AActor Interface.
 	// we don't want to have our components automatically destroyed by the Blueprint code
+#if WITH_EDITOR
 	virtual void RerunConstructionScripts() override {}
+#endif
 	virtual bool IsLevelBoundsRelevant() const override { return false; }
 
 public:
diff --git a/Engine/Source/Runtime/NavigationSystem/Private/NavigationData.cpp b/Engine/Source/Runtime/NavigationSystem/Private/NavigationData.cpp
--- a/Engine/Source/Runtime/NavigationSystem/Private/NavigationData.cpp
+++ b/Engine/Source/Runtime/NavigationSystem/Private/NavigationData.cpp
@@ -366,12 +366,14 @@ void ANavigationData::TickActor(float DeltaTime, enum ELevelTick TickType, FActo
 	}
 }
 
+#if WITH_EDITOR
 void ANavigationData::RerunConstructionScripts()
 {
 	Super::RerunConstructionScripts();
 
 	InstantiateAndRegisterRenderingComponent();
 }
+#endif
 
 void ANavigationData::OnRegistered() 
 { 
