Additional protections around interacting w/ the AudioStreamingCache if  FApp::CanEverRenderAudio() is false

diff --git a/Engine/Source/Runtime/Engine/Private/SoundWave.cpp b/Engine/Source/Runtime/Engine/Private/SoundWave.cpp
--- a/Engine/Source/Runtime/Engine/Private/SoundWave.cpp
+++ b/Engine/Source/Runtime/Engine/Private/SoundWave.cpp
@@ -885,7 +885,7 @@ void USoundWave::Serialize( FArchive& Ar )
 		const bool bShouldLoadChunks = bCooked || bBuiltStreamedAudio;
 
 		// If stream caching is enabled, here we determine if we should retain or prime this wave on load.
-		if (bShouldStreamSound && bShouldLoadChunks && FPlatformCompressionUtilities::IsCurrentPlatformUsingStreamCaching())
+		if (bShouldStreamSound && bShouldLoadChunks && FPlatformCompressionUtilities::IsCurrentPlatformUsingStreamCaching() && FApp::CanEverRenderAudio())
 		{
 			ESoundWaveLoadingBehavior CurrentLoadingBehavior = GetLoadingBehavior(false);
 
@@ -1464,6 +1464,7 @@ void USoundWave::PostLoad()
 
 			const bool bHasMultipleChunks = GetNumChunks() > 1;
 			bool bShouldPrime = (ActualLoadingBehavior == ESoundWaveLoadingBehavior::PrimeOnLoad);
+			bShouldPrime |= InternalProxy.IsValid();
 			bShouldPrime |= (GIsEditor && (ActualLoadingBehavior == ESoundWaveLoadingBehavior::RetainOnLoad)); // treat this scenario like PrimeOnLoad
 			
 			if (bShouldPrime && bHasMultipleChunks)
@@ -2920,7 +2921,7 @@ bool USoundWave::ShouldUseStreamCaching() const
 
 TArrayView<const uint8> USoundWave::GetZerothChunk(bool bForImmediatePlayback)
 {
-	if(IsTemplate() || IsRunningDedicatedServer())
+	if(IsTemplate() || IsRunningDedicatedServer() || !FApp::CanEverRenderAudio())
 	{
 		return TArrayView<const uint8>();
 	}
@@ -3057,7 +3058,7 @@ void USoundWave::RemovePlayingSource(const FSoundWaveClientPtr& Source)
 
 void USoundWave::UpdatePlatformData()
 {
-	if (IsStreaming(nullptr))
+	if (IsStreaming(nullptr) && FApp::CanEverRenderAudio())
 	{
 		if (!InternalProxy.IsValid())
 		{
@@ -3415,14 +3416,14 @@ void USoundWave::GetHandleForChunkOfAudio(TFunction<void(FAudioChunkHandle&&)> O
 {
 	ENamedThreads::Type ThreadToDispatchCallbackOn = (DispatchToGameThreadOnChunkRequestCVar != 0) ? ENamedThreads::GameThread : ENamedThreads::AnyThread;
 
-	if (!InternalProxy.IsValid())
+	if (!InternalProxy.IsValid() && FApp::CanEverRenderAudio())
 	{
 		InternalProxy = CreateSoundWaveProxy();
 	}
 
 	// if we are requesting a chunk that is out of bounds,
-	// early exit.
-	if (ChunkIndex >= static_cast<int32>(GetNumChunks()))
+	// or we shouldn't be interacting w/ the cache, early exit.
+	if (FApp::CanEverRenderAudio() || (ChunkIndex >= static_cast<int32>(GetNumChunks())))
 	{
 		FAudioChunkHandle EmptyChunkHandle;
 		OnLoadCompleted(MoveTemp(EmptyChunkHandle));
@@ -3485,7 +3486,7 @@ void USoundWave::RetainCompressedAudio(bool bForceSync /*= false*/)
 
 	// Since the zeroth chunk is always inlined and stored in memory,
 	// early exit if we only have one chunk.
-	if (GIsEditor || IsTemplate() || IsRunningDedicatedServer() || !IsStreaming() || DisableRetainingCVar || GetNumChunks() <= 1)
+	if (GIsEditor || IsTemplate() || IsRunningDedicatedServer() || !IsStreaming() || !FApp::CanEverRenderAudio() || DisableRetainingCVar || GetNumChunks() <= 1)
 	{
 		return;
 	}
