New HLSL translator supports vertex properties (WPO, custom UVs)

diff --git a/Engine/Shaders/Private/GeometryCollectionVertexFactory.ush b/Engine/Shaders/Private/GeometryCollectionVertexFactory.ush
--- a/Engine/Shaders/Private/GeometryCollectionVertexFactory.ush
+++ b/Engine/Shaders/Private/GeometryCollectionVertexFactory.ush
@@ -268,8 +268,9 @@ FMaterialVertexParameters GetMaterialVertexParameters(FVertexFactoryInput Input,
 	Result.Particle.WorldToParticle = GetPrimitiveData(Intermediates).WorldToLocal;
 #endif
 
+#if ENABLE_NEW_HLSL_GENERATOR
 	Result.MaterialVertexAttributes = EvaluateVertexMaterialAttributes(Result);
-
+#endif
 	return Result;
 }
 
diff --git a/Engine/Shaders/Private/LocalVertexFactory.ush b/Engine/Shaders/Private/LocalVertexFactory.ush
--- a/Engine/Shaders/Private/LocalVertexFactory.ush
+++ b/Engine/Shaders/Private/LocalVertexFactory.ush
@@ -660,8 +660,9 @@ FMaterialVertexParameters GetMaterialVertexParameters(FVertexFactoryInput Input,
 	Result.Particle.WorldToParticle = GetPrimitiveData(Intermediates).WorldToLocal;
 #endif
 
+#if ENABLE_NEW_HLSL_GENERATOR
 	Result.MaterialVertexAttributes = EvaluateVertexMaterialAttributes(Result);
-
+#endif
 	return Result;
 }
 
diff --git a/Engine/Shaders/Private/MaterialTemplate.ush b/Engine/Shaders/Private/MaterialTemplate.ush
--- a/Engine/Shaders/Private/MaterialTemplate.ush
+++ b/Engine/Shaders/Private/MaterialTemplate.ush
@@ -25,6 +25,10 @@
 // Merge conflicts on this line should be resolved by generating a new GUID
 // GUID = 5F290441-5611-4232-B26D-613F5E421262
 
+#ifndef ENABLE_NEW_HLSL_GENERATOR
+#define ENABLE_NEW_HLSL_GENERATOR 0
+#endif
+
 #if USES_SPEEDTREE
 	#include "/Engine/Private/SpeedTreeCommon.ush"
 #endif
@@ -2828,15 +2832,7 @@ float3 CalculateAnisotropyTangent(in out FMaterialPixelParameters Parameters, FP
 	return Tangent;
 }
 
-FMaterialAttributes EvaluateVertexMaterialAttributes(FMaterialVertexParameters Parameters)
-{
-	FMaterialAttributes DefaultMaterialAttributes = (FMaterialAttributes)0;
-%s
-	// Return the default, this may be skipped if the above code returns different attributes
-	return DefaultMaterialAttributes;
-}
-
-// EvaluatePixelMaterialAttributes
+// EvaluateMaterialAttributes
 %s
 
 // Be very, very careful about changing CalcPixelMaterialInputs() or CalcPixelMaterialInputsAnalyticDerivatives(). They apply the same basic calculation,
diff --git a/Engine/Source/Runtime/Engine/Private/HLSLTree/HLSLTree.cpp b/Engine/Source/Runtime/Engine/Private/HLSLTree/HLSLTree.cpp
--- a/Engine/Source/Runtime/Engine/Private/HLSLTree/HLSLTree.cpp
+++ b/Engine/Source/Runtime/Engine/Private/HLSLTree/HLSLTree.cpp
@@ -744,7 +744,7 @@ EExpressionEvaluation FPreparedType::GetFieldEvaluation(const FEmitScope& Scope,
 	EExpressionEvaluation Result = EExpressionEvaluation::None;
 	for (int32 Index = 0; Index < NumComponents; ++Index)
 	{
-		const FPreparedComponent Component = GetComponent(Index);
+		const FPreparedComponent Component = GetComponent(ComponentIndex + Index);
 		Result = CombineEvaluations(Result, Component.GetEvaluation(Scope));
 	}
 	return Result;
diff --git a/Engine/Source/Runtime/Engine/Private/HLSLTree/HLSLTreeCommon.cpp b/Engine/Source/Runtime/Engine/Private/HLSLTree/HLSLTreeCommon.cpp
--- a/Engine/Source/Runtime/Engine/Private/HLSLTree/HLSLTreeCommon.cpp
+++ b/Engine/Source/Runtime/Engine/Private/HLSLTree/HLSLTreeCommon.cpp
@@ -235,7 +235,7 @@ void FExpressionExternalInput::EmitValueShader(FEmitContext& Context, FEmitScope
 	if (IsTexCoord(InputType))
 	{
 		const int32 TexCoordIndex = TypeIndex - (int32)EExternalInput::TexCoord0;
-		Context.NumTexCoords = FMath::Max(Context.NumTexCoords, TexCoordIndex + 1);
+		Context.TexCoordMask[Context.ShaderFrequency] |= (1u << TexCoordIndex);
 		OutResult.Code = Context.EmitInlineExpression(Scope, Shader::EValueType::Float2, TEXT("Parameters.TexCoords[%].xy"), TexCoordIndex);
 	}
 	else if (IsTexCoord_Ddx(InputType))
@@ -575,12 +575,15 @@ void FExpressionSetStructField::EmitValuePreshader(FEmitContext& Context, FEmitS
 		OutPreshader.WriteOpcode(Shader::EPreshaderOpcode::ConstantZero).Write(Shader::FType(StructType));
 	}
 
-	FieldExpression->GetValuePreshader(Context, Scope, RequestedFieldType, OutPreshader);
+	if (FieldEvaluation != EExpressionEvaluation::None)
+	{
+		FieldExpression->GetValuePreshader(Context, Scope, RequestedFieldType, OutPreshader);
 
-	check(Context.PreshaderStackPosition > 0);
-	Context.PreshaderStackPosition--;
+		check(Context.PreshaderStackPosition > 0);
+		Context.PreshaderStackPosition--;
 
-	OutPreshader.WriteOpcode(Shader::EPreshaderOpcode::SetField).Write(Field->ComponentIndex).Write(Field->GetNumComponents());
+		OutPreshader.WriteOpcode(Shader::EPreshaderOpcode::SetField).Write(Field->ComponentIndex).Write(Field->GetNumComponents());
+	}
 }
 
 bool FExpressionSelect::PrepareValue(FEmitContext& Context, FEmitScope& Scope, const FRequestedType& RequestedType, FPrepareValueResult& OutResult) const
diff --git a/Engine/Source/Runtime/Engine/Private/Materials/HLSLMaterialTranslator.cpp b/Engine/Source/Runtime/Engine/Private/Materials/HLSLMaterialTranslator.cpp
--- a/Engine/Source/Runtime/Engine/Private/Materials/HLSLMaterialTranslator.cpp
+++ b/Engine/Source/Runtime/Engine/Private/Materials/HLSLMaterialTranslator.cpp
@@ -2215,7 +2215,6 @@ FString FHLSLMaterialTranslator::GetMaterialShaderCode()
 
 		EvaluateVertexCode += TranslatedAttributesCodeChunks[SF_Vertex];
 
-		LazyPrintf.PushParam(*EvaluateVertexCode);
 		LazyPrintf.PushParam(*TranslatedAttributesCodeChunks[SF_Pixel]);
 
 		FString EvaluateMaterialAttributesCode = TEXT("    FMaterialAttributes MaterialAttributes = EvaluatePixelMaterialAttributes(Parameters);" LINE_TERMINATOR);
@@ -2256,7 +2255,6 @@ FString FHLSLMaterialTranslator::GetMaterialShaderCode()
 	{
 		// skip material attributes code
 		LazyPrintf.PushParam(TEXT(""));
-		LazyPrintf.PushParam(TEXT(""));
 
 		for (int32 Iter = 0; Iter < CompiledPDV_MAX; Iter++)
 		{
diff --git a/Engine/Source/Runtime/Engine/Private/Materials/MaterialHLSLEmitter.cpp b/Engine/Source/Runtime/Engine/Private/Materials/MaterialHLSLEmitter.cpp
--- a/Engine/Source/Runtime/Engine/Private/Materials/MaterialHLSLEmitter.cpp
+++ b/Engine/Source/Runtime/Engine/Private/Materials/MaterialHLSLEmitter.cpp
@@ -47,6 +47,7 @@ static FString GenerateMaterialTemplateHLSL(EShaderPlatform ShaderPlatform,
 	const FMaterial& Material,
 	const UE::HLSLTree::FEmitContext& EmitContext,
 	const TCHAR* DeclarationsCode,
+	const TCHAR* VertexShaderCode,
 	const TCHAR* PixelShaderCodePhase0,
 	const TCHAR* PixelShaderCodePhase1)
 {
@@ -74,10 +75,14 @@ static FString GenerateMaterialTemplateHLSL(EShaderPlatform ShaderPlatform,
 
 	FLazyPrintf LazyPrintf(*MaterialTemplateSource);
 
-	const uint32 NumUserVertexTexCoords = EmitContext.NumTexCoords;
-	const uint32 NumUserTexCoords = EmitContext.NumTexCoords;
+	const uint32 VertexTexCoordMask = EmitContext.TexCoordMask[SF_Vertex];
+	const uint32 PixelTexCoordMask = EmitContext.TexCoordMask[SF_Pixel];
+	const uint32 NumTexCoords = PixelTexCoordMask ? (FMath::FloorLog2(PixelTexCoordMask) + 1u) : 0u;
+
+	const uint32 NumUserVertexTexCoords = NumTexCoords;
+	const uint32 NumUserTexCoords = NumTexCoords;
 	const uint32 NumCustomVectors = 0u;
-	const uint32 NumTexCoordVectors = EmitContext.NumTexCoords;
+	const uint32 NumTexCoordVectors = NumTexCoords;
 
 	LazyPrintf.PushParam(*FString::Printf(TEXT("%u"), NumUserVertexTexCoords));
 	LazyPrintf.PushParam(*FString::Printf(TEXT("%u"), NumUserTexCoords));
@@ -287,21 +292,35 @@ static FString GenerateMaterialTemplateHLSL(EShaderPlatform ShaderPlatform,
 
 	//if (bEnableExecutionFlow)
 	{
-		FString EvaluateVertexCode;
+		FString EvaluateMaterialDeclaration;
 
-		// Set default texcoords in the VS
-		for (uint32 TexCoordIndex = 0; TexCoordIndex < NumUserVertexTexCoords; ++TexCoordIndex)
-		{
-			const FString AttributeName = FMaterialAttributeDefinitionMap::GetAttributeName((EMaterialProperty)(MP_CustomizedUVs0 + TexCoordIndex));
+		EvaluateMaterialDeclaration += TEXT("FMaterialAttributes EvaluateVertexMaterialAttributesInternal(FMaterialVertexParameters Parameters)" LINE_TERMINATOR);
+		EvaluateMaterialDeclaration += TEXT("{" LINE_TERMINATOR);
+		EvaluateMaterialDeclaration += VertexShaderCode;
+		EvaluateMaterialDeclaration += TEXT("}" LINE_TERMINATOR);
 
-			EvaluateVertexCode += FString::Printf(TEXT("\tDefaultMaterialAttributes.%s = Parameters.TexCoords[%d];") LINE_TERMINATOR, *AttributeName, TexCoordIndex);
+		EvaluateMaterialDeclaration += TEXT("FMaterialAttributes EvaluateVertexMaterialAttributes(FMaterialVertexParameters Parameters)" LINE_TERMINATOR);
+		EvaluateMaterialDeclaration += TEXT("{" LINE_TERMINATOR);
+		EvaluateMaterialDeclaration += TEXT("    FMaterialAttributes Result = EvaluateVertexMaterialAttributesInternal(Parameters);" LINE_TERMINATOR);
+		for (uint32 TexCoordIndex = 0; TexCoordIndex < NumTexCoords; ++TexCoordIndex)
+		{
+			if ((PixelTexCoordMask & (1u << TexCoordIndex)) && !(VertexTexCoordMask & (1u << TexCoordIndex)))
+			{
+				const FString AttributeName = FMaterialAttributeDefinitionMap::GetAttributeName((EMaterialProperty)(MP_CustomizedUVs0 + TexCoordIndex));
+				EvaluateMaterialDeclaration += FString::Printf(TEXT("    Result.%s = Parameters.TexCoords[%d];") LINE_TERMINATOR, *AttributeName, TexCoordIndex);
+			}
 		}
+		EvaluateMaterialDeclaration += TEXT("   return Result;" LINE_TERMINATOR);
+		EvaluateMaterialDeclaration += TEXT("}" LINE_TERMINATOR);
 
-		//EvaluateVertexCode += TranslatedAttributesCodeChunks[SF_Vertex];
+		// Set default texcoords in the VS
+		/*for (uint32 TexCoordIndex = 0; TexCoordIndex < NumUserVertexTexCoords; ++TexCoordIndex)
+		{
+			const FString AttributeName = FMaterialAttributeDefinitionMap::GetAttributeName((EMaterialProperty)(MP_CustomizedUVs0 + TexCoordIndex));
 
-		LazyPrintf.PushParam(*EvaluateVertexCode);
+			EvaluateMaterialDeclaration += FString::Printf(TEXT("\tDefaultMaterialAttributes.%s = Parameters.TexCoords[%d];") LINE_TERMINATOR, *AttributeName, TexCoordIndex);
+		}*/
 
-		FString EvaluateMaterialDeclaration;
 		FString EvaluateMaterialAttributesPhase0 = TEXT("    FMaterialAttributes MaterialAttributesPhase0 = EvaluatePixelMaterialAttributesPhase0(Parameters);" LINE_TERMINATOR);
 		FString EvaluateMaterialAttributesPhase1;
 
@@ -377,6 +396,8 @@ static void GetMaterialEnvironment(EShaderPlatform InPlatform,
 {
 	bool bMaterialRequestsDualSourceBlending = false;
 
+	OutEnvironment.SetDefine(TEXT("ENABLE_NEW_HLSL_GENERATOR"), 1);
+
 	if (false)//bNeedsParticlePosition || Material->ShouldGenerateSphericalParticleNormals() || bUsesSphericalParticleOpacity)
 	{
 		OutEnvironment.SetDefine(TEXT("NEEDS_PARTICLE_POSITION"), 1);
@@ -687,6 +708,56 @@ public:
 	FStringBuilderMemstack(FMemStackBase& Allocator, int32 InSize) : TStringBuilderBase<TCHAR>((TCHAR*)Allocator.Alloc(sizeof(TCHAR)* InSize, alignof(TCHAR)), InSize) {}
 };
 
+static bool IsAttributeUsed(FMaterialHLSLGenerator& Generator,
+	UE::HLSLTree::FEmitContext& Context,
+	UE::HLSLTree::FEmitScope& Scope,
+	const UE::HLSLTree::FPreparedType& ResultType,
+	UE::HLSLTree::FExpression* Expression,
+	EMaterialProperty Property)
+{
+	using namespace UE::HLSLTree;
+	using namespace UE::Shader;
+
+	const FStructField* PropertyField = Generator.GetMaterialAttributesType()->FindFieldByName(*FMaterialAttributeDefinitionMap::GetAttributeName(Property));
+	if (!PropertyField)
+	{
+		return false;
+	}
+
+	const int32 NumComponents = PropertyField->GetNumComponents();
+	const EExpressionEvaluation Evaluation = ResultType.GetFieldEvaluation(Scope, PropertyField->ComponentIndex, NumComponents);
+	if (Evaluation == EExpressionEvaluation::None)
+	{
+		return false;
+	}
+	else if (Evaluation == EExpressionEvaluation::Constant)
+	{
+		FRequestedType RequestedType(Generator.GetMaterialAttributesType(), false);
+		RequestedType.SetFieldRequested(PropertyField);
+
+		const FValue& DefaultValue = Generator.GetMaterialAttributesDefaultValue();
+		const FValue ConstantValue = Expression->GetValueConstant(Context, Scope, RequestedType);
+		if (DefaultValue.GetType() != ConstantValue.GetType())
+		{
+			ensure(false); // expected types to match
+			return true;
+		}
+
+		for (int32 Index = 0; Index < NumComponents; ++Index)
+		{
+			const int32 ComponentIndex = PropertyField->ComponentIndex + Index;
+			if (DefaultValue.Component[ComponentIndex].Packed != ConstantValue.Component[ComponentIndex].Packed)
+			{
+				// Non-default value, flag as used
+				return true;
+			}
+		}
+		return false;
+	}
+
+	return true;
+}
+
 bool MaterialEmitHLSL(const FMaterialCompileTargetParameters& InCompilerTarget,
 	const FStaticParameterSet& InStaticParameters,
 	FMaterial& InOutMaterial,
@@ -696,6 +767,7 @@ bool MaterialEmitHLSL(const FMaterialCompileTargetParameters& InCompilerTarget,
 	using namespace UE::HLSLTree;
 	using namespace UE::Shader;
 
+	const TArray<FGuid>& OrderedVisibleAttributes = FMaterialAttributeDefinitionMap::GetOrderedVisibleAttributeList();
 	FMemory::Memzero(SharedPixelProperties);
 	SharedPixelProperties[MP_Normal] = true;
 	SharedPixelProperties[MP_Tangent] = true;
@@ -724,59 +796,114 @@ bool MaterialEmitHLSL(const FMaterialCompileTargetParameters& InCompilerTarget,
 		return false;
 	}
 
+	bool bUsesWorldPositionOffset = false;
+	bool bUsesPixelDepthOffset = false;
+
 	FEmitContext EmitContext(Allocator, Generator.GetErrors(), TypeRegistry);
 	EmitContext.Material = &InOutMaterial;
 	EmitContext.StaticParameters = &InStaticParameters;
 	EmitContext.MaterialCompilationOutput = &OutCompilationOutput;
 
-	const FStructField* NormalField = Generator.GetMaterialAttributesType()->FindFieldByName(*FMaterialAttributeDefinitionMap::GetAttributeName(MP_Normal));
-	check(NormalField);
-
-	FEmitScope* EmitResultScope = EmitContext.PrepareScope(&Generator.GetResultStatement()->GetParentScope());
-
-	// Prepare all fields *except* normal
-	FRequestedType RequestedMaterialAttributesType(Generator.GetMaterialAttributesType());
-	RequestedMaterialAttributesType.ClearFieldRequested(NormalField);
-	if (EmitContext.PrepareExpression(Generator.GetResultExpression(), *EmitResultScope, RequestedMaterialAttributesType).IsVoid())
+	// Prepare pixel shader code
+	FStringBuilderMemstack PixelCodePhase0(Allocator, 1024 * 1024);
+	FStringBuilderMemstack PixelCodePhase1(Allocator, 1024 * 1024);
 	{
-		return false;
-	}
+		const FStructField* NormalField = Generator.GetMaterialAttributesType()->FindFieldByName(*FMaterialAttributeDefinitionMap::GetAttributeName(MP_Normal));
+		check(NormalField);
 
+		EmitContext.ShaderFrequency = SF_Pixel;
+		FEmitScope* EmitResultScope = EmitContext.PrepareScope(&Generator.GetResultStatement()->GetParentScope());
 
-	const bool bReadMaterialNormal = EmitContext.bReadMaterialNormal;
-	FStringBuilderMemstack CodePhase0(Allocator, 1024 * 1024);
-	FStringBuilderMemstack CodePhase1(Allocator, 1024 * 1024);
-	if (!bReadMaterialNormal)
-	{
-		// No access to material normal, can execute everything in phase0
-		RequestedMaterialAttributesType.SetFieldRequested(NormalField);
-		if (EmitContext.PrepareExpression(Generator.GetResultExpression(), *EmitResultScope, RequestedMaterialAttributesType).IsVoid())
+		// Prepare all fields *except* normal
+		FRequestedType RequestedPixelAttributesType(Generator.GetMaterialAttributesType(), false);
+		for (const FGuid& AttributeID : OrderedVisibleAttributes)
+		{
+			if (FMaterialAttributeDefinitionMap::GetShaderFrequency(AttributeID) == SF_Pixel)
+			{
+				const FString& FieldName = FMaterialAttributeDefinitionMap::GetAttributeName(AttributeID);
+				const FStructField* Field = Generator.GetMaterialAttributesType()->FindFieldByName(*FieldName);
+				if (Field && Field != NormalField)
+				{
+					RequestedPixelAttributesType.SetFieldRequested(Field);
+				}
+			}
+		}
+
+		const FPreparedType PixelResultType0 = EmitContext.PrepareExpression(Generator.GetResultExpression(), *EmitResultScope, RequestedPixelAttributesType);
+		if (PixelResultType0.IsVoid())
 		{
 			return false;
 		}
 
-		HLSLTree->EmitShader(EmitContext, CodePhase0);
+		bUsesPixelDepthOffset = IsAttributeUsed(Generator, EmitContext, *EmitResultScope, PixelResultType0, Generator.GetResultExpression(), MP_PixelDepthOffset);
+
+		if (!EmitContext.bReadMaterialNormal)
+		{
+			// No access to material normal, can execute everything in phase0
+			RequestedPixelAttributesType.SetFieldRequested(NormalField);
+			const FPreparedType PixelResultType1 = EmitContext.PrepareExpression(Generator.GetResultExpression(), *EmitResultScope, RequestedPixelAttributesType);
+			if (PixelResultType1.IsVoid())
+			{
+				return false;
+			}
+
+			HLSLTree->EmitShader(EmitContext, PixelCodePhase0);
+		}
+		else
+		{
+			// Execute everything *except* normal in phase1
+			HLSLTree->EmitShader(EmitContext, PixelCodePhase1);
+
+			// Reset state
+			HLSLTree->ResetNodes();
+
+			EmitResultScope = EmitContext.PrepareScope(&Generator.GetResultStatement()->GetParentScope());
+
+			// Prepare code for just the normal
+			FRequestedType RequestedMaterialNormal(Generator.GetMaterialAttributesType(), false);
+			RequestedMaterialNormal.SetFieldRequested(NormalField);
+			const FPreparedType PixelResultType1 = EmitContext.PrepareExpression(Generator.GetResultExpression(), *EmitResultScope, RequestedMaterialNormal);
+			if (PixelResultType1.IsVoid())
+			{
+				return false;
+			}
+
+			// Execute the normal in phase0
+			HLSLTree->EmitShader(EmitContext, PixelCodePhase0);
+		}
 	}
-	else
+
+	// Prepare vertex shader code
+	FStringBuilderMemstack VertexCode(Allocator, 1024 * 1024);
 	{
-		// Execute everything *except* normal in phase1
-		HLSLTree->EmitShader(EmitContext, CodePhase1);
+		FRequestedType RequestedVertexAttributesType(Generator.GetMaterialAttributesType(), false);
+		for (const FGuid& AttributeID : OrderedVisibleAttributes)
+		{
+			if (FMaterialAttributeDefinitionMap::GetShaderFrequency(AttributeID) == SF_Vertex)
+			{
+				const FString& FieldName = FMaterialAttributeDefinitionMap::GetAttributeName(AttributeID);
+				const FStructField* Field = Generator.GetMaterialAttributesType()->FindFieldByName(*FieldName);
+				if (Field)
+				{
+					RequestedVertexAttributesType.SetFieldRequested(Field);
+				}
+			}
+		}
 
-		// Reset state
 		HLSLTree->ResetNodes();
 
-		EmitResultScope = EmitContext.PrepareScope(&Generator.GetResultStatement()->GetParentScope());
+		EmitContext.ShaderFrequency = SF_Vertex;
+		FEmitScope* EmitResultScope = EmitContext.PrepareScope(&Generator.GetResultStatement()->GetParentScope());
 
-		// Prepare code for just the normal
-		FRequestedType RequestedMaterialNormal(Generator.GetMaterialAttributesType(), false);
-		RequestedMaterialNormal.SetFieldRequested(NormalField);
-		if (EmitContext.PrepareExpression(Generator.GetResultExpression(), *EmitResultScope, RequestedMaterialNormal).IsVoid())
+		const FPreparedType VertexResultType = EmitContext.PrepareExpression(Generator.GetResultExpression(), *EmitResultScope, RequestedVertexAttributesType);
+		if (VertexResultType.IsVoid())
 		{
 			return false;
 		}
 
-		// Execute the normal in phase0
-		HLSLTree->EmitShader(EmitContext, CodePhase0);
+		bUsesWorldPositionOffset = IsAttributeUsed(Generator, EmitContext, *EmitResultScope, VertexResultType, Generator.GetResultExpression(), MP_WorldPositionOffset);
+
+		HLSLTree->EmitShader(EmitContext, VertexCode);
 	}
 
 	if (EmitContext.Errors->HasErrors())
@@ -784,6 +911,10 @@ bool MaterialEmitHLSL(const FMaterialCompileTargetParameters& InCompilerTarget,
 		return false;
 	}
 
+	OutCompilationOutput.bModifiesMeshPosition = bUsesPixelDepthOffset || bUsesWorldPositionOffset;
+	OutCompilationOutput.bUsesWorldPositionOffset = bUsesWorldPositionOffset;
+	OutCompilationOutput.bUsesPixelDepthOffset = bUsesPixelDepthOffset;
+
 	FStringBuilderMemstack Declarations(Allocator, 64 * 1024);
 	TypeRegistry.EmitDeclarationsCode(Declarations);
 
@@ -791,8 +922,9 @@ bool MaterialEmitHLSL(const FMaterialCompileTargetParameters& InCompilerTarget,
 		InOutMaterial,
 		EmitContext,
 		Declarations.ToString(),
-		CodePhase0.ToString(),
-		bReadMaterialNormal ? CodePhase1.ToString() : nullptr);
+		VertexCode.ToString(),
+		PixelCodePhase0.ToString(),
+		EmitContext.bReadMaterialNormal ? PixelCodePhase1.ToString() : nullptr);
 
 	OutMaterialEnvironment = new FSharedShaderCompilerEnvironment();
 	OutMaterialEnvironment->TargetPlatform = InCompilerTarget.TargetPlatform;
diff --git a/Engine/Source/Runtime/Engine/Public/HLSLTree/HLSLTreeEmit.h b/Engine/Source/Runtime/Engine/Public/HLSLTree/HLSLTreeEmit.h
--- a/Engine/Source/Runtime/Engine/Public/HLSLTree/HLSLTreeEmit.h
+++ b/Engine/Source/Runtime/Engine/Public/HLSLTree/HLSLTreeEmit.h
@@ -370,6 +370,7 @@ public:
 	FMemStackBase* Allocator = nullptr;
 	FErrorHandlerInterface* Errors = nullptr;
 	const Shader::FStructTypeRegistry* TypeRegistry = nullptr;
+	EShaderFrequency ShaderFrequency = SF_Pixel;
 
 	TArray<FEmitShaderNode*> EmitNodes;
 	TMap<const FScope*, FEmitScope*> EmitScopeMap;
@@ -391,10 +392,10 @@ public:
 	uint32 CurrentBoolUniformOffset = ~0u;
 	uint32 CurrentNumBoolComponents = 32u;
 	bool bReadMaterialNormal = false;
+	uint32 TexCoordMask[SF_NumFrequencies] = { 0u };
 
 	int32 NumExpressionLocals = 0;
 	int32 NumExpressionLocalPHIs = 0;
-	int32 NumTexCoords = 0;
 };
 
 } // namespace HLSLTree
