Fix vlog automation tests broken when supporting multi threading vlog

#rnx

diff --git a/Engine/Source/Runtime/Engine/Private/VisualLogger/VisualLogger.cpp b/Engine/Source/Runtime/Engine/Private/VisualLogger/VisualLogger.cpp
--- a/Engine/Source/Runtime/Engine/Private/VisualLogger/VisualLogger.cpp
+++ b/Engine/Source/Runtime/Engine/Private/VisualLogger/VisualLogger.cpp
@@ -385,10 +385,8 @@ void FVisualLogger::Tick(float DeltaTime)
 	}
 }
 
-void FVisualLogger::Flush()
+void FVisualLogger::FlushThreadsEntries()
 {
-	FWriteScopeLock Lock(EntryRWLock);
-
 	// We are not locking the thread entry map, we do not expect other threads adding new entries while flushing
 	// @todo we should add a thread access detector
 	for (FVisualLoggerObjectEntryMap* ThreadCurrentEntryMap : ThreadCurrentEntryMaps)
@@ -406,6 +404,14 @@ void FVisualLogger::Flush()
 			}
 		}
 	}
+}
+
+void FVisualLogger::Flush()
+{
+	FWriteScopeLock Lock(EntryRWLock);
+
+	FlushThreadsEntries();
+
 	for (auto &CurrentEntry : CurrentEntryPerObject)
 	{
 		if (CurrentEntry.Value.bIsInitialized)
diff --git a/Engine/Source/Runtime/Engine/Private/VisualLogger/VisualLoggerAutomationTests.cpp b/Engine/Source/Runtime/Engine/Private/VisualLogger/VisualLoggerAutomationTests.cpp
--- a/Engine/Source/Runtime/Engine/Private/VisualLogger/VisualLoggerAutomationTests.cpp
+++ b/Engine/Source/Runtime/Engine/Private/VisualLogger/VisualLoggerAutomationTests.cpp
@@ -9,6 +9,7 @@
 #include "VisualLogger/VisualLoggerTypes.h"
 #include "VisualLogger/VisualLogger.h"
 
+PRAGMA_DISABLE_OPTIMIZATION
 namespace
 {
 	UWorld* GetSimpleEngineAutomationTestWorld(const int32 TestFlags)
@@ -145,6 +146,7 @@ bool FVisualLogTest::RunTest(const FString& Parameters)
 
 			const float NewTimestamp = CurrentTimestamp + 0.1f;
 			FVisualLogEntry* NewEntry = FVisualLogger::Get().GetEntryToWrite(World, NewTimestamp); //generate new entry and serialize old one
+			FVisualLogger::Get().FlushThreadsEntries();
 			CurrentEntry = &Context.Device.LastEntry;
 			CHECK_NOT_NULL(CurrentEntry);
 			CHECK_SUCCESS(CurrentEntry->TimeStamp == CurrentTimestamp);
@@ -201,6 +203,7 @@ bool FVisualLogSegmentsTest::RunTest(const FString& Parameters)
 
 			const float NewTimestamp = CurrentTimestamp + 0.1f;
 			FVisualLogEntry* NewEntry = FVisualLogger::Get().GetEntryToWrite(World, NewTimestamp); //generate new entry and serialize old one
+			FVisualLogger::Get().FlushThreadsEntries();
 			CurrentEntry = &Context.Device.LastEntry;
 			CHECK_NOT_NULL(CurrentEntry);
 			CHECK_SUCCESS(CurrentEntry->TimeStamp == CurrentTimestamp);
@@ -282,6 +285,7 @@ bool FVisualLogEventsTest::RunTest(const FString& Parameters)
 
 		const float NewTimestamp = CurrentTimestamp + 0.1f;
 		FVisualLogEntry* NewEntry = FVisualLogger::Get().GetEntryToWrite(World, NewTimestamp); //generate new entry and serialize old one
+		FVisualLogger::Get().FlushThreadsEntries();
 		CurrentEntry = &Context.Device.LastEntry;
 		CHECK_NOT_NULL(CurrentEntry);
 		CHECK_SUCCESS(CurrentEntry->TimeStamp == CurrentTimestamp);
@@ -337,3 +341,4 @@ bool FVisualLogEventsTest::RunTest(const FString& Parameters)
 #undef CHECK_FAIL
 
 #endif //ENABLE_VISUAL_LOG
+PRAGMA_ENABLE_OPTIMIZATION
diff --git a/Engine/Source/Runtime/Engine/Public/VisualLogger/VisualLogger.h b/Engine/Source/Runtime/Engine/Public/VisualLogger/VisualLogger.h
--- a/Engine/Source/Runtime/Engine/Public/VisualLogger/VisualLogger.h
+++ b/Engine/Source/Runtime/Engine/Public/VisualLogger/VisualLogger.h
@@ -525,6 +525,8 @@ public:
 	FVisualLogEntry* GetLastEntryForObject(const UObject* Object);
 	/** flush and serialize data if timestamp allows it */
 	virtual void Flush() override;
+	/** Moves all threads entries into the global entry map */
+	void FlushThreadsEntries();
 
 	/** FileName getter to set project specific file name for vlogs - highly encouraged to use FVisualLogFilenameGetterDelegate::CreateUObject with this */
 	void SetLogFileNameGetter(const FVisualLogFilenameGetterDelegate& InLogFileNameGetter) { LogFileNameGetter = InLogFileNameGetter; }
