Fixing seeking in non-streaming sounds

diff --git a/Engine/Source/Runtime/Engine/Private/SoundWave.cpp b/Engine/Source/Runtime/Engine/Private/SoundWave.cpp
--- a/Engine/Source/Runtime/Engine/Private/SoundWave.cpp
+++ b/Engine/Source/Runtime/Engine/Private/SoundWave.cpp
@@ -2963,11 +2963,18 @@ TArrayView<const uint8> USoundWave::GetZerothChunk(bool bForImmediatePlayback)
 
 bool USoundWave::IsSeekable() const
 {
-	if (bIsSourceBus || bProcedural || SoundAssetCompressionType == ESoundAssetCompressionType::PlatformSpecific)
+	bool Result = true;
+
+	// Non-streaming codecs are seekable
+	// Streaming platform-specific codecs do not seek. Neither do source buses or procedural sources.
+	if (bIsSourceBus || bProcedural || (IsStreaming() && SoundAssetCompressionType == ESoundAssetCompressionType::PlatformSpecific))
 	{
-		return false;
+		Result = false;
 	}
-	return true;
+
+	SoundWaveDataPtr->bIsSeekable = Result; // update shared flags
+
+	return Result;
 }
 
 bool USoundWave::GetSoundWavesWithCookedAnalysisData(TArray<USoundWave*>& OutSoundWaves)
