Cloth - Fix cloth transforms when the simulation is disabled using p.ClothPhysics 0 after a cloth simulation has been created.
Chaos Cloth - Fix the sim data transform when rendering of the cloth starts after the component's transforms get updated. [Edigrate CL 18171882]

diff --git a/Engine/Source/Runtime/Engine/Private/Components/SkeletalMeshComponent.cpp b/Engine/Source/Runtime/Engine/Private/Components/SkeletalMeshComponent.cpp
--- a/Engine/Source/Runtime/Engine/Private/Components/SkeletalMeshComponent.cpp
+++ b/Engine/Source/Runtime/Engine/Private/Components/SkeletalMeshComponent.cpp
@@ -1432,7 +1432,7 @@ void USkeletalMeshComponent::TickComponent(float DeltaTime, enum ELevelTick Tick
 
 	// If we are suspended, we will not simulate clothing, but as clothing is simulated in local space
 	// relative to a root bone we need to extract simulation positions as this bone could be animated.
-	if(bClothingSimulationSuspended && ClothingSimulation)
+	if((!CVarEnableClothPhysics.GetValueOnGameThread() || bClothingSimulationSuspended) && ClothingSimulation)
 	{
 		CSV_SCOPED_TIMING_STAT(Animation, Cloth);
 
