Updating common memory reporting commands to properly use -csv output and added parsable delimiting to MemReport output. Also, general cleanup and fix for audiomemreport to properly redirect to a different OutputDevice (removes -skipcsvforaudio, adds -log).

#test memreport -csv -full -name=

diff --git a/Engine/Source/Runtime/Engine/Private/AudioDevice.cpp b/Engine/Source/Runtime/Engine/Private/AudioDevice.cpp
--- a/Engine/Source/Runtime/Engine/Private/AudioDevice.cpp
+++ b/Engine/Source/Runtime/Engine/Private/AudioDevice.cpp
@@ -1648,23 +1648,30 @@ bool FAudioDevice::HandleAudioMemoryInfo(const TCHAR* Cmd, FOutputDevice& Ar)
 		}
 	};
 
-	const FString PathName = *(FPaths::ProfilingDir() + TEXT("MemReports/"));
-	IFileManager::Get().MakeDirectory(*PathName);
+	// Default to writing our own .csv file unless -log is present
+	const bool bLogOutputToFile = !FParse::Param(Cmd, TEXT("LOG"));
 
-	const FString Filename = CreateProfileFilename(TEXT("_audio_memreport.csv"), true);
-	FString FilenameFull = PathName + Filename;
+	FOutputDevice* ReportAr = &Ar;
+	FArchive* FileAr = nullptr;
+	FOutputDeviceArchiveWrapper* FileArWrapper = nullptr;
 
-	FArchive* FileAr = IFileManager::Get().CreateDebugFileWriter(*FilenameFull);
-	FOutputDeviceArchiveWrapper* FileArWrapper = new FOutputDeviceArchiveWrapper(FileAr);
-	FOutputDevice* ReportAr = FileArWrapper;
-
-	if (FParse::Param(Cmd, TEXT("SkipCsvForAudio")))
+	if (bLogOutputToFile)
 	{
-		ReportAr = &Ar;
+		const FString PathName = *(FPaths::ProfilingDir() + TEXT("MemReports/"));
+		IFileManager::Get().MakeDirectory(*PathName);
+
+		const FString Filename = CreateProfileFilename(TEXT("_audio_memreport.csv"), true);
+		const FString FilenameFull = PathName + Filename;
+
+		FileAr = IFileManager::Get().CreateDebugFileWriter(*FilenameFull);
+		FileArWrapper = new FOutputDeviceArchiveWrapper(FileAr);
+		ReportAr = FileArWrapper;
+
+		UE_LOG(LogEngine, Log, TEXT("AudioMemReport: saving to %s"), *FilenameFull);
 	}
 	else
 	{
-		UE_LOG(LogEngine, Log, TEXT("AudioMemReport: saving to %s"), *FilenameFull);
+		UE_LOG(LogEngine, Log, TEXT("Use command \"AudioMemReport -log\" to output to client logs or the passed through FOutputDevice."));
 	}
 
 	// Get the sound wave class
@@ -1839,7 +1846,7 @@ bool FAudioDevice::HandleAudioMemoryInfo(const TCHAR* Cmd, FOutputDevice& Ar)
 				{
 					FSoundWaveGroupInfo* SubDirSize = SoundWaveGroupSizes.Find(SoundWaveGroupFolder);
 					check(SubDirSize);
-					ReportAr->Logf(TEXT("%s,%10.2f,%10.2f"), *SoundWaveGroupFolder, SubDirSize->ResourceSize.GetTotalMemoryBytes() / 1024.0f / 1024.0f, SubDirSize->CompressedResourceSize.GetTotalMemoryBytes() / 1024.0f / 1024.0f);
+					ReportAr->Logf(TEXT("%s,%.2f,%.2f"), *SoundWaveGroupFolder, SubDirSize->ResourceSize.GetTotalMemoryBytes() / 1024.0f / 1024.0f, SubDirSize->CompressedResourceSize.GetTotalMemoryBytes() / 1024.0f / 1024.0f);
 				}
 			}
 
@@ -1847,7 +1854,7 @@ bool FAudioDevice::HandleAudioMemoryInfo(const TCHAR* Cmd, FOutputDevice& Ar)
 			ReportAr->Log(TEXT("All Sound Wave Objects Sorted Alphebetically:"));
 			ReportAr->Log(TEXT(""));
 
-			ReportAr->Logf(TEXT("%s,%s,%s,%s,%s,%s,%s,%s"), TEXT("SoundWave"), TEXT("KB"), TEXT("MB"), TEXT("SoundGroup"), TEXT("Duration"), TEXT("CompressionState"), TEXT("Max Size in Cache (Unevictable, KB)"), TEXT("Max Size In Cache (Total, KB)"));
+			ReportAr->Logf(TEXT("%s,%s,%s,%s,%s,%s,%s,%s"), TEXT("SoundWave"), TEXT("KB"), TEXT("MB"), TEXT("SoundGroup"), TEXT("Duration"), TEXT("CompressionState"), TEXT("Max Size in Cache (Unevictable KB)"), TEXT("Max Size In Cache (Total KB)"));
 			for (const FSoundWaveInfo& Info : SoundWaveObjects)
 			{
 				float Kbytes = Info.ResourceSize.GetTotalMemoryBytes() / 1024.0f;
@@ -1868,14 +1875,17 @@ bool FAudioDevice::HandleAudioMemoryInfo(const TCHAR* Cmd, FOutputDevice& Ar)
 					break;
 				}
 
-				ReportAr->Logf(TEXT("%s,%10.2f,%10.2f,%s,%10.2f, %s, %10.2f, %10.2f"), *Info.SoundWave->GetPathName(), Kbytes, Kbytes / 1024.0f, *Info.SoundGroupName, Info.Duration, *LoadingTypeString, Info.MaxUnevictableSizeInCache / 1024.0f, Info.PotentialTotalSizeInCache / 1024.0f);
+				ReportAr->Logf(TEXT("%s,%.2f,%.2f,%s,%.2f,%s,%.2f,%.2f"), *Info.SoundWave->GetPathName(), Kbytes, Kbytes / 1024.0f, *Info.SoundGroupName, Info.Duration, *LoadingTypeString, Info.MaxUnevictableSizeInCache / 1024.0f, Info.PotentialTotalSizeInCache / 1024.0f);
 			}
 		}
 
 	}
 
-	// Shutdown and free archive resources
-	FileArWrapper->TearDown();
+	if (FileArWrapper != nullptr)
+	{
+		// Shutdown and free archive resources
+		FileArWrapper->TearDown();
+	}
 	delete FileArWrapper;
 	delete FileAr;
 
diff --git a/Engine/Source/Runtime/Engine/Private/UnrealEngine.cpp b/Engine/Source/Runtime/Engine/Private/UnrealEngine.cpp
--- a/Engine/Source/Runtime/Engine/Private/UnrealEngine.cpp
+++ b/Engine/Source/Runtime/Engine/Private/UnrealEngine.cpp
@@ -5942,7 +5942,7 @@ bool UEngine::HandleListTexturesCommand( const TCHAR* Cmd, FOutputDevice& Ar )
 
 		if (bCSV)
 		{
-			Ar.Logf(TEXT("%i, %i, %i, %s, %i, %i, %i, %s, %s, %s, %s, %s, %s, %i, %i, %s"),
+			Ar.Logf(TEXT("%i,%i,%i,%s,%i,%i,%i,%s,%s,%s,%s,%s,%s,%i,%i,%s"),
 				SortedTexture.MaxAllowedSizeX, SortedTexture.MaxAllowedSizeY, (SortedTexture.MaxAllowedSize + 512) / 1024,
 				*AuthoredBiasString,
 				SortedTexture.CurSizeX, SortedTexture.CurSizeY, (SortedTexture.CurrentSize + 512) / 1024,
@@ -6015,14 +6015,13 @@ bool UEngine::HandleListStaticMeshesCommand(const TCHAR* Cmd, FOutputDevice& Ar)
 	SortType = FParse::Param(Cmd, TEXT("MOBILESORT")) ? FCompareFSortedStaticMesh::ESortType::Mobile : SortType;
 	SortType = FParse::Param(Cmd, TEXT("INSTANCESORT")) ? FCompareFSortedStaticMesh::ESortType::InstanceCount : SortType;
 
+	const bool bCSV = FParse::Param(Cmd, TEXT("CSV"));
 	const bool bUsedComponents = FParse::Param(Cmd, TEXT("usedcomponents"));
 	const bool bUnknownRefOnly = FParse::Param(Cmd, TEXT("unknownrefonly"));
 	const bool bGameWorld = FParse::Param(Cmd, TEXT("game"));
 	const bool bUsage = FParse::Param(Cmd, TEXT("-?"));
 	
-
-
-	//non-editor builds literally don't have the data to determine mobile lods or vert data.  The data prints out incorrectly and is confusing, just remove it.
+	// Non-editor builds literally don't have the data to determine mobile lods or vert data. The data prints out incorrectly and is confusing. Just remove it.
 	const bool bHasMobileColumns = (bool)WITH_EDITORONLY_DATA;
 	Ar.Logf(TEXT("Listing all static meshes."));
 	if (bUsage)
@@ -6033,8 +6032,9 @@ bool UEngine::HandleListStaticMeshesCommand(const TCHAR* Cmd, FOutputDevice& Ar)
 			"-mobilesort: sort by mobile verts \n"
 			"-instancesort: sort by instance count \n"
 			"-game: only consider game world \n"
-			"-usedcomponents: print the all components used by each mesh\n"
-			"-unknownrefonly: list only static meshes using unknown ref heuristic for LOD streaming"));
+			"-usedcomponents: print the all components used by each mesh \n"
+			"-unknownrefonly: list only static meshes using unknown ref heuristic for LOD streaming \n"
+			"-csv: output in csv format"));
 	}
 
 	FRenderAssetStreamingManager* Streamer = nullptr;
@@ -6075,7 +6075,7 @@ bool UEngine::HandleListStaticMeshesCommand(const TCHAR* Cmd, FOutputDevice& Ar)
 	TArray<FSortedStaticMesh> SortedMeshes;
 	for(TObjectIterator<UStaticMesh> It; It; ++It)
 	{
-		UStaticMesh*		Mesh = *It;
+		UStaticMesh* Mesh = *It;
 
 		bool bUnknownRef = false;
 		bool bIsStreaming = (Mesh != nullptr ? Mesh->GetStreamingIndex() != INDEX_NONE : false);
@@ -6117,17 +6117,17 @@ bool UEngine::HandleListStaticMeshesCommand(const TCHAR* Cmd, FOutputDevice& Ar)
 			Mesh->GetRenderData()->LODResources[0].DistanceFieldData->GetResourceSizeEx(DistanceFieldSizeExc);
 		}
 
-		int32		NumKB = (Count.GetNum() + 512) / 1024;
-		int32		MaxKB = (Count.GetMax() + 512) / 1024;
-		int32		ResKBExc = (ResourceSizeExc.GetTotalMemoryBytes() + 512) / 1024;
-		int32		ResKBInc = (ResourceSizeInc.GetTotalMemoryBytes() + 512) / 1024;
-		int32		ResKBIncMobile = 0; //Update mobilesort once implemented
-		int32		DistanceFieldKB = (DistanceFieldSizeExc.GetTotalMemoryBytes() + 512) / 1024;
-		int32		LodCount = Mesh->GetNumLODs();
-		int32		VertexCountLod0 = LodCount > 0 ? Mesh->GetNumVertices(0) : 0;
-		int32		VertexCountLod1 = LodCount > 1 ? Mesh->GetNumVertices(1) : 0;
-		int32		VertexCountLod2 = LodCount > 2 ? Mesh->GetNumVertices(2) : 0;
-		int32		MobileMinLOD = -1;
+		int32 NumKB = (Count.GetNum() + 512) / 1024;
+		int32 MaxKB = (Count.GetMax() + 512) / 1024;
+		int32 ResKBExc = (ResourceSizeExc.GetTotalMemoryBytes() + 512) / 1024;
+		int32 ResKBInc = (ResourceSizeInc.GetTotalMemoryBytes() + 512) / 1024;
+		int32 ResKBIncMobile = 0; //Update mobilesort once implemented
+		int32 DistanceFieldKB = (DistanceFieldSizeExc.GetTotalMemoryBytes() + 512) / 1024;
+		int32 LodCount = Mesh->GetNumLODs();
+		int32 VertexCountLod0 = LodCount > 0 ? Mesh->GetNumVertices(0) : 0;
+		int32 VertexCountLod1 = LodCount > 1 ? Mesh->GetNumVertices(1) : 0;
+		int32 VertexCountLod2 = LodCount > 2 ? Mesh->GetNumVertices(2) : 0;
+		int32 MobileMinLOD = -1;
 
 #if WITH_EDITORONLY_DATA 
 		if (Mesh->IsMinLodQualityLevelEnable() && Mesh->GetQualityLevelMinLOD().PerQuality.Find(0/*Low*/) != nullptr)
@@ -6140,14 +6140,14 @@ bool UEngine::HandleListStaticMeshesCommand(const TCHAR* Cmd, FOutputDevice& Ar)
 		}
 #endif
 
-		int32		CollisionShapeCount = 0;
+		int32 CollisionShapeCount = 0;
 		if (Mesh->GetBodySetup())
 		{
 			CollisionShapeCount = Mesh->GetBodySetup()->AggGeom.GetElementCount();
 		}
 
-		int32		VertexCountTotal = 0;
-		int32		VertexCountTotalMobile = 0;
+		int32 VertexCountTotal = 0;
+		int32 VertexCountTotalMobile = 0;
 		int32 ResidentLodCount = 0;
 		int32 ResidentResKBExc = 0;
 		int32 NumMissingLODs = 0;
@@ -6170,7 +6170,7 @@ bool UEngine::HandleListStaticMeshesCommand(const TCHAR* Cmd, FOutputDevice& Ar)
 		}
 		ResidentResKBExc = (ResourceSizeExc.GetTotalMemoryBytes() - EvictedResourceSize.GetTotalMemoryBytes() + 512) / 1024;
 
-		int32		VertexCountCollision = 0;
+		int32 VertexCountCollision = 0;
 		if(Mesh->GetBodySetup())
 		{
 #if PHYSICS_INTERFACE_PHYSX
@@ -6187,7 +6187,7 @@ bool UEngine::HandleListStaticMeshesCommand(const TCHAR* Cmd, FOutputDevice& Ar)
 #endif
 		}
 
-		FString				Name = Mesh->GetFullName();
+		FString Name = Mesh->GetFullName();
 
 		new(SortedMeshes) FSortedStaticMesh(
 			Mesh,
@@ -6233,10 +6233,23 @@ bool UEngine::HandleListStaticMeshesCommand(const TCHAR* Cmd, FOutputDevice& Ar)
 	int64 TotalComponents = 0;
 	int64 TotalInstances = 0;
 
-	FString HeaderString(TEXT(",       NumKB,       MaxKB,    ResKBExc,    ResKBInc,   ResKBResi, DistFieldKB,    LODCount,ResiLODCount,   VertsLOD0,   VertsLOD1,   VertsLOD2, Verts Total,  Verts Coll, Coll Shapes,    NumComps,     NumInst,      Nanite, CPUAccessOverheadKB,  UnknownRef, Streaming"));
-	FString MobileHeaderString = bHasMobileColumns ? FString(TEXT(", ResKBIncMob,Verts Mobile,MobileMinLOD")) : FString();
-	
-	Ar.Logf(TEXT("%s%s, Name"), *HeaderString, *MobileHeaderString);	
+	FString HeaderString = bCSV ? FString(TEXT(",NumKB,MaxKB,ResKBExc,ResKBInc,ResKBResi,DistFieldKB,LODCount,ResiLODCount,VertsLOD0,VertsLOD1,VertsLOD2,Verts Total,Verts Coll,Coll Shapes,NumComps,NumInst,Nanite,CPUAccessOverheadKB,UnknownRef,Streaming"))
+		: FString(TEXT(",       NumKB,       MaxKB,    ResKBExc,    ResKBInc,   ResKBResi, DistFieldKB,    LODCount,ResiLODCount,   VertsLOD0,   VertsLOD1,   VertsLOD2, Verts Total,  Verts Coll, Coll Shapes,    NumComps,     NumInst,      Nanite, CPUAccessOverheadKB,  UnknownRef,  Streaming"));
+	FString MobileHeaderString = bCSV ? FString(TEXT(",ResKBIncMob,Verts Mobile,MobileMinLOD"))
+		: FString(TEXT(", ResKBIncMob,Verts Mobile,MobileMinLOD"));
+	if (bHasMobileColumns)
+	{
+		HeaderString = HeaderString.Append(MobileHeaderString);
+	}
+
+	if (bCSV)
+	{
+		Ar.Logf(TEXT("%s,Name"), *HeaderString);
+	}
+	else
+	{
+		Ar.Logf(TEXT("%s, Name"), *HeaderString);
+	}
 
 	for(int32 MeshIndex = 0; MeshIndex<SortedMeshes.Num(); MeshIndex++)
 	{
@@ -6310,7 +6323,16 @@ bool UEngine::HandleListStaticMeshesCommand(const TCHAR* Cmd, FOutputDevice& Ar)
 		TotalInstances += SortedMesh.InstanceUsageCount;
 	}
 
-	Ar.Logf(TEXT("Total NumKB: %lld KB, Total MaxKB: %lld KB, Total ResKB Exc: %lld KB, Total ResKB Inc %lld KB,  Total ResKB Inc Mobile %lld KB, Total ResKB Resident %lld KB, Total Distance Field %lld KB, Total CPUAccess Overhead %lld KB, Total Vertex Count: %i, Total Vertex Count Mobile: %i, Component Count: %lld, Instance Count: %lld, Static Mesh Count=%d"), TotalNumKB, TotalMaxKB, TotalResKBExc, TotalResKBInc, TotalResKBIncMobile, TotalResKBResident, TotalDistanceFieldKB, TotalCPUAccessOverheadKB, TotalVertexCount, TotalVertexCountMobile, TotalComponents, TotalInstances, SortedMeshes.Num());
+	if (bCSV)
+	{
+		Ar.Logf(TEXT("Total NumKB: %lld,Total MaxKB: %lld,Total ResKB Exc: %lld,Total ResKB Inc: %lld,Total ResKB Inc Mobile: %lld,Total ResKB Resident: %lld,Total Distance Field: %lld,Total CPUAccess Overhead: %lld,Total Vertex Count: %i,Total Vertex Count Mobile: %i,Component Count: %i,Instance Count: %i,Static Mesh Count: %i"),
+			TotalNumKB, TotalMaxKB, TotalResKBExc, TotalResKBInc, TotalResKBIncMobile, TotalResKBResident, TotalDistanceFieldKB, TotalCPUAccessOverheadKB, TotalVertexCount, TotalVertexCountMobile, TotalComponents, TotalInstances, SortedMeshes.Num());
+	}
+	else
+	{
+		Ar.Logf(TEXT("Total NumKB: %lld KB, Total MaxKB: %lld KB, Total ResKB Exc: %lld KB, Total ResKB Inc: %lld KB, Total ResKB Inc Mobile: %lld KB, Total ResKB Resident: %lld KB, Total Distance Field: %lld KB, Total CPUAccess Overhead: %lld KB, Total Vertex Count: %i, Total Vertex Count Mobile: %i, Component Count: %i, Instance Count: %i, Static Mesh Count: %i"),
+			TotalNumKB, TotalMaxKB, TotalResKBExc, TotalResKBInc, TotalResKBIncMobile, TotalResKBResident, TotalDistanceFieldKB, TotalCPUAccessOverheadKB, TotalVertexCount, TotalVertexCountMobile, TotalComponents, TotalInstances, SortedMeshes.Num());
+	}
 
 	if (bUsedComponents)
 	{
@@ -6524,7 +6546,7 @@ bool UEngine::HandleListSkeletalMeshesCommand(const TCHAR* Cmd, FOutputDevice& A
 
 	if (bCSV)
 	{
-		Ar.Logf(TEXT(",NumKB, MaxKB, ResKBExc, ResKBInc, ResKBIncMobile, ResKBResident, LOD Count, Resident LOD Count, MobileMinLOD, DisableMinLODStripping, bSupportLODStreaming, MaxNumStreamedLODs, MaxNumOptionalLODs, Verts LOD0, Verts LOD1, Verts LOD2, Verts Total, Verts Total Mobile, Verts Collision, CPUAccessOverheadKB, UnknownRef, Streaming, Name"));
+		Ar.Logf(TEXT(",NumKB,MaxKB,ResKBExc,ResKBInc,ResKBIncMobile,ResKBResident,LOD Count,Resident LOD Count,MobileMinLOD,DisableMinLODStripping,bSupportLODStreaming,MaxNumStreamedLODs,MaxNumOptionalLODs,Verts LOD0,Verts LOD1,Verts LOD2,Verts Total,Verts Total Mobile,Verts Collision,CPUAccessOverheadKB,UnknownRef,Streaming,Name"));
 	}
 	else
 	{
@@ -6537,7 +6559,7 @@ bool UEngine::HandleListSkeletalMeshesCommand(const TCHAR* Cmd, FOutputDevice& A
 
 		if (bCSV)
 		{
-			Ar.Logf(TEXT(",%i, %i, %i, %i, %i, %i, %i, %i, %i, %i, %i, %i, %i, %i, %i, %i, %i, %i, %i, %i, %s, %s, %s"),
+			Ar.Logf(TEXT(",%i,%i,%i,%i,%i,%i,%i,%i,%i,%i,%i,%i,%i,%i,%i,%i,%i,%i,%i,%i,%s,%s,%s"),
 				SortedMesh.NumKB, SortedMesh.MaxKB, SortedMesh.ResKBExc, SortedMesh.ResKBInc, SortedMesh.ResKBIncMobile, SortedMesh.ResidentResKBExc, SortedMesh.LodCount, SortedMesh.ResidentLodCount, SortedMesh.MobileMinLOD,
 				SortedMesh.MeshDisablesMinLODStripping, SortedMesh.bSupportLODStreaming, SortedMesh.MaxNumStreamedLODs, SortedMesh.MaxNumOptionalLODs,
 				SortedMesh.VertexCountLod0, SortedMesh.VertexCountLod1, SortedMesh.VertexCountLod2, SortedMesh.VertexCountTotal, SortedMesh.VertexCountTotalMobile, SortedMesh.VertexCountCollision,
@@ -6565,7 +6587,16 @@ bool UEngine::HandleListSkeletalMeshesCommand(const TCHAR* Cmd, FOutputDevice& A
 		TotalCPUAccessOverheadKB += SortedMesh.CPUAccessOverheadKB;
 	}
 
-	Ar.Logf(TEXT("Total NumKB: %lld KB, Total MaxKB: %lld KB, Total ResKB Exc: %lld KB, Total ResKB Inc %lld KB,  Total ResKB Inc Mobile %lld KB, Total ResKB Resident %lld KB, Total CPUAccess Overhead %lld KB, Total Vertex Count: %i, Total Vertex Count Mobile: %i, Static Mesh Count=%d"), TotalNumKB, TotalMaxKB, TotalResKBExc, TotalResKBInc, TotalResKBIncMobile, TotalResKBResident, TotalCPUAccessOverheadKB, TotalVertexCount, TotalVertexCountMobile, SortedMeshes.Num());
+	if (bCSV)
+	{
+		Ar.Logf(TEXT("Total NumKB: %lld,Total MaxKB: %lld,Total ResKB Exc: %lld,Total ResKB Inc %lld,Total ResKB Inc Mobile %lld,Total ResKB Resident %lld,Total CPUAccess Overhead %lld,Total Vertex Count: %i,Total Vertex Count Mobile: %i,Static Mesh Count: %i"),
+			TotalNumKB, TotalMaxKB, TotalResKBExc, TotalResKBInc, TotalResKBIncMobile, TotalResKBResident, TotalCPUAccessOverheadKB, TotalVertexCount, TotalVertexCountMobile, SortedMeshes.Num());
+	}
+	else
+	{
+		Ar.Logf(TEXT("Total NumKB: %lld KB, Total MaxKB: %lld KB, Total ResKB Exc: %lld KB, Total ResKB Inc %lld KB,  Total ResKB Inc Mobile %lld KB, Total ResKB Resident %lld KB, Total CPUAccess Overhead %lld KB, Total Vertex Count: %i, Total Vertex Count Mobile: %i, Static Mesh Count: %i"),
+			TotalNumKB, TotalMaxKB, TotalResKBExc, TotalResKBInc, TotalResKBIncMobile, TotalResKBResident, TotalCPUAccessOverheadKB, TotalVertexCount, TotalVertexCountMobile, SortedMeshes.Num());
+	}
 
 	return true;
 }
@@ -6577,7 +6608,6 @@ bool UEngine::HandleListAnimsCommand(const TCHAR* Cmd, FOutputDevice& Ar)
 
 	Ar.Logf(TEXT("Listing all animations."));
 
-
 	TArray<FSortedAnimAsset> SortedAnimAssets;
 	for (TObjectIterator<UAnimationAsset> It; It; ++It)
 	{
@@ -6645,14 +6675,28 @@ bool UEngine::HandleListAnimsCommand(const TCHAR* Cmd, FOutputDevice& Ar)
 	int64 TotalResKBExc = 0;
 	int64 TotalResKBInc = 0;
 
-	Ar.Logf(TEXT(",NumKB, MaxKB, ResKBExc, ResKBInc, Type, Name, Length, RateScale, NumCurves, TranslationFormat, RotationFormat, ScaleFormat"));
+	if (bCSV)
+	{
+		Ar.Logf(TEXT(",NumKB,MaxKB,ResKBExc,ResKBInc,Type,Name,Length,RateScale,NumCurves,TranslationFormat,RotationFormat,ScaleFormat"));
+	}
+	else
+	{
+		Ar.Logf(TEXT(", NumKB, MaxKB, ResKBExc, ResKBInc, Type, Name, Length, RateScale, NumCurves, TranslationFormat, RotationFormat, ScaleFormat"));
+	}
 
 	for (const FSortedAnimAsset& SortedAsset : SortedAnimAssets)
 	{
 
-		Ar.Logf(TEXT(",%i, %i, %i, %i, %s, %s, %f, %f, %d, %s"),
-			SortedAsset.NumKB, SortedAsset.MaxKB, SortedAsset.ResKBExc, SortedAsset.ResKBInc, *SortedAsset.AnimAssetType.ToString(), *SortedAsset.Name, SortedAsset.SequenceLength, 
-			SortedAsset.RateScale, SortedAsset.NumCurves, *SortedAsset.DebugString);
+		if (bCSV)
+		{
+			Ar.Logf(TEXT(",%i,%i,%i,%i,%s,%s,%f,%f,%i,%s"),
+				SortedAsset.NumKB, SortedAsset.MaxKB, SortedAsset.ResKBExc, SortedAsset.ResKBInc, *SortedAsset.AnimAssetType.ToString(), *SortedAsset.Name, SortedAsset.SequenceLength, SortedAsset.RateScale, SortedAsset.NumCurves, *SortedAsset.DebugString);
+		} 
+		else
+		{
+			Ar.Logf(TEXT(", %i, %i, %i, %i, %s, %s, %f, %f, %i, %s"),
+				SortedAsset.NumKB, SortedAsset.MaxKB, SortedAsset.ResKBExc, SortedAsset.ResKBInc, *SortedAsset.AnimAssetType.ToString(), *SortedAsset.Name, SortedAsset.SequenceLength, SortedAsset.RateScale, SortedAsset.NumCurves, *SortedAsset.DebugString);
+		}
 
 		TotalNumKB += SortedAsset.NumKB;
 		TotalMaxKB += SortedAsset.MaxKB;
@@ -6660,7 +6704,16 @@ bool UEngine::HandleListAnimsCommand(const TCHAR* Cmd, FOutputDevice& Ar)
 		TotalResKBInc += SortedAsset.ResKBInc;
 	}
 
-	Ar.Logf(TEXT("Total NumKB: %lld KB, Total MaxKB: %lld KB, Total ResKB Exc: %lld KB, Total ResKB Inc %lld KB, Anim Count=%d"), TotalNumKB, TotalMaxKB, TotalResKBExc, TotalResKBInc, SortedAnimAssets.Num());
+	if (bCSV)
+	{
+		Ar.Logf(TEXT("Total NumKB: %lld,Total MaxKB: %lld,Total ResKB Exc: %lld,Total ResKB Inc: %lld,Anim Count: %i"),
+			TotalNumKB, TotalMaxKB, TotalResKBExc, TotalResKBInc, SortedAnimAssets.Num());
+	} 
+	else
+	{
+		Ar.Logf(TEXT("Total NumKB: %lld KB, Total MaxKB: %lld KB, Total ResKB Exc: %lld KB, Total ResKB Inc: %lld KB, Anim Count: %i"),
+			TotalNumKB, TotalMaxKB, TotalResKBExc, TotalResKBInc, SortedAnimAssets.Num());
+	}
 
 	return true;
 }
@@ -7040,7 +7093,7 @@ bool UEngine::HandleMemReportDeferredCommand( const TCHAR* Cmd, FOutputDevice& A
 		ReportAr->Logf(TEXT("Device Profile: %s"), *DeviceProfile);
 
 		ReportAr->Logf(TEXT("CommandLine Options: %s"), FCommandLine::Get());
-	ReportAr->Logf(TEXT("Time Since Boot: %.02f Seconds") LINE_TERMINATOR, FPlatformTime::Seconds() - GStartTime);
+		ReportAr->Logf(TEXT("Time Since Boot: %.02f Seconds") LINE_TERMINATOR, FPlatformTime::Seconds() - GStartTime);
 	}
 
 
@@ -7057,8 +7110,9 @@ bool UEngine::HandleMemReportDeferredCommand( const TCHAR* Cmd, FOutputDevice& A
 				Command += " -csv";
 			}
 
-			Exec( InWorld, *Command, *ReportAr );
-			ReportAr->Logf( LINE_TERMINATOR );
+			ReportAr->Logf(TEXT("MemReport: Begin command \"%s\""), *Command);
+			Exec(InWorld, *Command, *ReportAr);
+			ReportAr->Logf(TEXT("MemReport: End command \"%s\"") LINE_TERMINATOR, *Command);
 		}
 	}
 
@@ -7076,8 +7130,9 @@ bool UEngine::HandleMemReportDeferredCommand( const TCHAR* Cmd, FOutputDevice& A
 					Command += " -csv";
 				}
 
-				Exec( InWorld, *Command, *ReportAr );
-				ReportAr->Logf( LINE_TERMINATOR );
+				ReportAr->Logf(TEXT("MemReport: Begin command \"%s\""), *Command);
+				Exec(InWorld, *Command, *ReportAr);
+				ReportAr->Logf(TEXT("MemReport: End command \"%s\"") LINE_TERMINATOR, *Command);
 			}
 		}
 	}
@@ -8650,7 +8705,7 @@ bool UEngine::HandleObjCommand( const TCHAR* Cmd, FOutputDevice& Ar )
 
 			if (bCSV)
 			{
-				Ar.Logf(TEXT(", Object, NumKB, MaxKB, ResExcKB, ResExcDedSysKB, ResExcDedVidKB, ResExcUnkKB"));
+				Ar.Logf(TEXT(",Object,NumKB,MaxKB,ResExcKB,ResExcDedSysKB,ResExcDedVidKB,ResExcUnkKB"));
 			}
 			else
 			{
@@ -8671,7 +8726,7 @@ bool UEngine::HandleObjCommand( const TCHAR* Cmd, FOutputDevice& Ar )
 				if (bCSV)
 				{
 					Ar.Logf(
-						TEXT(", %s, %.2f, %.2f, %.2f, %.2f, %.2f, %.2f"),
+						TEXT(",%s,%.2f,%.2f,%.2f,%.2f,%.2f,%.2f"),
 						*ObjItem.Object->GetFullName(),
 						ObjItem.Num / 1024.0f,
 						ObjItem.Max / 1024.0f,
@@ -8726,7 +8781,7 @@ bool UEngine::HandleObjCommand( const TCHAR* Cmd, FOutputDevice& Ar )
 			{
 				if (bShowMemory)
 				{
-					Ar.Logf(TEXT(", Class, Count, NumKB, MaxKB, ResExcKB, ResExcDedSysKB, ResExcDedVidKB, ResExcUnkKB"));
+					Ar.Logf(TEXT(",Class,Count,NumKB,MaxKB,ResExcKB,ResExcDedSysKB,ResExcDedVidKB,ResExcUnkKB"));
 				}
 				else
 				{
@@ -8761,7 +8816,7 @@ bool UEngine::HandleObjCommand( const TCHAR* Cmd, FOutputDevice& Ar )
 				{
 					if (bShowMemory)
 					{
-						Ar.Logf(TEXT(", %s, %i, %f, %f, %f, %f, %f, %f"),
+						Ar.Logf(TEXT(",%s,%i,%f,%f,%f,%f,%f,%f"),
 							bShowFullClassName ? *ClassData.Class->GetFullName() : *ClassData.Class->GetName(),
 							(int32)ClassData.Count,
 							ClassData.Num / 1024.0f,
