PR #8805: Fix RemoteViewPitch being generated "incorrectly" (Contributed by [REDACTED_PERSON])

diff --git a/Engine/Source/Runtime/Core/Public/Misc/NetworkVersion.h b/Engine/Source/Runtime/Core/Public/Misc/NetworkVersion.h
--- a/Engine/Source/Runtime/Core/Public/Misc/NetworkVersion.h
+++ b/Engine/Source/Runtime/Core/Public/Misc/NetworkVersion.h
@@ -55,6 +55,7 @@ enum EEngineNetworkVersionHistory
 	HISTORY_MONTAGE_PLAY_INST_ID_SERIALIZATION = 21,// Bump version to support net serialization of FGameplayAbilityRepAnimMontage, addition of PlayInstanceId and removal of bForcePlayBit
 	HISTORY_SERIALIZE_DOUBLE_VECTORS_AS_DOUBLES	= 22,// Bump version to support net serialization of double vector types
 	HISTORY_PACKED_VECTOR_LWC_SUPPORT = 23,			// Bump version to support quantized LWC FVector net serialization
+	HISTORY_PAWN_REMOTEVIEWPITCH = 24,				// Bump version to support serialization changes to RemoteViewPitch
 	// New history items go above here.
 
 	HISTORY_ENGINENETVERSION_PLUS_ONE,
diff --git a/Engine/Source/Runtime/Engine/Classes/Engine/DemoNetDriver.h b/Engine/Source/Runtime/Engine/Classes/Engine/DemoNetDriver.h
--- a/Engine/Source/Runtime/Engine/Classes/Engine/DemoNetDriver.h
+++ b/Engine/Source/Runtime/Engine/Classes/Engine/DemoNetDriver.h
@@ -610,6 +610,9 @@ public:
 	uint32 GetPlaybackDemoChangelist() const { return ReplayHelper.PlaybackDemoHeader.EngineVersion.GetChangelist(); }
 	uint32 GetPlaybackDemoVersion() const { return ReplayHelper.PlaybackDemoHeader.Version; }
 
+	uint32 GetPlaybackEngineNetworkProtocolVersion() const { return ReplayHelper.PlaybackDemoHeader.EngineNetworkProtocolVersion; }
+	uint32 GetPlaybackGameNetworkProtocolVersion() const { return ReplayHelper.PlaybackDemoHeader.GameNetworkProtocolVersion; }
+
 	FString GetDemoPath() const;
 
 private:
diff --git a/Engine/Source/Runtime/Engine/Private/Pawn.cpp b/Engine/Source/Runtime/Engine/Private/Pawn.cpp
--- a/Engine/Source/Runtime/Engine/Private/Pawn.cpp
+++ b/Engine/Source/Runtime/Engine/Private/Pawn.cpp
@@ -30,6 +30,8 @@
 #include "Components/PawnNoiseEmitterComponent.h"
 #include "GameFramework/GameNetworkManager.h"
 #include "GameFramework/InputSettings.h"
+#include "Engine/DemoNetDriver.h"
+#include "Misc/NetworkVersion.h"
 
 DEFINE_LOG_CATEGORY(LogDamage);
 DEFINE_LOG_CATEGORY_STATIC(LogPawn, Warning, All);
@@ -272,8 +274,7 @@ float APawn::GetDefaultHalfHeight() const
 void APawn::SetRemoteViewPitch(float NewRemoteViewPitch)
 {
 	// Compress pitch to 1 byte
-	NewRemoteViewPitch = FRotator::ClampAxis(NewRemoteViewPitch);
-	RemoteViewPitch = (uint8)(NewRemoteViewPitch * 255.f/360.f);
+	RemoteViewPitch = FRotator::CompressAxisToByte(NewRemoteViewPitch);
 }
 
 
@@ -944,8 +945,18 @@ FRotator APawn::GetBaseAimRotation() const
 		else
 		{
 			// Else use the RemoteViewPitch
-			POVRot.Pitch = RemoteViewPitch;
-			POVRot.Pitch = POVRot.Pitch * 360.0f / 255.0f;
+			const UWorld* World = GetWorld();
+			const UDemoNetDriver* DemoNetDriver = World ? World->GetDemoNetDriver() : nullptr;
+
+			if (DemoNetDriver && DemoNetDriver->IsPlaying() && (DemoNetDriver->GetPlaybackEngineNetworkProtocolVersion() < EEngineNetworkVersionHistory::HISTORY_PAWN_REMOTEVIEWPITCH))
+			{
+				POVRot.Pitch = RemoteViewPitch;
+				POVRot.Pitch = POVRot.Pitch * 360.0f / 255.0f;
+			}
+			else
+			{
+				POVRot.Pitch = FRotator::DecompressAxisFromByte(RemoteViewPitch);
+			}
 		}
 	}
 
