More networking LLM tag adjustments

diff --git a/Engine/Source/Runtime/Engine/Private/DataChannel.cpp b/Engine/Source/Runtime/Engine/Private/DataChannel.cpp
--- a/Engine/Source/Runtime/Engine/Private/DataChannel.cpp
+++ b/Engine/Source/Runtime/Engine/Private/DataChannel.cpp
@@ -3059,6 +3059,7 @@ int32 GNumReplicateActorCalls = 0;
 
 int64 UActorChannel::ReplicateActor()
 {
+	LLM_SCOPE_BYTAG(NetChannel);
 	SCOPE_CYCLE_COUNTER(STAT_NetReplicateActorTime);
 
 	check(Actor);
diff --git a/Engine/Source/Runtime/Engine/Private/DataReplication.cpp b/Engine/Source/Runtime/Engine/Private/DataReplication.cpp
--- a/Engine/Source/Runtime/Engine/Private/DataReplication.cpp
+++ b/Engine/Source/Runtime/Engine/Private/DataReplication.cpp
@@ -1729,11 +1729,13 @@ bool FObjectReplicator::CanSkipUpdate(FReplicationFlags RepFlags)
 
 bool FObjectReplicator::ReplicateProperties(FOutBunch& Bunch, FReplicationFlags RepFlags, FNetBitWriter& Writer)
 {
+	LLM_SCOPE_BYTAG(NetObjReplicator);
 	return ReplicateProperties_r(Bunch, RepFlags, Writer);
 }
 
 bool FObjectReplicator::ReplicateProperties(FOutBunch& Bunch, FReplicationFlags RepFlags)
 {
+	LLM_SCOPE_BYTAG(NetObjReplicator);
 	FNetBitWriter Writer(Bunch.PackageMap, 8192);
 	return ReplicateProperties_r(Bunch, RepFlags, Writer);
 }
diff --git a/Engine/Source/Runtime/Engine/Private/NetConnection.cpp b/Engine/Source/Runtime/Engine/Private/NetConnection.cpp
--- a/Engine/Source/Runtime/Engine/Private/NetConnection.cpp
+++ b/Engine/Source/Runtime/Engine/Private/NetConnection.cpp
@@ -3725,6 +3725,7 @@ int32 UNetConnection::GetFreeChannelIndex(const FName& ChName) const
 
 UChannel* UNetConnection::CreateChannelByName(const FName& ChName, EChannelCreateFlags CreateFlags, int32 ChIndex)
 {
+	LLM_SCOPE_BYTAG(NetConnection);
 	QUICK_SCOPE_CYCLE_COUNTER(STAT_NetConnection_CreateChannelByName);
 
 	check(Driver->IsKnownChannelName(ChName));
diff --git a/Engine/Source/Runtime/Engine/Private/NetDriver.cpp b/Engine/Source/Runtime/Engine/Private/NetDriver.cpp
--- a/Engine/Source/Runtime/Engine/Private/NetDriver.cpp
+++ b/Engine/Source/Runtime/Engine/Private/NetDriver.cpp
@@ -4810,6 +4810,7 @@ int32 UNetDriver::ServerReplicateActors_ProcessPrioritizedActors( UNetConnection
 
 int64 UNetDriver::SendDestructionInfo(UNetConnection* Connection, FActorDestructionInfo* DestructionInfo)
 {
+	LLM_SCOPE_BYTAG(NetDriver);
 	QUICK_SCOPE_CYCLE_COUNTER(STAT_NetSendDestructionInfo);
 	CSV_SCOPED_TIMING_STAT_EXCLUSIVE(SendDestructionInfo);
 
