Fix for a skeletal mesh causing a crash in cooked content

#tests Generated a cooked build, loaded into associated environment and crashed activating the skeletal mesh, proceeded to use the minidump from the crash to find the line of code that would break. After adding the code changes I proceeded to through the testing steps again, there were some minor visual errors but the crash did not occur.

diff --git a/Engine/Source/Runtime/Engine/Private/Components/SkinnedMeshComponent.cpp b/Engine/Source/Runtime/Engine/Private/Components/SkinnedMeshComponent.cpp
--- a/Engine/Source/Runtime/Engine/Private/Components/SkinnedMeshComponent.cpp
+++ b/Engine/Source/Runtime/Engine/Private/Components/SkinnedMeshComponent.cpp
@@ -4228,7 +4228,13 @@ FVector3f GetTypedSkinnedVertexPosition(
 	for (int32 InfluenceIndex = 0; InfluenceIndex < MaxBoneInfluences; InfluenceIndex++)
 #endif
 	{
-		const int32 MeshBoneIndex = Section.BoneMap[SkinWeightVertexBuffer.GetBoneIndex(BufferVertIndex, InfluenceIndex)];
+		const int32 BoneMapIndex = SkinWeightVertexBuffer.GetBoneIndex(BufferVertIndex, InfluenceIndex);
+		if (!ensureMsgf(Section.BoneMap.IsValidIndex(BoneMapIndex), TEXT("%s has attempted to access a BoneMap of size %i with an invalid index of %i in GetTypedSkinnedVertexPosition()"), *SkinnedComp->SkeletalMesh->GetFullName(), Section.BoneMap.Num(), BoneMapIndex))
+		{
+			continue;
+		}
+
+		const int32 MeshBoneIndex = Section.BoneMap[BoneMapIndex];
 		int32 TransformBoneIndex = MeshBoneIndex;
 
 		if (MasterPoseComponentInst)
@@ -4242,8 +4248,11 @@ FVector3f GetTypedSkinnedVertexPosition(
 		{
 			if (bCachedMatrices)
 			{
-				const FMatrix44f& RefToLocal = RefToLocals[MeshBoneIndex];
-				SkinnedPos += RefToLocal.TransformPosition(PositionVertexBuffer.VertexPosition(BufferVertIndex)) * Weight;
+				if (ensureMsgf(RefToLocals.IsValidIndex(MeshBoneIndex), TEXT("%s has attempted to access a RefToLocals of size %i with an invalid index of %i in GetTypedSkinnedVertexPosition()"), *SkinnedComp->SkeletalMesh->GetFullName(), RefToLocals.Num(), MeshBoneIndex))
+				{
+					const FMatrix44f& RefToLocal = RefToLocals[MeshBoneIndex];
+					SkinnedPos += RefToLocal.TransformPosition(PositionVertexBuffer.VertexPosition(BufferVertIndex)) * Weight;
+				}
 			}
 			else
 			{
