Merging CL#18840220

[Backout] - CL18778522 and 18817723
Original CL Desc
-----------------------------------------------------------------
- No longer create a sound proxy if the app can't render (play) audio
#p4v-preflight-copy 18774992

diff --git a/Engine/Source/Runtime/Engine/Private/SoundWave.cpp b/Engine/Source/Runtime/Engine/Private/SoundWave.cpp
--- a/Engine/Source/Runtime/Engine/Private/SoundWave.cpp
+++ b/Engine/Source/Runtime/Engine/Private/SoundWave.cpp
@@ -885,7 +885,7 @@ void USoundWave::Serialize( FArchive& Ar )
 		const bool bShouldLoadChunks = bCooked || bBuiltStreamedAudio;
 
 		// If stream caching is enabled, here we determine if we should retain or prime this wave on load.
-		if (bShouldStreamSound && bShouldLoadChunks && FPlatformCompressionUtilities::IsCurrentPlatformUsingStreamCaching() && FApp::CanEverRenderAudio())
+		if (bShouldStreamSound && bShouldLoadChunks && FPlatformCompressionUtilities::IsCurrentPlatformUsingStreamCaching())
 		{
 			ESoundWaveLoadingBehavior CurrentLoadingBehavior = GetLoadingBehavior(false);
 
@@ -1444,8 +1444,7 @@ void USoundWave::PostLoad()
 	// If our loading behavior is defined by a sound class, we need to update whether this sound wave actually needs to retain its audio data or not.
 	ESoundWaveLoadingBehavior ActualLoadingBehavior = GetLoadingBehavior();
 
-	// no need for a proxy if we can't render audio!
-	if (!InternalProxy.IsValid() && ActualLoadingBehavior != ESoundWaveLoadingBehavior::ForceInline && FApp::CanEverRenderAudio())
+	if (!InternalProxy.IsValid() && ActualLoadingBehavior != ESoundWaveLoadingBehavior::ForceInline)
 	{
 		InternalProxy = CreateSoundWaveProxy();
 	}
@@ -1464,7 +1463,6 @@ void USoundWave::PostLoad()
 
 			const bool bHasMultipleChunks = GetNumChunks() > 1;
 			bool bShouldPrime = (ActualLoadingBehavior == ESoundWaveLoadingBehavior::PrimeOnLoad);
-			bShouldPrime |= InternalProxy.IsValid();
 			bShouldPrime |= (GIsEditor && (ActualLoadingBehavior == ESoundWaveLoadingBehavior::RetainOnLoad)); // treat this scenario like PrimeOnLoad
 			
 			if (bShouldPrime && bHasMultipleChunks)
@@ -2204,6 +2202,8 @@ void USoundWave::BakeEnvelopeAnalysis()
 
 void USoundWave::PostEditChangeProperty(FPropertyChangedEvent& PropertyChangedEvent)
 {
+	// maxtodo: test this, stale data? looping? channel count?
+
 	Super::PostEditChangeProperty(PropertyChangedEvent);
 	check(SoundWaveDataPtr);
 
@@ -2921,7 +2921,7 @@ bool USoundWave::ShouldUseStreamCaching() const
 
 TArrayView<const uint8> USoundWave::GetZerothChunk(bool bForImmediatePlayback)
 {
-	if(IsTemplate() || IsRunningDedicatedServer() || !FApp::CanEverRenderAudio())
+	if(IsTemplate() || IsRunningDedicatedServer())
 	{
 		return TArrayView<const uint8>();
 	}
@@ -3065,7 +3065,7 @@ void USoundWave::RemovePlayingSource(const FSoundWaveClientPtr& Source)
 
 void USoundWave::UpdatePlatformData()
 {
-	if (IsStreaming(nullptr) && FApp::CanEverRenderAudio())
+	if (IsStreaming(nullptr))
 	{
 		if (!InternalProxy.IsValid())
 		{
@@ -3423,14 +3423,14 @@ void USoundWave::GetHandleForChunkOfAudio(TFunction<void(FAudioChunkHandle&&)> O
 {
 	ENamedThreads::Type ThreadToDispatchCallbackOn = (DispatchToGameThreadOnChunkRequestCVar != 0) ? ENamedThreads::GameThread : ENamedThreads::AnyThread;
 
-	if (!InternalProxy.IsValid() && FApp::CanEverRenderAudio())
+	if (!InternalProxy.IsValid())
 	{
 		InternalProxy = CreateSoundWaveProxy();
 	}
 
 	// if we are requesting a chunk that is out of bounds,
-	// or we shouldn't be interacting w/ the cache, early exit.
-	if (FApp::CanEverRenderAudio() || (ChunkIndex >= static_cast<int32>(GetNumChunks())))
+	// early exit.
+	if (ChunkIndex >= static_cast<int32>(GetNumChunks()))
 	{
 		FAudioChunkHandle EmptyChunkHandle;
 		OnLoadCompleted(MoveTemp(EmptyChunkHandle));
@@ -3493,7 +3493,7 @@ void USoundWave::RetainCompressedAudio(bool bForceSync /*= false*/)
 
 	// Since the zeroth chunk is always inlined and stored in memory,
 	// early exit if we only have one chunk.
-	if (GIsEditor || IsTemplate() || IsRunningDedicatedServer() || !IsStreaming() || !FApp::CanEverRenderAudio() || DisableRetainingCVar || GetNumChunks() <= 1)
+	if (GIsEditor || IsTemplate() || IsRunningDedicatedServer() || !IsStreaming() || DisableRetainingCVar || GetNumChunks() <= 1)
 	{
 		return;
 	}
