Fix LevelInstanceTransform deserializing doubles instead of floats on existing data
Fix MedievalGame crashing on boot

diff --git a/Engine/Source/Runtime/Engine/Private/WorldPartition/LevelInstance/LevelInstanceActorDesc.cpp b/Engine/Source/Runtime/Engine/Private/WorldPartition/LevelInstance/LevelInstanceActorDesc.cpp
--- a/Engine/Source/Runtime/Engine/Private/WorldPartition/LevelInstance/LevelInstanceActorDesc.cpp
+++ b/Engine/Source/Runtime/Engine/Private/WorldPartition/LevelInstance/LevelInstanceActorDesc.cpp
@@ -129,7 +129,18 @@ void FLevelInstanceActorDesc::Serialize(FArchive& Ar)
 	Ar.UsingCustomVersion(FUE5ReleaseStreamObjectVersion::GUID);
 	Ar.UsingCustomVersion(FUE5MainStreamObjectVersion::GUID);
 
-	Ar << LevelPackage << LevelInstanceTransform;
+	Ar << LevelPackage;
+
+	if (Ar.CustomVer(FUE5ReleaseStreamObjectVersion::GUID) < FUE5ReleaseStreamObjectVersion::LargeWorldCoordinates)
+	{
+		FTransform3f LevelInstanceTransformFlt;
+		Ar << LevelInstanceTransformFlt;
+		LevelInstanceTransform = FTransform(LevelInstanceTransformFlt);
+	}
+	else
+	{
+		Ar << LevelInstanceTransform;
+	}
 
 	if (Ar.CustomVer(FUE5ReleaseStreamObjectVersion::GUID) >= FUE5ReleaseStreamObjectVersion::LevelInstanceSerializeRuntimeBehavior)
 	{
