Fix for forced property compares caused by scoped role downgrade using SetReplicates

diff --git a/Engine/Source/Runtime/Engine/Private/DataChannel.cpp b/Engine/Source/Runtime/Engine/Private/DataChannel.cpp
--- a/Engine/Source/Runtime/Engine/Private/DataChannel.cpp
+++ b/Engine/Source/Runtime/Engine/Private/DataChannel.cpp
@@ -3037,15 +3037,7 @@ public:
 	~FScopedRoleDowngrade()
 	{
 		// Upgrade role back to autonomous proxy if needed
-		if ( Actor->GetRemoteRole() != ActualRemoteRole )
-		{
-			Actor->SetReplicates( ActualRemoteRole != ROLE_None );
-
-			if ( ActualRemoteRole == ROLE_AutonomousProxy )
-			{
-				Actor->SetAutonomousProxy( true, false );
-			}
-		}
+		Actor->SetAutonomousProxy(ActualRemoteRole == ROLE_AutonomousProxy, false);
 	}
 
 private:
