Adding functionality to compile remaining shaders for materials.

When running with ODSC in the editor Materials can have incomplete or partial shader maps.  In other words some shaders _might_ be missing.  If a render system requires all shaders to be present and can't use ODSC functionality to compile the necessary shaders we need a mechanism for compiling just the missing shaders.

- Added UMaterialInterface::CacheShaders() which is implemented for UMatetrial and UMaterialInstance.  The call is just forwarded along to CacheResourceShadersForRendering and InitStaticPermutation respectively.
- You can call CacheResourceShadersForRendering multiple times (even with a completed shader map) and it will not recache uniform expressions.  If the Resource's shader map is complete the function is a no-op.
- FObjectCacheContext can accept a UWorld and will only return primitives that are in that world.
- Added a function SubmitRemainingJobsForWorld(UWorld* World, EMaterialShaderPrecompileMode CompileMode) that will iterate all materials used by primitives in the world and compile the missing shaders.
- SubmitRemainingJobsForWorld() is used in various scenarios in the editor to ensure _all_ shaders are compiled for all materials before proceeding.
- Explicitly try to compile UI and PP materials.  Since we don't know which ones are used in the given world we just have to iterate all loaded ones.  This is a potential area of optimization.
- If there are already complete shader maps (if you are not running with the shader job cache), SubmitRemainingJobsForWorld is fast and results in a no-op for each material.

diff --git a/Engine/Plugins/Experimental/AutomationUtils/Source/AutomationUtils/Private/AutomationUtilsBlueprintLibrary.cpp b/Engine/Plugins/Experimental/AutomationUtils/Source/AutomationUtils/Private/AutomationUtilsBlueprintLibrary.cpp
--- a/Engine/Plugins/Experimental/AutomationUtils/Source/AutomationUtils/Private/AutomationUtilsBlueprintLibrary.cpp
+++ b/Engine/Plugins/Experimental/AutomationUtils/Source/AutomationUtils/Private/AutomationUtilsBlueprintLibrary.cpp
@@ -23,6 +23,7 @@
 #include "Containers/Ticker.h"
 #include "Engine/GameEngine.h"
 #include "Stats/Stats.h"
+#include "Materials/MaterialInterface.h"
 
 //Private Helper Class Definitions
 class FAutomationUtilsGameplayViewExtension : public FSceneViewExtensionBase
@@ -347,22 +348,25 @@ void UAutomationUtilsBlueprintLibrary::TakeGameplayAutomationScreenshot(const FS
 {
 	FlushAsyncLoading();
 
-	//Finish Loading Before Screenshot
-	if (!FPlatformProperties::RequiresCookedData())
-	{
-		//Finish Compiling all shaders and other async assets
-		FAssetCompilingManager::Get().FinishAllCompilation();
-	}
-
+	UWorld* CurrentWorld{ nullptr };
 	// Make sure we finish all level streaming
 	if (UGameEngine* GameEngine = Cast<UGameEngine>(GEngine))
 	{
 		if (UWorld* GameWorld = GameEngine->GetGameWorld())
 		{
+			CurrentWorld = GameWorld;
 			GameWorld->FlushLevelStreaming(EFlushLevelStreamingType::Full);
 		}
 	}
 
+	//Finish Loading Before Screenshot
+	if (!FPlatformProperties::RequiresCookedData())
+	{
+		UMaterialInterface::SubmitRemainingJobsForWorld(CurrentWorld);
+		//Finish Compiling all shaders and other async assets
+		FAssetCompilingManager::Get().FinishAllCompilation();
+	}
+
 	//Stream in everything
 	IStreamingManager::Get().StreamAllResources(0.0f);
 
diff --git a/Engine/Source/Developer/FunctionalTesting/Private/AutomationBlueprintFunctionLibrary.cpp b/Engine/Source/Developer/FunctionalTesting/Private/AutomationBlueprintFunctionLibrary.cpp
--- a/Engine/Source/Developer/FunctionalTesting/Private/AutomationBlueprintFunctionLibrary.cpp
+++ b/Engine/Source/Developer/FunctionalTesting/Private/AutomationBlueprintFunctionLibrary.cpp
@@ -49,6 +49,8 @@
 #include "IImageWrapperModule.h"
 #include "ImageWrapperHelper.h"
 #include "Misc/FileHelper.h"
+#include "Materials/MaterialInterface.h"
+#include "AssetCompilingManager.h"
 
 #if WITH_EDITOR
 #include "SLevelViewport.h"
@@ -676,24 +678,27 @@ UAutomationBlueprintFunctionLibrary::UAutomationBlueprintFunctionLibrary(const c
 
 void UAutomationBlueprintFunctionLibrary::FinishLoadingBeforeScreenshot()
 {
-	// Finish compiling the shaders if the platform doesn't require cooked data.
-	if (!FPlatformProperties::RequiresCookedData())
-	{
-		GShaderCompilingManager->FinishAllCompilation();
-		FModuleManager::GetModuleChecked<IAutomationControllerModule>("AutomationController").GetAutomationController()->ResetAutomationTestTimeout(TEXT("shader compilation"));
-	}
-
 	FlushAsyncLoading();
 
+	UWorld* CurrentWorld{ nullptr };
 	// Make sure we finish all level streaming
 	if (UGameEngine* GameEngine = Cast<UGameEngine>(GEngine))
 	{
 		if (UWorld* GameWorld = GameEngine->GetGameWorld())
 		{
+			CurrentWorld = GameWorld;
 			GameWorld->FlushLevelStreaming(EFlushLevelStreamingType::Full);
 		}
 	}
 
+	// Finish compiling the shaders if the platform doesn't require cooked data.
+	if (!FPlatformProperties::RequiresCookedData())
+	{
+		UMaterialInterface::SubmitRemainingJobsForWorld(CurrentWorld);
+		FAssetCompilingManager::Get().FinishAllCompilation();
+		FModuleManager::GetModuleChecked<IAutomationControllerModule>("AutomationController").GetAutomationController()->ResetAutomationTestTimeout(TEXT("shader compilation"));
+	}
+
 	// Force all mip maps to load before taking the screenshot.
 	UTexture::ForceUpdateTextureStreaming();
 
diff --git a/Engine/Source/Editor/Blutility/Private/AsyncCaptureScene.cpp b/Engine/Source/Editor/Blutility/Private/AsyncCaptureScene.cpp
--- a/Engine/Source/Editor/Blutility/Private/AsyncCaptureScene.cpp
+++ b/Engine/Source/Editor/Blutility/Private/AsyncCaptureScene.cpp
@@ -16,6 +16,8 @@
 #include "ShaderCompiler.h"
 #include "IAutomationControllerManager.h"
 #include "IAutomationControllerModule.h"
+#include "Materials/MaterialInterface.h"
+#include "AssetCompilingManager.h"
 
 //----------------------------------------------------------------------//
 // UAsyncCaptureScene
@@ -96,25 +98,28 @@ void UAsyncCaptureScene::NotifyComplete(UTextureRenderTarget2D* InTexture)
 
 void UAsyncCaptureScene::FinishLoadingBeforeScreenshot()
 {
-	// Finish compiling the shaders if the platform doesn't require cooked data.
-	if (!FPlatformProperties::RequiresCookedData())
-	{
-		GShaderCompilingManager->FinishAllCompilation();
-		IAutomationControllerModule& AutomationController = IAutomationControllerModule::Get();
-		AutomationController.GetAutomationController()->ResetAutomationTestTimeout(TEXT("shader compilation"));
-	}
-
 	FlushAsyncLoading();
 
+	UWorld* CurrentWorld{ nullptr };
 	// Make sure we finish all level streaming
 	if (UGameEngine* GameEngine = Cast<UGameEngine>(GEngine))
 	{
 		if (UWorld* GameWorld = GameEngine->GetGameWorld())
 		{
+			CurrentWorld = GameWorld;
 			GameWorld->FlushLevelStreaming(EFlushLevelStreamingType::Full);
 		}
 	}
 
+	// Finish compiling the shaders if the platform doesn't require cooked data.
+	if (!FPlatformProperties::RequiresCookedData())
+	{
+		UMaterialInterface::SubmitRemainingJobsForWorld(CurrentWorld);
+		FAssetCompilingManager::Get().FinishAllCompilation();
+		IAutomationControllerModule& AutomationController = IAutomationControllerModule::Get();
+		AutomationController.GetAutomationController()->ResetAutomationTestTimeout(TEXT("shader compilation"));
+	}
+
 	// Force all mip maps to load before taking the screenshot.
 	UTexture::ForceUpdateTextureStreaming();
 
diff --git a/Engine/Source/Editor/UnrealEd/Private/EditorEngine.cpp b/Engine/Source/Editor/UnrealEd/Private/EditorEngine.cpp
--- a/Engine/Source/Editor/UnrealEd/Private/EditorEngine.cpp
+++ b/Engine/Source/Editor/UnrealEd/Private/EditorEngine.cpp
@@ -99,6 +99,7 @@
 #include "Interfaces/ITargetPlatform.h"
 #include "Misc/AutomationTest.h"
 #include "ActorFolder.h"
+#include "Materials/MaterialInterface.h"
 
 // needed for the RemotePropagator
 #include "AudioDevice.h"
@@ -243,6 +244,7 @@
 #include "DerivedDataBuildLocalExecutor.h"
 #include "DerivedDataBuildRemoteExecutor.h"
 #include "DerivedDataBuildWorkers.h"
+#include "AssetCompilingManager.h"
 
 #if WITH_CHAOS
 #include "ChaosSolversModule.h"
@@ -3739,9 +3741,10 @@ void UEditorEngine::BuildReflectionCaptures(UWorld* World)
 	GWarn->StatusUpdate(0, 1, StatusText);
 
 	// Wait for shader compiling to finish so we don't capture the default material
-	if (GShaderCompilingManager != NULL)
+	if (GShaderCompilingManager != nullptr)
 	{
-		GShaderCompilingManager->FinishAllCompilation();
+		UMaterialInterface::SubmitRemainingJobsForWorld(World);
+		FAssetCompilingManager::Get().FinishAllCompilation();
 	}
 
 	// Process any outstanding captures before we start operating on scenarios
diff --git a/Engine/Source/Runtime/Engine/Classes/Materials/Material.h b/Engine/Source/Runtime/Engine/Classes/Materials/Material.h
--- a/Engine/Source/Runtime/Engine/Classes/Materials/Material.h
+++ b/Engine/Source/Runtime/Engine/Classes/Materials/Material.h
@@ -1484,6 +1484,8 @@ private:
 	/** Caches shader maps for an array of material resources. */
 	void CacheShadersForResources(EShaderPlatform ShaderPlatform, const TArray<FMaterialResource*>& ResourcesToCache, EMaterialShaderPrecompileMode PrecompileMode = EMaterialShaderPrecompileMode::Default, const ITargetPlatform* TargetPlatform = nullptr);
 
+	ENGINE_API virtual void CacheShaders(EMaterialShaderPrecompileMode CompileMode) override;
+
 #if WITH_EDITOR
 	/**
 	 * If there is some texture reference used by a TextureProperty node in any expressions, this function
diff --git a/Engine/Source/Runtime/Engine/Classes/Materials/MaterialInstance.h b/Engine/Source/Runtime/Engine/Classes/Materials/MaterialInstance.h
--- a/Engine/Source/Runtime/Engine/Classes/Materials/MaterialInstance.h
+++ b/Engine/Source/Runtime/Engine/Classes/Materials/MaterialInstance.h
@@ -913,6 +913,8 @@ protected:
 	/** Caches shader maps for an array of material resources. */
 	void CacheShadersForResources(EShaderPlatform ShaderPlatform, const TArray<FMaterialResource*>& ResourcesToCache, EMaterialShaderPrecompileMode PrecompileMode = EMaterialShaderPrecompileMode::Default, const ITargetPlatform* TargetPlatform = nullptr);
 
+	ENGINE_API virtual void CacheShaders(EMaterialShaderPrecompileMode CompileMode) override;
+
 	/** 
 	 * Copies over material instance parameters from the base material given a material interface.
 	 * This is a slow operation that is needed for the editor.
diff --git a/Engine/Source/Runtime/Engine/Classes/Materials/MaterialInterface.h b/Engine/Source/Runtime/Engine/Classes/Materials/MaterialInterface.h
--- a/Engine/Source/Runtime/Engine/Classes/Materials/MaterialInterface.h
+++ b/Engine/Source/Runtime/Engine/Classes/Materials/MaterialInterface.h
@@ -766,6 +766,28 @@ public:
 	 */
 	ENGINE_API static void RecacheAllMaterialUniformExpressions(bool bRecreateUniformBuffer);
 
+	/**
+	 * @brief Submits shaders to be compiled for all the materials in the world.
+	 *
+	 * This function will submit any remaining shaders to be compiled for all the materials in the passed in world.  By default
+	 * these shader compilation jobs will be compiled in the background so if you need the results immediately you can call
+	 * FinishAllCompilation() to block on the results.
+	 *
+	 * If the world is nullptr this will submit remaining shaders to be compiled for all the loaded materials.
+	 *
+	 * @code
+	 * GShaderCompilingManager->SubmitRemainingJobsForWorld(World);
+	 * GShaderCompilingManager->FinishAllCompilation();
+	 * @endcode
+	 *
+	 * @param World Only compile shaders for materials that are used on primitives in this world.
+	 * @param CompileMode Controls whether or not we block on the shader compile results.
+	 *
+	 * @note This will only submit shader compile jobs for missing shaders on each material.  Calling this multiple times on the same world
+	 * will result in a no-op.
+	 */
+	ENGINE_API static void SubmitRemainingJobsForWorld(UWorld* World, EMaterialShaderPrecompileMode CompileMode = EMaterialShaderPrecompileMode::Default);
+
 	/**
 	 * Re-caches uniform expressions for this material interface                   
 	 * Set bRecreateUniformBuffer to true if uniform buffer layout will change (e.g. FMaterial is being recompiled).
@@ -775,6 +797,15 @@ public:
 	 */
 	virtual void RecacheUniformExpressions(bool bRecreateUniformBuffer) const {}
 
+	/** @brief Submits remaining shaders for recompilation.
+	*
+	* This function will submit any remaining shaders to be compiled for the given material.  By default
+	* these shader compilation jobs will be compiled in the background.
+	*
+	* @param CompileMode Controls whether or not we block on the shader compile results.
+	*/
+	ENGINE_API virtual void CacheShaders(EMaterialShaderPrecompileMode CompileMode = EMaterialShaderPrecompileMode::Default) {}
+
 #if WITH_EDITOR
 	/** Clears the shader cache and recompiles the shader for rendering. */
 	ENGINE_API virtual void ForceRecompileForRendering() {}
diff --git a/Engine/Source/Runtime/Engine/Private/Materials/Material.cpp b/Engine/Source/Runtime/Engine/Private/Materials/Material.cpp
--- a/Engine/Source/Runtime/Engine/Private/Materials/Material.cpp
+++ b/Engine/Source/Runtime/Engine/Private/Materials/Material.cpp
@@ -2136,6 +2136,8 @@ void UMaterial::CacheResourceShadersForRendering(bool bRegenerateId, EMaterialSh
 
 	if (FApp::CanEverRender())
 	{
+		bool bCachedShaders{ false };
+
 		const EMaterialQualityLevel::Type ActiveQualityLevel = GetCachedScalabilityCVars().MaterialQualityLevel;
 		uint32 FeatureLevelsToCompile = GetFeatureLevelsToCompileForRendering();
 
@@ -2151,6 +2153,12 @@ void UMaterial::CacheResourceShadersForRendering(bool bRegenerateId, EMaterialSh
 			FMaterialResource* CurrentResource = FindOrCreateMaterialResource(MaterialResources, this, nullptr, FeatureLevel, ActiveQualityLevel);
 			check(CurrentResource);
 
+			// If this resource is complete there is nothing to do here.
+			if (CurrentResource->IsGameThreadShaderMapComplete())
+			{
+				continue;
+			}
+
 #if STORE_ONLY_ACTIVE_SHADERMAPS
 			if (CurrentResource && !CurrentResource->GetGameThreadShaderMap())
 			{
@@ -2171,6 +2179,7 @@ void UMaterial::CacheResourceShadersForRendering(bool bRegenerateId, EMaterialSh
 			ResourcesToCache.Reset();
 			ResourcesToCache.Add(CurrentResource);
 			CacheShadersForResources(ShaderPlatform, ResourcesToCache, PrecompileMode);
+			bCachedShaders = true;
 		}
 
 		FString AdditionalFormatToCache = GCompileMaterialsForShaderFormatCVar->GetString();
@@ -2190,7 +2199,10 @@ void UMaterial::CacheResourceShadersForRendering(bool bRegenerateId, EMaterialSh
 			}
 		}
 
-		RecacheUniformExpressions(true);
+		if (bCachedShaders)
+		{
+			RecacheUniformExpressions(true);
+		}
 	}
 
 	FMaterial::DeferredDeleteArray(ResourcesToFree);
@@ -2289,6 +2301,11 @@ void UMaterial::CacheShadersForResources(EShaderPlatform ShaderPlatform, const T
 	}
 }
 
+void UMaterial::CacheShaders(EMaterialShaderPrecompileMode CompileMode)
+{
+	CacheResourceShadersForRendering(false, CompileMode);
+}
+
 void UMaterial::ReleaseResourcesAndMutateDDCKey(const FGuid& TransformationId)
 {
 	if (TransformationId.IsValid())
diff --git a/Engine/Source/Runtime/Engine/Private/Materials/MaterialInstance.cpp b/Engine/Source/Runtime/Engine/Private/Materials/MaterialInstance.cpp
--- a/Engine/Source/Runtime/Engine/Private/Materials/MaterialInstance.cpp
+++ b/Engine/Source/Runtime/Engine/Private/Materials/MaterialInstance.cpp
@@ -1985,6 +1985,8 @@ void UMaterialInstance::CacheResourceShadersForRendering(EMaterialShaderPrecompi
 		check(IsA(UMaterialInstanceConstant::StaticClass()));
 		UMaterial* BaseMaterial = GetMaterial();
 
+		bool bCachedShaders{ false };
+
 		uint32 FeatureLevelsToCompile = GetFeatureLevelsToCompileForRendering();
 		const EMaterialQualityLevel::Type ActiveQualityLevel = GetCachedScalabilityCVars().MaterialQualityLevel;
 
@@ -2000,6 +2002,12 @@ void UMaterialInstance::CacheResourceShadersForRendering(EMaterialShaderPrecompi
 			FMaterialResource* CurrentResource = FindOrCreateMaterialResource(StaticPermutationMaterialResources, BaseMaterial, this, FeatureLevel, ActiveQualityLevel);
 			check(CurrentResource);
 
+			// If this resource is complete there is nothing to do here.
+			if (CurrentResource->IsGameThreadShaderMapComplete())
+			{
+				continue;
+			}
+
 #if STORE_ONLY_ACTIVE_SHADERMAPS
 			if (!CurrentResource->GetGameThreadShaderMap())
 			{
@@ -2020,10 +2028,14 @@ void UMaterialInstance::CacheResourceShadersForRendering(EMaterialShaderPrecompi
 			ResourcesToCache.Reset();
 			ResourcesToCache.Add(CurrentResource);
 			CacheShadersForResources(ShaderPlatform, ResourcesToCache, PrecompileMode);
+			bCachedShaders = true;
 		}
-	}
 
-	RecacheUniformExpressions(true);
+		if (bCachedShaders)
+		{
+			RecacheUniformExpressions(true);
+		}
+	}
 
 	InitResources();
 }
@@ -2124,6 +2136,11 @@ void UMaterialInstance::CacheShadersForResources(EShaderPlatform ShaderPlatform,
 	}
 }
 
+void UMaterialInstance::CacheShaders(EMaterialShaderPrecompileMode CompileMode)
+{
+	InitStaticPermutation(CompileMode);
+}
+
 bool UMaterialInstance::GetMaterialLayers(FMaterialLayersFunctions& OutLayers, TMicRecursionGuard RecursionGuard) const
 {
 	if (StaticParameters.bHasMaterialLayers)
diff --git a/Engine/Source/Runtime/Engine/Private/Materials/MaterialInterface.cpp b/Engine/Source/Runtime/Engine/Private/Materials/MaterialInterface.cpp
--- a/Engine/Source/Runtime/Engine/Private/Materials/MaterialInterface.cpp
+++ b/Engine/Source/Runtime/Engine/Private/Materials/MaterialInterface.cpp
@@ -28,6 +28,7 @@
 #include "TextureCompiler.h"
 #include "MaterialShaderQualitySettings.h"
 #include "ShaderPlatformQualitySettings.h"
+#include "ObjectCacheContext.h"
 
 #if WITH_EDITOR
 #include "ObjectCacheEventSink.h"
@@ -445,6 +446,50 @@ void UMaterialInterface::RecacheAllMaterialUniformExpressions(bool bRecreateUnif
 	}
 }
 
+void UMaterialInterface::SubmitRemainingJobsForWorld(UWorld* World, EMaterialShaderPrecompileMode CompileMode)
+{
+	TRACE_CPUPROFILER_EVENT_SCOPE(UMaterialInterface::SubmitRemainingJobsForWorld);
+
+	TSet<UMaterialInterface*> MaterialsToCache;
+	FObjectCacheContextScope ObjectCacheScope;
+
+	for (UPrimitiveComponent* PrimitiveComponent : ObjectCacheScope.GetContext().GetPrimitiveComponents())
+	{
+		if (World && !World->ContainsActor(PrimitiveComponent->GetOwner()))
+		{
+			continue;
+		}
+
+		if (PrimitiveComponent->IsRenderStateCreated())
+		{
+			TObjectCacheIterator<UMaterialInterface> UsedMaterials = ObjectCacheScope.GetContext().GetUsedMaterials(PrimitiveComponent);
+			for (UMaterialInterface* MaterialInterface : UsedMaterials)
+			{
+				if (MaterialInterface)
+				{
+					MaterialsToCache.Add(MaterialInterface);
+				}
+			}
+		}
+	}
+
+	// Add UI and PP Materials.
+	for (TObjectIterator<UMaterialInterface> It; It; ++It)
+	{
+		UMaterial* Material = It->GetMaterial();
+		if (Material &&
+			(Material->IsUIMaterial() || Material->IsPostProcessMaterial()))
+		{
+			MaterialsToCache.Add(*It);
+		}
+	}
+
+	for (UMaterialInterface* Material : MaterialsToCache)
+	{
+		Material->CacheShaders(CompileMode);
+	}
+}
+
 bool UMaterialInterface::IsReadyForFinishDestroy()
 {
 	bool bIsReady = Super::IsReadyForFinishDestroy();
diff --git a/Engine/Source/Runtime/Engine/Private/Materials/MaterialShared.cpp b/Engine/Source/Runtime/Engine/Private/Materials/MaterialShared.cpp
--- a/Engine/Source/Runtime/Engine/Private/Materials/MaterialShared.cpp
+++ b/Engine/Source/Runtime/Engine/Private/Materials/MaterialShared.cpp
@@ -2361,8 +2361,19 @@ bool FMaterial::CacheShaders(const FMaterialShaderMapId& ShaderMapId, EShaderPla
 	else
 	{
 #if WITH_EDITOR
-		// Clear outdated compile errors as we're not calling Translate on this path
-		CompileErrors.Empty();
+		// We have a shader map, the shader map is incomplete, and we've been asked to compile.
+		if (AllowShaderCompiling() && 
+			!IsGameThreadShaderMapComplete() && 
+			(PrecompileMode != EMaterialShaderPrecompileMode::None))
+		{
+			// Submit the remaining shaders in the map for compilation.
+			SubmitCompileJobs(EShaderCompileJobPriority::High);
+		}
+		else
+		{
+			// Clear outdated compile errors as we're not calling Translate on this path
+			CompileErrors.Empty();
+		}
 #endif // WITH_EDITOR
 	}
 
