Fix for crash when not skipping first update transition and using a conduit as the entry sate

diff --git a/Engine/Source/Runtime/Engine/Private/Animation/AnimNode_StateMachine.cpp b/Engine/Source/Runtime/Engine/Private/Animation/AnimNode_StateMachine.cpp
--- a/Engine/Source/Runtime/Engine/Private/Animation/AnimNode_StateMachine.cpp
+++ b/Engine/Source/Runtime/Engine/Private/Animation/AnimNode_StateMachine.cpp
@@ -450,32 +450,36 @@ void FAnimNode_StateMachine::Update_AnyThread(const FAnimationUpdateContext& Con
 				}
 			}
 
-			// Push the transition onto the stack
-			const FAnimationTransitionBetweenStates& ReferenceTransition = GetTransitionInfo(PotentialTransition.TransitionRule->TransitionIndex); //-V595
-			FAnimationActiveTransitionEntry* NewTransition = new (ActiveTransitionArray) FAnimationActiveTransitionEntry(NextState, ExistingWeightOfNextState, PreviousTransitionForNextState, PreviousState, ReferenceTransition, PotentialTransition);
-			if (NewTransition && PotentialTransition.TransitionRule)
+			// Don't add a transition if the previous state is a conduit (likely means we're finding the best entry state)
+			if (!IsAConduitState(PreviousState))
 			{
-				NewTransition->InitializeCustomGraphLinks(Context, *(PotentialTransition.TransitionRule));
-
-				if (ReferenceTransition.LogicType == ETransitionLogicType::TLT_Inertialization)
+				// Push the transition onto the stack
+				const FAnimationTransitionBetweenStates& ReferenceTransition = GetTransitionInfo(PotentialTransition.TransitionRule->TransitionIndex); //-V595
+				FAnimationActiveTransitionEntry* NewTransition = new (ActiveTransitionArray) FAnimationActiveTransitionEntry(NextState, ExistingWeightOfNextState, PreviousTransitionForNextState, PreviousState, ReferenceTransition, PotentialTransition);
+				if (NewTransition && PotentialTransition.TransitionRule)
 				{
-					UE::Anim::IInertializationRequester* InertializationRequester = Context.GetMessage<UE::Anim::IInertializationRequester>();
-					if (InertializationRequester)
-					{
-						InertializationRequester->RequestInertialization(ReferenceTransition.CrossfadeDuration, ReferenceTransition.BlendProfile);
-						InertializationRequester->AddDebugRecord(*Context.AnimInstanceProxy, Context.GetCurrentNodeId());
-					}
-					else
+					NewTransition->InitializeCustomGraphLinks(Context, *(PotentialTransition.TransitionRule));
+
+					if (ReferenceTransition.LogicType == ETransitionLogicType::TLT_Inertialization)
 					{
-						LogInertializationRequestError(Context, PreviousState, NextState);
+						UE::Anim::IInertializationRequester* InertializationRequester = Context.GetMessage<UE::Anim::IInertializationRequester>();
+						if (InertializationRequester)
+						{
+							InertializationRequester->RequestInertialization(ReferenceTransition.CrossfadeDuration, ReferenceTransition.BlendProfile);
+							InertializationRequester->AddDebugRecord(*Context.AnimInstanceProxy, Context.GetCurrentNodeId());
+						}
+						else
+						{
+							LogInertializationRequestError(Context, PreviousState, NextState);
+						}
 					}
-				}
 
-				NewTransition->SourceTransitionIndices = PotentialTransition.SourceTransitionIndices;
+					NewTransition->SourceTransitionIndices = PotentialTransition.SourceTransitionIndices;
 
-				if (!bFirstUpdate)
-				{
-					Context.AnimInstanceProxy->AddAnimNotifyFromGeneratedClass(NewTransition->StartNotify);
+					if (!bFirstUpdate)
+					{
+						Context.AnimInstanceProxy->AddAnimNotifyFromGeneratedClass(NewTransition->StartNotify);
+					}
 				}
 			}
 
