Sequencer: Encode recorded sound wave assets as bink audio so that they are seekable

Audio: Added a SetSoundAssetCompressionType to set the compression type

diff --git a/Engine/Plugins/Runtime/AudioCapture/Source/AudioCaptureEditor/Private/AudioRecordingManager.cpp b/Engine/Plugins/Runtime/AudioCapture/Source/AudioCaptureEditor/Private/AudioRecordingManager.cpp
--- a/Engine/Plugins/Runtime/AudioCapture/Source/AudioCaptureEditor/Private/AudioRecordingManager.cpp
+++ b/Engine/Plugins/Runtime/AudioCapture/Source/AudioCaptureEditor/Private/AudioRecordingManager.cpp
@@ -427,6 +427,7 @@ void FAudioRecordingManager::StopRecording(TArray<USoundWave*>& OutSoundWaves)
 					NewSoundWave->Duration = (float)(NumRecordedSamples / NumInputChannels) / WAVE_FILE_SAMPLERATE;
 					NewSoundWave->SetSampleRate(WAVE_FILE_SAMPLERATE);
 					NewSoundWave->NumChannels = NumChannelsToSerialize;
+					NewSoundWave->SetSoundAssetCompressionType(ESoundAssetCompressionType::BinkAudio);
 
 					if (bCreatedPackage)
 					{
diff --git a/Engine/Source/Runtime/Engine/Classes/Sound/SoundWave.h b/Engine/Source/Runtime/Engine/Classes/Sound/SoundWave.h
--- a/Engine/Source/Runtime/Engine/Classes/Sound/SoundWave.h
+++ b/Engine/Source/Runtime/Engine/Classes/Sound/SoundWave.h
@@ -417,8 +417,13 @@ public:
 	}
 
 	/** Returns the sound's asset compression type. */
+	UFUNCTION(BlueprintPure, Category = "Audio")
 	ESoundAssetCompressionType GetSoundAssetCompressionType() const;
 
+	/** Procedurally set the compression type. */
+	UFUNCTION(BlueprintCallable, Category = "Audio")
+	void SetSoundAssetCompressionType(ESoundAssetCompressionType InSoundAssetCompressionType);
+
 private:
 	// cached proxy
 	FSoundWaveProxyPtr InternalProxy{ nullptr };
@@ -548,6 +553,11 @@ private:
 	// Caches resolved loading behavior from the SoundClass graph. Must be called on the game thread.
 	void CacheInheritedLoadingBehavior() const;
 
+	// Called when we change any properties about the underlying audio asset
+#if WITH_EDITOR
+	void UpdateAsset();
+#endif
+
 public:
 
 	/** Set to true for programmatically generated audio. */
diff --git a/Engine/Source/Runtime/Engine/Private/SoundWave.cpp b/Engine/Source/Runtime/Engine/Private/SoundWave.cpp
--- a/Engine/Source/Runtime/Engine/Private/SoundWave.cpp
+++ b/Engine/Source/Runtime/Engine/Private/SoundWave.cpp
@@ -923,6 +923,15 @@ ESoundAssetCompressionType USoundWave::GetSoundAssetCompressionType() const
 	return SoundAssetCompressionType;
 }
 
+void USoundWave::SetSoundAssetCompressionType(ESoundAssetCompressionType InSoundAssetCompressionType)
+{
+#if WITH_EDITOR
+	SoundAssetCompressionType = InSoundAssetCompressionType;
+	SoundWaveDataPtr->bIsSeekable = IsSeekable();
+	UpdateAsset();
+#endif // #if WITH_EDITOR
+}
+
 float USoundWave::GetSubtitlePriority() const
 {
 	return SubtitlePriority;
@@ -2251,30 +2260,7 @@ void USoundWave::PostEditChangeProperty(FPropertyChangedEvent& PropertyChangedEv
 
 			if (Name == CompressionQualityFName || Name == SampleRateFName || Name == StreamingFName || Name == SoundAssetCompressionTypeFName || Name == LoadingBehaviorFName)
 			{
-				InvalidateCompressedData();
-				FreeResources();
-				UpdatePlatformData();
-				MarkPackageDirty();
-
-				// if we are force inline, we need to make sure the shared data is pulled from the DDC
-				// before we attempt to use a decoder on the proxy (not using stream caching)
-				if (LoadingBehavior == ESoundWaveLoadingBehavior::ForceInline)
-				{
-					if (GEngine)
-					{
-						FAudioDevice* LocalAudioDevice = GEngine->GetMainAudioDeviceRaw();
-						if (LocalAudioDevice)
-						{
-							FName RuntimeFormat = LocalAudioDevice->GetRuntimeFormat(this);
-
-							if (LoadingBehavior == ESoundWaveLoadingBehavior::ForceInline && !GetResourceData())
-							{
-								InitAudioResource(RuntimeFormat);
-								check(SoundWaveDataPtr->ResourceSize > 0);
-							}
-						}
-					}
-				}
+				UpdateAsset();
 			}
 
 			if (AnyFFTAnalysisPropertiesChanged(Name))
@@ -3596,6 +3582,36 @@ void USoundWave::OverrideLoadingBehavior(ESoundWaveLoadingBehavior InLoadingBeha
 	}
 }
 
+#if WITH_EDITOR
+void USoundWave::UpdateAsset()
+{
+	InvalidateCompressedData();
+	FreeResources();
+	UpdatePlatformData();
+	MarkPackageDirty();
+
+	// if we are force inline, we need to make sure the shared data is pulled from the DDC
+	// before we attempt to use a decoder on the proxy (not using stream caching)
+	if (LoadingBehavior == ESoundWaveLoadingBehavior::ForceInline)
+	{
+		if (GEngine)
+		{
+			FAudioDevice* LocalAudioDevice = GEngine->GetMainAudioDeviceRaw();
+			if (LocalAudioDevice)
+			{
+				FName RuntimeFormat = LocalAudioDevice->GetRuntimeFormat(this);
+
+				if (LoadingBehavior == ESoundWaveLoadingBehavior::ForceInline && !GetResourceData())
+				{
+					InitAudioResource(RuntimeFormat);
+					check(SoundWaveDataPtr->ResourceSize > 0);
+				}
+			}
+		}
+	}
+}
+#endif // #if WITH_EDITOR
+
 void USoundWave::CacheInheritedLoadingBehavior() const
 {
 	check(IsInGameThread());
