StaticMesh: Added a debug name to static mesh cache puts

#rnx

diff --git a/Engine/Source/Runtime/Engine/Private/StaticMesh.cpp b/Engine/Source/Runtime/Engine/Private/StaticMesh.cpp
--- a/Engine/Source/Runtime/Engine/Private/StaticMesh.cpp
+++ b/Engine/Source/Runtime/Engine/Private/StaticMesh.cpp
@@ -2917,7 +2917,7 @@ void FStaticMeshRenderData::Cache(const ITargetPlatform* TargetPlatform, UStatic
 				EstimatedCompressedSize = uint64(Value.GetData().GetCompressedSize() * DDCSizeToEstimateFactor) + EstimatedNaniteStreamingCompressedSize;
 
 				FRequestOwner RequestOwner(UE::DerivedData::EPriority::Blocking);
-				const FCachePutRequest PutRequest = { FSharedString(TEXT("StaticMesh")), RecordBuilder.Build(), ECachePolicy::Default | ECachePolicy::KeepAlive };
+				const FCachePutRequest PutRequest = { FSharedString(Owner->GetPathName()), RecordBuilder.Build(), ECachePolicy::Default | ECachePolicy::KeepAlive };
 				GetCache().Put(MakeArrayView(&PutRequest, 1), RequestOwner,
 					[&bSavedToDDC](FCachePutResponse&& Response)
 					{
