Fixes for BlendSpace smoothing: Incorrect preview position, and not being able to enable smoothing without closing/re-opening the asset

diff --git a/Engine/Source/Runtime/Engine/Classes/Animation/AnimationAsset.h b/Engine/Source/Runtime/Engine/Classes/Animation/AnimationAsset.h
--- a/Engine/Source/Runtime/Engine/Classes/Animation/AnimationAsset.h
+++ b/Engine/Source/Runtime/Engine/Classes/Animation/AnimationAsset.h
@@ -201,11 +201,6 @@ struct FBlendFilter
 	{
 	}
 	
-	bool IsValid() const
-	{
-		return FilterPerAxis.Num() > 0;
-	}
-	
 	FVector GetFilterLastOutput() const
 	{
 		return FVector(
diff --git a/Engine/Source/Runtime/Engine/Classes/Animation/BlendSpace.h b/Engine/Source/Runtime/Engine/Classes/Animation/BlendSpace.h
--- a/Engine/Source/Runtime/Engine/Classes/Animation/BlendSpace.h
+++ b/Engine/Source/Runtime/Engine/Classes/Animation/BlendSpace.h
@@ -518,8 +518,11 @@ public:
 	*/
 	ENGINE_API bool GetSamplesFromBlendInput(const FVector &BlendInput, TArray<FBlendSampleData> & OutSampleDataList, int32& InOutCachedTriangulationIndex, bool bCombineAnimations) const;
 
-	/** Initialize BlendSpace filtering for runtime. **/
-	ENGINE_API void InitializeFilter(FBlendFilter* Filter) const;
+	/** 
+	 * Initialize BlendSpace filtering for runtime. Filtering supports multiple dimensions, defaulting to
+	 * two (since we don't have 3D BlendSpaces yet) 
+	**/
+	ENGINE_API void InitializeFilter(FBlendFilter* Filter, int NumDimensions = 2) const;
 
 	/** Update BlendSpace filtering parameters - values that don't require a full initialization **/
 	ENGINE_API void UpdateFilterParams(FBlendFilter* Filter) const;
diff --git a/Engine/Source/Runtime/Engine/Private/Animation/AnimInterpFilter.cpp b/Engine/Source/Runtime/Engine/Private/Animation/AnimInterpFilter.cpp
--- a/Engine/Source/Runtime/Engine/Private/Animation/AnimInterpFilter.cpp
+++ b/Engine/Source/Runtime/Engine/Private/Animation/AnimInterpFilter.cpp
@@ -184,6 +184,15 @@ void FFIRFilterTimeBased::WrapToValue(float Input, float Range)
 
 float FFIRFilterTimeBased::UpdateAndGetFilteredData(float Input, float DeltaTime)
 {
+	// Early return if there is no smoothing - applies to all smoothing types. Note that if
+	// WindowDuration changes we will have been re-initialized, so we don't need to worry about
+	// updating the filter in this case.
+	if (WindowDuration <= 0.0f)
+	{
+		LastOutput = Input;
+		return Input;
+	}
+
 	if (DeltaTime > KINDA_SMALL_NUMBER)
 	{
 		float Result;
diff --git a/Engine/Source/Runtime/Engine/Private/Animation/BlendSpace.cpp b/Engine/Source/Runtime/Engine/Private/Animation/BlendSpace.cpp
--- a/Engine/Source/Runtime/Engine/Private/Animation/BlendSpace.cpp
+++ b/Engine/Source/Runtime/Engine/Private/Animation/BlendSpace.cpp
@@ -1052,44 +1052,29 @@ bool UBlendSpace::GetSamplesFromBlendInput(
 	return (OutSampleDataList.Num() != 0);
 }
 
-void UBlendSpace::InitializeFilter(FBlendFilter* Filter) const
+void UBlendSpace::InitializeFilter(FBlendFilter* Filter, int NumDimensions) const
 {
 	if (Filter)
 	{
-		Filter->FilterPerAxis.Empty();
-
-		int32 NumActiveAxes = 0;
-		for (int32 FilterIndex = 0; FilterIndex != 3; ++FilterIndex)
+		Filter->FilterPerAxis.SetNum(NumDimensions);
+		for (int32 FilterIndex = 0; FilterIndex != NumDimensions; ++FilterIndex)
 		{
-			if(InterpolationParam[FilterIndex].InterpolationTime > 0.0f)
-			{
-				NumActiveAxes = FMath::Max(NumActiveAxes, FilterIndex + 1);
-			}
-		}
-
-		if(NumActiveAxes > 0)
-		{
-			Filter->FilterPerAxis.SetNum(NumActiveAxes);
-			
-			for (int32 FilterIndex = 0; FilterIndex < Filter->FilterPerAxis.Num(); ++FilterIndex)
-			{
-				Filter->FilterPerAxis[FilterIndex].Initialize(InterpolationParam[FilterIndex].InterpolationTime,
-															  InterpolationParam[FilterIndex].InterpolationType,
-															  InterpolationParam[FilterIndex].DampingRatio,
-															  BlendParameters[FilterIndex].Min,
-															  BlendParameters[FilterIndex].Max,
-															  InterpolationParam[FilterIndex].MaxSpeed,
-															  !BlendParameters[FilterIndex].bWrapInput);
-			}
+			Filter->FilterPerAxis[FilterIndex].Initialize(InterpolationParam[FilterIndex].InterpolationTime,
+														  InterpolationParam[FilterIndex].InterpolationType,
+														  InterpolationParam[FilterIndex].DampingRatio,
+														  BlendParameters[FilterIndex].Min,
+														  BlendParameters[FilterIndex].Max,
+														  InterpolationParam[FilterIndex].MaxSpeed,
+														  !BlendParameters[FilterIndex].bWrapInput);
 		}
 	}
 }
 
 void UBlendSpace::UpdateFilterParams(FBlendFilter* Filter) const
 {
-	if (Filter && Filter->IsValid())
+	if (Filter)
 	{
-		for (int32 FilterIndex = 0; FilterIndex < Filter->FilterPerAxis.Num(); ++FilterIndex)
+		for (int32 FilterIndex = 0; FilterIndex != Filter->FilterPerAxis.Num(); ++FilterIndex)
 		{
 			Filter->FilterPerAxis[FilterIndex].SetParams(InterpolationParam[FilterIndex].DampingRatio,
 			                                             BlendParameters[FilterIndex].Min,
@@ -1790,31 +1775,34 @@ FVector UBlendSpace::FilterInput(FBlendFilter* Filter, const FVector& BlendInput
 	if(Filter)
 	{
 #if WITH_EDITOR
-		// Check
-		for (int32 AxisIndex = 0; AxisIndex < Filter->FilterPerAxis.Num(); ++AxisIndex)
+		if (Filter->FilterPerAxis.IsEmpty())
+		{
+			InitializeFilter(Filter);
+		}
+		else
 		{
-			if (Filter->FilterPerAxis[AxisIndex].NeedsUpdate(InterpolationParam[AxisIndex].InterpolationType,
-															 InterpolationParam[AxisIndex].InterpolationTime))
+			for (int32 AxisIndex = 0; AxisIndex < Filter->FilterPerAxis.Num(); ++AxisIndex)
 			{
-				InitializeFilter(Filter);
-				break;
+				if (Filter->FilterPerAxis[AxisIndex].NeedsUpdate(InterpolationParam[AxisIndex].InterpolationType,
+																 InterpolationParam[AxisIndex].InterpolationTime))
+				{
+					InitializeFilter(Filter);
+					break;
+				}
 			}
 		}
 		// Note that if we expose the damping ratio etc as pins, this should be called outside of the editor too.
 		UpdateFilterParams(Filter);
 #endif
-		if(Filter->IsValid())
+		for (int32 AxisIndex = 0; AxisIndex < Filter->FilterPerAxis.Num(); ++AxisIndex)
 		{
-			for (int32 AxisIndex = 0; AxisIndex < Filter->FilterPerAxis.Num(); ++AxisIndex)
+			if (BlendParameters[AxisIndex].bWrapInput)
 			{
-				if (BlendParameters[AxisIndex].bWrapInput)
-				{
-					Filter->FilterPerAxis[AxisIndex].WrapToValue(
-						BlendInput[AxisIndex], BlendParameters[AxisIndex].Max - BlendParameters[AxisIndex].Min);
-				}
-				FilteredBlendInput[AxisIndex] = Filter->FilterPerAxis[AxisIndex].UpdateAndGetFilteredData(
-					BlendInput[AxisIndex], DeltaTime);
+				Filter->FilterPerAxis[AxisIndex].WrapToValue(
+					BlendInput[AxisIndex], BlendParameters[AxisIndex].Max - BlendParameters[AxisIndex].Min);
 			}
+			FilteredBlendInput[AxisIndex] = Filter->FilterPerAxis[AxisIndex].UpdateAndGetFilteredData(
+				BlendInput[AxisIndex], DeltaTime);
 		}
 	}
 	return FilteredBlendInput;
diff --git a/Engine/Source/Runtime/Engine/Public/AnimInterpFilter.h b/Engine/Source/Runtime/Engine/Public/AnimInterpFilter.h
--- a/Engine/Source/Runtime/Engine/Public/AnimInterpFilter.h
+++ b/Engine/Source/Runtime/Engine/Public/AnimInterpFilter.h
@@ -119,8 +119,7 @@ public:
 	void Initialize(float InWindowDuration, EFilterInterpolationType InInterpolationType, float InDampingRatio,
 	                float InMinValue, float InMaxValue, float InMaxSpeed, bool bInClamp)
 	{
-		FilterWindow.Empty(10);
-		FilterWindow.AddZeroed(10);
+		FilterWindow.Empty();
 		InterpolationType = InInterpolationType;
 		NumValidFilter = 0;
 		CurrentStackIndex = 0;
