Fix for sounds played w/ Quartz not being stoppable while cached on the game thread (before they have taken a voice slot).

diff --git a/Engine/Source/Runtime/Engine/Classes/Components/AudioComponent.h b/Engine/Source/Runtime/Engine/Classes/Components/AudioComponent.h
--- a/Engine/Source/Runtime/Engine/Classes/Components/AudioComponent.h
+++ b/Engine/Source/Runtime/Engine/Classes/Components/AudioComponent.h
@@ -542,6 +542,7 @@ private:
 		EAudioFaderCurve FadeCurve{ EAudioFaderCurve::Linear };
 		uint32 CommandID{ (uint32)INDEX_NONE };
 		TWeakObjectPtr<UQuartzClockHandle> ClockHandle;
+		bool bHasBeenStoppedWhileQueued{ false };
 	};
 
 	TArray<FAudioComponentPendingQuartzCommandData> PendingQuartzCommandData;
diff --git a/Engine/Source/Runtime/Engine/Classes/Sound/QuartzQuantizationUtilities.cpp b/Engine/Source/Runtime/Engine/Classes/Sound/QuartzQuantizationUtilities.cpp
--- a/Engine/Source/Runtime/Engine/Classes/Sound/QuartzQuantizationUtilities.cpp
+++ b/Engine/Source/Runtime/Engine/Classes/Sound/QuartzQuantizationUtilities.cpp
@@ -445,12 +445,18 @@ namespace Audio
 		OnQueuedCustom(InCommandInitInfo);
 	}
 
-	void IQuartzQuantizedCommand::FailedToQueue()
+	void IQuartzQuantizedCommand::FailedToQueue(FQuartzQuantizedRequestData& InGameThreadData)
 	{
 		TRACE_CPUPROFILER_EVENT_SCOPE(QuartzQuantizedCommand::FailedToQueue);
+
+		GameThreadSubscribers.Append(InGameThreadData.GameThreadSubscribers);
+		GameThreadDelegateID = InGameThreadData.GameThreadDelegateID;
+
 		if (GameThreadSubscribers.Num())
 		{
 			FQuartzQuantizedCommandDelegateData Data;
+			Data.CommandType = GetCommandType();
+			Data.DelegateSubType = EQuartzCommandDelegateSubType::CommandOnFailedToQueue;
 			Data.DelegateID = GameThreadDelegateID;
 
 
diff --git a/Engine/Source/Runtime/Engine/Classes/Sound/QuartzQuantizationUtilities.h b/Engine/Source/Runtime/Engine/Classes/Sound/QuartzQuantizationUtilities.h
--- a/Engine/Source/Runtime/Engine/Classes/Sound/QuartzQuantizationUtilities.h
+++ b/Engine/Source/Runtime/Engine/Classes/Sound/QuartzQuantizationUtilities.h
@@ -507,7 +507,7 @@ namespace Audio
 		void OnQueued(const FQuartzQuantizedCommandInitInfo& InCommandInitInfo);
 
 		// Perhaps the associated sound failed concurrency and will not be playing
-		void FailedToQueue();
+		void FailedToQueue(FQuartzQuantizedRequestData& InGameThreadData);
 
 		// Called 2x Assumed thread latency before OnFinalCallback()
 		void AboutToStart();
diff --git a/Engine/Source/Runtime/Engine/Private/Components/AudioComponent.cpp b/Engine/Source/Runtime/Engine/Private/Components/AudioComponent.cpp
--- a/Engine/Source/Runtime/Engine/Private/Components/AudioComponent.cpp
+++ b/Engine/Source/Runtime/Engine/Private/Components/AudioComponent.cpp
@@ -571,6 +571,14 @@ void UAudioComponent::PlayQueuedQuantizedInternal(const UObject* WorldContextObj
 				bIsValidCommand = false;
 			}
 
+			// was the sound already stopped while we were caching it?
+			if (PendingData.bHasBeenStoppedWhileQueued)
+			{
+				InternalRequestData.QuantizedRequestData.QuantizedCommandPtr->FailedToQueue(InternalRequestData.QuantizedRequestData);
+				UE_LOG(LogAudioQuartz, Verbose, TEXT("Sound (%s) to be played (on Clock: %s) was stopped before being evaluated to play internally"), *InternalRequestData.QuantizedRequestData.ClockName.ToString(), *this->Sound->GetName());
+				bIsValidCommand = false;
+			}
+
 			if (bIsValidCommand)
 			{
 				InternalRequestData.QuantizedRequestData.GameThreadSubscribers.Add(CommandQueuePtr);
@@ -956,6 +964,11 @@ void UAudioComponent::Stop()
 {
 	if (!IsActive())
 	{
+		for (auto& Command : PendingQuartzCommandData)
+		{
+			Command.bHasBeenStoppedWhileQueued = true;
+		}
+
 		return;
 	}
 
