New HLSL translator supports previous WPO

diff --git a/Engine/Source/Runtime/Engine/Private/HLSLTree/HLSLTree.cpp b/Engine/Source/Runtime/Engine/Private/HLSLTree/HLSLTree.cpp
--- a/Engine/Source/Runtime/Engine/Private/HLSLTree/HLSLTree.cpp
+++ b/Engine/Source/Runtime/Engine/Private/HLSLTree/HLSLTree.cpp
@@ -15,6 +15,22 @@ namespace UE
 namespace HLSLTree
 {
 
+enum ELocalPHIChainType : uint8
+{
+	Ddx,
+	Ddy,
+	PreviousFrame,
+};
+
+struct FLocalPHIChainEntry
+{
+	FLocalPHIChainEntry() = default;
+	FLocalPHIChainEntry(const ELocalPHIChainType& InType, const FRequestedType& InRequestedType) : Type(InType), RequestedType(InRequestedType) {}
+
+	ELocalPHIChainType Type;
+	FRequestedType RequestedType;
+};
+
 /**
  * Represents a phi node (see various topics on single static assignment)
  * A phi node takes on a value based on the previous scope that was executed.
@@ -32,14 +48,15 @@ public:
 		}
 	}
 
-	FExpressionLocalPHI(const FExpressionLocalPHI* Source, EDerivativeCoordinate Coord);
+	FExpressionLocalPHI(const FExpressionLocalPHI* Source, ELocalPHIChainType Type, const FRequestedType& RequestedType = FRequestedType());
 
 	virtual void ComputeAnalyticDerivatives(FTree& Tree, FExpressionDerivatives& OutResult) const override;
+	virtual FExpression* ComputePreviousFrame(FTree& Tree, const FRequestedType& RequestedType) const override;
 	virtual bool PrepareValue(FEmitContext& Context, FEmitScope& Scope, const FRequestedType& RequestedType, FPrepareValueResult& OutResult) const override;
 	virtual void EmitValueShader(FEmitContext& Context, FEmitScope& Scope, const FRequestedType& RequestedType, FEmitValueShaderResult& OutResult) const override;
 	virtual void EmitValuePreshader(FEmitContext& Context, FEmitScope& Scope, const FRequestedType& RequestedType, Shader::FPreshaderData& OutPreshader) const override;
 
-	TArray<EDerivativeCoordinate, TInlineAllocator<8>> DerivativeChain;
+	TArray<FLocalPHIChainEntry, TInlineAllocator<8>> Chain;
 	FName LocalName;
 	FScope* Scopes[MaxNumPreviousScopes];
 	FExpression* Values[MaxNumPreviousScopes];
@@ -198,12 +215,12 @@ FScope* FScope::FindSharedParent(FScope* Lhs, FScope* Rhs)
 	return Scope0;
 }
 
-FExpressionLocalPHI::FExpressionLocalPHI(const FExpressionLocalPHI* Source, EDerivativeCoordinate Coord)
-	: DerivativeChain(Source->DerivativeChain)
+FExpressionLocalPHI::FExpressionLocalPHI(const FExpressionLocalPHI* Source, ELocalPHIChainType Type, const FRequestedType& RequestedType)
+	: Chain(Source->Chain)
 	, LocalName(Source->LocalName)
 	, NumValues(Source->NumValues)
 {
-	DerivativeChain.Add(Coord);
+	Chain.Emplace(Type, RequestedType);
 	for (int32 i = 0; i < NumValues; ++i)
 	{
 		Scopes[i] = Source->Scopes[i];
@@ -214,8 +231,13 @@ void FExpressionLocalPHI::ComputeAnalyticDerivatives(FTree& Tree, FExpressionDer
 {
 	// We don't have values assigned at the time analytic derivatives are computed
 	// It's possible the derivatives will be end up being invalid, but that case will need to be detected later, during PrepareValue
-	OutResult.ExpressionDdx = Tree.NewExpression<FExpressionLocalPHI>(this, EDerivativeCoordinate::Ddx);
-	OutResult.ExpressionDdy = Tree.NewExpression<FExpressionLocalPHI>(this, EDerivativeCoordinate::Ddy);
+	OutResult.ExpressionDdx = Tree.NewExpression<FExpressionLocalPHI>(this, ELocalPHIChainType::Ddx);
+	OutResult.ExpressionDdy = Tree.NewExpression<FExpressionLocalPHI>(this, ELocalPHIChainType::Ddy);
+}
+
+FExpression* FExpressionLocalPHI::ComputePreviousFrame(FTree& Tree, const FRequestedType& RequestedType) const
+{
+	return Tree.NewExpression<FExpressionLocalPHI>(this, ELocalPHIChainType::PreviousFrame, RequestedType);
 }
 
 bool FExpressionLocalPHI::PrepareValue(FEmitContext& Context, FEmitScope& Scope, const FRequestedType& RequestedType, FPrepareValueResult& OutResult) const
@@ -633,9 +655,31 @@ EExpressionEvaluation FPreparedComponent::GetEvaluation(const FEmitScope& Scope)
 
 FPreparedComponent CombineComponents(const FPreparedComponent& Lhs, const FPreparedComponent& Rhs)
 {
-	const EExpressionEvaluation Evaluation = CombineEvaluations(Lhs.Evaluation, Rhs.Evaluation);
-	FEmitScope* LoopScope = IsLoopEvaluation(Evaluation) ? FEmitScope::FindSharedParent(Lhs.LoopScope, Rhs.LoopScope) : nullptr;
-	return FPreparedComponent(Evaluation, LoopScope);
+	if (Lhs.IsNone())
+	{
+		return Rhs;
+	}
+	else if (Rhs.IsNone())
+	{
+		return Lhs;
+	}
+	else
+	{
+		const EExpressionEvaluation Evaluation = CombineEvaluations(Lhs.Evaluation, Rhs.Evaluation);
+		FEmitScope* LoopScope = nullptr;
+		if (IsLoopEvaluation(Evaluation))
+		{
+			LoopScope = FEmitScope::FindSharedParent(Lhs.LoopScope, Rhs.LoopScope);
+		}
+		FExpression* ForwardValue = nullptr;
+		int32 ForwardComponentIndex = INDEX_NONE;
+		if (Lhs.ForwardValue == Rhs.ForwardValue && Lhs.ForwardComponentIndex == Rhs.ForwardComponentIndex)
+		{
+			ForwardValue = Lhs.ForwardValue;
+			ForwardComponentIndex = Lhs.ForwardComponentIndex;
+		}
+		return FPreparedComponent(Evaluation, LoopScope, ForwardValue, ForwardComponentIndex);
+	}
 }
 
 FPreparedType::FPreparedType(const Shader::FType& InType)
@@ -829,6 +873,19 @@ void FPreparedType::SetLoopEvaluation(FEmitScope& Scope, const FRequestedType& R
 	}
 }
 
+void FPreparedType::SetForwardValue(FExpression* ForwardValue, const FRequestedType& RequestedType)
+{
+	for (int32 Index = 0; Index < PreparedComponents.Num(); ++Index)
+	{
+		FPreparedComponent& Component = PreparedComponents[Index];
+		if (RequestedType.IsComponentRequested(Index) && !Component.IsNone())
+		{
+			Component.ForwardValue = ForwardValue;
+			Component.ForwardComponentIndex = Index;
+		}
+	}
+}
+
 void FPreparedType::SetField(const Shader::FStructField* Field, const FPreparedType& FieldType)
 {
 	for (int32 Index = 0; Index < Field->GetNumComponents(); ++Index)
@@ -890,12 +947,11 @@ FPreparedType MergePreparedTypes(const FPreparedType& Lhs, const FPreparedType&
 bool FPrepareValueResult::TryMergePreparedType(FEmitContext& Context, const Shader::FStructType* StructType, Shader::EValueComponentType ComponentType)
 {
 	// If we previously had a forwarded value set, reset that and start over
-	if (ForwardValue || !PreparedType.IsInitialized())
+	if (!PreparedType.IsInitialized())
 	{
 		PreparedType.PreparedComponents.Reset();
 		PreparedType.ValueComponentType = ComponentType;
 		PreparedType.StructType = StructType;
-		ForwardValue = nullptr;
 		return true;
 	}
 
@@ -924,7 +980,6 @@ bool FPrepareValueResult::SetTypeVoid()
 	PreparedType.PreparedComponents.Reset();
 	PreparedType.ValueComponentType = Shader::EValueComponentType::Void;
 	PreparedType.StructType = nullptr;
-	ForwardValue = nullptr;
 	return false;
 }
 
@@ -967,11 +1022,8 @@ bool FPrepareValueResult::SetType(FEmitContext& Context, const FRequestedType& R
 
 bool FPrepareValueResult::SetForwardValue(FEmitContext& Context, FEmitScope& Scope, const FRequestedType& RequestedType, FExpression* InForwardValue)
 {
-	if (InForwardValue != ForwardValue)
-	{
-		PreparedType = Context.PrepareExpression(InForwardValue, Scope, RequestedType);
-		ForwardValue = InForwardValue;
-	}
+	const FPreparedType ForwardPreparedType = Context.PrepareExpression(InForwardValue, Scope, RequestedType);
+	PreparedType = MergePreparedTypes(PreparedType, ForwardPreparedType);
 	return !PreparedType.IsVoid();
 }
 
@@ -985,6 +1037,11 @@ void FExpression::ComputeAnalyticDerivatives(FTree& Tree, FExpressionDerivatives
 	// nop
 }
 
+FExpression* FExpression::ComputePreviousFrame(FTree& Tree, const FRequestedType& RequestedType) const
+{
+	return nullptr;
+}
+
 void FExpression::EmitValueShader(FEmitContext& Context, FEmitScope& Scope, const FRequestedType& RequestedType, FEmitValueShaderResult& OutResult) const
 {
 	check(false);
@@ -995,12 +1052,79 @@ void FExpression::EmitValuePreshader(FEmitContext& Context, FEmitScope& Scope, c
 	check(false);
 }
 
+FExpression* FExpression::GetForwardValue(const FRequestedType& RequestedType)
+{
+	TArray<FExpression*, TInlineAllocator<16>> VisitedExpressions;
+	FExpression* Expression = this;
+	int32 BaseComponentIndex = 0;
+
+	while (true)
+	{
+		const FPreparedType& PreparedType = Expression->PrepareValueResult.PreparedType;
+
+		VisitedExpressions.Add(Expression);
+
+		FExpression* ForwardValue = nullptr;
+		int32 ForwardComponentIndex = INDEX_NONE;
+		int32 StartComponentIndex = INDEX_NONE;
+		// Scan the required components, and see if they all have the same ForwardValue, in the correct order
+		for (int32 Index = 0; Index < RequestedType.RequestedComponents.Num(); ++Index)
+		{
+			if (RequestedType.RequestedComponents[Index])
+			{
+				const FPreparedComponent Component = PreparedType.GetComponent(BaseComponentIndex + Index);
+				if (!Component.ForwardValue)
+				{
+					return nullptr;
+				}
+
+				if (StartComponentIndex == INDEX_NONE)
+				{
+					StartComponentIndex = Index;
+					ForwardValue = Component.ForwardValue;
+					ForwardComponentIndex = Component.ForwardComponentIndex;
+				}
+				else if (ForwardValue != Component.ForwardValue ||
+					ForwardComponentIndex + Index - StartComponentIndex != Component.ForwardComponentIndex)
+				{
+					return nullptr;
+				}
+			}
+		}
+
+		check(ForwardValue);
+		if (ForwardValue == this)
+		{
+			// Don't forward to ourselves
+			return nullptr;
+		}
+		if (ForwardComponentIndex == StartComponentIndex)
+		{
+			// Forward to a value with no component offset, can use that directly
+			return ForwardValue;
+		}
+
+		// If there's a component offset, we can't forward directly
+		// Instead, recursively check to see if the ForwardValue itself has a Forward value
+		// This can happen when forwarding across multiple set/get struct field expressions
+		BaseComponentIndex = ForwardComponentIndex - StartComponentIndex;
+		Expression = ForwardValue;
+
+		// Make sure we're not in a loop
+		check(!VisitedExpressions.Contains(Expression));
+	}
+
+	checkNoEntry();
+	return nullptr;
+}
+
 FEmitShaderExpression* FExpression::GetValueShader(FEmitContext& Context, FEmitScope& Scope, const FRequestedType& RequestedType, const Shader::FType& ResultType)
 {
 	FOwnerScope OwnerScope(*Context.Errors, GetOwner());
-	if (PrepareValueResult.ForwardValue)
+	FExpression* ForwardValue = GetForwardValue(RequestedType);
+	if (ForwardValue)
 	{
-		return PrepareValueResult.ForwardValue->GetValueShader(Context, Scope, RequestedType, ResultType);
+		return ForwardValue->GetValueShader(Context, Scope, RequestedType, ResultType);
 	}
 
 	const EExpressionEvaluation Evaluation = PrepareValueResult.PreparedType.GetEvaluation(Scope, RequestedType);
@@ -1036,9 +1160,10 @@ FEmitShaderExpression* FExpression::GetValueShader(FEmitContext& Context, FEmitS
 void FExpression::GetValuePreshader(FEmitContext& Context, FEmitScope& Scope, const FRequestedType& RequestedType, Shader::FPreshaderData& OutPreshader)
 {
 	FOwnerScope OwnerScope(*Context.Errors, GetOwner());
-	if (PrepareValueResult.ForwardValue)
+	FExpression* ForwardValue = GetForwardValue(RequestedType);
+	if (ForwardValue)
 	{
-		return PrepareValueResult.ForwardValue->GetValuePreshader(Context, Scope, RequestedType, OutPreshader);
+		return ForwardValue->GetValuePreshader(Context, Scope, RequestedType, OutPreshader);
 	}
 
 	const int32 PrevStackPosition = Context.PreshaderStackPosition;
@@ -1061,9 +1186,10 @@ void FExpression::GetValuePreshader(FEmitContext& Context, FEmitScope& Scope, co
 Shader::FValue FExpression::GetValueConstant(FEmitContext& Context, FEmitScope& Scope, const FRequestedType& RequestedType)
 {
 	FOwnerScope OwnerScope(*Context.Errors, GetOwner());
-	if (PrepareValueResult.ForwardValue)
+	FExpression* ForwardValue = GetForwardValue(RequestedType);
+	if (ForwardValue)
 	{
-		return PrepareValueResult.ForwardValue->GetValueConstant(Context, Scope, RequestedType);
+		return ForwardValue->GetValueConstant(Context, Scope, RequestedType);
 	}
 
 	check(!bReentryFlag);
@@ -1187,10 +1313,19 @@ bool FTree::Finalize()
 				return false;
 			}
 
-			for(EDerivativeCoordinate DerivativeCoord : Expression->DerivativeChain)
+			for(const FLocalPHIChainEntry& Entry : Expression->Chain)
 			{
-				const FExpressionDerivatives Derivatives = GetAnalyticDerivatives(LocalValue);
-				LocalValue = Derivatives.Get(DerivativeCoord);
+				const ELocalPHIChainType ChainType = Entry.Type;
+				if (ChainType == ELocalPHIChainType::Ddx || ChainType == ELocalPHIChainType::Ddy)
+				{
+					const FExpressionDerivatives Derivatives = GetAnalyticDerivatives(LocalValue);
+					LocalValue = (ChainType == ELocalPHIChainType::Ddx) ? Derivatives.ExpressionDdx : Derivatives.ExpressionDdy;
+				}
+				else
+				{
+					check(ChainType == ELocalPHIChainType::PreviousFrame);
+					LocalValue = GetPreviousFrame(LocalValue, Entry.RequestedType);
+				}
 			}
 			// May be nullptr if derivatives are not valid
 			Expression->Values[i] = LocalValue;
@@ -1327,6 +1462,23 @@ const FExpressionDerivatives& FTree::GetAnalyticDerivatives(FExpression* InExpre
 	return InExpression->Derivatives;
 }
 
+FExpression* FTree::GetPreviousFrame(FExpression* InExpression, const FRequestedType& RequestedType)
+{
+	FExpression* Result = InExpression;
+	if (Result && !RequestedType.IsVoid())
+	{
+		FExpressionReentryScope ReentryScope(InExpression);
+		FOwnerScope OwnerScope(*this, InExpression->GetOwner()); // Associate any newly created nodes with the same owner as the input expression
+
+		FExpression* PrevFrameExpression = InExpression->ComputePreviousFrame(*this, RequestedType);
+		if (PrevFrameExpression)
+		{
+			Result = PrevFrameExpression;
+		}
+	}
+	return Result;
+}
+
 FScope* FTree::NewScope(FScope& Scope)
 {
 	FScope* NewScope = NewNode<FScope>();
diff --git a/Engine/Source/Runtime/Engine/Private/HLSLTree/HLSLTreeCommon.cpp b/Engine/Source/Runtime/Engine/Private/HLSLTree/HLSLTreeCommon.cpp
--- a/Engine/Source/Runtime/Engine/Private/HLSLTree/HLSLTreeCommon.cpp
+++ b/Engine/Source/Runtime/Engine/Private/HLSLTree/HLSLTreeCommon.cpp
@@ -49,10 +49,15 @@ FExternalInputDescription GetExternalInputDescription(EExternalInput Input)
 	case EExternalInput::TexCoord6_Ddy: return FExternalInputDescription(TEXT("TexCoord6_Ddy"), Shader::EValueType::Float2);
 	case EExternalInput::TexCoord7_Ddy: return FExternalInputDescription(TEXT("TexCoord7_Ddy"), Shader::EValueType::Float2);
 
-	case EExternalInput::WorldPosition: return FExternalInputDescription(TEXT("WorldPosition"), Shader::EValueType::Double3, EExternalInput::WorldPosition_Ddx, EExternalInput::WorldPosition_Ddy);
-	case EExternalInput::WorldPosition_NoOffsets: return FExternalInputDescription(TEXT("WorldPosition_NoOffsets"), Shader::EValueType::Double3, EExternalInput::WorldPosition_Ddx, EExternalInput::WorldPosition_Ddy);
-	case EExternalInput::TranslatedWorldPosition: return FExternalInputDescription(TEXT("TranslatedWorldPosition"), Shader::EValueType::Float3, EExternalInput::WorldPosition_Ddx, EExternalInput::WorldPosition_Ddy);
-	case EExternalInput::TranslatedWorldPosition_NoOffsets: return FExternalInputDescription(TEXT("TranslatedWorldPosition_NoOffsets"), Shader::EValueType::Float3, EExternalInput::WorldPosition_Ddx, EExternalInput::WorldPosition_Ddy);
+	case EExternalInput::WorldPosition: return FExternalInputDescription(TEXT("WorldPosition"), Shader::EValueType::Double3, EExternalInput::WorldPosition_Ddx, EExternalInput::WorldPosition_Ddy, EExternalInput::PrevWorldPosition);
+	case EExternalInput::WorldPosition_NoOffsets: return FExternalInputDescription(TEXT("WorldPosition_NoOffsets"), Shader::EValueType::Double3, EExternalInput::WorldPosition_Ddx, EExternalInput::WorldPosition_Ddy, EExternalInput::PrevWorldPosition_NoOffsets);
+	case EExternalInput::TranslatedWorldPosition: return FExternalInputDescription(TEXT("TranslatedWorldPosition"), Shader::EValueType::Float3, EExternalInput::WorldPosition_Ddx, EExternalInput::WorldPosition_Ddy, EExternalInput::PrevTranslatedWorldPosition);
+	case EExternalInput::TranslatedWorldPosition_NoOffsets: return FExternalInputDescription(TEXT("TranslatedWorldPosition_NoOffsets"), Shader::EValueType::Float3, EExternalInput::WorldPosition_Ddx, EExternalInput::WorldPosition_Ddy, EExternalInput::PrevTranslatedWorldPosition_NoOffsets);
+
+	case EExternalInput::PrevWorldPosition: return FExternalInputDescription(TEXT("PrevWorldPosition"), Shader::EValueType::Double3, EExternalInput::WorldPosition_Ddx, EExternalInput::WorldPosition_Ddy);
+	case EExternalInput::PrevWorldPosition_NoOffsets: return FExternalInputDescription(TEXT("PrevWorldPosition_NoOffsets"), Shader::EValueType::Double3, EExternalInput::WorldPosition_Ddx, EExternalInput::WorldPosition_Ddy);
+	case EExternalInput::PrevTranslatedWorldPosition: return FExternalInputDescription(TEXT("PrevTranslatedWorldPosition"), Shader::EValueType::Float3, EExternalInput::WorldPosition_Ddx, EExternalInput::WorldPosition_Ddy);
+	case EExternalInput::PrevTranslatedWorldPosition_NoOffsets: return FExternalInputDescription(TEXT("PrevTranslatedWorldPosition_NoOffsets"), Shader::EValueType::Float3, EExternalInput::WorldPosition_Ddx, EExternalInput::WorldPosition_Ddy);
 
 	case EExternalInput::WorldPosition_Ddx: return FExternalInputDescription(TEXT("WorldPosition_Ddx"), Shader::EValueType::Float3);
 	case EExternalInput::WorldPosition_Ddy: return FExternalInputDescription(TEXT("WorldPosition_Ddx"), Shader::EValueType::Float3);
@@ -223,6 +228,16 @@ void FExpressionExternalInput::ComputeAnalyticDerivatives(FTree& Tree, FExpressi
 	}
 }
 
+FExpression* FExpressionExternalInput::ComputePreviousFrame(FTree& Tree, const FRequestedType& RequestedType) const
+{
+	const FExternalInputDescription InputDesc = GetExternalInputDescription(InputType);
+	if (InputDesc.PreviousFrame != EExternalInput::None)
+	{
+		return Tree.NewExpression<FExpressionExternalInput>(InputDesc.PreviousFrame);
+	}
+	return nullptr;
+}
+
 bool FExpressionExternalInput::PrepareValue(FEmitContext& Context, FEmitScope& Scope, const FRequestedType& RequestedType, FPrepareValueResult& OutResult) const
 {
 	const FExternalInputDescription InputDesc = GetExternalInputDescription(InputType);
@@ -251,32 +266,24 @@ void FExpressionExternalInput::EmitValueShader(FEmitContext& Context, FEmitScope
 	else
 	{
 		const FExternalInputDescription InputDesc = GetExternalInputDescription(InputType);
-
-		// TODO - handle PrevFrame position
+		const TCHAR* Code = nullptr;
 		switch (InputType)
 		{
-		case EExternalInput::WorldPosition:
-			OutResult.Code = Context.EmitInlineExpression(Scope, InputDesc.Type, TEXT("GetWorldPosition(Parameters)"));
-			break;
-		case EExternalInput::WorldPosition_NoOffsets:
-			OutResult.Code = Context.EmitInlineExpression(Scope, InputDesc.Type, TEXT("GetWorldPosition_NoMaterialOffsets(Parameters)"));
-			break;
-		case EExternalInput::TranslatedWorldPosition:
-			OutResult.Code = Context.EmitInlineExpression(Scope, InputDesc.Type, TEXT("GetTranslatedWorldPosition(Parameters)"));
-			break;
-		case EExternalInput::TranslatedWorldPosition_NoOffsets:
-			OutResult.Code = Context.EmitInlineExpression(Scope, InputDesc.Type, TEXT("GetTranslatedWorldPosition_NoMaterialOffsets(Parameters)"));
-			break;
-		case EExternalInput::WorldPosition_Ddx:
-			OutResult.Code = Context.EmitInlineExpression(Scope, InputDesc.Type, TEXT("Parameters.WorldPosition_DDX"));
-			break;
-		case EExternalInput::WorldPosition_Ddy:
-			OutResult.Code = Context.EmitInlineExpression(Scope, InputDesc.Type, TEXT("Parameters.WorldPosition_DDY"));
-			break;
+		case EExternalInput::WorldPosition: Code = TEXT("GetWorldPosition(Parameters)"); break;
+		case EExternalInput::WorldPosition_NoOffsets: Code = TEXT("GetWorldPosition_NoMaterialOffsets(Parameters)"); break;
+		case EExternalInput::TranslatedWorldPosition: Code = TEXT("GetTranslatedWorldPosition(Parameters)"); break;
+		case EExternalInput::TranslatedWorldPosition_NoOffsets: Code = TEXT("GetTranslatedWorldPosition_NoMaterialOffsets(Parameters)"); break;
+		case EExternalInput::PrevWorldPosition: Code = TEXT("GetPrevWorldPosition(Parameters)"); break;
+		case EExternalInput::PrevWorldPosition_NoOffsets: Code = TEXT("GetPrevWorldPosition_NoMaterialOffsets(Parameters)"); break;
+		case EExternalInput::PrevTranslatedWorldPosition: Code = TEXT("GetPrevTranslatedWorldPosition(Parameters)"); break;
+		case EExternalInput::PrevTranslatedWorldPosition_NoOffsets: Code = TEXT("GetPrevTranslatedWorldPosition_NoMaterialOffsets(Parameters)"); break;
+		case EExternalInput::WorldPosition_Ddx: Code = TEXT("Parameters.WorldPosition_DDX"); break;
+		case EExternalInput::WorldPosition_Ddy: Code = TEXT("Parameters.WorldPosition_DDY"); break;
 		default:
 			checkNoEntry();
 			break;
 		}
+		OutResult.Code = Context.EmitInlineExpression(Scope, InputDesc.Type, Code);
 	}
 }
 
@@ -443,6 +450,13 @@ void FExpressionGetStructField::ComputeAnalyticDerivatives(FTree& Tree, FExpress
 	}
 }
 
+FExpression* FExpressionGetStructField::ComputePreviousFrame(FTree& Tree, const FRequestedType& RequestedType) const
+{
+	FRequestedType RequestedStructType(StructType, false);
+	RequestedStructType.SetField(Field, RequestedType);
+	return Tree.NewExpression<FExpressionGetStructField>(StructType, Field, Tree.GetPreviousFrame(StructExpression, RequestedStructType));
+}
+
 bool FExpressionGetStructField::PrepareValue(FEmitContext& Context, FEmitScope& Scope, const FRequestedType& RequestedType, FPrepareValueResult& OutResult) const
 {
 	FRequestedType RequestedStructType(StructType, false);
@@ -495,6 +509,18 @@ void FExpressionSetStructField::ComputeAnalyticDerivatives(FTree& Tree, FExpress
 	}
 }
 
+FExpression* FExpressionSetStructField::ComputePreviousFrame(FTree& Tree, const FRequestedType& RequestedType) const
+{
+	FRequestedType RequestedStructType(RequestedType);
+	RequestedStructType.ClearFieldRequested(Field);
+	FExpression* PrevStructExpression = Tree.GetPreviousFrame(StructExpression, RequestedStructType);
+
+	const FRequestedType RequestedFieldType = RequestedType.GetField(Field);
+	FExpression* PrevFieldExpression = Tree.GetPreviousFrame(FieldExpression, RequestedFieldType);
+
+	return Tree.NewExpression<FExpressionSetStructField>(StructType, Field, PrevStructExpression, PrevFieldExpression);
+}
+
 bool FExpressionSetStructField::PrepareValue(FEmitContext& Context, FEmitScope& Scope, const FRequestedType& RequestedType, FPrepareValueResult& OutResult) const
 {
 	FRequestedType RequestedStructType(RequestedType);
@@ -507,20 +533,13 @@ bool FExpressionSetStructField::PrepareValue(FEmitContext& Context, FEmitScope&
 	}
 
 	const FPreparedType& FieldPreparedType = Context.PrepareExpression(FieldExpression, Scope, RequestedType.GetField(Field));
-	if (FieldPreparedType.IsVoid())
-	{
-		return OutResult.SetForwardValue(Context, Scope, RequestedType, StructExpression);
-	}
-	else
+	FPreparedType ResultType(StructPreparedType);
+	if (ResultType.IsVoid())
 	{
-		FPreparedType ResultType(StructPreparedType);
-		if (ResultType.IsVoid())
-		{
-			ResultType = StructType;
-		}
-		ResultType.SetField(Field, FieldPreparedType);
-		return OutResult.SetType(Context, RequestedType, ResultType);
+		ResultType = StructType;
 	}
+	ResultType.SetField(Field, FieldPreparedType);
+	return OutResult.SetType(Context, RequestedType, ResultType);
 }
 
 void FExpressionSetStructField::EmitValueShader(FEmitContext& Context, FEmitScope& Scope, const FRequestedType& RequestedType, FEmitValueShaderResult& OutResult) const
@@ -654,6 +673,11 @@ void FExpressionUnaryOp::ComputeAnalyticDerivatives(FTree& Tree, FExpressionDeri
 	}
 }
 
+FExpression* FExpressionUnaryOp::ComputePreviousFrame(FTree& Tree, const FRequestedType& RequestedType) const
+{
+	return Tree.NewUnaryOp(Op, Tree.GetPreviousFrame(Input, RequestedType));
+}
+
 bool FExpressionUnaryOp::PrepareValue(FEmitContext& Context, FEmitScope& Scope, const FRequestedType& RequestedType, FPrepareValueResult& OutResult) const
 {
 	const FPreparedType& InputType = Context.PrepareExpression(Input, Scope, RequestedType);
@@ -779,6 +803,11 @@ void FExpressionBinaryOp::ComputeAnalyticDerivatives(FTree& Tree, FExpressionDer
 	}
 }
 
+FExpression* FExpressionBinaryOp::ComputePreviousFrame(FTree& Tree, const FRequestedType& RequestedType) const
+{
+	return Tree.NewBinaryOp(Op, Tree.GetPreviousFrame(Lhs, RequestedType), Tree.GetPreviousFrame(Rhs, RequestedType));
+}
+
 bool FExpressionBinaryOp::PrepareValue(FEmitContext& Context, FEmitScope& Scope, const FRequestedType& RequestedType, FPrepareValueResult& OutResult) const
 {
 	const FPreparedType& LhsType = Context.PrepareExpression(Lhs, Scope, RequestedType);
@@ -916,6 +945,12 @@ void FExpressionSwizzle::ComputeAnalyticDerivatives(FTree& Tree, FExpressionDeri
 	}
 }
 
+FExpression* FExpressionSwizzle::ComputePreviousFrame(FTree& Tree, const FRequestedType& RequestedType) const
+{
+	const FRequestedType RequestedInputType = Parameters.GetRequestedInputType(RequestedType);
+	return Tree.NewExpression<FExpressionSwizzle>(Parameters, Tree.GetPreviousFrame(Input, RequestedInputType));
+}
+
 bool FExpressionSwizzle::PrepareValue(FEmitContext& Context, FEmitScope& Scope, const FRequestedType& RequestedType, FPrepareValueResult& OutResult) const
 {
 	const FRequestedType RequestedInputType = Parameters.GetRequestedInputType(RequestedType);
@@ -986,6 +1021,12 @@ void FExpressionAppend::ComputeAnalyticDerivatives(FTree& Tree, FExpressionDeriv
 	}
 }
 
+FExpression* FExpressionAppend::ComputePreviousFrame(FTree& Tree, const FRequestedType& RequestedType) const
+{
+	// TODO - requested type?
+	return Tree.NewExpression<FExpressionAppend>(Tree.GetPreviousFrame(Lhs, RequestedType), Tree.GetPreviousFrame(Rhs, RequestedType));
+}
+
 bool FExpressionAppend::PrepareValue(FEmitContext& Context, FEmitScope& Scope, const FRequestedType& RequestedType, FPrepareValueResult& OutResult) const
 {
 	const FPreparedType& LhsType = Context.PrepareExpression(Lhs, Scope, RequestedType);
diff --git a/Engine/Source/Runtime/Engine/Private/HLSLTree/HLSLTreeEmit.cpp b/Engine/Source/Runtime/Engine/Private/HLSLTree/HLSLTreeEmit.cpp
--- a/Engine/Source/Runtime/Engine/Private/HLSLTree/HLSLTreeEmit.cpp
+++ b/Engine/Source/Runtime/Engine/Private/HLSLTree/HLSLTreeEmit.cpp
@@ -222,7 +222,7 @@ FPreparedType FEmitContext::PrepareExpression(FExpression* InExpression, FEmitSc
 		FEmitScope* LoopScope = Scope.FindLoop();
 		check(LoopScope);
 		Result.PreparedType.SetLoopEvaluation(*LoopScope, RequestedType);
-		return Result.GetPreparedType();
+		return Result.PreparedType;
 	}
 
 	bool bResult = false;
@@ -231,16 +231,15 @@ FPreparedType FEmitContext::PrepareExpression(FExpression* InExpression, FEmitSc
 		bResult = InExpression->PrepareValue(*this, Scope, RequestedType, Result);
 	}
 
-	if (!bResult)
+	FPreparedType ResultType;
+	if (bResult)
 	{
-		Result.SetTypeVoid();
-	}
-	else
-	{
-		check(!Result.GetPreparedType().IsVoid());
+		ResultType = Result.PreparedType;
+		check(!ResultType.IsVoid());
+		ResultType.SetForwardValue(InExpression, RequestedType);
 	}
 
-	return Result.GetPreparedType();
+	return ResultType;
 }
 
 FEmitScope* FEmitContext::InternalPrepareScope(FScope* InScope, FScope* InParentScope)
diff --git a/Engine/Source/Runtime/Engine/Private/Materials/MaterialHLSLEmitter.cpp b/Engine/Source/Runtime/Engine/Private/Materials/MaterialHLSLEmitter.cpp
--- a/Engine/Source/Runtime/Engine/Private/Materials/MaterialHLSLEmitter.cpp
+++ b/Engine/Source/Runtime/Engine/Private/Materials/MaterialHLSLEmitter.cpp
@@ -215,7 +215,7 @@ static FString GenerateMaterialTemplateHLSL(EShaderPlatform ShaderPlatform,
 	LazyPrintf.PushParam(*FString::Printf(TEXT("return %.5f"), Material.GetOpacityMaskClipValue()));
 
 	LazyPrintf.PushParam(TEXT("return Parameters.MaterialVertexAttributes.WorldPositionOffset"));
-	LazyPrintf.PushParam(TEXT("return 0.0f"));
+	LazyPrintf.PushParam(TEXT("return Parameters.MaterialVertexAttributes.PrevWorldPositionOffset"));
 	LazyPrintf.PushParam(TEXT("return 0.0f"));
 	LazyPrintf.PushParam(TEXT("return 0.0f"));
 
@@ -889,6 +889,7 @@ bool MaterialEmitHLSL(const FMaterialCompileTargetParameters& InCompilerTarget,
 				}
 			}
 		}
+		RequestedVertexAttributesType.SetFieldRequested(Generator.GetMaterialAttributesType()->FindFieldByName(TEXT("PrevWorldPositionOffset")));
 
 		HLSLTree->ResetNodes();
 
diff --git a/Engine/Source/Runtime/Engine/Private/Materials/MaterialHLSLGenerator.cpp b/Engine/Source/Runtime/Engine/Private/Materials/MaterialHLSLGenerator.cpp
--- a/Engine/Source/Runtime/Engine/Private/Materials/MaterialHLSLGenerator.cpp
+++ b/Engine/Source/Runtime/Engine/Private/Materials/MaterialHLSLGenerator.cpp
@@ -96,7 +96,7 @@ FMaterialHLSLGenerator::FMaterialHLSLGenerator(const FMaterialCompileTargetParam
 	FFunctionCallEntry* RootFunctionEntry = new(InOutTree.GetAllocator()) FFunctionCallEntry();
 	FunctionCallStack.Add(RootFunctionEntry);
 
-	TArray<UE::Shader::FStructFieldInitializer, TInlineAllocator<MP_MAX>> MaterialAttributeFields;
+	TArray<UE::Shader::FStructFieldInitializer, TInlineAllocator<MP_MAX + 16>> MaterialAttributeFields;
 
 	const TArray<FGuid>& OrderedVisibleAttributes = FMaterialAttributeDefinitionMap::GetOrderedVisibleAttributeList();
 	for (const FGuid& AttributeID : OrderedVisibleAttributes)
@@ -123,6 +123,9 @@ FMaterialHLSLGenerator::FMaterialHLSLGenerator(const FMaterialCompileTargetParam
 		}
 	}
 
+	MaterialAttributeFields.Emplace(TEXT("PrevWorldPositionOffset"), UE::Shader::EValueType::Float3);
+	MaterialAttributesDefaultValue.Component.Append({ 0.0f, 0.0f, 0.0f });
+
 	UE::Shader::FStructTypeInitializer MaterialAttributesInitializer;
 	MaterialAttributesInitializer.Name = TEXT("FMaterialAttributes");
 	MaterialAttributesInitializer.Fields = MaterialAttributeFields;
@@ -254,6 +257,7 @@ static UE::HLSLTree::FExpression* CompileMaterialInput(FMaterialHLSLGenerator& G
 bool FMaterialHLSLGenerator::GenerateResult(UE::HLSLTree::FScope& Scope)
 {
 	using namespace UE::HLSLTree;
+	using namespace UE::Shader;
 
 	FFunctionCallEntry* FunctionEntry = FunctionCallStack.Last();
 
@@ -289,6 +293,17 @@ bool FMaterialHLSLGenerator::GenerateResult(UE::HLSLTree::FScope& Scope)
 				{
 					check(InputDescription.Type == UE::Shader::EValueType::Struct);
 					AttributesExpression = InputDescription.Input->AcquireHLSLExpression(*this, Scope);
+
+					const FString& WPOName = FMaterialAttributeDefinitionMap::GetAttributeName(MP_WorldPositionOffset);
+					const FStructField* WPOField = GetMaterialAttributesType()->FindFieldByName(*WPOName);
+					const FStructField* PrevWPOField = GetMaterialAttributesType()->FindFieldByName(TEXT("PrevWorldPositionOffset"));
+
+					FRequestedType PrevRequestedType(GetMaterialAttributesType(), false);
+					PrevRequestedType.SetFieldRequested(WPOField);
+
+					FExpression* PrevAttributesExpression = HLSLTree->GetPreviousFrame(AttributesExpression, PrevRequestedType);
+					FExpression* PrevWPOExpression = HLSLTree->NewExpression<FExpressionGetStructField>(GetMaterialAttributesType(), WPOField, PrevAttributesExpression);
+					AttributesExpression = HLSLTree->NewExpression<FExpressionSetStructField>(GetMaterialAttributesType(), PrevWPOField, AttributesExpression, PrevWPOExpression);
 				}
 			}
 			else
diff --git a/Engine/Source/Runtime/Engine/Public/HLSLTree/HLSLTree.h b/Engine/Source/Runtime/Engine/Public/HLSLTree/HLSLTree.h
--- a/Engine/Source/Runtime/Engine/Public/HLSLTree/HLSLTree.h
+++ b/Engine/Source/Runtime/Engine/Public/HLSLTree/HLSLTree.h
@@ -115,56 +115,66 @@ class FHasher
 {
 public:
 	FXxHash64 Finalize() { return Builder.Finalize(); }
-	FHasher& AppendData(const void* Data, uint64 Size) { Builder.Update(Data, Size); return *this; }
+	void AppendData(const void* Data, uint64 Size) { Builder.Update(Data, Size); }
 
-	template<typename T>
-	inline FHasher& AppendValue(const T& Value)
-	{
-		return AppendData(&Value, sizeof(Value));
-	}
+private:
+	FXxHash64Builder Builder;
+};
 
-	template<typename T>
-	inline FHasher& AppendValue(TArrayView<T> Value)
-	{
-		for (const T& Element : Value)
-		{
-			AppendValue(Element);
-		}
-		return *this;
-	}
+template<typename T>
+inline void AppendHash(FHasher& Hasher, const T& Value)
+{
+	Hasher.AppendData(&Value, sizeof(Value));
+}
 
-	inline FHasher& AppendValue(const FName& Value)
+template<typename T>
+inline void AppendHash(FHasher& Hasher, TArrayView<T> Value)
+{
+	for (const T& Element : Value)
 	{
-		return AppendValue(Value.GetComparisonIndex()).AppendValue(Value.GetNumber());
+		AppendHash(Hasher, Element);
 	}
+}
 
-	inline FHasher& AppendValue(const Shader::FType& Type)
-	{
-		if (Type.IsStruct()) return AppendValue(Type.StructType);
-		else return AppendValue(Type.ValueType);
-	}
 
-	inline FHasher& AppendValue(const Shader::FValue& Value)
-	{
-		AppendValue(Value.Type);
-		for (int32 i = 0; i < Value.Type.GetNumComponents(); ++i)
-		{
-			AppendValue(Value.TryGetComponent(i));
-		}
-		return *this;
-	}
+template<typename Allocator>
+inline void AppendHash(FHasher& Hasher, const TBitArray<Allocator>& Value)
+{
+	const uint32 NumWords = FBitSet::CalculateNumWords(Value.Num());
+	const uint32* Data = Value.GetData();
+	Hasher.AppendData(Data, NumWords * sizeof(uint32));
+}
 
-	inline FHasher& AppendValues() { return *this; }
+inline void AppendHash(FHasher& Hasher, const FName& Value)
+{
+	AppendHash(Hasher, Value.GetComparisonIndex());
+	AppendHash(Hasher, Value.GetNumber());
+}
 
-	template<typename T, typename... ArgTypes>
-	inline FHasher& AppendValues(const T& Value, ArgTypes&&... Args)
+inline void AppendHash(FHasher& Hasher, const Shader::FType& Type)
+{
+	if (Type.IsStruct()) AppendHash(Hasher, Type.StructType);
+	else AppendHash(Hasher, Type.ValueType);
+}
+
+inline void AppendHash(FHasher& Hasher, const Shader::FValue& Value)
+{
+	AppendHash(Hasher, Value.Type);
+	for (int32 i = 0; i < Value.Type.GetNumComponents(); ++i)
 	{
-		return AppendValue(Value).AppendValues(Forward<ArgTypes>(Args)...);
+		AppendHash(Hasher, Value.TryGetComponent(i));
 	}
+}
+
+inline void AppendHashes(FHasher& Hasher) {}
+
+template<typename T, typename... ArgTypes>
+inline void AppendHashes(FHasher& Hasher, const T& Value, ArgTypes&&... Args)
+{
+	AppendHash(Hasher, Value);
+	AppendHashes(Hasher, Forward<ArgTypes>(Args)...);
+}
 
-private:
-	FXxHash64Builder Builder;
-};
 
 /** Root class of the HLSL AST */
 class FNode
@@ -282,13 +292,20 @@ inline bool operator!=(const FRequestedType& Lhs, const FRequestedType& Rhs)
 	return !operator==(Lhs, Rhs);
 }
 
+inline void AppendHash(FHasher& Hasher, const FRequestedType& Value)
+{
+	AppendHash(Hasher, Value.StructType);
+	AppendHash(Hasher, Value.ValueComponentType);
+	AppendHash(Hasher, Value.RequestedComponents);
+}
+
 FRequestedType MakeRequestedType(Shader::EValueComponentType ComponentType, const FRequestedType& RequestedComponents);
 
 struct FPreparedComponent
 {
 	FPreparedComponent() = default;
-	FPreparedComponent(EExpressionEvaluation InEvaluation, FEmitScope* InLoopScope = nullptr)
-		: LoopScope(InLoopScope), Evaluation(InEvaluation)
+	FPreparedComponent(EExpressionEvaluation InEvaluation, FEmitScope* InLoopScope = nullptr, FExpression* InForwardValue = nullptr, int32 InForwardComponentIndex = INDEX_NONE)
+		: LoopScope(InLoopScope), ForwardValue(InForwardValue), ForwardComponentIndex(InForwardComponentIndex), Evaluation(InEvaluation)
 	{
 		check(!IsLoopEvaluation(InEvaluation) || InLoopScope);
 	}
@@ -298,12 +315,16 @@ struct FPreparedComponent
 	EExpressionEvaluation GetEvaluation(const FEmitScope& Scope) const;
 
 	FEmitScope* LoopScope = nullptr;
+	FExpression* ForwardValue = nullptr;
+	int32 ForwardComponentIndex = INDEX_NONE;
 	EExpressionEvaluation Evaluation = EExpressionEvaluation::None;
 };
 inline bool operator==(const FPreparedComponent& Lhs, const FPreparedComponent& Rhs)
 {
 	return Lhs.Evaluation == Rhs.Evaluation &&
-		Lhs.LoopScope == Rhs.LoopScope;
+		Lhs.LoopScope == Rhs.LoopScope &&
+		Lhs.ForwardValue == Rhs.ForwardValue &&
+		Lhs.ForwardComponentIndex == Rhs.ForwardComponentIndex;
 }
 inline bool operator!=(const FPreparedComponent& Lhs, const FPreparedComponent& Rhs)
 {
@@ -326,6 +347,7 @@ public:
 	void SetEvaluation(EExpressionEvaluation Evaluation);
 	void MergeEvaluation(EExpressionEvaluation Evaluation);
 	void SetLoopEvaluation(FEmitScope& Scope, const FRequestedType& RequestedType);
+	void SetForwardValue(FExpression* ForwardValue, const FRequestedType& RequestedType);
 
 	void SetField(const Shader::FStructField* Field, const FPreparedType& FieldType);
 	FPreparedType GetFieldType(const Shader::FStructField* Field) const;
@@ -353,7 +375,7 @@ public:
 	void EnsureNumComponents(int32 NumComponents);
 
 	/** Evaluation type for each component, may be 'None' for components that are unused */
-	TArray<FPreparedComponent, TInlineAllocator<16>> PreparedComponents;
+	TArray<FPreparedComponent, TInlineAllocator<4>> PreparedComponents;
 };
 inline bool operator==(const FPreparedType& Lhs, const FPreparedType& Rhs)
 {
@@ -383,7 +405,6 @@ public:
 private:
 	bool TryMergePreparedType(FEmitContext& Context, const Shader::FStructType* StructType, Shader::EValueComponentType ComponentType);
 
-	FExpression* ForwardValue = nullptr;
 	FPreparedType PreparedType;
 
 	friend class FExpression;
@@ -424,6 +445,7 @@ public:
 	FRequestedType GetRequestedType() const { return PrepareValueResult.PreparedType.GetRequestedType(); }
 	Shader::FType GetType() const { return PrepareValueResult.PreparedType.GetType(); }
 	EExpressionEvaluation GetEvaluation(const FEmitScope& Scope, const FRequestedType& RequestedType) const { return PrepareValueResult.PreparedType.GetEvaluation(Scope, RequestedType); }
+	FExpression* GetForwardValue(const FRequestedType& RequestedType);
 
 	virtual void Reset() override;
 
@@ -436,6 +458,7 @@ public:
 
 protected:
 	virtual void ComputeAnalyticDerivatives(FTree& Tree, FExpressionDerivatives& OutResult) const;
+	virtual FExpression* ComputePreviousFrame(FTree& Tree, const FRequestedType& RequestedType) const;
 	virtual bool PrepareValue(FEmitContext& Context, FEmitScope& Scope, const FRequestedType& RequestedType, FPrepareValueResult& OutResult) const = 0;
 	virtual void EmitValueShader(FEmitContext& Context, FEmitScope& Scope, const FRequestedType& RequestedType, FEmitValueShaderResult& OutResult) const;
 	virtual void EmitValuePreshader(FEmitContext& Context, FEmitScope& Scope, const FRequestedType& RequestedType, Shader::FPreshaderData& OutPreshader) const;
@@ -569,8 +592,8 @@ public:
 	inline FExpression* NewExpression(ArgTypes&&... Args)
 	{
 		FHasher Hasher;
-		Hasher.AppendValue(Private::TTypeHasher<T>::GetHash());
-		Hasher.AppendValues(Forward<ArgTypes>(Args)...);
+		AppendHash(Hasher, Private::TTypeHasher<T>::GetHash());
+		AppendHashes(Hasher, Forward<ArgTypes>(Args)...);
 		const FXxHash64 Hash = Hasher.Finalize();
 		FExpression* Expression = FindExpression(Hash);
 		if (!Expression)
@@ -596,6 +619,7 @@ public:
 	FExpression* NewFunctionCall(FScope& Scope, FFunction* Function, int32 OutputIndex);
 
 	const FExpressionDerivatives& GetAnalyticDerivatives(FExpression* InExpression);
+	FExpression* GetPreviousFrame(FExpression* InExpression, const FRequestedType& RequestedType);
 
 	FScope* NewScope(FScope& Scope);
 	FScope* NewOwnedScope(FStatement& Owner);
diff --git a/Engine/Source/Runtime/Engine/Public/HLSLTree/HLSLTreeCommon.h b/Engine/Source/Runtime/Engine/Public/HLSLTree/HLSLTreeCommon.h
--- a/Engine/Source/Runtime/Engine/Public/HLSLTree/HLSLTreeCommon.h
+++ b/Engine/Source/Runtime/Engine/Public/HLSLTree/HLSLTreeCommon.h
@@ -30,7 +30,8 @@ class FExpressionMaterialParameter : public FExpression
 public:
 	explicit FExpressionMaterialParameter(EMaterialParameterType InType, const FName& InName, const Shader::FValue& InDefaultValue)
 		: ParameterName(InName), DefaultValue(InDefaultValue), ParameterType(InType)
-	{}
+	{
+	}
 
 	FName ParameterName;
 	Shader::FValue DefaultValue;
@@ -77,6 +78,11 @@ enum class EExternalInput : uint8
 	TranslatedWorldPosition,
 	TranslatedWorldPosition_NoOffsets,
 
+	PrevWorldPosition,
+	PrevWorldPosition_NoOffsets,
+	PrevTranslatedWorldPosition,
+	PrevTranslatedWorldPosition_NoOffsets,
+
 	WorldPosition_Ddx,
 	WorldPosition_Ddy,
 };
@@ -84,14 +90,15 @@ static constexpr int32 NumTexCoords = 8;
 
 struct FExternalInputDescription
 {
-	FExternalInputDescription(const TCHAR* InName, Shader::EValueType InType, EExternalInput InDdx = EExternalInput::None, EExternalInput InDdy = EExternalInput::None)
-		: Name(InName), Type(InType), Ddx(InDdx), Ddy(InDdy)
+	FExternalInputDescription(const TCHAR* InName, Shader::EValueType InType, EExternalInput InDdx = EExternalInput::None, EExternalInput InDdy = EExternalInput::None, EExternalInput InPreviousFrame = EExternalInput::None)
+		: Name(InName), Type(InType), Ddx(InDdx), Ddy(InDdy), PreviousFrame(InPreviousFrame)
 	{}
 
 	const TCHAR* Name;
 	Shader::EValueType Type;
 	EExternalInput Ddx;
 	EExternalInput Ddy;
+	EExternalInput PreviousFrame;
 };
 
 FExternalInputDescription GetExternalInputDescription(EExternalInput Input);
@@ -122,6 +129,7 @@ public:
 	EExternalInput InputType;
 
 	virtual void ComputeAnalyticDerivatives(FTree& Tree, FExpressionDerivatives& OutResult) const override;
+	virtual FExpression* ComputePreviousFrame(FTree& Tree, const FRequestedType& RequestedType) const override;
 	virtual bool PrepareValue(FEmitContext& Context, FEmitScope& Scope, const FRequestedType& RequestedType, FPrepareValueResult& OutResult) const override;
 	virtual void EmitValueShader(FEmitContext& Context, FEmitScope& Scope, const FRequestedType& RequestedType, FEmitValueShaderResult& OutResult) const override;
 };
@@ -162,6 +170,7 @@ public:
 	FExpression* StructExpression;
 
 	virtual void ComputeAnalyticDerivatives(FTree& Tree, FExpressionDerivatives& OutResult) const override;
+	virtual FExpression* ComputePreviousFrame(FTree& Tree, const FRequestedType& RequestedType) const override;
 	virtual bool PrepareValue(FEmitContext& Context, FEmitScope& Scope, const FRequestedType& RequestedType, FPrepareValueResult& OutResult) const override;
 	virtual void EmitValueShader(FEmitContext& Context, FEmitScope& Scope, const FRequestedType& RequestedType, FEmitValueShaderResult& OutResult) const override;
 	virtual void EmitValuePreshader(FEmitContext& Context, FEmitScope& Scope, const FRequestedType& RequestedType, Shader::FPreshaderData& OutPreshader) const override;
@@ -186,6 +195,7 @@ public:
 	FExpression* FieldExpression;
 
 	virtual void ComputeAnalyticDerivatives(FTree& Tree, FExpressionDerivatives& OutResult) const override;
+	virtual FExpression* ComputePreviousFrame(FTree& Tree, const FRequestedType& RequestedType) const override;
 	virtual bool PrepareValue(FEmitContext& Context, FEmitScope& Scope, const FRequestedType& RequestedType, FPrepareValueResult& OutResult) const override;
 	virtual void EmitValueShader(FEmitContext& Context, FEmitScope& Scope, const FRequestedType& RequestedType, FEmitValueShaderResult& OutResult) const override;
 	virtual void EmitValuePreshader(FEmitContext& Context, FEmitScope& Scope, const FRequestedType& RequestedType, Shader::FPreshaderData& OutPreshader) const override;
@@ -221,6 +231,7 @@ public:
 	FExpression* Input;
 
 	virtual void ComputeAnalyticDerivatives(FTree& Tree, FExpressionDerivatives& OutResult) const override;
+	virtual FExpression* ComputePreviousFrame(FTree& Tree, const FRequestedType& RequestedType) const override;
 	virtual bool PrepareValue(FEmitContext& Context, FEmitScope& Scope, const FRequestedType& RequestedType, FPrepareValueResult& OutResult) const override;
 	virtual void EmitValueShader(FEmitContext& Context, FEmitScope& Scope, const FRequestedType& RequestedType, FEmitValueShaderResult& OutResult) const override;
 	virtual void EmitValuePreshader(FEmitContext& Context, FEmitScope& Scope, const FRequestedType& RequestedType, Shader::FPreshaderData& OutPreshader) const override;
@@ -240,6 +251,7 @@ public:
 	FExpression* Rhs;
 
 	virtual void ComputeAnalyticDerivatives(FTree& Tree, FExpressionDerivatives& OutResult) const override;
+	virtual FExpression* ComputePreviousFrame(FTree& Tree, const FRequestedType& RequestedType) const override;
 	virtual bool PrepareValue(FEmitContext& Context, FEmitScope& Scope, const FRequestedType& RequestedType, FPrepareValueResult& OutResult) const override;
 	virtual void EmitValueShader(FEmitContext& Context, FEmitScope& Scope, const FRequestedType& RequestedType, FEmitValueShaderResult& OutResult) const override;
 	virtual void EmitValuePreshader(FEmitContext& Context, FEmitScope& Scope, const FRequestedType& RequestedType, Shader::FPreshaderData& OutPreshader) const override;
@@ -270,6 +282,7 @@ public:
 	FExpression* Input;
 
 	virtual void ComputeAnalyticDerivatives(FTree& Tree, FExpressionDerivatives& OutResult) const override;
+	virtual FExpression* ComputePreviousFrame(FTree& Tree, const FRequestedType& RequestedType) const override;
 	virtual bool PrepareValue(FEmitContext& Context, FEmitScope& Scope, const FRequestedType& RequestedType, FPrepareValueResult& OutResult) const override;
 	virtual void EmitValueShader(FEmitContext& Context, FEmitScope& Scope, const FRequestedType& RequestedType, FEmitValueShaderResult& OutResult) const override;
 	virtual void EmitValuePreshader(FEmitContext& Context, FEmitScope& Scope, const FRequestedType& RequestedType, Shader::FPreshaderData& OutPreshader) const override;
@@ -287,6 +300,7 @@ public:
 	FExpression* Rhs;
 
 	virtual void ComputeAnalyticDerivatives(FTree& Tree, FExpressionDerivatives& OutResult) const override;
+	virtual FExpression* ComputePreviousFrame(FTree& Tree, const FRequestedType& RequestedType) const override;
 	virtual bool PrepareValue(FEmitContext& Context, FEmitScope& Scope, const FRequestedType& RequestedType, FPrepareValueResult& OutResult) const override;
 	virtual void EmitValueShader(FEmitContext& Context, FEmitScope& Scope, const FRequestedType& RequestedType, FEmitValueShaderResult& OutResult) const override;
 	virtual void EmitValuePreshader(FEmitContext& Context, FEmitScope& Scope, const FRequestedType& RequestedType, Shader::FPreshaderData& OutPreshader) const override;
diff --git a/Engine/Source/Runtime/Engine/Public/HLSLTree/HLSLTreeEmit.h b/Engine/Source/Runtime/Engine/Public/HLSLTree/HLSLTreeEmit.h
--- a/Engine/Source/Runtime/Engine/Public/HLSLTree/HLSLTreeEmit.h
+++ b/Engine/Source/Runtime/Engine/Public/HLSLTree/HLSLTreeEmit.h
@@ -8,6 +8,7 @@
 #include "Misc/StringBuilder.h"
 #include "Misc/MemStack.h"
 #include "Hash/xxhash.h"
+#include "RHIDefinitions.h"
 #include "HLSLTree/HLSLTreeTypes.h"
 
 class FMaterial;
