Broadcast Wave support for timecodes
#swarm 18820277
#p4v-preflight-copy 18855272

diff --git a/Engine/Source/Editor/AudioEditor/Private/Factories/SoundFactory.cpp b/Engine/Source/Editor/AudioEditor/Private/Factories/SoundFactory.cpp
--- a/Engine/Source/Editor/AudioEditor/Private/Factories/SoundFactory.cpp
+++ b/Engine/Source/Editor/AudioEditor/Private/Factories/SoundFactory.cpp
@@ -325,13 +325,16 @@ UObject* USoundFactory::CreateObject
 			// Attempt to convert to 16 bit audio
 			if (Audio::ConvertAudioToWav(RawWaveData, ConvertedRawWaveData))
 			{
-				WaveInfo = FWaveModInfo();
+				// Icky, Reencoding with SNDFILE will strip the timecode info, so back it up.
+				auto CachedTimeCodeInfo = MoveTemp(WaveInfo.TimecodeInfo);
+				WaveInfo = FWaveModInfo();				
 				if (!WaveInfo.ReadWaveInfo(ConvertedRawWaveData.GetData(), ConvertedRawWaveData.Num(), &ErrorMessage))
 				{
 					Warn->Logf(ELogVerbosity::Error, TEXT("Failed to convert to 16 bit WAV source on import."));
 					GEditor->GetEditorSubsystem<UImportSubsystem>()->BroadcastAssetPostImport(this, nullptr);
 					return nullptr;
 				}
+				WaveInfo.TimecodeInfo = MoveTemp(CachedTimeCodeInfo);
 			}
 
 			// Copy over the data
@@ -554,6 +557,16 @@ UObject* USoundFactory::CreateObject
 			Sound->CuePoints.Add(NewCuePoint);
 		}
 
+		// If we've read some time-code.
+		if (WaveInfo.TimecodeInfo)
+		{
+			Sound->SetTimecodeInfo(*WaveInfo.TimecodeInfo);
+		}
+		else
+		{
+			Sound->SetTimecodeInfo(FSoundWaveTimecodeInfo{});
+		}
+				
 		// Compressed data is now out of date.
 		const bool bRebuildStreamingChunks = FPlatformCompressionUtilities::IsCurrentPlatformUsingStreamCaching();
 		Sound->InvalidateCompressedData(true /* bFreeResources */, bRebuildStreamingChunks);
diff --git a/Engine/Source/Runtime/Engine/Classes/Sound/SoundBase.h b/Engine/Source/Runtime/Engine/Classes/Sound/SoundBase.h
--- a/Engine/Source/Runtime/Engine/Classes/Sound/SoundBase.h
+++ b/Engine/Source/Runtime/Engine/Classes/Sound/SoundBase.h
@@ -7,6 +7,7 @@
 #include "IAudioExtensionPlugin.h"
 #include "Sound/AudioSettings.h"
 #include "Sound/SoundClass.h"
+#include "Sound/SoundTimecodeOffset.h"
 #include "SoundConcurrency.h"
 #include "SoundModulationDestination.h"
 #include "SoundSourceBusSend.h"
@@ -196,6 +197,12 @@ public:
 	UPROPERTY(EditAnywhere, AdvancedDisplay, Instanced, Category = Advanced)
 	TArray<TObjectPtr<UAssetUserData>> AssetUserData;
 
+#if WITH_EDITORONLY_DATA	
+private:
+	UPROPERTY()
+	FSoundTimecodeOffset TimecodeOffset;
+#endif //WITH_EDITORONLY_DATA
+
 public:
 	//~ Begin UObject Interface.
 #if WITH_EDITORONLY_DATA
@@ -312,4 +319,9 @@ public:
 
 	/** Gets all the default parameters for this Asset.  */
 	virtual bool GetAllDefaultParameters(TArray<FAudioParameter>& OutParameters) const { return false; }
+
+#if WITH_EDITORONLY_DATA
+	void SetTimecodeOffset(const FSoundTimecodeOffset& InTimecodeOffset);
+	TOptional<FSoundTimecodeOffset> GetTimecodeOffset() const;
+#endif //WITH_EDITORONLY_DATA
 };
diff --git a/Engine/Source/Runtime/Engine/Classes/Sound/SoundTimecodeOffset.h b/Engine/Source/Runtime/Engine/Classes/Sound/SoundTimecodeOffset.h
new file mode 100644
--- /dev/null
+++ b/Engine/Source/Runtime/Engine/Classes/Sound/SoundTimecodeOffset.h
@@ -0,0 +1,19 @@
+// Copyright Epic Games, Inc. All Rights Reserved.
+#pragma once
+
+#include "SoundTimecodeOffset.generated.h"
+
+USTRUCT()
+struct FSoundTimecodeOffset
+{
+	GENERATED_BODY()
+
+	UPROPERTY()
+	double NumOfSecondsSinceMidnight = 0.0;
+
+	constexpr bool operator==(const FSoundTimecodeOffset& InRhs) const
+	{
+		// We largely only care about the 0 case, so double equality is ok.
+		return NumOfSecondsSinceMidnight == InRhs.NumOfSecondsSinceMidnight;
+	}
+};
diff --git a/Engine/Source/Runtime/Engine/Classes/Sound/SoundWave.h b/Engine/Source/Runtime/Engine/Classes/Sound/SoundWave.h
--- a/Engine/Source/Runtime/Engine/Classes/Sound/SoundWave.h
+++ b/Engine/Source/Runtime/Engine/Classes/Sound/SoundWave.h
@@ -12,6 +12,7 @@
 #include "Misc/Guid.h"
 #include "Async/AsyncWork.h"
 #include "Sound/SoundBase.h"
+#include "Sound/SoundWaveTimecodeInfo.h"
 #include "Serialization/BulkData.h"
 #include "Serialization/BulkDataBuffer.h"
 #include "Sound/SoundGroups.h"
@@ -712,6 +713,14 @@ public:
 
 #endif // WITH_EDITORONLY_DATA
 
+#if WITH_EDITORONLY_DATA
+	
+	/** Information about the time-code from import, if available.  */
+	UPROPERTY()
+	FSoundWaveTimecodeInfo TimecodeInfo;
+
+#endif // WITH_EDITORONLY_DATA
+
 #if WITH_EDITORONLY_DATA
 	UPROPERTY()
 	FString SourceFilePath_DEPRECATED;
@@ -845,6 +854,11 @@ public:
 #endif
 	}
 
+#if WITH_EDITORONLY_DATA
+	void SetTimecodeInfo(const FSoundWaveTimecodeInfo& InTimecode);
+	TOptional<FSoundWaveTimecodeInfo> GetTimecodeInfo() const;
+#endif //WITH_EDITORONLY_DATA	
+
 	/**
 	* Overwrite sample rate. Used for procedural soundwaves, as well as sound waves that are resampled on compress/decompress.
 	*/
diff --git a/Engine/Source/Runtime/Engine/Classes/Sound/SoundWaveTimecodeInfo.h b/Engine/Source/Runtime/Engine/Classes/Sound/SoundWaveTimecodeInfo.h
new file mode 100644
--- /dev/null
+++ b/Engine/Source/Runtime/Engine/Classes/Sound/SoundWaveTimecodeInfo.h
@@ -0,0 +1,48 @@
+// Copyright Epic Games, Inc. All Rights Reserved.
+#pragma once
+
+#include "SoundWaveTimecodeInfo.generated.h"
+
+USTRUCT()
+struct FSoundWaveTimecodeInfo
+{
+	GENERATED_USTRUCT_BODY()
+
+	UPROPERTY(Category = Timecode, VisibleAnywhere)
+	uint64 NumSamplesSinceMidnight = ~0UL;
+
+	UPROPERTY(Category = Timecode, VisibleAnywhere)
+	uint32 NumSamplesPerSecond = 0;
+
+	UPROPERTY(Category = Timecode, VisibleAnywhere)
+	FString Description;
+
+	UPROPERTY(Category = Timecode, VisibleAnywhere)
+	FString OriginatorTime;
+
+	UPROPERTY(Category = Timecode, VisibleAnywhere)
+	FString OriginatorDate;
+
+	UPROPERTY(Category = Timecode, VisibleAnywhere)
+	FString OriginatorDescription;
+
+	UPROPERTY(Category = Timecode, VisibleAnywhere)
+	FString OriginatorReference;
+
+	inline bool operator==(const FSoundWaveTimecodeInfo& InRhs) const
+	{
+		// Note, we don't compare the strings.
+		return NumSamplesSinceMidnight == InRhs.NumSamplesSinceMidnight &&
+			NumSamplesPerSecond == InRhs.NumSamplesPerSecond;
+	}
+
+	inline double GetNumSecondsSinceMidnight() const
+	{
+		if( NumSamplesSinceMidnight != ~0 && NumSamplesPerSecond > 0)
+		{
+			return (double)NumSamplesSinceMidnight / NumSamplesPerSecond;
+		}
+		return 0.0;
+	}
+};
+
diff --git a/Engine/Source/Runtime/Engine/Private/Audio.cpp b/Engine/Source/Runtime/Engine/Private/Audio.cpp
--- a/Engine/Source/Runtime/Engine/Private/Audio.cpp
+++ b/Engine/Source/Runtime/Engine/Private/Audio.cpp
@@ -23,6 +23,7 @@
 #include "Sound/SoundSubmix.h"
 #include "Sound/SoundNodeWavePlayer.h"
 #include "Sound/SoundWave.h"
+#include "Sound/SoundWaveTimecodeInfo.h"
 #include "Sound/QuartzQuantizationUtilities.h"
 #include "UObject/UObjectHash.h"
 #include "UObject/UObjectIterator.h"
@@ -1198,6 +1199,38 @@ struct FRiffLabeledTextChunk
 	uint16 CodePage;		// Unused
 };
 
+// Specification of the Broadcast Wave Format(BWF)
+// https://tech.ebu.ch/docs/tech/tech3285.pdf
+struct FRiffBroadcastAudioExtension 
+{
+	uint32 ChunkID;					// 'bext'
+	uint32 ChunkDataSize;			// Depends on contained text
+	uint8 Description[256];			// ASCII : «Description of the sound sequence» 
+	uint8 Originator[32];			// ASCII : «Name of the originator» 
+	uint8 OriginatorReference[32];	// ASCII : «Reference of the originator» 
+	uint8 OriginationDate[10];		// ASCII : «yyyy:mm:dd» 
+	uint8 OriginationTime[8];		// ASCII : «hh:mm:ss» 
+	uint32 TimeReferenceLow;		// First sample count since midnight, low word 
+	uint32 TimeReferenceHigh;		// First sample count since midnight, high word	
+
+	uint16 Version;					// Version of the BWF; unsigned binary number 
+	uint8 UMID[64];					// Binary SMPTE UMID 
+	uint16 LoudnessValue;			// WORD : «Integrated Loudness Value of the filein LUFS (multiplied by 100) » 
+	uint16 LoudnessRange;			// WORD : «Loudness Range of the file in LU (multiplied by 100) » 
+	uint16 MaxTruePeakLevel;		// WORD : «Maximum True Peak Level of the file expressed as dBTP (multiplied by 100) » 
+	uint16 MaxMomentaryLoudness;	// WORD : «Highest value of the Momentary Loudness Level of the file in LUFS (multiplied by 100) » 
+	uint16 MaxShortTermLoudness;	// WORD : «Highest value of the Short-TermLoudness Level of the file in LUFS (multiplied by 100) » 
+	uint16 Reserved[180];			// 180 bytes, reserved for future use, set to “NULL” 
+	uint8 CodingHistory[1];			// ASCII : « History coding » (truncated here are don't care about it) 
+};
+
+struct FRiffIXmlChunk
+{
+	uint32 ChunkID;					// 'iXML'
+	uint32 ChunkDataSize;			// Depends on contained text
+	uint8 XmlText[1];				// Raw XML blob. (truncated as the size is based on the chunk size).
+};	
+
 // FExtendedFormatChunk subformat GUID.
 struct FSubformatGUID
 {
@@ -1256,6 +1289,8 @@ bool IsKnownChunkId(uint32 ChunkId)
 		UE_mmioFOURCC('s', 'm', 'p', 'l'),
 		UE_mmioFOURCC('i', 'n', 's', 't'),
 		UE_mmioFOURCC('a', 'c', 'i', 'd'),
+		UE_mmioFOURCC('b', 'e', 'x', 't'),
+		UE_mmioFOURCC('i', 'X', 'M', 'L'),
 	};
 
 	return Algo::Find(KnownIds, ChunkId) != nullptr;
@@ -1636,6 +1671,38 @@ bool FWaveModInfo::ReadWaveInfo( const uint8* WaveData, int32 WaveDataSize, FStr
 		}
 	}
 
+	// Look for the cue chunks
+	RiffChunk = FindRiffChunk(RiffChunkStart, WaveDataEnd, UE_mmioFOURCC('b', 'e', 'x', 't'));
+
+	if (RiffChunk != nullptr)
+	{
+		const FRiffBroadcastAudioExtension* BextChunk = (const FRiffBroadcastAudioExtension*)((uint8*)RiffChunk);
+		uint64 NumSamplesSinceMidnight = ((uint64) BextChunk->TimeReferenceHigh << 32) | (uint64) BextChunk->TimeReferenceLow;
+
+		// Only record info if there's a valid non-zero time-code.
+		if(NumSamplesSinceMidnight > 0 && FmtChunk && FmtChunk->nSamplesPerSec)
+		{
+			TimecodeInfo = MakePimpl<FSoundWaveTimecodeInfo, EPimplPtrMode::DeepCopy>();
+				
+			TimecodeInfo->NumSamplesSinceMidnight   = NumSamplesSinceMidnight;
+			TimecodeInfo->NumSamplesPerSecond		= FmtChunk->nSamplesPerSec;
+			TimecodeInfo->Description				= StringCast<TCHAR>((ANSICHAR*)BextChunk->Description).Get();
+			TimecodeInfo->OriginatorDescription		= StringCast<TCHAR>((ANSICHAR*)BextChunk->Originator).Get();
+			TimecodeInfo->OriginatorDate			= StringCast<TCHAR>((ANSICHAR*)BextChunk->OriginationDate).Get();
+			TimecodeInfo->OriginatorReference		= StringCast<TCHAR>((ANSICHAR*)BextChunk->OriginatorReference).Get();
+			TimecodeInfo->OriginatorTime			= StringCast<TCHAR>((ANSICHAR*)BextChunk->OriginationTime).Get();
+		}
+	}
+
+	// Look for the cue chunks
+	RiffChunk = FindRiffChunk(RiffChunkStart, WaveDataEnd, UE_mmioFOURCC('i', 'X', 'M', 'L'));
+
+	if (RiffChunk != nullptr)
+	{
+		const FRiffIXmlChunk* IXmlChunk = (const FRiffIXmlChunk*)((uint8*)RiffChunk);
+		FString XmlString = StringCast<TCHAR>((ANSICHAR*)IXmlChunk->XmlText).Get();
+		// TODO.
+	}
 
 	return true;
 }
diff --git a/Engine/Source/Runtime/Engine/Private/SoundBase.cpp b/Engine/Source/Runtime/Engine/Private/SoundBase.cpp
--- a/Engine/Source/Runtime/Engine/Private/SoundBase.cpp
+++ b/Engine/Source/Runtime/Engine/Private/SoundBase.cpp
@@ -490,3 +490,22 @@ bool USoundBase::IsParameterValid(const FAudioParameter& InParameter) const
 
 	return true;
 }
+
+#if WITH_EDITORONLY_DATA
+
+void USoundBase::SetTimecodeOffset(const FSoundTimecodeOffset& InTimecodeOffset)
+{
+	TimecodeOffset = InTimecodeOffset;
+}
+
+TOptional<FSoundTimecodeOffset> USoundBase::GetTimecodeOffset() const
+{
+	static const FSoundTimecodeOffset Defaults;
+	if(TimecodeOffset == Defaults)
+	{
+		return {};
+	}
+	return TimecodeOffset;
+}
+
+#endif //WITH_EDITORONLY_DATA
\ No newline at end of file
diff --git a/Engine/Source/Runtime/Engine/Private/SoundWave.cpp b/Engine/Source/Runtime/Engine/Private/SoundWave.cpp
--- a/Engine/Source/Runtime/Engine/Private/SoundWave.cpp
+++ b/Engine/Source/Runtime/Engine/Private/SoundWave.cpp
@@ -650,6 +650,32 @@ void USoundWave::GetResourceSizeEx(FResourceSizeEx& CumulativeResourceSize)
 }
 
 
+#if WITH_EDITORONLY_DATA
+
+void USoundWave::SetTimecodeInfo(const FSoundWaveTimecodeInfo& InTimecode)
+{
+	// Keep editor info on the SoundWave.
+	TimecodeInfo = InTimecode;
+
+	// Set SoundBase offset.
+	FSoundTimecodeOffset Offset;
+	Offset.NumOfSecondsSinceMidnight = InTimecode.GetNumSecondsSinceMidnight();
+	SetTimecodeOffset(Offset);
+}
+
+TOptional<FSoundWaveTimecodeInfo> USoundWave::GetTimecodeInfo() const
+{
+	// Do we have any time-code info that's not the defaults?
+	static const FSoundWaveTimecodeInfo Defaults;
+	if(Defaults == TimecodeInfo)
+	{
+		return {};
+	}
+	return TimecodeInfo;
+}
+
+#endif //WITH_EDITORONLY_DATA
+
 int32 USoundWave::GetResourceSizeForFormat(FName Format)
 {
 	return GetCompressedDataSize(Format);
diff --git a/Engine/Source/Runtime/Engine/Public/Audio.h b/Engine/Source/Runtime/Engine/Public/Audio.h
--- a/Engine/Source/Runtime/Engine/Public/Audio.h
+++ b/Engine/Source/Runtime/Engine/Public/Audio.h
@@ -92,6 +92,7 @@ class USoundNode;
 struct FWaveInstance;
 struct FReverbSettings;
 struct FSampleLoop;
+struct FSoundWaveTimecodeInfo;
 
 enum ELoopingMode
 {
@@ -875,6 +876,9 @@ public:
 	// List of cues parsed from the wave file
 	TArray<FWaveCue> WaveCues;
 
+	// Timecode data if it was found on import.
+	TPimplPtr<FSoundWaveTimecodeInfo, EPimplPtrMode::DeepCopy> TimecodeInfo;
+
 	// Constructor.
 	FWaveModInfo()
 	{
