Add support for time/panner with new HLSL translator

diff --git a/Engine/Source/Runtime/Engine/Classes/Materials/MaterialExpressionDeltaTime.h b/Engine/Source/Runtime/Engine/Classes/Materials/MaterialExpressionDeltaTime.h
--- a/Engine/Source/Runtime/Engine/Classes/Materials/MaterialExpressionDeltaTime.h
+++ b/Engine/Source/Runtime/Engine/Classes/Materials/MaterialExpressionDeltaTime.h
@@ -15,6 +15,7 @@ class UMaterialExpressionDeltaTime : public UMaterialExpression
 
 	//~ Begin UMaterialExpression Interface
 #if WITH_EDITOR
+	virtual bool GenerateHLSLExpression(FMaterialHLSLGenerator& Generator, UE::HLSLTree::FScope& Scope, int32 OutputIndex, UE::HLSLTree::FExpression*& OutExpression) override;
 	virtual int32 Compile(class FMaterialCompiler* Compiler, int32 OutputIndex) override;
 	virtual void GetCaption(TArray<FString>& OutCaptions) const override;
 #endif
diff --git a/Engine/Source/Runtime/Engine/Classes/Materials/MaterialExpressionPanner.h b/Engine/Source/Runtime/Engine/Classes/Materials/MaterialExpressionPanner.h
--- a/Engine/Source/Runtime/Engine/Classes/Materials/MaterialExpressionPanner.h
+++ b/Engine/Source/Runtime/Engine/Classes/Materials/MaterialExpressionPanner.h
@@ -40,6 +40,7 @@ class UMaterialExpressionPanner : public UMaterialExpression
 
 	//~ Begin UMaterialExpression Interface
 #if WITH_EDITOR
+	virtual bool GenerateHLSLExpression(FMaterialHLSLGenerator& Generator, UE::HLSLTree::FScope& Scope, int32 OutputIndex, UE::HLSLTree::FExpression*& OutExpression) override;
 	virtual int32 Compile(class FMaterialCompiler* Compiler, int32 OutputIndex) override;
 	virtual void GetCaption(TArray<FString>& OutCaptions) const override;
 	virtual bool NeedsRealtimePreview() override { return Time.Expression==NULL && (SpeedX != 0.f || SpeedY != 0.f); }
diff --git a/Engine/Source/Runtime/Engine/Classes/Materials/MaterialExpressionTime.h b/Engine/Source/Runtime/Engine/Classes/Materials/MaterialExpressionTime.h
--- a/Engine/Source/Runtime/Engine/Classes/Materials/MaterialExpressionTime.h
+++ b/Engine/Source/Runtime/Engine/Classes/Materials/MaterialExpressionTime.h
@@ -27,6 +27,7 @@ class UMaterialExpressionTime : public UMaterialExpression
 
 	//~ Begin UMaterialExpression Interface
 #if WITH_EDITOR
+	virtual bool GenerateHLSLExpression(FMaterialHLSLGenerator& Generator, UE::HLSLTree::FScope& Scope, int32 OutputIndex, UE::HLSLTree::FExpression*& OutExpression) override;
 	virtual int32 Compile(class FMaterialCompiler* Compiler, int32 OutputIndex) override;
 	virtual void GetCaption(TArray<FString>& OutCaptions) const override;
 	virtual bool NeedsRealtimePreview() override { return true; }
diff --git a/Engine/Source/Runtime/Engine/Private/HLSLTree/HLSLTree.cpp b/Engine/Source/Runtime/Engine/Private/HLSLTree/HLSLTree.cpp
--- a/Engine/Source/Runtime/Engine/Private/HLSLTree/HLSLTree.cpp
+++ b/Engine/Source/Runtime/Engine/Private/HLSLTree/HLSLTree.cpp
@@ -103,6 +103,7 @@ FUnaryOpDescription GetUnaryOpDesription(EUnaryOp Op)
 	case EUnaryOp::None: return FUnaryOpDescription(TEXT("None"), TEXT(""), Shader::EPreshaderOpcode::Nop); break;
 	case EUnaryOp::Neg: return FUnaryOpDescription(TEXT("Neg"), TEXT("-"), Shader::EPreshaderOpcode::Neg); break;
 	case EUnaryOp::Rcp: return FUnaryOpDescription(TEXT("Rcp"), TEXT("/"), Shader::EPreshaderOpcode::Rcp); break;
+	case EUnaryOp::Frac: return FUnaryOpDescription(TEXT("Frac"), TEXT("frac"), Shader::EPreshaderOpcode::Frac); break;
 	default: checkNoEntry(); return FUnaryOpDescription();
 	}
 }
@@ -116,6 +117,7 @@ FBinaryOpDescription GetBinaryOpDesription(EBinaryOp Op)
 	case EBinaryOp::Sub: return FBinaryOpDescription(TEXT("Subtract"), TEXT("-"), Shader::EPreshaderOpcode::Sub); break;
 	case EBinaryOp::Mul: return FBinaryOpDescription(TEXT("Multiply"), TEXT("*"), Shader::EPreshaderOpcode::Mul); break;
 	case EBinaryOp::Div: return FBinaryOpDescription(TEXT("Divide"), TEXT("/"), Shader::EPreshaderOpcode::Div); break;
+	case EBinaryOp::Fmod: return FBinaryOpDescription(TEXT("Fmod"), TEXT("%"), Shader::EPreshaderOpcode::Fmod); break;
 	case EBinaryOp::Less: return FBinaryOpDescription(TEXT("Less"), TEXT("<"), Shader::EPreshaderOpcode::Less); break;
 	default: checkNoEntry(); return FBinaryOpDescription();
 	}
diff --git a/Engine/Source/Runtime/Engine/Private/HLSLTree/HLSLTreeCommon.cpp b/Engine/Source/Runtime/Engine/Private/HLSLTree/HLSLTreeCommon.cpp
--- a/Engine/Source/Runtime/Engine/Private/HLSLTree/HLSLTreeCommon.cpp
+++ b/Engine/Source/Runtime/Engine/Private/HLSLTree/HLSLTreeCommon.cpp
@@ -59,6 +59,13 @@ FExternalInputDescription GetExternalInputDescription(EExternalInput Input)
 	case EExternalInput::PrevTranslatedWorldPosition: return FExternalInputDescription(TEXT("PrevTranslatedWorldPosition"), Shader::EValueType::Float3, EExternalInput::WorldPosition_Ddx, EExternalInput::WorldPosition_Ddy);
 	case EExternalInput::PrevTranslatedWorldPosition_NoOffsets: return FExternalInputDescription(TEXT("PrevTranslatedWorldPosition_NoOffsets"), Shader::EValueType::Float3, EExternalInput::WorldPosition_Ddx, EExternalInput::WorldPosition_Ddy);
 
+	case EExternalInput::GameTime: return FExternalInputDescription(TEXT("GameTime"), Shader::EValueType::Float1, EExternalInput::None, EExternalInput::None, EExternalInput::PrevGameTime);
+	case EExternalInput::RealTime: return FExternalInputDescription(TEXT("RealTime"), Shader::EValueType::Float1, EExternalInput::None, EExternalInput::None, EExternalInput::PrevRealTime);
+	case EExternalInput::DeltaTime: return FExternalInputDescription(TEXT("DeltaTime"), Shader::EValueType::Float1);
+
+	case EExternalInput::PrevGameTime: return FExternalInputDescription(TEXT("PrevGameTime"), Shader::EValueType::Float1);
+	case EExternalInput::PrevRealTime: return FExternalInputDescription(TEXT("PrevRealTime"), Shader::EValueType::Float1);
+
 	case EExternalInput::WorldPosition_Ddx: return FExternalInputDescription(TEXT("WorldPosition_Ddx"), Shader::EValueType::Float3);
 	case EExternalInput::WorldPosition_Ddy: return FExternalInputDescription(TEXT("WorldPosition_Ddx"), Shader::EValueType::Float3);
 	default: checkNoEntry(); return FExternalInputDescription(TEXT("Invalid"), Shader::EValueType::Void);
@@ -277,6 +284,11 @@ void FExpressionExternalInput::EmitValueShader(FEmitContext& Context, FEmitScope
 		case EExternalInput::PrevWorldPosition_NoOffsets: Code = TEXT("GetPrevWorldPosition_NoMaterialOffsets(Parameters)"); break;
 		case EExternalInput::PrevTranslatedWorldPosition: Code = TEXT("GetPrevTranslatedWorldPosition(Parameters)"); break;
 		case EExternalInput::PrevTranslatedWorldPosition_NoOffsets: Code = TEXT("GetPrevTranslatedWorldPosition_NoMaterialOffsets(Parameters)"); break;
+		case EExternalInput::GameTime: Code = TEXT("View.GameTime"); break;
+		case EExternalInput::RealTime: Code = TEXT("View.RealTime"); break;
+		case EExternalInput::DeltaTime: Code = TEXT("View.DeltaTime"); break;
+		case EExternalInput::PrevGameTime: Code = TEXT("View.PrevFrameGameTime"); break;
+		case EExternalInput::PrevRealTime: Code = TEXT("View.PrevFrameRealTime"); break;
 		case EExternalInput::WorldPosition_Ddx: Code = TEXT("Parameters.WorldPosition_DDX"); break;
 		case EExternalInput::WorldPosition_Ddy: Code = TEXT("Parameters.WorldPosition_DDY"); break;
 		default:
@@ -666,6 +678,9 @@ void FExpressionUnaryOp::ComputeAnalyticDerivatives(FTree& Tree, FExpressionDeri
 			OutResult.ExpressionDdy = Tree.NewMul(dFdA, InputDerivatives.ExpressionDdy);
 			break;
 		}
+		case EUnaryOp::Frac:
+			OutResult = InputDerivatives;
+			break;
 		default:
 			checkNoEntry();
 			break;
@@ -731,6 +746,17 @@ void FExpressionUnaryOp::EmitValueShader(FEmitContext& Context, FEmitScope& Scop
 			OutResult.Code = Context.EmitExpression(Scope, ResultType, TEXT("rcp(%)"), InputValue);
 		}
 		break;
+	case EUnaryOp::Frac:
+		if (InputTypeDesc.ComponentType == Shader::EValueComponentType::Double)
+		{
+			ResultType = Shader::MakeValueType(Shader::EValueComponentType::Float, ResultType.GetNumComponents());
+			OutResult.Code = Context.EmitExpression(Scope, ResultType, TEXT("LWCFrac(%)"), InputValue);
+		}
+		else
+		{
+			OutResult.Code = Context.EmitExpression(Scope, ResultType, TEXT("frac(%)"), InputValue);
+		}
+		break;
 	default:
 		checkNoEntry();
 		break;
@@ -796,6 +822,11 @@ void FExpressionBinaryOp::ComputeAnalyticDerivatives(FTree& Tree, FExpressionDer
 			OutResult.ExpressionDdy = Tree.NewAdd(Tree.NewMul(dFdA, LhsDerivatives.ExpressionDdy), Tree.NewMul(dFdB, RhsDerivatives.ExpressionDdy));
 			break;
 		}
+		case EBinaryOp::Fmod:
+			// Only valid when B derivatives are zero.
+			// We can't really do anything meaningful in the non-zero case.
+			OutResult = LhsDerivatives;
+			break;
 		default:
 			checkNoEntry();
 			break;
@@ -845,73 +876,102 @@ bool FExpressionBinaryOp::PrepareValue(FEmitContext& Context, FEmitScope& Scope,
 
 namespace Private
 {
-FRequestedType GetBinaryOpInputType(const FRequestedType& RequestedType, EBinaryOp Op, const FRequestedType& LhsType, const FRequestedType& RhsType)
+struct FBinaryOpTypes
 {
-	const int32 NumRequestedComponents = RequestedType.GetNumComponents();
+	FBinaryOpTypes(const FRequestedType& InType) : LhsType(InType), RhsType(LhsType), ResultType(InType.GetType()) {}
+	FBinaryOpTypes(const FRequestedType& InLhsType, const FRequestedType& InRhsType, const Shader::FType& InResultType) : LhsType(InLhsType), RhsType(InRhsType), ResultType(InResultType) {}
+
+	FRequestedType LhsType;
+	FRequestedType RhsType;
+	Shader::FType ResultType;
+};
+FBinaryOpTypes GetBinaryOpTypes(const FRequestedType& RequestedType, EBinaryOp Op, const FRequestedType& LhsType, const FRequestedType& RhsType)
+{
+	const int32 NumComponents = RequestedType.GetNumComponents();
 	const Shader::EValueComponentType InputComponentType = Shader::CombineComponentTypes(LhsType.ValueComponentType, RhsType.ValueComponentType);
-	return MakeRequestedType(InputComponentType, RequestedType);
+	const Shader::EValueComponentType NonLWCInputComponentType = Shader::MakeNonLWCType(InputComponentType);
+	const FRequestedType InputType = MakeRequestedType(InputComponentType, RequestedType);
+	const FRequestedType NonLWCInputType = MakeRequestedType(NonLWCInputComponentType, RequestedType);
+
+	switch (Op)
+	{
+	case EBinaryOp::Less:
+		return FBinaryOpTypes(InputType, InputType, Shader::MakeValueType(Shader::EValueComponentType::Bool, NumComponents));
+	case EBinaryOp::Fmod:
+		return FBinaryOpTypes(InputType, NonLWCInputType, NonLWCInputType.GetType());
+	default:
+		return InputType;
+	}
 }
 } // namespace Private
 
 void FExpressionBinaryOp::EmitValueShader(FEmitContext& Context, FEmitScope& Scope, const FRequestedType& RequestedType, FEmitValueShaderResult& OutResult) const
 {
-	const FRequestedType InputType = Private::GetBinaryOpInputType(RequestedType, Op, Lhs->GetRequestedType(), Rhs->GetRequestedType());
-	const Shader::EValueComponentType InputComponentType = InputType.ValueComponentType;
-	Shader::FType ResultType = InputType.GetType();
-	FEmitShaderExpression* LhsValue = Lhs->GetValueShader(Context, Scope, InputType);
-	FEmitShaderExpression* RhsValue = Rhs->GetValueShader(Context, Scope, InputType);
+	const Private::FBinaryOpTypes Types = Private::GetBinaryOpTypes(RequestedType, Op, Lhs->GetRequestedType(), Rhs->GetRequestedType());
+	const bool bLWC = Types.LhsType.ValueComponentType == Shader::EValueComponentType::Double;
+	FEmitShaderExpression* LhsValue = Lhs->GetValueShader(Context, Scope, Types.LhsType);
+	FEmitShaderExpression* RhsValue = Rhs->GetValueShader(Context, Scope, Types.RhsType);
 
 	switch (Op)
 	{
 	case EBinaryOp::Add:
-		if (InputComponentType == Shader::EValueComponentType::Double)
+		if (bLWC)
 		{
-			OutResult.Code = Context.EmitExpression(Scope, ResultType, TEXT("LWCAdd(%, %)"), LhsValue, RhsValue);
+			OutResult.Code = Context.EmitExpression(Scope, Types.ResultType, TEXT("LWCAdd(%, %)"), LhsValue, RhsValue);
 		}
 		else
 		{
-			OutResult.Code = Context.EmitExpression(Scope, ResultType, TEXT("(% + %)"), LhsValue, RhsValue);
+			OutResult.Code = Context.EmitExpression(Scope, Types.ResultType, TEXT("(% + %)"), LhsValue, RhsValue);
 		}
 		break;
 	case EBinaryOp::Sub:
-		if (InputComponentType == Shader::EValueComponentType::Double)
+		if (bLWC)
 		{
-			OutResult.Code = Context.EmitExpression(Scope, ResultType, TEXT("LWCSubtract(%, %)"), LhsValue, RhsValue);
+			OutResult.Code = Context.EmitExpression(Scope, Types.ResultType, TEXT("LWCSubtract(%, %)"), LhsValue, RhsValue);
 		}
 		else
 		{
-			OutResult.Code = Context.EmitExpression(Scope, ResultType, TEXT("(% - %)"), LhsValue, RhsValue);
+			OutResult.Code = Context.EmitExpression(Scope, Types.ResultType, TEXT("(% - %)"), LhsValue, RhsValue);
 		}
 		break;
 	case EBinaryOp::Mul:
-		if (InputComponentType == Shader::EValueComponentType::Double)
+		if (bLWC)
+		{
+			OutResult.Code = Context.EmitExpression(Scope, Types.ResultType, TEXT("LWCMultiply(%, %)"), LhsValue, RhsValue);
+		}
+		else
 		{
-			OutResult.Code = Context.EmitExpression(Scope, ResultType, TEXT("LWCMultiply(%, %)"), LhsValue, RhsValue);
+			OutResult.Code = Context.EmitExpression(Scope, Types.ResultType, TEXT("(% * %)"), LhsValue, RhsValue);
+		}
+		break;
+	case EBinaryOp::Fmod:
+		if (bLWC)
+		{
+			OutResult.Code = Context.EmitExpression(Scope, Types.ResultType, TEXT("LWCFmod(%, %)"), LhsValue, RhsValue);
 		}
 		else
 		{
-			OutResult.Code = Context.EmitExpression(Scope, ResultType, TEXT("(% * %)"), LhsValue, RhsValue);
+			OutResult.Code = Context.EmitExpression(Scope, Types.ResultType, TEXT("fmod(%, %)"), LhsValue, RhsValue);
 		}
 		break;
 	case EBinaryOp::Div:
-		if (InputComponentType == Shader::EValueComponentType::Double)
+		if (bLWC)
 		{
-			OutResult.Code = Context.EmitExpression(Scope, ResultType, TEXT("LWCDivide(%, %)"), LhsValue, RhsValue);
+			OutResult.Code = Context.EmitExpression(Scope, Types.ResultType, TEXT("LWCDivide(%, %)"), LhsValue, RhsValue);
 		}
 		else
 		{
-			OutResult.Code = Context.EmitExpression(Scope, ResultType, TEXT("(% / %)"), LhsValue, RhsValue);
+			OutResult.Code = Context.EmitExpression(Scope, Types.ResultType, TEXT("(% / %)"), LhsValue, RhsValue);
 		}
 		break;
 	case EBinaryOp::Less:
-		ResultType = Shader::MakeValueType(Shader::EValueComponentType::Bool, ResultType.GetNumComponents());
-		if (InputComponentType == Shader::EValueComponentType::Double)
+		if (bLWC)
 		{
-			OutResult.Code = Context.EmitExpression(Scope, ResultType, TEXT("LWCLess(%, %)"), LhsValue, RhsValue);
+			OutResult.Code = Context.EmitExpression(Scope, Types.ResultType, TEXT("LWCLess(%, %)"), LhsValue, RhsValue);
 		}
 		else
 		{
-			OutResult.Code = Context.EmitExpression(Scope, ResultType, TEXT("(% < %)"), LhsValue, RhsValue);
+			OutResult.Code = Context.EmitExpression(Scope, Types.ResultType, TEXT("(% < %)"), LhsValue, RhsValue);
 		}
 		break;
 	default:
@@ -922,12 +982,12 @@ void FExpressionBinaryOp::EmitValueShader(FEmitContext& Context, FEmitScope& Sco
 
 void FExpressionBinaryOp::EmitValuePreshader(FEmitContext& Context, FEmitScope& Scope, const FRequestedType& RequestedType, Shader::FPreshaderData& OutPreshader) const
 {
-	const FRequestedType InputType = Private::GetBinaryOpInputType(RequestedType, Op, Lhs->GetRequestedType(), Rhs->GetRequestedType());
+	const Private::FBinaryOpTypes Types = Private::GetBinaryOpTypes(RequestedType, Op, Lhs->GetRequestedType(), Rhs->GetRequestedType());
 	const FBinaryOpDescription OpDesc = GetBinaryOpDesription(Op);
 	check(OpDesc.PreshaderOpcode != Shader::EPreshaderOpcode::Nop);
 
-	Lhs->GetValuePreshader(Context, Scope, InputType, OutPreshader);
-	Rhs->GetValuePreshader(Context, Scope, InputType, OutPreshader);
+	Lhs->GetValuePreshader(Context, Scope, Types.LhsType, OutPreshader);
+	Rhs->GetValuePreshader(Context, Scope, Types.RhsType, OutPreshader);
 
 	check(Context.PreshaderStackPosition > 0);
 	Context.PreshaderStackPosition--;
diff --git a/Engine/Source/Runtime/Engine/Private/Materials/MaterialExpressionHLSL.cpp b/Engine/Source/Runtime/Engine/Private/Materials/MaterialExpressionHLSL.cpp
--- a/Engine/Source/Runtime/Engine/Private/Materials/MaterialExpressionHLSL.cpp
+++ b/Engine/Source/Runtime/Engine/Private/Materials/MaterialExpressionHLSL.cpp
@@ -21,6 +21,9 @@
 #include "Materials/MaterialExpressionScalarParameter.h"
 #include "Materials/MaterialExpressionStaticBoolParameter.h"
 #include "Materials/MaterialExpressionWorldPosition.h"
+#include "Materials/MaterialExpressionTime.h"
+#include "Materials/MaterialExpressionDeltaTime.h"
+#include "Materials/MaterialExpressionPanner.h"
 #include "Materials/MaterialExpressionTextureCoordinate.h"
 #include "Materials/MaterialExpressionTextureSample.h"
 #include "Materials/MaterialExpressionTextureSampleParameter.h"
@@ -140,6 +143,7 @@ bool UMaterialExpressionStaticBoolParameter::GenerateHLSLExpression(FMaterialHLS
 	OutExpression = Generator.GetTree().NewExpression<UE::HLSLTree::FExpressionMaterialParameter>(EMaterialParameterType::StaticSwitch, ParameterName, (bool)DefaultValue);
 	return true;
 }
+
 bool UMaterialExpressionWorldPosition::GenerateHLSLExpression(FMaterialHLSLGenerator& Generator, UE::HLSLTree::FScope& Scope, int32 OutputIndex, UE::HLSLTree::FExpression*& OutExpression)
 {
 	UE::HLSLTree::EExternalInput InputType = UE::HLSLTree::EExternalInput::None;
@@ -157,6 +161,86 @@ bool UMaterialExpressionWorldPosition::GenerateHLSLExpression(FMaterialHLSLGener
 	return true;
 }
 
+bool UMaterialExpressionTime::GenerateHLSLExpression(FMaterialHLSLGenerator& Generator, UE::HLSLTree::FScope& Scope, int32 OutputIndex, UE::HLSLTree::FExpression*& OutExpression)
+{
+	using namespace UE::HLSLTree;
+
+	if (bOverride_Period && Period == 0.0f)
+	{
+		OutExpression = Generator.NewConstant(0.0f);
+		return true;
+	}
+
+	EExternalInput InputType = bIgnorePause ? EExternalInput::RealTime : EExternalInput::GameTime;
+	OutExpression = Generator.GetTree().NewExpression<UE::HLSLTree::FExpressionExternalInput>(InputType);
+	if (bOverride_Period)
+	{
+		OutExpression = Generator.GetTree().NewFmod(OutExpression, Generator.NewConstant(Period));
+	}
+	return true;
+}
+
+bool UMaterialExpressionDeltaTime::GenerateHLSLExpression(FMaterialHLSLGenerator& Generator, UE::HLSLTree::FScope& Scope, int32 OutputIndex, UE::HLSLTree::FExpression*& OutExpression)
+{
+	OutExpression = Generator.GetTree().NewExpression<UE::HLSLTree::FExpressionExternalInput>(UE::HLSLTree::EExternalInput::DeltaTime);
+	return true;
+}
+
+bool UMaterialExpressionPanner::GenerateHLSLExpression(FMaterialHLSLGenerator& Generator, UE::HLSLTree::FScope& Scope, int32 OutputIndex, UE::HLSLTree::FExpression*& OutExpression)
+{
+#if 0
+	int32 TimeArg = Time.GetTracedInput().Expression ? Time.Compile(Compiler) : Compiler->GameTime(false, 0.0f);
+	bool bIsSpeedExpressionValid = Speed.GetTracedInput().Expression != nullptr;
+	int32 SpeedVectorArg = bIsSpeedExpressionValid ? Speed.Compile(Compiler) : INDEX_NONE;
+	int32 SpeedXArg = bIsSpeedExpressionValid ? Compiler->ComponentMask(SpeedVectorArg, true, false, false, false) : Compiler->Constant(SpeedX);
+	int32 SpeedYArg = bIsSpeedExpressionValid ? Compiler->ComponentMask(SpeedVectorArg, false, true, false, false) : Compiler->Constant(SpeedY);
+	int32 Arg1;
+	int32 Arg2;
+	if (bFractionalPart)
+	{
+		// Note: this is to avoid (delay) divergent accuracy issues as GameTime increases.
+		// TODO: C++ to calculate its phase via per frame time delta.
+		Arg1 = Compiler->PeriodicHint(Compiler->Frac(Compiler->Mul(TimeArg, SpeedXArg)));
+		Arg2 = Compiler->PeriodicHint(Compiler->Frac(Compiler->Mul(TimeArg, SpeedYArg)));
+	}
+	else
+	{
+		Arg1 = Compiler->PeriodicHint(Compiler->Mul(TimeArg, SpeedXArg));
+		Arg2 = Compiler->PeriodicHint(Compiler->Mul(TimeArg, SpeedYArg));
+	}
+
+	int32 Arg3 = Coordinate.GetTracedInput().Expression ? Coordinate.Compile(Compiler) : Compiler->TextureCoordinate(ConstCoordinate, false, false);
+	return Compiler->Add(
+		Compiler->AppendVector(
+			Arg1,
+			Arg2
+		),
+		Arg3
+	);
+#endif
+	using namespace UE::HLSLTree;
+
+	FExpression* ExpressionTime = Time.TryAcquireHLSLExpression(Generator, Scope);
+	if (!ExpressionTime)
+	{
+		ExpressionTime = Generator.GetTree().NewExpression<FExpressionExternalInput>(EExternalInput::GameTime);
+	}
+	FExpression* ExpressionSpeed = Speed.AcquireHLSLExpressionOrConstant(Generator, Scope, FVector2f(SpeedX, SpeedY));
+	FExpression* ExpressionOffset = Generator.GetTree().NewMul(ExpressionSpeed, ExpressionTime);
+	if (bFractionalPart)
+	{
+		ExpressionOffset = Generator.GetTree().NewFrac(ExpressionOffset);
+	}
+	FExpression* ExpressionTexCoord = Coordinate.TryAcquireHLSLExpression(Generator, Scope);
+	if (!ExpressionTexCoord)
+	{
+		ExpressionTexCoord = Generator.NewTexCoord(ConstCoordinate);
+	}
+
+	OutExpression = Generator.GetTree().NewAdd(ExpressionTexCoord, ExpressionOffset);
+	return true;
+}
+
 bool UMaterialExpressionTextureCoordinate::GenerateHLSLExpression(FMaterialHLSLGenerator& Generator, UE::HLSLTree::FScope& Scope, int32 OutputIndex, UE::HLSLTree::FExpression*& OutExpression)
 {
 	// TODO - tiling, mirroring
diff --git a/Engine/Source/Runtime/Engine/Public/HLSLTree/HLSLTree.h b/Engine/Source/Runtime/Engine/Public/HLSLTree/HLSLTree.h
--- a/Engine/Source/Runtime/Engine/Public/HLSLTree/HLSLTree.h
+++ b/Engine/Source/Runtime/Engine/Public/HLSLTree/HLSLTree.h
@@ -632,11 +632,13 @@ public:
 
 	FExpression* NewNeg(FExpression* Input) { return NewUnaryOp(EUnaryOp::Neg, Input); }
 	FExpression* NewRcp(FExpression* Input) { return NewUnaryOp(EUnaryOp::Rcp, Input); }
+	FExpression* NewFrac(FExpression* Input) { return NewUnaryOp(EUnaryOp::Frac, Input); }
 
 	FExpression* NewAdd(FExpression* Lhs, FExpression* Rhs) { return NewBinaryOp(EBinaryOp::Add, Lhs, Rhs); }
 	FExpression* NewSub(FExpression* Lhs, FExpression* Rhs) { return NewBinaryOp(EBinaryOp::Sub, Lhs, Rhs); }
 	FExpression* NewMul(FExpression* Lhs, FExpression* Rhs) { return NewBinaryOp(EBinaryOp::Mul, Lhs, Rhs); }
 	FExpression* NewDiv(FExpression* Lhs, FExpression* Rhs) { return NewBinaryOp(EBinaryOp::Div, Lhs, Rhs); }
+	FExpression* NewFmod(FExpression* Lhs, FExpression* Rhs) { return NewBinaryOp(EBinaryOp::Fmod, Lhs, Rhs); }
 	FExpression* NewLess(FExpression* Lhs, FExpression* Rhs) { return NewBinaryOp(EBinaryOp::Less, Lhs, Rhs); }
 
 	FTextureParameterDeclaration* NewTextureParameterDeclaration(const FName& Name, const FTextureDescription& DefaultValue);
diff --git a/Engine/Source/Runtime/Engine/Public/HLSLTree/HLSLTreeCommon.h b/Engine/Source/Runtime/Engine/Public/HLSLTree/HLSLTreeCommon.h
--- a/Engine/Source/Runtime/Engine/Public/HLSLTree/HLSLTreeCommon.h
+++ b/Engine/Source/Runtime/Engine/Public/HLSLTree/HLSLTreeCommon.h
@@ -83,6 +83,13 @@ enum class EExternalInput : uint8
 	PrevTranslatedWorldPosition,
 	PrevTranslatedWorldPosition_NoOffsets,
 
+	GameTime,
+	RealTime,
+	DeltaTime,
+
+	PrevGameTime,
+	PrevRealTime,
+
 	WorldPosition_Ddx,
 	WorldPosition_Ddy,
 };
diff --git a/Engine/Source/Runtime/Engine/Public/HLSLTree/HLSLTreeTypes.h b/Engine/Source/Runtime/Engine/Public/HLSLTree/HLSLTreeTypes.h
--- a/Engine/Source/Runtime/Engine/Public/HLSLTree/HLSLTreeTypes.h
+++ b/Engine/Source/Runtime/Engine/Public/HLSLTree/HLSLTreeTypes.h
@@ -64,6 +64,7 @@ enum class EUnaryOp : uint8
 	None,
 	Neg,
 	Rcp,
+	Frac,
 };
 
 struct FUnaryOpDescription
@@ -85,6 +86,7 @@ enum class EBinaryOp : uint8
 	Sub,
 	Mul,
 	Div,
+	Fmod,
 	Less,
 };
 
diff --git a/Engine/Source/Runtime/Engine/Public/Shader/ShaderTypes.h b/Engine/Source/Runtime/Engine/Public/Shader/ShaderTypes.h
--- a/Engine/Source/Runtime/Engine/Public/Shader/ShaderTypes.h
+++ b/Engine/Source/Runtime/Engine/Public/Shader/ShaderTypes.h
@@ -28,6 +28,8 @@ const TCHAR* GetComponentTypeName(EValueComponentType Type);
 uint32 GetComponentTypeSizeInBytes(EValueComponentType Type);
 EValueComponentType CombineComponentTypes(EValueComponentType Lhs, EValueComponentType Rhs);
 
+inline EValueComponentType MakeNonLWCType(EValueComponentType Type) { return Type == EValueComponentType::Double ? EValueComponentType::Float : Type; }
+
 struct FValueTypeDescription
 {
 	FValueTypeDescription() : Name(nullptr), ComponentType(EValueComponentType::Void), NumComponents(0) {}
