Implement PT shape to GT shape. This is not 100% correct if GT is changing, but will at least be safe

diff --git a/Engine/Source/Runtime/Engine/Private/Collision/CollisionConversions.cpp b/Engine/Source/Runtime/Engine/Private/Collision/CollisionConversions.cpp
--- a/Engine/Source/Runtime/Engine/Private/Collision/CollisionConversions.cpp
+++ b/Engine/Source/Runtime/Engine/Private/Collision/CollisionConversions.cpp
@@ -243,10 +243,26 @@ static void SetHitResultFromShapeAndFaceIndex(const FPhysicsShape& Shape,  const
 const FPhysicsActor* GetGTActor(const FPhysicsActor* GTActor) { return GTActor; }
 
 template <bool bGT>
-const FPhysicsShape* GetGTShape(const FPhysicsShape* GTShape) { return GTShape; }
+const FPhysicsShape* GetGTShape(const FPhysicsShape* GTShape, const FPhysicsActor* GTActor) { return GTShape; }
 
 template <>
-const FPhysicsShape* GetGTShape<false>(const FPhysicsShape* GTShape) { return GTShape; ensure(false); }	//TODO: implement this
+const FPhysicsShape* GetGTShape<false>(const FPhysicsShape* PTShape, const FPhysicsActor* GTActor)
+{
+	//This assumes that the geometry hasn't changed.
+	//If it has you will either get null, or the wrong GT shape (note this gt shape is valid and safe to access, user will just get incorrect result)
+	//This is expected because you are accessing GT timeline data from the PT timeline
+	if (GTActor)
+	{
+		const Chaos::FShapesArray& Shapes = GTActor->ShapesArray();
+		const int32 Idx = PTShape->GetShapeIndex();
+		if (Idx < Shapes.Num())
+		{
+			return Shapes[Idx].Get();
+		}
+	}
+
+	return nullptr;
+}
 
 #if WITH_CHAOS
 const FPhysicsActor* GetGTActor(const Chaos::FGeometryParticleHandle* PTActor)
@@ -257,8 +273,6 @@ const FPhysicsActor* GetGTActor(const Chaos::FGeometryParticleHandle* PTActor)
 	auto Proxy = static_cast<const Chaos::FSingleParticlePhysicsProxy*>(PTActor->PhysicsProxy());
 	return Proxy->GetParticle_LowLevel();
 }
-
-const FPhysicsShape* GetGTShape(const Chaos::FPerShapeData* GTShape) { return GTShape; }
 #endif
 
 template <typename THitLocation>
@@ -315,7 +329,7 @@ EConvertQueryResult ConvertQueryImpactHitImp(const UWorld* World, const THitLoca
 	const FPhysicsActor* GTActor = GetGTActor(pHitActor);
 	//Note that a non-gt query could still happen on gt if we're in a fixed tick
 	constexpr bool bIsGTQuery = std::is_same<THitLocation, FHitLocation>::value;
-	const FPhysicsShape* GTShape = GetGTShape<bIsGTQuery>(pHitShape);
+	const FPhysicsShape* GTShape = GetGTShape<bIsGTQuery>(pHitShape, GTActor);
 
 	// See if this is a 'blocking' hit
 	const FCollisionFilterData ShapeFilter = GetQueryFilterData(HitShape);
@@ -601,7 +615,7 @@ static bool ConvertOverlappedShapeToImpactHit(const UWorld* World, const THitLoc
 
 	//Note that a non-gt query could still happen on gt if we're in a fixed tick
 	constexpr bool bIsGTQuery = std::is_same<THitLocation, FHitLocation>::value;
-	const FPhysicsShape* GTHitShape = GetGTShape<bIsGTQuery>(&HitShape);
+	const FPhysicsShape* GTHitShape = GetGTShape<bIsGTQuery>(&HitShape, GTHitActor);
 	if(GTHitActor && GTHitShape)
 	{
 		SetHitResultFromShapeAndFaceIndex(*GTHitShape, *GTHitActor, GetInternalFaceIndex(Hit), OutResult.ImpactPoint, OutResult, bReturnPhysMat);
