Fix crash when a material is not converted to Strata yet.

diff --git a/Engine/Source/Runtime/Engine/Private/Materials/MaterialShared.cpp b/Engine/Source/Runtime/Engine/Private/Materials/MaterialShared.cpp
--- a/Engine/Source/Runtime/Engine/Private/Materials/MaterialShared.cpp
+++ b/Engine/Source/Runtime/Engine/Private/Materials/MaterialShared.cpp
@@ -1718,12 +1718,15 @@ bool FMaterialResource::HasMaterialPropertyConnected(EMaterialProperty In) const
 	{
 		// Strata material traversal is cached as this is an expensive operation
 		#if WITH_EDITOR
-		check(Material->HasStrataFrontMaterialConnected());
-		if (!CachedStrataMaterialInfo.IsValid())
+		if (Material->HasStrataFrontMaterialConnected())
 		{
-			if (Material->FrontMaterial.Expression->IsResultStrataMaterial(Material->FrontMaterial.OutputIndex))
+			if (!CachedStrataMaterialInfo.IsValid())
 			{
-				Material->FrontMaterial.Expression->GatherStrataMaterialInfo(CachedStrataMaterialInfo, Material->FrontMaterial.OutputIndex);
+				check(Material->HasStrataFrontMaterialConnected());
+				if (Material->FrontMaterial.Expression->IsResultStrataMaterial(Material->FrontMaterial.OutputIndex))
+				{
+					Material->FrontMaterial.Expression->GatherStrataMaterialInfo(CachedStrataMaterialInfo, Material->FrontMaterial.OutputIndex);
+				}
 			}
 		}
 		return CachedStrataMaterialInfo.HasPropertyConnected(In);
