Allow async RenderAssetStreaming updates to run past past the StreamingManager tick, and only sync where needed.
Enable r.Streaming.OverlapAssetAndLevelTicks for all non-editor platforms

Saves between 0.1ms to 1ms on the game thread, depending on the platform

diff --git a/Engine/Source/Runtime/Engine/Private/Streaming/AsyncTextureStreaming.cpp b/Engine/Source/Runtime/Engine/Private/Streaming/AsyncTextureStreaming.cpp
--- a/Engine/Source/Runtime/Engine/Private/Streaming/AsyncTextureStreaming.cpp
+++ b/Engine/Source/Runtime/Engine/Private/Streaming/AsyncTextureStreaming.cpp
@@ -331,7 +331,7 @@ void FRenderAssetStreamingMipCalcTask::ApplyPakStateChanges_Async()
 
 	if (bRecacheAllFiles || MountedStateDirtyFiles.Num())
 	{
-		for (FStreamingRenderAsset& StreamingRenderAsset : StreamingManager.StreamingRenderAssets)
+		for (FStreamingRenderAsset& StreamingRenderAsset : StreamingManager.AsyncUnsafeStreamingRenderAssets)
 		{
 			if (IsAborted()) break;
 
@@ -349,7 +349,7 @@ void FRenderAssetStreamingMipCalcTask::ApplyPakStateChanges_Async()
 
 void FRenderAssetStreamingMipCalcTask::TryDropMaxResolutions(TArray<int32>& PrioritizedRenderAssets, int64& MemoryBudgeted, const int64 InMemoryBudget)
 {
-	TArray<FStreamingRenderAsset>& StreamingRenderAssets = StreamingManager.StreamingRenderAssets;
+	TArray<FStreamingRenderAsset>& StreamingRenderAssets = StreamingManager.AsyncUnsafeStreamingRenderAssets;
 	const FRenderAssetStreamingSettings& Settings = StreamingManager.Settings;
 
 	// When using mip bias per texture/mesh, we first reduce the maximum resolutions (if used) in order to fit.
@@ -409,7 +409,7 @@ void FRenderAssetStreamingMipCalcTask::TryDropMaxResolutions(TArray<int32>& Prio
 
 void FRenderAssetStreamingMipCalcTask::TryDropMips(TArray<int32>& PrioritizedRenderAssets, int64& MemoryBudgeted, const int64 InMemoryBudget)
 {
-	TArray<FStreamingRenderAsset>& StreamingRenderAssets = StreamingManager.StreamingRenderAssets;
+	TArray<FStreamingRenderAsset>& StreamingRenderAssets = StreamingManager.AsyncUnsafeStreamingRenderAssets;
 	const FRenderAssetStreamingSettings& Settings = StreamingManager.Settings;
 
 	while (MemoryBudgeted > InMemoryBudget && !IsAborted())
@@ -470,7 +470,7 @@ void FRenderAssetStreamingMipCalcTask::TryDropMips(TArray<int32>& PrioritizedRen
 
 void FRenderAssetStreamingMipCalcTask::TryKeepMips(TArray<int32>& PrioritizedRenderAssets, int64& MemoryBudgeted, const int64 InMemoryBudget)
 {
-	TArray<FStreamingRenderAsset>& StreamingRenderAssets = StreamingManager.StreamingRenderAssets;
+	TArray<FStreamingRenderAsset>& StreamingRenderAssets = StreamingManager.AsyncUnsafeStreamingRenderAssets;
 	bool bBudgetIsChanging = true;
 	
 	while (MemoryBudgeted < InMemoryBudget && bBudgetIsChanging && !IsAborted())
@@ -513,7 +513,7 @@ void FRenderAssetStreamingMipCalcTask::UpdateBudgetedMips_Async()
 	// Update Budget
 	//*************************************
 
-	TArray<FStreamingRenderAsset>& StreamingRenderAssets = StreamingManager.StreamingRenderAssets;
+	TArray<FStreamingRenderAsset>& StreamingRenderAssets = StreamingManager.AsyncUnsafeStreamingRenderAssets;
 	const FRenderAssetStreamingSettings& Settings = StreamingManager.Settings;
 
 	TArray<int32> PrioritizedRenderAssets;
@@ -777,7 +777,7 @@ void FRenderAssetStreamingMipCalcTask::UpdateBudgetedMips_Async()
 
 void FRenderAssetStreamingMipCalcTask::UpdateLoadAndCancelationRequests_Async()
 {
-	TArray<FStreamingRenderAsset>& StreamingRenderAssets = StreamingManager.StreamingRenderAssets;
+	TArray<FStreamingRenderAsset>& StreamingRenderAssets = StreamingManager.AsyncUnsafeStreamingRenderAssets;
 	const FRenderAssetStreamingSettings& Settings = StreamingManager.Settings;
 
 	LoadRequests.Empty();
@@ -861,7 +861,7 @@ void FRenderAssetStreamingMipCalcTask::UpdateLoadAndCancelationRequests_Async()
 
 void FRenderAssetStreamingMipCalcTask::UpdatePendingStreamingStatus_Async()
 {
-	TArray<FStreamingRenderAsset>& StreamingRenderAssets = StreamingManager.StreamingRenderAssets;
+	TArray<FStreamingRenderAsset>& StreamingRenderAssets = StreamingManager.AsyncUnsafeStreamingRenderAssets;
 	const bool bIsStreamingPaused = StreamingManager.bPauseRenderAssetStreaming;
 
 	PendingUpdateDirties.Empty();
@@ -885,7 +885,7 @@ void FRenderAssetStreamingMipCalcTask::DoWork()
 	// While the async task is runnning, the StreamingRenderAssets are guarantied not to be reallocated.
 	// 2 things can happen : a texture can be removed, in which case the texture will be set to null
 	// or some members can be updated following calls to UpdateDynamicData().
-	TArray<FStreamingRenderAsset>& StreamingRenderAssets = StreamingManager.StreamingRenderAssets;
+	TArray<FStreamingRenderAsset>& StreamingRenderAssets = StreamingManager.AsyncUnsafeStreamingRenderAssets;
 	const FRenderAssetStreamingSettings& Settings = StreamingManager.Settings;
 
 	StreamingData.ComputeViewInfoExtras(Settings);
@@ -928,7 +928,7 @@ void FRenderAssetStreamingMipCalcTask::UpdateStats_Async()
 #if STATS
 	FRenderAssetStreamingStats& Stats = StreamingManager.GatheredStats;
 	FRenderAssetStreamingSettings& Settings = StreamingManager.Settings;
-	TArray<FStreamingRenderAsset>& StreamingRenderAssets = StreamingManager.StreamingRenderAssets;
+	TArray<FStreamingRenderAsset>& StreamingRenderAssets = StreamingManager.AsyncUnsafeStreamingRenderAssets;
 
 	Stats.RenderAssetPool = PoolSize;
 	// Stats.StreamingPool = MemoryBudget;
@@ -1064,7 +1064,7 @@ void FRenderAssetStreamingMipCalcTask::UpdateStats_Async()
 
 void FRenderAssetStreamingMipCalcTask::UpdateCSVOnlyStats_Async()
 {
-	const TArray<FStreamingRenderAsset>& StreamingRenderAssets = StreamingManager.StreamingRenderAssets;
+	const TArray<FStreamingRenderAsset>& StreamingRenderAssets = StreamingManager.AsyncUnsafeStreamingRenderAssets;
 
 	FRenderAssetStreamingStats& Stats = StreamingManager.GatheredStats;
 
diff --git a/Engine/Source/Runtime/Engine/Private/Streaming/StreamingManagerTexture.cpp b/Engine/Source/Runtime/Engine/Private/Streaming/StreamingManagerTexture.cpp
--- a/Engine/Source/Runtime/Engine/Private/Streaming/StreamingManagerTexture.cpp
+++ b/Engine/Source/Runtime/Engine/Private/Streaming/StreamingManagerTexture.cpp
@@ -41,7 +41,7 @@ CSV_DEFINE_CATEGORY(TextureStreaming, true);
 
 static TAutoConsoleVariable<int32> CVarStreamingOverlapAssetAndLevelTicks(
 	TEXT("r.Streaming.OverlapAssetAndLevelTicks"),
-	!WITH_EDITOR && (PLATFORM_PS4),
+	!WITH_EDITOR,
 	TEXT("Ticks render asset streaming info on a high priority task thread while ticking levels on GT"),
 	ECVF_Default);
 
@@ -175,7 +175,7 @@ FRenderAssetStreamingManager::FRenderAssetStreamingManager()
 		MarkMountedStateDirty(MakeIoFilenameHash(FileName));
 	});
 
-	StreamingRenderAssets.Reserve(UE_STREAMINGRENDERASSETS_ARRAY_DEFAULT_RESERVED_SIZE);
+	AsyncUnsafeStreamingRenderAssets.Reserve(UE_STREAMINGRENDERASSETS_ARRAY_DEFAULT_RESERVED_SIZE);
 }
 
 FRenderAssetStreamingManager::~FRenderAssetStreamingManager()
@@ -192,6 +192,16 @@ FRenderAssetStreamingManager::~FRenderAssetStreamingManager()
 	STAT(DisplayedStats.Apply();)
 }
 
+TArray<FStreamingRenderAsset>& FRenderAssetStreamingManager::GetStreamingRenderAssetsAsyncSafe()
+{
+	if (StreamingRenderAssetsSyncEvent.IsValid())
+	{
+		StreamingRenderAssetsSyncEvent->Wait(ENamedThreads::GameThread);
+	}
+
+	return AsyncUnsafeStreamingRenderAssets;
+}
+
 void FRenderAssetStreamingManager::OnPreGarbageCollect()
 {
 	TRACE_CPUPROFILER_EVENT_SCOPE(FRenderAssetStreamingManager::OnPreGarbageCollect);
@@ -201,6 +211,11 @@ void FRenderAssetStreamingManager::OnPreGarbageCollect()
 #endif
 
 	FScopeLock ScopeLock(&CriticalSection);
+	if (StreamingRenderAssetsSyncEvent.IsValid())
+	{
+		StreamingRenderAssetsSyncEvent->Wait(ENamedThreads::GameThread);
+	}
+
 	QUICK_SCOPE_CYCLE_COUNTER(STAT_FRenderAssetStreamingManager_OnPreGarbageCollect);
 
 	FRemovedRenderAssetArray RemovedRenderAssets;
@@ -245,6 +260,8 @@ void FRenderAssetStreamingManager::CancelForcedResources()
 {
 	FScopeLock ScopeLock(&CriticalSection);
 
+	TArray<FStreamingRenderAsset>& StreamingRenderAssets = GetStreamingRenderAssetsAsyncSafe();
+
 	// Update textures/meshes that are Forced on a timer.
 	for ( int32 Idx=0; Idx < StreamingRenderAssets.Num(); ++Idx )
 	{
@@ -295,6 +312,7 @@ bool FRenderAssetStreamingManager::StreamOutRenderAssetData( int64 RequiredMemor
 
 	const int64 MaxTempMemoryAllowed = static_cast<int64>(Settings.MaxTempMemoryAllowed) * 1024 * 1024;
 	const bool CachedPauseTextureStreaming = bPauseRenderAssetStreaming;
+	TArray<FStreamingRenderAsset>& StreamingRenderAssets = GetStreamingRenderAssetsAsyncSafe();
 
 	// Pause texture streaming to prevent sending load requests.
 	bPauseRenderAssetStreaming = true;
@@ -392,6 +410,8 @@ void FRenderAssetStreamingManager::TickFastResponseAssets()
 		return;
 	}
 
+	TArray<FStreamingRenderAsset>& StreamingRenderAssets = GetStreamingRenderAssetsAsyncSafe();
+
 	for (TSet<UStreamableRenderAsset*>::TIterator It(FastResponseRenderAssets); It; ++It)
 	{
 		UStreamableRenderAsset* RenderAsset = *It;
@@ -443,6 +463,8 @@ void FRenderAssetStreamingManager::TickFastResponseAssets()
 
 void FRenderAssetStreamingManager::ProcessRemovedRenderAssets()
 {
+	TArray<FStreamingRenderAsset>& StreamingRenderAssets = GetStreamingRenderAssetsAsyncSafe();
+
 	for (int32 AssetIndex : RemovedRenderAssetIndices)
 	{
 		// Remove swap all elements, until this entry has a valid texture/mesh.
@@ -471,6 +493,7 @@ void FRenderAssetStreamingManager::ProcessRemovedRenderAssets()
 void FRenderAssetStreamingManager::ProcessAddedRenderAssets()
 {
 	// Add new textures or meshes.
+	TArray<FStreamingRenderAsset>& StreamingRenderAssets = GetStreamingRenderAssetsAsyncSafe();
 	StreamingRenderAssets.Reserve(StreamingRenderAssets.Num() + PendingStreamingRenderAssets.Num());
 	for (int32 Idx = 0; Idx < PendingStreamingRenderAssets.Num(); ++Idx)
 	{
@@ -499,7 +522,7 @@ void FRenderAssetStreamingManager::ConditionalUpdateStaticData()
 	{
 		STAT(GatheredStats.SetupAsyncTaskCycles += FPlatformTime::Cycles();)
 		// Update each texture static data.
-		for (FStreamingRenderAsset& StreamingRenderAsset : StreamingRenderAssets)
+		for (FStreamingRenderAsset& StreamingRenderAsset : AsyncUnsafeStreamingRenderAssets)
 		{
 			StreamingRenderAsset.UpdateStaticData(Settings);
 
@@ -591,12 +614,12 @@ void FRenderAssetStreamingManager::ProcessLevelsToReferenceToStreamedTextures()
 		for (; RenderAssetIterator; ++RenderAssetIterator)
 		{
 			const UStreamableRenderAsset* RenderAsset = *RenderAssetIterator;
-			if (RenderAsset == nullptr || !ReferencedRenderAssets.Contains(RenderAsset) || !StreamingRenderAssets.IsValidIndex(RenderAsset->StreamingIndex))
+			if (RenderAsset == nullptr || !ReferencedRenderAssets.Contains(RenderAsset) || !AsyncUnsafeStreamingRenderAssets.IsValidIndex(RenderAsset->StreamingIndex))
 			{
 				continue;
 			}
 
-			FStreamingRenderAsset& StreamingRenderAsset = StreamingRenderAssets[RenderAsset->StreamingIndex];
+			FStreamingRenderAsset& StreamingRenderAsset = AsyncUnsafeStreamingRenderAssets[RenderAsset->StreamingIndex];
 
 			check(StreamingRenderAsset.RenderAsset == RenderAsset);
 
@@ -805,6 +828,7 @@ void FRenderAssetStreamingManager::AddStreamingRenderAsset(UStreamableRenderAsse
 void FRenderAssetStreamingManager::RemoveStreamingRenderAsset( UStreamableRenderAsset* RenderAsset )
 {
 	FScopeLock ScopeLock(&CriticalSection);
+	TArray<FStreamingRenderAsset>& StreamingRenderAssets = GetStreamingRenderAssetsAsyncSafe();
 
 	STAT(GatheredStats.CallbacksCycles = -(int32)FPlatformTime::Cycles();)
 
@@ -1064,9 +1088,9 @@ FStreamingRenderAsset* FRenderAssetStreamingManager::GetStreamingRenderAsset( co
 {
 	FScopeLock ScopeLock(&CriticalSection);
 
-	if (RenderAsset && StreamingRenderAssets.IsValidIndex(RenderAsset->StreamingIndex))
+	if (RenderAsset && AsyncUnsafeStreamingRenderAssets.IsValidIndex(RenderAsset->StreamingIndex))
 	{
-		FStreamingRenderAsset* StreamingRenderAsset = &StreamingRenderAssets[RenderAsset->StreamingIndex];
+		FStreamingRenderAsset* StreamingRenderAsset = &AsyncUnsafeStreamingRenderAssets[RenderAsset->StreamingIndex];
 
 		// If the texture/mesh don't match, this means the texture/mesh is pending in PendingStreamingRenderAssets, for which no FStreamingRenderAsset* is yet allocated.
 		// If this is not acceptable, the caller should first synchronize everything through SyncStates
@@ -1112,6 +1136,7 @@ void FRenderAssetStreamingManager::UpdateIndividualRenderAsset( UStreamableRende
 void FRenderAssetStreamingManager::FastForceFullyResident(UStreamableRenderAsset* RenderAsset)
 {
 	check(IsInGameThread());
+	TArray<FStreamingRenderAsset>& StreamingRenderAssets = GetStreamingRenderAssetsAsyncSafe();
 
 	if (CVarStreamingAllowFastForceResident.GetValueOnGameThread()
 		&& IStreamingManager::Get().IsStreamingEnabled()
@@ -1144,6 +1169,8 @@ void FRenderAssetStreamingManager::FastForceFullyResident(UStreamableRenderAsset
  */
 void FRenderAssetStreamingManager::UpdateStreamingRenderAssets(int32 StageIndex, int32 NumUpdateStages, bool bWaitForMipFading, bool bAsync)
 {
+	TArray<FStreamingRenderAsset>& StreamingRenderAssets = GetStreamingRenderAssetsAsyncSafe();
+
 	if ( StageIndex == 0 )
 	{
 		CurrentUpdateStreamingRenderAssetIndex = 0;
@@ -1291,6 +1318,7 @@ static FORCEINLINE bool ShouldAmortizeMipCopies()
  */
 void FRenderAssetStreamingManager::StreamRenderAssets( bool bProcessEverything )
 {
+	TArray<FStreamingRenderAsset>& StreamingRenderAssets = GetStreamingRenderAssetsAsyncSafe();
 	const FRenderAssetStreamingMipCalcTask& AsyncTask = AsyncWork->GetTask();
 
 	// Note that render asset indices referred by the async task could be outdated if UpdatePendingStates() was called between the
@@ -1382,6 +1410,7 @@ void FRenderAssetStreamingManager::ProcessPendingMipCopyRequests()
 		return;
 	}
 
+	TArray<FStreamingRenderAsset>& StreamingRenderAssets = GetStreamingRenderAssetsAsyncSafe();
 	int32 NumRemainingRequests = CVarTextureStreamingMaxNumTexturesToStreamPerFrame.GetValueOnGameThread();
 
 	while (NumRemainingRequests
@@ -1420,6 +1449,11 @@ void FRenderAssetStreamingManager::TickDeferredMipLevelChangeCallbacks()
 	{
 		check(IsInGameThread());
 
+		if(StreamingRenderAssetsSyncEvent.IsValid())
+		{
+			StreamingRenderAssetsSyncEvent->Wait(ENamedThreads::GameThread);
+		}
+
 		for (int32 AssetIdx = 0; AssetIdx < DeferredTickCBAssets.Num(); ++AssetIdx)
 		{
 			UStreamableRenderAsset* Asset = DeferredTickCBAssets[AssetIdx];
@@ -1547,20 +1581,17 @@ static TAutoConsoleVariable<int32> CVarUseBackgroundThreadPool(
 
 class FUpdateStreamingRenderAssetsTask
 {
-	FEvent* CompletionEvent;
 	FRenderAssetStreamingManager* Manager;
 	int32 StageIdx;
 	int32 NumUpdateStages;
 	bool bWaitForMipFading;
 public:
 	FUpdateStreamingRenderAssetsTask(
-		FEvent* InCompletionEvent,
 		FRenderAssetStreamingManager* InManager,
 		int32 InStageIdx,
 		int32 InNumUpdateStages,
 		bool bInWaitForMipFading)
-		: CompletionEvent(InCompletionEvent)
-		, Manager(InManager)
+		: Manager(InManager)
 		, StageIdx(InStageIdx)
 		, NumUpdateStages(InNumUpdateStages)
 		, bWaitForMipFading(bInWaitForMipFading)
@@ -1582,13 +1613,13 @@ public:
 	{
 		FOptionalTaskTagScope Scope(ETaskTag::EParallelGameThread);
 		Manager->UpdateStreamingRenderAssets(StageIdx, NumUpdateStages, bWaitForMipFading, true);
-		CompletionEvent->Trigger();
 	}
 };
 
 void FRenderAssetStreamingManager::UpdateResourceStreaming( float DeltaTime, bool bProcessEverything/*=false*/ )
 {
 	FScopeLock ScopeLock(&CriticalSection);
+	TArray<FStreamingRenderAsset>& StreamingRenderAssets = GetStreamingRenderAssetsAsyncSafe();
 
 	SCOPE_CYCLE_COUNTER(STAT_RenderAssetStreaming_GameThreadUpdateTime);
 	CSV_SCOPED_TIMING_STAT_EXCLUSIVE(RenderAssetStreaming);
@@ -1672,7 +1703,6 @@ void FRenderAssetStreamingManager::UpdateResourceStreaming( float DeltaTime, boo
 
 		TickFastResponseAssets();
 
-		FEvent* SyncEvent = nullptr;
 		// Optimization: overlapping UpdateStreamingRenderAssets() and IncrementalUpdate();
 		// Restrict this optimization to platforms tested to have a win;
 		// Platforms tested and results (ave exec time of UpdateResourceStreaming):
@@ -1681,10 +1711,13 @@ void FRenderAssetStreamingManager::UpdateResourceStreaming( float DeltaTime, boo
 		const bool bOverlappedExecution = bUseThreadingForPerf && CVarStreamingOverlapAssetAndLevelTicks.GetValueOnGameThread();
 		if (bOverlappedExecution)
 		{
-			SyncEvent = FPlatformProcess::GetSynchEventFromPool(false);
-			check(SyncEvent);
-			TGraphTask<FUpdateStreamingRenderAssetsTask>::CreateTask(nullptr, ENamedThreads::GameThread)
-				.ConstructAndDispatchWhenReady(SyncEvent, this, ProcessingStage - 1, NumRenderAssetProcessingStages, DeltaTime > 0.f);
+			if (StreamingRenderAssetsSyncEvent.IsValid() && StreamingRenderAssetsSyncEvent->IsComplete())
+			{
+				StreamingRenderAssetsSyncEvent = nullptr;
+			}
+
+			StreamingRenderAssetsSyncEvent = TGraphTask<FUpdateStreamingRenderAssetsTask>::CreateTask(nullptr, ENamedThreads::GameThread)
+			.ConstructAndDispatchWhenReady(this, ProcessingStage - 1, NumRenderAssetProcessingStages, DeltaTime > 0.f);
 		}
 		else
 		{
@@ -1694,11 +1727,6 @@ void FRenderAssetStreamingManager::UpdateResourceStreaming( float DeltaTime, boo
 		IncrementalUpdate(1.f / (float)FMath::Max(NumRenderAssetProcessingStages - 1, 1), true); // -1 since we don't want to do anything at stage 0.
 		++ProcessingStage;
 
-		if (bOverlappedExecution)
-		{
-			SyncEvent->Wait();
-			FPlatformProcess::ReturnSynchEventToPool(SyncEvent);
-		}
 		STAT(GatheredStats.UpdateStreamingDataCycles = FMath::Max<uint32>(ProcessingStage > 2 ? GatheredStats.UpdateStreamingDataCycles : 0, FPlatformTime::Cycles() - StartTime);)
 	}
 	else if (AsyncWork->IsDone())
@@ -1756,6 +1784,7 @@ void FRenderAssetStreamingManager::UpdateResourceStreaming( float DeltaTime, boo
 int32 FRenderAssetStreamingManager::BlockTillAllRequestsFinished( float TimeLimit /*= 0.0f*/, bool bLogResults /*= false*/ )
 {
 	FScopeLock ScopeLock(&CriticalSection);
+	TArray<FStreamingRenderAsset>& StreamingRenderAssets = GetStreamingRenderAssetsAsyncSafe();
 	QUICK_SCOPE_CYCLE_COUNTER(STAT_FRenderAssetStreamingManager_BlockTillAllRequestsFinished);
 
 	double StartTime = FPlatformTime::Seconds();
@@ -1861,8 +1890,8 @@ void FRenderAssetStreamingManager::GetAssetComponents(const UStreamableRenderAss
 		QUICK_SCOPE_CYCLE_COUNTER(STAT_FRenderAssetStreamingManager_GetAssetComponents);
 
 		const int32 Idx = RenderAsset->StreamingIndex;
-		check(StreamingRenderAssets.IsValidIndex(Idx) && StreamingRenderAssets[Idx].RenderAsset == RenderAsset);
-		const FStreamingRenderAsset& StreamingRenderAsset = StreamingRenderAssets[Idx];
+		check(AsyncUnsafeStreamingRenderAssets.IsValidIndex(Idx) && AsyncUnsafeStreamingRenderAssets[Idx].RenderAsset == RenderAsset);
+		const FStreamingRenderAsset& StreamingRenderAsset = AsyncUnsafeStreamingRenderAssets[Idx];
 
 		for (TConstSetBitIterator<> It(StreamingRenderAsset.LevelIndexUsage); It; ++It)
 		{
@@ -1912,7 +1941,7 @@ void FRenderAssetStreamingManager::AddRenderedTextureStats(TMap<FString, FRender
 {
 	FScopeLock ScopeLock(&CriticalSection);
 	
-	for (const FStreamingRenderAsset& StreamingRenderAsset : StreamingRenderAssets)
+	for (const FStreamingRenderAsset& StreamingRenderAsset : AsyncUnsafeStreamingRenderAssets)
 	{
 		const UStreamableRenderAsset* RenderAsset = StreamingRenderAsset.RenderAsset;
 		if (RenderAsset == nullptr || StreamingRenderAsset.bUseUnkownRefHeuristic
@@ -1970,6 +1999,7 @@ bool FRenderAssetStreamingManager::HandleDumpTextureStreamingStatsCommand( const
 bool FRenderAssetStreamingManager::HandleListStreamingRenderAssetsCommand( const TCHAR* Cmd, FOutputDevice& Ar )
 {
 	FScopeLock ScopeLock(&CriticalSection);
+	TArray<FStreamingRenderAsset>& StreamingRenderAssets = GetStreamingRenderAssetsAsyncSafe();
 
 	SyncStates(true);
 
@@ -2355,6 +2385,7 @@ bool FRenderAssetStreamingManager::HandlePauseRenderAssetStreamingCommand( const
 bool FRenderAssetStreamingManager::HandleStreamingManagerMemoryCommand( const TCHAR* Cmd, FOutputDevice& Ar, UWorld* InWorld )
 {
 	FScopeLock ScopeLock(&CriticalSection);
+	TArray<FStreamingRenderAsset>& StreamingRenderAssets = GetStreamingRenderAssetsAsyncSafe();
 
 	SyncStates(true);
 
@@ -2488,8 +2519,8 @@ bool FRenderAssetStreamingManager::HandleLODGroupsCommand( const TCHAR* Cmd, FOu
 
 bool FRenderAssetStreamingManager::HandleInvestigateRenderAssetCommand(const TCHAR* Cmd, FOutputDevice& Ar, UWorld* InWorld)
 {
-
 	FScopeLock ScopeLock(&CriticalSection);
+	TArray<FStreamingRenderAsset>& StreamingRenderAssets = GetStreamingRenderAssetsAsyncSafe();
 	SyncStates(true);
 
 	FString InvestigateAssetName(FParse::Token(Cmd, 0));
diff --git a/Engine/Source/Runtime/Engine/Private/Streaming/StreamingManagerTexture.h b/Engine/Source/Runtime/Engine/Private/Streaming/StreamingManagerTexture.h
--- a/Engine/Source/Runtime/Engine/Private/Streaming/StreamingManagerTexture.h
+++ b/Engine/Source/Runtime/Engine/Private/Streaming/StreamingManagerTexture.h
@@ -251,8 +251,11 @@ private:
 			}
 		}
 
-		/** All streaming texture or mesh objects. */
-		TArray<FStreamingRenderAsset> StreamingRenderAssets;
+		/** All streaming texture or mesh objects. 
+		* Use directly only if it's safe to overlap with UpdateStreamingRenderAssets.
+		* For an async safe version, use GetStreamingRenderAssetsAsyncSafe
+		*/
+		TArray<FStreamingRenderAsset> AsyncUnsafeStreamingRenderAssets;
 
 		/** All the textures/meshes referenced in StreamingRenderAssets. Used to handled deleted textures/meshes.  */
 		TSet<const UStreamableRenderAsset*> ReferencedRenderAssets;
@@ -409,5 +412,9 @@ private:
 	// A critical section use around code that could be called in parallel with NotifyPrimitiveUpdated() or NotifyPrimitiveUpdated_Concurrent().
 	FCriticalSection CriticalSection;
 
+	// An event used to prevent FUpdateStreamingRenderAssetsTask from overlapping with related work
+	FGraphEventRef StreamingRenderAssetsSyncEvent;
+	TArray<FStreamingRenderAsset>& GetStreamingRenderAssetsAsyncSafe();
+
 	friend bool TrackRenderAssetEvent( FStreamingRenderAsset* StreamingRenderAsset, UStreamableRenderAsset* RenderAsset, bool bForceMipLevelsToBeResident, const FRenderAssetStreamingManager* Manager);
 };
