Add light function support for Strata.

diff --git a/Engine/Shaders/Private/LightFunctionCommon.ush b/Engine/Shaders/Private/LightFunctionCommon.ush
--- a/Engine/Shaders/Private/LightFunctionCommon.ush
+++ b/Engine/Shaders/Private/LightFunctionCommon.ush
@@ -4,6 +4,10 @@
 	LightFunctionCommon.usf: Utility functions for light functions.
 =============================================================================*/
 
+#if STRATA_ENABLED
+#include "/Engine/Private/Strata/Strata.ush" 
+#endif
+
 /** Tan of spotlight cone outer angle in x, ShadowFadeFraction in y, IsSpotLight in z, IsPointLight in w. */
 float4 LightFunctionParameters;
 float4x4 LightFunctionTranslatedWorldToLight;
@@ -50,9 +54,21 @@ float3 GetLightFunctionColor(float3 LightVector, float3 TranslatedWorldSpace)
 	FPixelMaterialInputs PixelMaterialInputs;
 	CalcPixelMaterialInputs(MaterialParameters, PixelMaterialInputs);
 
-#if COMPUTESHADER
+#if STRATA_ENABLED
+
+  #if COMPUTESHADER
+	return 0.f; // STRATA_TODO: This path doesn't seem to be used
+  #else
+	return StrataGetBSDFEmissive(PixelMaterialInputs.FrontMaterial.Layers[0].BSDFs[0]);
+  #endif
+
+#else // STRATA_ENABLED
+
+   #if COMPUTESHADER
 	return GetMaterialEmissiveForCS(MaterialParameters);
-#else
+   #else
 	return GetMaterialEmissive(PixelMaterialInputs);
-#endif
+   #endif
+
+#endif // STRATA_ENABLED
 }
diff --git a/Engine/Shaders/Private/Strata/Strata.ush b/Engine/Shaders/Private/Strata/Strata.ush
--- a/Engine/Shaders/Private/Strata/Strata.ush
+++ b/Engine/Shaders/Private/Strata/Strata.ush
@@ -1677,6 +1677,11 @@ float StrataGetBSDFAnisotropy(in FStrataBSDF BSDF)
 	return 0.f;
 }
 
+float3 StrataGetBSDFEmissive(in FStrataBSDF BSDF)
+{	
+	return BSDF_GETEMISSIVE(BSDF);
+}
+
 float3 StrataGetWorldNormal(in FStrataPixelHeader Header, in FStrataBSDF BSDF, in const FStrataAddressing Addressing)
 {
 	return StrataGetBSDFSharedBasis_DeferredShading(Header, BSDF, Addressing)[2];
diff --git a/Engine/Source/Runtime/Engine/Private/Materials/Material.cpp b/Engine/Source/Runtime/Engine/Private/Materials/Material.cpp
--- a/Engine/Source/Runtime/Engine/Private/Materials/Material.cpp
+++ b/Engine/Source/Runtime/Engine/Private/Materials/Material.cpp
@@ -2979,7 +2979,7 @@ void UMaterial::ConvertMaterialToStrataMaterial()
 			FrontMaterial.Connect(0, VolBSDF);
 			bInvalidateShader = true;
 		}
-		else if (MaterialDomain == MD_PostProcess)
+		else if (MaterialDomain == MD_PostProcess || MaterialDomain == MD_LightFunction)
 		{
 			// Some post-process material don't have their shading mode set correctly to Unlit. Since only Unlit is supported, forcing it here.
 			ShadingModel = MSM_Unlit;
@@ -4049,7 +4049,7 @@ void UMaterial::RebuildShadingModelField()
 			// Now derive some properties from the material
 			if (StrataMaterialInfo.HasOnlyShadingModel(SSM_Unlit))
 			{
-				if (MaterialDomain != EMaterialDomain::MD_Surface && MaterialDomain != EMaterialDomain::MD_PostProcess)
+				if (MaterialDomain != EMaterialDomain::MD_Surface && MaterialDomain != EMaterialDomain::MD_PostProcess && MaterialDomain != EMaterialDomain::MD_LightFunction)
 				{
 					MaterialDomain = EMaterialDomain::MD_Surface;
 				}
@@ -5389,7 +5389,14 @@ static bool IsPropertyActive_Internal(EMaterialProperty InProperty,
 	else if (Domain == MD_LightFunction)
 	{
 		// light functions should already use MSM_Unlit but we also we don't want WorldPosOffset
-		return InProperty == MP_EmissiveColor;
+		if (bStrataEnabled)
+		{
+			return InProperty == MP_FrontMaterial;
+		}
+		else
+		{
+			return InProperty == MP_EmissiveColor;
+		}
 	}
 	else if (Domain == MD_DeferredDecal)
 	{
diff --git a/Engine/Source/Runtime/Engine/Private/Materials/MaterialShared.cpp b/Engine/Source/Runtime/Engine/Private/Materials/MaterialShared.cpp
--- a/Engine/Source/Runtime/Engine/Private/Materials/MaterialShared.cpp
+++ b/Engine/Source/Runtime/Engine/Private/Materials/MaterialShared.cpp
@@ -1710,6 +1710,7 @@ bool FMaterialResource::HasMaterialPropertyConnected(EMaterialProperty In) const
 	// STRATA_TODO: temporary validation until we have converted all domains
 	const bool bIsStrataSupportedDomain = 
 		Material->MaterialDomain == MD_PostProcess || 
+		Material->MaterialDomain == MD_LightFunction ||
 		Material->MaterialDomain == MD_DeferredDecal || 
 		Material->MaterialDomain == MD_Surface || 
 		Material->MaterialDomain == MD_Volume;
diff --git a/Engine/Source/Runtime/Renderer/Private/LightFunctionRendering.cpp b/Engine/Source/Runtime/Renderer/Private/LightFunctionRendering.cpp
--- a/Engine/Source/Runtime/Renderer/Private/LightFunctionRendering.cpp
+++ b/Engine/Source/Runtime/Renderer/Private/LightFunctionRendering.cpp
@@ -89,6 +89,12 @@ public:
 		return Parameters.MaterialParameters.MaterialDomain == MD_LightFunction && IsFeatureLevelSupported(Parameters.Platform, ERHIFeatureLevel::SM5);
 	}
 
+	static void ModifyCompilationEnvironment(const FMaterialShaderPermutationParameters& Parameters, FShaderCompilerEnvironment& OutEnvironment)
+	{
+		FMaterialShader::ModifyCompilationEnvironment(Parameters, OutEnvironment);
+		OutEnvironment.SetDefine(TEXT("STRATA_ENABLED"), Strata::IsStrataEnabled() ? 1u : 0u);
+	}
+
 	FLightFunctionPS() {}
 	FLightFunctionPS(const ShaderMetaType::CompiledShaderInitializerType& Initializer)
 		: FMaterialShader(Initializer)
diff --git a/Engine/Source/Runtime/Renderer/Private/TranslucentLighting.cpp b/Engine/Source/Runtime/Renderer/Private/TranslucentLighting.cpp
--- a/Engine/Source/Runtime/Renderer/Private/TranslucentLighting.cpp
+++ b/Engine/Source/Runtime/Renderer/Private/TranslucentLighting.cpp
@@ -689,6 +689,7 @@ public:
 		FVirtualShadowMapArray::SetShaderDefines(OutEnvironment);
 		FMaterialShader::ModifyCompilationEnvironment(Parameters, OutEnvironment);
 		OutEnvironment.SetDefine(TEXT("INJECTION_PIXEL_SHADER"), 1);
+		OutEnvironment.SetDefine(TEXT("STRATA_ENABLED"), Strata::IsStrataEnabled() ? 1u : 0u);
 	}
 
 	/**
diff --git a/Engine/Source/Runtime/Renderer/Private/VolumetricFogLightFunction.cpp b/Engine/Source/Runtime/Renderer/Private/VolumetricFogLightFunction.cpp
--- a/Engine/Source/Runtime/Renderer/Private/VolumetricFogLightFunction.cpp
+++ b/Engine/Source/Runtime/Renderer/Private/VolumetricFogLightFunction.cpp
@@ -78,6 +78,12 @@ class FVolumetricFogLightFunctionPS : public FMaterialShader
 			false);
 	}
 
+	static void ModifyCompilationEnvironment(const FMaterialShaderPermutationParameters& Parameters, FShaderCompilerEnvironment& OutEnvironment)
+	{
+		FMaterialShader::ModifyCompilationEnvironment(Parameters, OutEnvironment);
+		OutEnvironment.SetDefine(TEXT("STRATA_ENABLED"), Strata::IsStrataEnabled() ? 1u : 0u);
+	}
+
 	static bool ShouldCompilePermutation(const FMaterialShaderPermutationParameters& Parameters)
 	{
 		return Parameters.MaterialParameters.MaterialDomain == MD_LightFunction && DoesPlatformSupportVolumetricFog(Parameters.Platform);
