Fixed a GameThread stall (waiting for render thread) whenever a CVar is changed during gameplay.

It seems like the intent of the callback sink for r.FreeSkeletalMeshBuffers was to free eligible CPU buffers when the CVar was raised to 1, but unfortunately we're currently relying on this sink to free certain helper structures.

This change allows those structures to be freed after submitting the buffers to RHI so that we can change the callback sink to only do work on the rising edge (which generally only happens when switching game-mode/device profile and not during gameplay).

diff --git a/Engine/Source/Runtime/Engine/Private/PositionVertexBuffer.cpp b/Engine/Source/Runtime/Engine/Private/PositionVertexBuffer.cpp
--- a/Engine/Source/Runtime/Engine/Private/PositionVertexBuffer.cpp
+++ b/Engine/Source/Runtime/Engine/Private/PositionVertexBuffer.cpp
@@ -196,22 +196,7 @@ void FPositionVertexBuffer::operator=(const FPositionVertexBuffer &Other)
 template <bool bRenderThread>
 FBufferRHIRef FPositionVertexBuffer::CreateRHIBuffer_Internal()
 {
-	if (NumVertices)
-	{
-		FResourceArrayInterface* RESTRICT ResourceArray = VertexData ? VertexData->GetResourceArray() : nullptr;
-		const uint32 SizeInBytes = ResourceArray ? ResourceArray->GetResourceDataSize() : 0;
-		FRHIResourceCreateInfo CreateInfo(TEXT("FPositionVertexBuffer"), ResourceArray);
-		CreateInfo.bWithoutNativeResource = !VertexData;
-		if (bRenderThread)
-		{
-			return RHICreateVertexBuffer(SizeInBytes, BUF_Static | BUF_ShaderResource, CreateInfo);
-		}
-		else
-		{
-			return RHIAsyncCreateVertexBuffer(SizeInBytes, BUF_Static | BUF_ShaderResource, CreateInfo);
-		}
-	}
-	return nullptr;
+	return CreateRHIBuffer<bRenderThread>(VertexData, NumVertices, BUF_Static | BUF_ShaderResource, TEXT("FPositionVertexBuffer"));
 }
 
 FBufferRHIRef FPositionVertexBuffer::CreateRHIBuffer_RenderThread()
diff --git a/Engine/Source/Runtime/Engine/Private/Rendering/ColorVertexBuffer.cpp b/Engine/Source/Runtime/Engine/Private/Rendering/ColorVertexBuffer.cpp
--- a/Engine/Source/Runtime/Engine/Private/Rendering/ColorVertexBuffer.cpp
+++ b/Engine/Source/Runtime/Engine/Private/Rendering/ColorVertexBuffer.cpp
@@ -394,22 +394,7 @@ uint32 FColorVertexBuffer::GetAllocatedSize() const
 template <bool bRenderThread>
 FBufferRHIRef FColorVertexBuffer::CreateRHIBuffer_Internal()
 {
-	if (NumVertices)
-	{
-		FResourceArrayInterface* RESTRICT ResourceArray = VertexData ? VertexData->GetResourceArray() : nullptr;
-		const uint32 SizeInBytes = ResourceArray ? ResourceArray->GetResourceDataSize() : 0;
-		FRHIResourceCreateInfo CreateInfo(TEXT("FColorVertexBuffer"), ResourceArray);
-		CreateInfo.bWithoutNativeResource = !VertexData;
-		if (bRenderThread)
-		{
-			return RHICreateVertexBuffer(SizeInBytes, BUF_Static | BUF_ShaderResource, CreateInfo);
-		}
-		else
-		{
-			return RHIAsyncCreateVertexBuffer(SizeInBytes, BUF_Static | BUF_ShaderResource, CreateInfo);
-		}
-	}
-	return nullptr;
+	return CreateRHIBuffer<bRenderThread>(VertexData, NumVertices, BUF_Static | BUF_ShaderResource, TEXT("FColorVertexBuffer"));
 }
 
 FBufferRHIRef FColorVertexBuffer::CreateRHIBuffer_RenderThread()
diff --git a/Engine/Source/Runtime/Engine/Private/Rendering/SkeletalMeshVertexClothBuffer.cpp b/Engine/Source/Runtime/Engine/Private/Rendering/SkeletalMeshVertexClothBuffer.cpp
--- a/Engine/Source/Runtime/Engine/Private/Rendering/SkeletalMeshVertexClothBuffer.cpp
+++ b/Engine/Source/Runtime/Engine/Private/Rendering/SkeletalMeshVertexClothBuffer.cpp
@@ -74,24 +74,7 @@ void FSkeletalMeshVertexClothBuffer::ClearMetaData()
 template <bool bRenderThread>
 FBufferRHIRef FSkeletalMeshVertexClothBuffer::CreateRHIBuffer_Internal()
 {
-	if (NumVertices)
-	{
-		FResourceArrayInterface* ResourceArray = VertexData ? VertexData->GetResourceArray() : nullptr;
-		const uint32 SizeInBytes = ResourceArray ? ResourceArray->GetResourceDataSize() : 0;
-		const EBufferUsageFlags BufferFlags = BUF_Static | BUF_ShaderResource;
-		FRHIResourceCreateInfo CreateInfo(TEXT("FSkeletalMeshVertexClothBuffer"), ResourceArray);
-		CreateInfo.bWithoutNativeResource = !VertexData;
-
-		if (bRenderThread)
-		{
-			return RHICreateVertexBuffer(SizeInBytes, BufferFlags, CreateInfo);
-		}
-		else
-		{
-			return RHIAsyncCreateVertexBuffer(SizeInBytes, BufferFlags, CreateInfo);
-		}
-	}
-	return nullptr;
+	return CreateRHIBuffer<bRenderThread>(VertexData, NumVertices, BUF_Static | BUF_ShaderResource, TEXT("FSkeletalMeshVertexClothBuffer"));
 }
 
 FBufferRHIRef FSkeletalMeshVertexClothBuffer::CreateRHIBuffer_RenderThread()
diff --git a/Engine/Source/Runtime/Engine/Private/Rendering/SkinWeightVertexBuffer.cpp b/Engine/Source/Runtime/Engine/Private/Rendering/SkinWeightVertexBuffer.cpp
--- a/Engine/Source/Runtime/Engine/Private/Rendering/SkinWeightVertexBuffer.cpp
+++ b/Engine/Source/Runtime/Engine/Private/Rendering/SkinWeightVertexBuffer.cpp
@@ -179,25 +179,7 @@ void FSkinWeightLookupVertexBuffer::SetWeightOffsetAndInfluenceCount(uint32 Vert
 template <bool bRenderThread>
 FBufferRHIRef FSkinWeightLookupVertexBuffer::CreateRHIBuffer_Internal()
 {
-	if (NumVertices)
-	{
-		// Create the vertex buffer.
-		FResourceArrayInterface* ResourceArray = LookupData ? LookupData->GetResourceArray() : nullptr;
-		const uint32 SizeInBytes = ResourceArray ? ResourceArray->GetResourceDataSize() : 0;
-		const EBufferUsageFlags BufferFlags = BUF_Static | BUF_ShaderResource | BUF_SourceCopy;
-		FRHIResourceCreateInfo CreateInfo(TEXT("FSkinWeightLookupVertexBuffer"), ResourceArray);
-		CreateInfo.bWithoutNativeResource = !LookupData;
-
-		if (bRenderThread)
-		{
-			return RHICreateVertexBuffer(SizeInBytes, BufferFlags, CreateInfo);
-		}
-		else
-		{
-			return RHIAsyncCreateVertexBuffer(SizeInBytes, BufferFlags, CreateInfo);
-		}
-	}
-	return nullptr;
+	return CreateRHIBuffer<bRenderThread>(LookupData, NumVertices, BUF_Static | BUF_ShaderResource | BUF_SourceCopy, TEXT("FSkinWeightLookupVertexBuffer"));
 }
 
 /*-----------------------------------------------------------------------------
@@ -389,26 +371,8 @@ void FSkinWeightDataVertexBuffer::CopyMetaData(const FSkinWeightDataVertexBuffer
 template <bool bRenderThread>
 FBufferRHIRef FSkinWeightDataVertexBuffer::CreateRHIBuffer_Internal()
 {
-	if (NumBones)
-	{
-		// Create the vertex buffer.
-		FResourceArrayInterface* ResourceArray = WeightData ? WeightData->GetResourceArray() : nullptr;
-		const uint32 SizeInBytes = ResourceArray ? ResourceArray->GetResourceDataSize() : 0;
-		const EBufferUsageFlags BufferFlags = BUF_Static | BUF_ShaderResource | BUF_SourceCopy;
-		FRHIResourceCreateInfo CreateInfo(TEXT("FSkinWeightDataVertexBuffer"), ResourceArray);
-		CreateInfo.bWithoutNativeResource = !WeightData;
-
-		// BUF_ShaderResource is needed for support of the SkinCache (we could make is dependent on GEnableGPUSkinCacheShaders or are there other users?)
-		if (bRenderThread)
-		{
-			return RHICreateVertexBuffer(SizeInBytes, BufferFlags, CreateInfo);
-		}
-		else
-		{
-			return RHIAsyncCreateVertexBuffer(SizeInBytes, BufferFlags, CreateInfo);
-		}
-	}
-	return nullptr;
+	// BUF_ShaderResource is needed for support of the SkinCache (we could make is dependent on GEnableGPUSkinCacheShaders or are there other users?)
+	return CreateRHIBuffer<bRenderThread>(WeightData, NumBones, BUF_Static | BUF_ShaderResource | BUF_SourceCopy, TEXT("FSkinWeightDataVertexBuffer"));
 }
 
 FBufferRHIRef FSkinWeightDataVertexBuffer::CreateRHIBuffer_RenderThread()
diff --git a/Engine/Source/Runtime/Engine/Private/Rendering/StaticMeshVertexBuffer.cpp b/Engine/Source/Runtime/Engine/Private/Rendering/StaticMeshVertexBuffer.cpp
--- a/Engine/Source/Runtime/Engine/Private/Rendering/StaticMeshVertexBuffer.cpp
+++ b/Engine/Source/Runtime/Engine/Private/Rendering/StaticMeshVertexBuffer.cpp
@@ -269,22 +269,7 @@ void FStaticMeshVertexBuffer::operator=(const FStaticMeshVertexBuffer &Other)
 template <bool bRenderThread>
 FBufferRHIRef FStaticMeshVertexBuffer::CreateTangentsRHIBuffer_Internal()
 {
-	if (GetNumVertices())
-	{
- 		FResourceArrayInterface* RESTRICT ResourceArray = TangentsData ? TangentsData->GetResourceArray() : nullptr;
-		const uint32 SizeInBytes = ResourceArray ? ResourceArray->GetResourceDataSize() : 0;
-		FRHIResourceCreateInfo CreateInfo(TEXT("TangentsRHIBuffer"), ResourceArray);
-		CreateInfo.bWithoutNativeResource = !TangentsData;
-		if (bRenderThread)
-		{
-			return RHICreateVertexBuffer(SizeInBytes, BUF_Static | BUF_ShaderResource, CreateInfo);
-		}
-		else
-		{
-			return RHIAsyncCreateVertexBuffer(SizeInBytes, BUF_Static | BUF_ShaderResource, CreateInfo);
-		}
-	}
-	return nullptr;
+	return CreateRHIBuffer<bRenderThread>(TangentsData, GetNumVertices(), BUF_Static | BUF_ShaderResource, TEXT("TangentsRHIBuffer"));
 }
 
 FBufferRHIRef FStaticMeshVertexBuffer::CreateTangentsRHIBuffer_RenderThread()
@@ -300,22 +285,7 @@ FBufferRHIRef FStaticMeshVertexBuffer::CreateTangentsRHIBuffer_Async()
 template <bool bRenderThread>
 FBufferRHIRef FStaticMeshVertexBuffer::CreateTexCoordRHIBuffer_Internal()
 {
-	if (GetNumTexCoords())
-	{
-		FResourceArrayInterface* RESTRICT ResourceArray = TexcoordData ? TexcoordData->GetResourceArray() : nullptr;
-		const uint32 SizeInBytes = ResourceArray ? ResourceArray->GetResourceDataSize() : 0;
-		FRHIResourceCreateInfo CreateInfo(TEXT("TexCoordRHIBuffer"), ResourceArray);
-		CreateInfo.bWithoutNativeResource = !TexcoordData;
-		if (bRenderThread)
-		{
-			return RHICreateVertexBuffer(SizeInBytes, BUF_Static | BUF_ShaderResource, CreateInfo);
-		}
-		else
-		{
-			return RHIAsyncCreateVertexBuffer(SizeInBytes, BUF_Static | BUF_ShaderResource, CreateInfo);
-		}
-	}
-	return nullptr;
+	return CreateRHIBuffer<bRenderThread>(TexcoordData, GetNumTexCoords(), BUF_Static | BUF_ShaderResource, TEXT("TexCoordRHIBuffer"));
 }
 
 FBufferRHIRef FStaticMeshVertexBuffer::CreateTexCoordRHIBuffer_RenderThread()
diff --git a/Engine/Source/Runtime/Engine/Private/SkeletalMesh.cpp b/Engine/Source/Runtime/Engine/Private/SkeletalMesh.cpp
--- a/Engine/Source/Runtime/Engine/Private/SkeletalMesh.cpp
+++ b/Engine/Source/Runtime/Engine/Private/SkeletalMesh.cpp
@@ -282,29 +282,28 @@ FArchive& operator<<(FArchive& Ar, FClothBufferIndexMapping& ClothBufferIndexMap
 		<< ClothBufferIndexMapping.LODBiasStride;
 }
 
-/*-----------------------------------------------------------------------------
-FreeSkeletalMeshBuffersSinkCallback
------------------------------------------------------------------------------*/
-
 void FreeSkeletalMeshBuffersSinkCallback()
 {
-	// If r.FreeSkeletalMeshBuffers==1 then CPU buffer copies are to be released.
-	static TOptional<bool> bLastFreeSkeletalMeshBuffers;
-	static const auto CVar = IConsoleManager::Get().FindTConsoleVariableDataInt(TEXT("r.FreeSkeletalMeshBuffers"));	
-	bool bFreeSkeletalMeshBuffers = CVar->GetValueOnGameThread() == 1;
-	if(!bLastFreeSkeletalMeshBuffers.IsSet() || (bFreeSkeletalMeshBuffers != bLastFreeSkeletalMeshBuffers.GetValue()))
-	{
-		TRACE_CPUPROFILER_EVENT_SCOPE(FreeSkeletalMeshBuffersSinkCallback);
+	static const auto CVar = IConsoleManager::Get().FindTConsoleVariableDataInt(TEXT("r.FreeSkeletalMeshBuffers"));
+	static bool bPreviousFreeSkeletalMeshBuffers = false;
+	const bool bFreeSkeletalMeshBuffers = CVar->GetValueOnGameThread() == 1;
 
-		FlushRenderingCommands();
-		for (TObjectIterator<USkeletalMesh> It;It;++It)
+	// We previously relied on this sink occurring periodically to free helper structures, so run it if we're not using the new logic that frees on submit.
+	if (!GFreeStructuresOnRHIBufferCreation || (bPreviousFreeSkeletalMeshBuffers != bFreeSkeletalMeshBuffers))
+	{
+		bPreviousFreeSkeletalMeshBuffers = bFreeSkeletalMeshBuffers;
+		if (bFreeSkeletalMeshBuffers)
 		{
-			if (!It->HasPendingInitOrStreaming() && !It->GetResourceForRendering()->RequiresCPUSkinning(GMaxRHIFeatureLevel))
+			FlushRenderingCommands();
+			for (TObjectIterator<USkeletalMesh> It; It; ++It)
 			{
-				It->ReleaseCPUResources();
+				USkeletalMesh* SkeletalMesh = *It;
+				if (!SkeletalMesh->HasPendingInitOrStreaming() && !SkeletalMesh->GetResourceForRendering()->RequiresCPUSkinning(GMaxRHIFeatureLevel))
+				{
+					SkeletalMesh->ReleaseCPUResources();
+				}
 			}
 		}
-		bLastFreeSkeletalMeshBuffers = bFreeSkeletalMeshBuffers;
 	}
 }
 
diff --git a/Engine/Source/Runtime/RenderCore/Private/RenderResource.cpp b/Engine/Source/Runtime/RenderCore/Private/RenderResource.cpp
--- a/Engine/Source/Runtime/RenderCore/Private/RenderResource.cpp
+++ b/Engine/Source/Runtime/RenderCore/Private/RenderResource.cpp
@@ -26,9 +26,15 @@ FAutoConsoleVariableRef CVarMaxVertexBytesAllocatedPerFrame(
 int32 GGlobalBufferNumFramesUnusedThresold = 30;
 FAutoConsoleVariableRef CVarReadBufferNumFramesUnusedThresold(
 	TEXT("r.NumFramesUnusedBeforeReleasingGlobalResourceBuffers"),
-	GGlobalBufferNumFramesUnusedThresold ,
+	GGlobalBufferNumFramesUnusedThresold,
 	TEXT("Number of frames after which unused global resource allocations will be discarded. Set 0 to ignore. (default=30)"));
 
+bool GFreeStructuresOnRHIBufferCreation = true;
+FAutoConsoleVariableRef CVarFreeStructuresOnRHIBufferCreation(
+	TEXT("r.FreeStructuresOnRHIBufferCreation"),
+	GFreeStructuresOnRHIBufferCreation,
+	TEXT("Toggles experimental method for freeing helper structures that own the resource arrays after submitting to RHI instead of in the callback sink."));
+
 int32 GVarDebugForceRuntimeBLAS = 0;
 FAutoConsoleVariableRef CVarDebugForceRuntimeBLAS(
 	TEXT("r.Raytracing.DebugForceRuntimeBLAS"),
diff --git a/Engine/Source/Runtime/RenderCore/Public/RenderResource.h b/Engine/Source/Runtime/RenderCore/Public/RenderResource.h
--- a/Engine/Source/Runtime/RenderCore/Public/RenderResource.h
+++ b/Engine/Source/Runtime/RenderCore/Public/RenderResource.h
@@ -16,6 +16,9 @@
 /** Number of frames after which unused global resource allocations will be discarded. */
 extern int32 GGlobalBufferNumFramesUnusedThresold;
 
+/** Experimental: whether we free helper structures after submitting to RHI. */
+extern RENDERCORE_API bool GFreeStructuresOnRHIBufferCreation;
+
 /**
  * A rendering resource which is owned by the rendering thread.
  * NOTE - Adding new virtual methods to this class may require stubs added to FViewport/FDummyViewport, otherwise certain modules may have link errors
@@ -158,6 +161,39 @@ protected:
 	const FStaticFeatureLevel GetFeatureLevel() const { return FeatureLevel == ERHIFeatureLevel::Num ? FStaticFeatureLevel(GMaxRHIFeatureLevel) : FeatureLevel; }
 	FORCEINLINE bool HasValidFeatureLevel() const { return FeatureLevel < ERHIFeatureLevel::Num; }
 
+	// Helper for submitting a resource array to RHI and freeing eligible CPU memory
+	template<bool bRenderThread, typename T>
+	FBufferRHIRef CreateRHIBuffer(T& InOutResourceObject, const uint32 ResourceCount, EBufferUsageFlags InBufferUsageFlags, const TCHAR* InDebugName)
+	{
+		FBufferRHIRef Buffer;
+		FResourceArrayInterface* RESTRICT ResourceArray = InOutResourceObject ? InOutResourceObject->GetResourceArray() : nullptr;
+		if (ResourceCount != 0)
+		{
+			const uint32 SizeInBytes = ResourceArray ? ResourceArray->GetResourceDataSize() : 0;
+			FRHIResourceCreateInfo CreateInfo(InDebugName, ResourceArray);
+			CreateInfo.bWithoutNativeResource = !InOutResourceObject;
+
+			if (bRenderThread)
+			{
+				Buffer = RHICreateVertexBuffer(SizeInBytes, InBufferUsageFlags, CreateInfo);
+			}
+			else
+			{
+				Buffer = RHIAsyncCreateVertexBuffer(SizeInBytes, InBufferUsageFlags, CreateInfo);
+			}
+		}
+
+		if (GFreeStructuresOnRHIBufferCreation && InOutResourceObject && (!ResourceArray || !ResourceArray->GetAllowCPUAccess()))
+		{
+			// The resource should have been freed as part of the buffer creation
+			check(!ResourceArray || !ResourceArray->GetResourceDataSize());
+			delete InOutResourceObject;
+			InOutResourceObject = nullptr;
+		}
+
+		return Buffer;
+	}
+
 private:
 	TEnumAsByte<ERHIFeatureLevel::Type> FeatureLevel;
 };
@@ -520,6 +556,7 @@ public:
 	{
 		VertexBufferRHI.SafeRelease();
 	}
+
 	virtual FString GetFriendlyName() const override { return TEXT("FVertexBuffer"); }
 };
 
