Replay subsystem updates for Replay Demo Splitting
- Adding delegates to report when recording is started and stopped
- Preventing new recordings from starting when there's already one in progress

diff --git a/Engine/Source/Runtime/Engine/Private/ReplaySubsystem.cpp b/Engine/Source/Runtime/Engine/Private/ReplaySubsystem.cpp
--- a/Engine/Source/Runtime/Engine/Private/ReplaySubsystem.cpp
+++ b/Engine/Source/Runtime/Engine/Private/ReplaySubsystem.cpp
@@ -103,7 +103,7 @@ void UReplaySubsystem::RecordReplay(const FString& Name, const FString& Friendly
 
 	if (FParse::Param(FCommandLine::Get(), TEXT("NOREPLAYS")))
 	{
-		UE_LOG(LogDemo, Warning, TEXT("UReplaySubsystem::StartRecordingReplay: Rejected due to -noreplays option"));
+		UE_LOG(LogDemo, Warning, TEXT("UReplaySubsystem::RecordReplay: Rejected due to -noreplays option"));
 		return;
 	}
 
@@ -111,13 +111,19 @@ void UReplaySubsystem::RecordReplay(const FString& Name, const FString& Friendly
 
 	if (CurrentWorld == nullptr)
 	{
-		UE_LOG(LogDemo, Warning, TEXT("UReplaySubsystem::StartRecordingReplay: GetWorld() is null"));
+		UE_LOG(LogDemo, Warning, TEXT("UReplaySubsystem::RecordReplay: GetWorld() is null"));
 		return;
 	}
 
 	if (CurrentWorld->IsPlayingReplay())
 	{
-		UE_LOG(LogDemo, Warning, TEXT("UReplaySubsystem::StartRecordingReplay: A replay is already playing, cannot begin recording another one."));
+		UE_LOG(LogDemo, Warning, TEXT("UReplaySubsystem::RecordReplay: A replay is already playing, cannot begin recording another one."));
+		return;
+	}
+
+	if(IsRecording())
+	{
+		UE_LOG(LogDemo, Warning, TEXT("UReplaySubsystem::RecordReplay: A replay is already being recorded, cannot begin recording another one."));
 		return;
 	}
 
@@ -154,6 +160,7 @@ void UReplaySubsystem::RecordReplay(const FString& Name, const FString& Friendly
 		UE_LOG(LogDemo, Log, TEXT("UReplaySubsystem::RecordReplay: Starting recording with replay connection.  Name: %s FriendlyName: %s"), *Name, *FriendlyName);
 
 		Connection->StartRecording();
+		OnReplayRecordingStarted.Broadcast();
 		return;
 	}
 
@@ -211,6 +218,7 @@ void UReplaySubsystem::RecordReplay(const FString& Name, const FString& Friendly
 	}
 
 	UE_LOG(LogDemo, Verbose, TEXT("Num Network Actors: %i"), DemoNetDriver->GetNetworkObjectList().GetActiveObjects().Num());
+	OnReplayRecordingStarted.Broadcast();
 }
 
 bool UReplaySubsystem::PlayReplay(const FString& Name, UWorld* WorldOverride, const TArray<FString>& AdditionalOptions)
@@ -303,6 +311,8 @@ void UReplaySubsystem::StopExistingReplays(UWorld* InWorld)
 		Connection->CleanUp();
 		ReplayConnection = nullptr;
 	}
+
+	OnReplayRecordingStopped.Broadcast();
 }
 
 FString UReplaySubsystem::GetActiveReplayName() const
diff --git a/Engine/Source/Runtime/Engine/Public/ReplaySubsystem.h b/Engine/Source/Runtime/Engine/Public/ReplaySubsystem.h
--- a/Engine/Source/Runtime/Engine/Public/ReplaySubsystem.h
+++ b/Engine/Source/Runtime/Engine/Public/ReplaySubsystem.h
@@ -8,6 +8,9 @@
 
 class UReplayNetConnection;
 
+DECLARE_DYNAMIC_MULTICAST_DELEGATE(FOnReplayRecordingStarted);
+DECLARE_DYNAMIC_MULTICAST_DELEGATE(FOnReplayRecordingStopped);
+
 UCLASS(DisplayName = "Replay Subsystem")
 class ENGINE_API UReplaySubsystem : public UGameInstanceSubsystem
 {
@@ -121,6 +124,12 @@ public:
 	UPROPERTY(EditAnywhere, Category=Replay)
 	bool bLoadDefaultMapOnStop = true;
 
+	UPROPERTY(BlueprintAssignable)
+	FOnReplayRecordingStarted OnReplayRecordingStarted;
+
+	UPROPERTY(BlueprintAssignable)
+	FOnReplayRecordingStopped OnReplayRecordingStopped;
+
 private:
 	void StopExistingReplays(UWorld* InWorld);
 
