Fix network granular mem tracking formatting for top level kb

diff --git a/Engine/Source/Runtime/Engine/Private/Net/NetworkGranularMemoryLogging.cpp b/Engine/Source/Runtime/Engine/Private/Net/NetworkGranularMemoryLogging.cpp
--- a/Engine/Source/Runtime/Engine/Private/Net/NetworkGranularMemoryLogging.cpp
+++ b/Engine/Source/Runtime/Engine/Private/Net/NetworkGranularMemoryLogging.cpp
@@ -180,7 +180,7 @@ namespace GranularNetworkMemoryTrackingPrivate
 			for (auto It = ScopeStack->TopLevelScopes.CreateIterator(); It; ++It)
 			{
 				It.Value()->GenerateRows(EmptyPrefix, Rows);
-				TopLevelKBRows.Add(FString::Printf(TEXT("%s KB\r\n%d"), *It.Key(), static_cast<double>(It.Value()->GetTotalBytes()) / double(1024.f)));
+				TopLevelKBRows.Add(FString::Printf(TEXT("%s KB\r\n%llu"), *It.Key(), FMath::DivideAndRoundUp(It.Value()->GetTotalBytes(), (uint64)1024)));
 			}
 
 			Rows.Append(TopLevelKBRows);
