Fix inverted condition (CL 18738937)

#rnx

diff --git a/Engine/Source/Runtime/Engine/Private/WorldComposition.cpp b/Engine/Source/Runtime/Engine/Private/WorldComposition.cpp
--- a/Engine/Source/Runtime/Engine/Private/WorldComposition.cpp
+++ b/Engine/Source/Runtime/Engine/Private/WorldComposition.cpp
@@ -482,7 +482,7 @@ void UWorldComposition::OnTileInfoUpdated(const FName& InPackageName, const FWor
 	UPackage* LevelPackage = Cast<UPackage>(StaticFindObjectFast(UPackage::StaticClass(), NULL, Tile->PackageName));
 	if (LevelPackage)
 	{
-		if (LevelPackage->GetWorldTileInfo())
+		if (!LevelPackage->GetWorldTileInfo())
 		{
 			LevelPackage->SetWorldTileInfo(MakeUnique<FWorldTileInfo>(Tile->Info));
 			PackageDirty = true;
