Lumen - fixed Lumen show flags duplicate localization names

diff --git a/Engine/Source/Runtime/Engine/Public/ShowFlagsValues.inl b/Engine/Source/Runtime/Engine/Public/ShowFlagsValues.inl
--- a/Engine/Source/Runtime/Engine/Public/ShowFlagsValues.inl
+++ b/Engine/Source/Runtime/Engine/Public/ShowFlagsValues.inl
@@ -409,19 +409,19 @@ SHOWFLAG_FIXED_IN_SHIPPING(0, VirtualTextureResidency, SFG_Hidden, NSLOCTEXT("Un
 SHOWFLAG_FIXED_IN_SHIPPING(0, DVRStreamingSystem, SFG_Hidden, NSLOCTEXT("UnrealEd", "DVRStreamingSystem", "DVR Streaming System"))
 
 /** Use screen space tracing in Lumen */
-SHOWFLAG_ALWAYS_ACCESSIBLE(LumenScreenTraces, SFG_Lumen, NSLOCTEXT("UnrealEd", "LumenSF", "Screen Traces"))
+SHOWFLAG_ALWAYS_ACCESSIBLE(LumenScreenTraces, SFG_Lumen, NSLOCTEXT("UnrealEd", "LumenScreenTracesSF", "Screen Traces"))
 /** Use detail tracing in Lumen */
-SHOWFLAG_ALWAYS_ACCESSIBLE(LumenDetailTraces, SFG_Lumen, NSLOCTEXT("UnrealEd", "LumenSF", "Detail Traces"))
+SHOWFLAG_ALWAYS_ACCESSIBLE(LumenDetailTraces, SFG_Lumen, NSLOCTEXT("UnrealEd", "LumenDetailTracesSF", "Detail Traces"))
 /** Use global traces in Lumen */
-SHOWFLAG_ALWAYS_ACCESSIBLE(LumenGlobalTraces, SFG_Lumen, NSLOCTEXT("UnrealEd", "LumenSF", "Global Traces"))
+SHOWFLAG_ALWAYS_ACCESSIBLE(LumenGlobalTraces, SFG_Lumen, NSLOCTEXT("UnrealEd", "LumenGlobalTracesSF", "Global Traces"))
 /** Use far field traces in Lumen */
-SHOWFLAG_ALWAYS_ACCESSIBLE(LumenFarFieldTraces, SFG_Lumen, NSLOCTEXT("UnrealEd", "LumenSF", "Far Field Traces"))
+SHOWFLAG_ALWAYS_ACCESSIBLE(LumenFarFieldTraces, SFG_Lumen, NSLOCTEXT("UnrealEd", "LumenFarFieldTracesSF", "Far Field Traces"))
 /** Compute secondary bounces in Lumen */
-SHOWFLAG_ALWAYS_ACCESSIBLE(LumenSecondaryBounces, SFG_Lumen, NSLOCTEXT("UnrealEd", "LumenSF", "Secondary Bounces"))
+SHOWFLAG_ALWAYS_ACCESSIBLE(LumenSecondaryBounces, SFG_Lumen, NSLOCTEXT("UnrealEd", "LumenSecondaryBouncesSF", "Secondary Bounces"))
 /** Compute screen space directional occlusion in Lumen */
-SHOWFLAG_ALWAYS_ACCESSIBLE(LumenScreenSpaceDirectionalOcclusion, SFG_Lumen, NSLOCTEXT("UnrealEd", "LumenSF", "Screen Space Directional Occlusion"))
+SHOWFLAG_ALWAYS_ACCESSIBLE(LumenScreenSpaceDirectionalOcclusion, SFG_Lumen, NSLOCTEXT("UnrealEd", "LumenScreenSpaceDirectionalOcclusionSF", "Screen Space Directional Occlusion"))
 /** Whether to reuse shadowmaps when calculating shadowing.  Can be disabled to debug view dependent lighting from shadowing technique mismatches. */
-SHOWFLAG_ALWAYS_ACCESSIBLE(LumenReuseShadowMaps, SFG_Lumen, NSLOCTEXT("UnrealEd", "LumenSF", "Reuse Shadow Maps"))
+SHOWFLAG_ALWAYS_ACCESSIBLE(LumenReuseShadowMaps, SFG_Lumen, NSLOCTEXT("UnrealEd", "LumenReuseShadowMapsSF", "Reuse Shadow Maps"))
 
 #undef SHOWFLAG_ALWAYS_ACCESSIBLE
 #undef SHOWFLAG_FIXED_IN_SHIPPING
