Add Median and Average of int array to KismetMathLibrary

#rnx

diff --git a/Engine/Source/Runtime/Engine/Classes/Kismet/KismetMathLibrary.h b/Engine/Source/Runtime/Engine/Classes/Kismet/KismetMathLibrary.h
--- a/Engine/Source/Runtime/Engine/Classes/Kismet/KismetMathLibrary.h
+++ b/Engine/Source/Runtime/Engine/Classes/Kismet/KismetMathLibrary.h
@@ -795,6 +795,14 @@ class ENGINE_API UKismetMathLibrary : public UBlueprintFunctionLibrary
 	UFUNCTION(BlueprintPure, Category="Math|Integer")
 	static void MinOfIntArray(const TArray<int32>& IntArray, int32& IndexOfMinValue, int32& MinValue);
 
+	/** Returns median of all array entries. Returns value of 0 if the supplied array is empty. */
+	UFUNCTION(BlueprintPure, Category = "Math|Integer")
+	static void MedianOfIntArray(TArray<int32> IntArray, float& MedianValue);
+
+	/** Returns average of all array entries. Returns value of 0 if the supplied array is empty. */
+	UFUNCTION(BlueprintPure, Category = "Math|Integer")
+	static void AverageOfIntArray(const TArray<int32>& IntArray, float& AverageValue);
+
 	/** Returns max of all array entries and the index at which it was found. Returns value of 0 and index of -1 if the supplied array is empty. */
 	UFUNCTION(BlueprintPure, Category="Math|Float")
 	static void MaxOfFloatArray(const TArray<float>& FloatArray, int32& IndexOfMaxValue, float& MaxValue);
diff --git a/Engine/Source/Runtime/Engine/Private/KismetMathLibrary.cpp b/Engine/Source/Runtime/Engine/Private/KismetMathLibrary.cpp
--- a/Engine/Source/Runtime/Engine/Private/KismetMathLibrary.cpp
+++ b/Engine/Source/Runtime/Engine/Private/KismetMathLibrary.cpp
@@ -315,6 +315,39 @@ void UKismetMathLibrary::MinOfIntArray(const TArray<int32>& IntArray, int32& Ind
 	MinValue = FMath::Min<int32>(IntArray, &IndexOfMinValue);
 }
 
+void UKismetMathLibrary::MedianOfIntArray(TArray<int32> IntArray, float& MedianValue)
+{
+	if (IntArray.IsEmpty())
+	{
+		MedianValue = 0;
+		return;
+	}
+
+	IntArray.Sort();
+	
+	if (IntArray.Num() % 2 == 1)
+	{
+		MedianValue = IntArray[(IntArray.Num() - 1) / 2];
+	}
+	else
+	{
+		MedianValue = IntArray[IntArray.Num() / 2] + IntArray[(IntArray.Num() / 2) - 1];
+		MedianValue /= 2;
+	}
+}
+
+void UKismetMathLibrary::AverageOfIntArray(const TArray<int32>& IntArray, float& AverageValue)
+{
+	AverageValue = 0;
+
+	for (int32 Index = 0; Index < IntArray.Num(); ++Index)
+	{
+		AverageValue += IntArray[Index];
+	}
+
+	AverageValue /= IntArray.Num();
+}
+
 void UKismetMathLibrary::MaxOfFloatArray(const TArray<float>& FloatArray, int32& IndexOfMaxValue, float& MaxValue)
 {
 	MaxValue = FMath::Max(FloatArray, &IndexOfMaxValue);
