Issue: Replay playback does not start at 0.
Problem: We are updating DemoCurrentTime in TickDemoPlayback before the ReplayStreamer has any data.
Fix: Make sure we dont update time before we have data to play.

diff --git a/Engine/Source/Runtime/Engine/Private/DemoNetDriver.cpp b/Engine/Source/Runtime/Engine/Private/DemoNetDriver.cpp
--- a/Engine/Source/Runtime/Engine/Private/DemoNetDriver.cpp
+++ b/Engine/Source/Runtime/Engine/Private/DemoNetDriver.cpp
@@ -2529,6 +2529,12 @@ void UDemoNetDriver::TickDemoPlayback(float DeltaSeconds)
 		return;
 	}
 
+	// If we don't have data on frame 0 wait until we have it
+	if (!GetReplayStreamer()->IsDataAvailable() && ReplayHelper.DemoFrameNum == 0)
+	{
+		return;
+	}
+
 	// If the ExitAfterReplay option is set, automatically shut down at the end of the replay.
 	// Use AtEnd() of the archive instead of checking DemoCurrentTime/DemoTotalTime, because the DemoCurrentTime may never catch up to DemoTotalTime.
 	if (FArchive* const StreamingArchive = ReplayHelper.ReplayStreamer->GetStreamingArchive())
diff --git a/Engine/Source/Runtime/NetworkReplayStreaming/HttpNetworkReplayStreaming/Private/HttpNetworkReplayStreaming.cpp b/Engine/Source/Runtime/NetworkReplayStreaming/HttpNetworkReplayStreaming/Private/HttpNetworkReplayStreaming.cpp
--- a/Engine/Source/Runtime/NetworkReplayStreaming/HttpNetworkReplayStreaming/Private/HttpNetworkReplayStreaming.cpp
+++ b/Engine/Source/Runtime/NetworkReplayStreaming/HttpNetworkReplayStreaming/Private/HttpNetworkReplayStreaming.cpp
@@ -283,8 +283,6 @@ void FHttpNetworkReplayStreamer::StartStreaming(const FStartStreamingParameters&
 	CheckpointArchive.SetIsLoading(!Params.bRecord);
 	CheckpointArchive.SetIsSaving(Params.bRecord);
 
-	LastChunkTime = FPlatformTime::Seconds();
-
 	TotalDemoTimeInMS = 0;
 
 	StreamTimeRangeStart	= 0;
diff --git a/Engine/Source/Runtime/NetworkReplayStreaming/LocalFileNetworkReplayStreaming/Private/LocalFileNetworkReplayStreaming.cpp b/Engine/Source/Runtime/NetworkReplayStreaming/LocalFileNetworkReplayStreaming/Private/LocalFileNetworkReplayStreaming.cpp
--- a/Engine/Source/Runtime/NetworkReplayStreaming/LocalFileNetworkReplayStreaming/Private/LocalFileNetworkReplayStreaming.cpp
+++ b/Engine/Source/Runtime/NetworkReplayStreaming/LocalFileNetworkReplayStreaming/Private/LocalFileNetworkReplayStreaming.cpp
@@ -698,8 +698,6 @@ void FLocalFileNetworkReplayStreamer::StartStreaming(const FStartStreamingParame
 	StreamDataOffset = 0;
 	StreamChunkIndex = 0;
 
-	LastChunkTime = FPlatformTime::Seconds();
-
 	if (!Params.bRecord)
 	{
 		// We are playing
