Re-submitting (updated changes) to USoundWave to not create a FSoundWaveProxy or access the Audio Streaming Cache when FApp::CanEverRenderAudio() is FALSE.

diff --git a/Engine/Source/Runtime/Engine/Private/SoundWave.cpp b/Engine/Source/Runtime/Engine/Private/SoundWave.cpp
--- a/Engine/Source/Runtime/Engine/Private/SoundWave.cpp
+++ b/Engine/Source/Runtime/Engine/Private/SoundWave.cpp
@@ -911,7 +911,7 @@ void USoundWave::Serialize( FArchive& Ar )
 		const bool bShouldLoadChunks = bCooked || bBuiltStreamedAudio;
 
 		// If stream caching is enabled, here we determine if we should retain or prime this wave on load.
-		if (bShouldStreamSound && bShouldLoadChunks && FPlatformCompressionUtilities::IsCurrentPlatformUsingStreamCaching())
+		if (bShouldStreamSound && bShouldLoadChunks && FPlatformCompressionUtilities::IsCurrentPlatformUsingStreamCaching() && FApp::CanEverRenderAudio())
 		{
 			ESoundWaveLoadingBehavior CurrentLoadingBehavior = GetLoadingBehavior(false);
 
@@ -1479,12 +1479,13 @@ void USoundWave::PostLoad()
 	// If our loading behavior is defined by a sound class, we need to update whether this sound wave actually needs to retain its audio data or not.
 	ESoundWaveLoadingBehavior ActualLoadingBehavior = GetLoadingBehavior();
 
-	if (!InternalProxy.IsValid() && ActualLoadingBehavior != ESoundWaveLoadingBehavior::ForceInline)
+	// no need for a proxy if we can't render audio!
+	if (!InternalProxy.IsValid() && ActualLoadingBehavior != ESoundWaveLoadingBehavior::ForceInline && FApp::CanEverRenderAudio())
 	{
 		InternalProxy = CreateSoundWaveProxy();
 	}
 
-	if (ShouldUseStreamCaching() && ActualLoadingBehavior != GetLoadingBehavior(false))
+	if (FApp::CanEverRenderAudio() && ShouldUseStreamCaching() && ActualLoadingBehavior != GetLoadingBehavior(false))
 	{
 		if (!GIsEditor && !DisableRetainingCVar && ActualLoadingBehavior == ESoundWaveLoadingBehavior::RetainOnLoad)
 		{
@@ -2237,8 +2238,6 @@ void USoundWave::BakeEnvelopeAnalysis()
 
 void USoundWave::PostEditChangeProperty(FPropertyChangedEvent& PropertyChangedEvent)
 {
-	// maxtodo: test this, stale data? looping? channel count?
-
 	Super::PostEditChangeProperty(PropertyChangedEvent);
 	check(SoundWaveDataPtr);
 
@@ -2933,7 +2932,7 @@ bool USoundWave::ShouldUseStreamCaching() const
 
 TArrayView<const uint8> USoundWave::GetZerothChunk(bool bForImmediatePlayback)
 {
-	if(IsTemplate() || IsRunningDedicatedServer())
+	if(IsTemplate() || IsRunningDedicatedServer() || !FApp::CanEverRenderAudio())
 	{
 		return TArrayView<const uint8>();
 	}
@@ -3077,7 +3076,7 @@ void USoundWave::RemovePlayingSource(const FSoundWaveClientPtr& Source)
 
 void USoundWave::UpdatePlatformData()
 {
-	if (IsStreaming(nullptr))
+	if (IsStreaming(nullptr) && FApp::CanEverRenderAudio())
 	{
 		if (!InternalProxy.IsValid())
 		{
@@ -3435,14 +3434,14 @@ void USoundWave::GetHandleForChunkOfAudio(TFunction<void(FAudioChunkHandle&&)> O
 {
 	ENamedThreads::Type ThreadToDispatchCallbackOn = (DispatchToGameThreadOnChunkRequestCVar != 0) ? ENamedThreads::GameThread : ENamedThreads::AnyThread;
 
-	if (!InternalProxy.IsValid())
+	if (!InternalProxy.IsValid() && FApp::CanEverRenderAudio())
 	{
 		InternalProxy = CreateSoundWaveProxy();
 	}
 
 	// if we are requesting a chunk that is out of bounds,
-	// early exit.
-	if (ChunkIndex >= static_cast<int32>(GetNumChunks()))
+	// or we shouldn't be interacting w/ the cache, early exit.
+	if (!FApp::CanEverRenderAudio() || (ChunkIndex >= static_cast<int32>(GetNumChunks())))
 	{
 		FAudioChunkHandle EmptyChunkHandle;
 		OnLoadCompleted(MoveTemp(EmptyChunkHandle));
@@ -3505,7 +3504,7 @@ void USoundWave::RetainCompressedAudio(bool bForceSync /*= false*/)
 
 	// Since the zeroth chunk is always inlined and stored in memory,
 	// early exit if we only have one chunk.
-	if (GIsEditor || IsTemplate() || IsRunningDedicatedServer() || !IsStreaming() || DisableRetainingCVar || GetNumChunks() <= 1)
+	if (GIsEditor || IsTemplate() || IsRunningDedicatedServer() || !IsStreaming() || !FApp::CanEverRenderAudio() || DisableRetainingCVar || GetNumChunks() <= 1)
 	{
 		return;
 	}
