Use correct parameter when mixing a modulator into a destination value (that set by the destination, not the referenced modulator's OutputProperty field)
#rnx

diff --git a/Engine/Source/Runtime/Engine/Private/SoundModulationDestination.cpp b/Engine/Source/Runtime/Engine/Private/SoundModulationDestination.cpp
--- a/Engine/Source/Runtime/Engine/Private/SoundModulationDestination.cpp
+++ b/Engine/Source/Runtime/Engine/Private/SoundModulationDestination.cpp
@@ -219,7 +219,7 @@ namespace Audio
 			return;
 		}
 
-		Audio::FModulationParameter OutputParam = InModulator->GetOutputParameter();
+		Audio::FModulationParameter OutputParam = Audio::GetModulationParameter(ParameterName);
 		auto UpdateHandleLambda = [this, ModSettings = InModulator->CreateProxySettings(), Parameter = MoveTemp(OutputParam)]() mutable
 		{
 			if (FAudioDevice* AudioDevice = FAudioDeviceManager::Get()->GetAudioDeviceRaw(DeviceId))
