Fix VT UVs when analytic derivatives are not valid

diff --git a/Engine/Source/Runtime/Engine/Private/Materials/HLSLMaterialTranslator.cpp b/Engine/Source/Runtime/Engine/Private/Materials/HLSLMaterialTranslator.cpp
--- a/Engine/Source/Runtime/Engine/Private/Materials/HLSLMaterialTranslator.cpp
+++ b/Engine/Source/Runtime/Engine/Private/Materials/HLSLMaterialTranslator.cpp
@@ -6003,25 +6003,21 @@ int32 FHLSLMaterialTranslator::TextureSample(
 		AddCodeChunk(MCT_Float, TEXT("MaterialStoreTexCoordScale(Parameters, %s, %d)"), *CoerceParameter(NonLWCCoordinateIndex, UVsType), (int)TextureReferenceIndex);
 	}
 
-	FString UV_Value = TEXT("0.0f");
+	const FString UV_Value = CoerceParameter(NonLWCCoordinateIndex, UVsType);
 	FString UV_Ddx = TEXT("0.0f");
 	FString UV_Ddy = TEXT("0.0f");
 	FString UV_Scale = TEXT("1.0f");
-	if (IsAnalyticDerivEnabled() && IsDerivativeValid(UvDerivativeStatus))
+	if (IsAnalyticDerivEnabled() && UvDerivativeStatus == EDerivativeStatus::Valid)
 	{
-		UV_Value = CoerceParameter(NonLWCCoordinateIndex, UVsType);
-		if (UvDerivativeStatus == EDerivativeStatus::Valid)
+		const EMaterialValueType SourceUVsType = GetParameterType(CoordinateIndex);
+		const FString UVAnalytic = GetParameterCodeDeriv(CoordinateIndex, CompiledPDV_Analytic);
+		UV_Ddx = CoerceValue(UVAnalytic + TEXT(".Ddx"), MakeNonLWCType(SourceUVsType), UVsType); // Ddx/y are never LWC scale
+		UV_Ddy = CoerceValue(UVAnalytic + TEXT(".Ddy"), MakeNonLWCType(SourceUVsType), UVsType);
+		if (MipValueMode == TMVM_MipBias)
 		{
-			const EMaterialValueType SourceUVsType = GetParameterType(CoordinateIndex);
-			const FString UVAnalytic = GetParameterCodeDeriv(CoordinateIndex, CompiledPDV_Analytic);
-			UV_Ddx = CoerceValue(UVAnalytic + TEXT(".Ddx"), MakeNonLWCType(SourceUVsType), UVsType); // Ddx/y are never LWC scale
-			UV_Ddy = CoerceValue(UVAnalytic + TEXT(".Ddy"), MakeNonLWCType(SourceUVsType), UVsType);
-			if (MipValueMode == TMVM_MipBias)
-			{
-				UV_Scale = FString::Printf(TEXT("exp2(%s)"), *MipValue0Code);
-				UV_Ddx = FString::Printf(TEXT("(%s)*exp2(%s)"), *UV_Ddx, *MipValue0Code);
-				UV_Ddy = FString::Printf(TEXT("(%s)*exp2(%s)"), *UV_Ddy, *MipValue0Code);
-			}
+			UV_Scale = FString::Printf(TEXT("exp2(%s)"), *MipValue0Code);
+			UV_Ddx = FString::Printf(TEXT("(%s)*exp2(%s)"), *UV_Ddx, *MipValue0Code);
+			UV_Ddy = FString::Printf(TEXT("(%s)*exp2(%s)"), *UV_Ddy, *MipValue0Code);
 		}
 	}
 
@@ -6112,8 +6108,7 @@ int32 FHLSLMaterialTranslator::TextureSample(
 	else
 	{
 		// Non-VT MipValueMode logic
-		const FString UVs = CoerceParameter(NonLWCCoordinateIndex, UVsType);
-
+		// 
 		// Re-route decal texture sampling so platforms may add specific workarounds there
 		if (bDecal)
 		{
@@ -6123,19 +6118,19 @@ int32 FHLSLMaterialTranslator::TextureSample(
 		FString SampleCodeFinite;
 		if (MipValueMode == TMVM_None)
 		{
-			SampleCodeFinite = FString::Printf(TEXT("%s(%s,%s,%s)"), *TextureTypeName, *TextureName, *SamplerStateCode, *UVs);
+			SampleCodeFinite = FString::Printf(TEXT("%s(%s,%s,%s)"), *TextureTypeName, *TextureName, *SamplerStateCode, *UV_Value);
 		}
 		else if (MipValueMode == TMVM_MipLevel)
 		{
-			SampleCodeFinite = FString::Printf(TEXT("%sLevel(%s,%s,%s,%s)"), *TextureTypeName, *TextureName, *SamplerStateCode, *UVs, *MipValue0Code);
+			SampleCodeFinite = FString::Printf(TEXT("%sLevel(%s,%s,%s,%s)"), *TextureTypeName, *TextureName, *SamplerStateCode, *UV_Value, *MipValue0Code);
 		}
 		else if (MipValueMode == TMVM_MipBias)
 		{
-			SampleCodeFinite = FString::Printf(TEXT("%sBias(%s,%s,%s,%s)"), *TextureTypeName, *TextureName, *SamplerStateCode, *UVs, *MipValue0Code);
+			SampleCodeFinite = FString::Printf(TEXT("%sBias(%s,%s,%s,%s)"), *TextureTypeName, *TextureName, *SamplerStateCode, *UV_Value, *MipValue0Code);
 		}
 		else if (MipValueMode == TMVM_Derivative)
 		{
-			SampleCodeFinite = FString::Printf(TEXT("%sGrad(%s,%s,%s,%s,%s)"), *TextureTypeName, *TextureName, *SamplerStateCode, *UVs, *MipValue0Code, *MipValue1Code);
+			SampleCodeFinite = FString::Printf(TEXT("%sGrad(%s,%s,%s,%s,%s)"), *TextureTypeName, *TextureName, *SamplerStateCode, *UV_Value, *MipValue0Code, *MipValue1Code);
 		}
 		else
 		{
@@ -6151,9 +6146,9 @@ int32 FHLSLMaterialTranslator::TextureSample(
 			{
 				
 				if (SamplerDebugSupported(TextureType, bVirtualTexture, bDecal))
-					SampleCodeAnalytic = FString::Printf(TEXT("Debug%sGrad(%s,%s,%s,%s,%s,%s)"), *TextureTypeName, *TextureName, *SamplerStateCode, *UVs, *UV_Ddx, *UV_Ddy, *UV_Scale);
+					SampleCodeAnalytic = FString::Printf(TEXT("Debug%sGrad(%s,%s,%s,%s,%s,%s)"), *TextureTypeName, *TextureName, *SamplerStateCode, *UV_Value, *UV_Ddx, *UV_Ddy, *UV_Scale);
 				else
-					SampleCodeAnalytic = FString::Printf(TEXT("%sGrad(%s,%s,%s,%s,%s)"), *TextureTypeName, *TextureName, *SamplerStateCode, *UVs, *UV_Ddx, *UV_Ddy);
+					SampleCodeAnalytic = FString::Printf(TEXT("%sGrad(%s,%s,%s,%s,%s)"), *TextureTypeName, *TextureName, *SamplerStateCode, *UV_Value, *UV_Ddx, *UV_Ddy);
 
 				SampleCodeAnalytic = ApplySamplerType(SampleCodeAnalytic, SamplerType);
 			}
