Add more misc functionality to new HLSL translator

diff --git a/Engine/Source/Runtime/Engine/Classes/Materials/MaterialExpressionClamp.h b/Engine/Source/Runtime/Engine/Classes/Materials/MaterialExpressionClamp.h
--- a/Engine/Source/Runtime/Engine/Classes/Materials/MaterialExpressionClamp.h
+++ b/Engine/Source/Runtime/Engine/Classes/Materials/MaterialExpressionClamp.h
@@ -47,6 +47,7 @@ class UMaterialExpressionClamp : public UMaterialExpression
 
 	//~ Begin UMaterialExpression Interface
 #if WITH_EDITOR
+	virtual bool GenerateHLSLExpression(FMaterialHLSLGenerator& Generator, UE::HLSLTree::FScope& Scope, int32 OutputIndex, UE::HLSLTree::FExpression*& OutExpression) override;
 	virtual int32 Compile(class FMaterialCompiler* Compiler, int32 OutputIndex) override;
 	virtual void GetCaption(TArray<FString>& OutCaptions) const override;
 #endif
diff --git a/Engine/Source/Runtime/Engine/Classes/Materials/MaterialExpressionDotProduct.h b/Engine/Source/Runtime/Engine/Classes/Materials/MaterialExpressionDotProduct.h
--- a/Engine/Source/Runtime/Engine/Classes/Materials/MaterialExpressionDotProduct.h
+++ b/Engine/Source/Runtime/Engine/Classes/Materials/MaterialExpressionDotProduct.h
@@ -23,6 +23,7 @@ class UMaterialExpressionDotProduct : public UMaterialExpression
 
 	//~ Begin UMaterialExpression Interface
 #if WITH_EDITOR
+	virtual bool GenerateHLSLExpression(FMaterialHLSLGenerator& Generator, UE::HLSLTree::FScope& Scope, int32 OutputIndex, UE::HLSLTree::FExpression*& OutExpression) override;
 	virtual int32 Compile(class FMaterialCompiler* Compiler, int32 OutputIndex) override;
 	virtual void GetCaption(TArray<FString>& OutCaptions) const override;
 #endif
diff --git a/Engine/Source/Runtime/Engine/Classes/Materials/MaterialExpressionLinearInterpolate.h b/Engine/Source/Runtime/Engine/Classes/Materials/MaterialExpressionLinearInterpolate.h
--- a/Engine/Source/Runtime/Engine/Classes/Materials/MaterialExpressionLinearInterpolate.h
+++ b/Engine/Source/Runtime/Engine/Classes/Materials/MaterialExpressionLinearInterpolate.h
@@ -38,6 +38,7 @@ class UMaterialExpressionLinearInterpolate : public UMaterialExpression
 
 	//~ Begin UMaterialExpression Interface
 #if WITH_EDITOR
+	virtual bool GenerateHLSLExpression(FMaterialHLSLGenerator& Generator, UE::HLSLTree::FScope& Scope, int32 OutputIndex, UE::HLSLTree::FExpression*& OutExpression) override;
 	virtual int32 Compile(class FMaterialCompiler* Compiler, int32 OutputIndex) override;
 	virtual void GetCaption(TArray<FString>& OutCaptions) const override;
 	virtual FText GetKeywords() const override {return FText::FromString(TEXT("lerp"));}
diff --git a/Engine/Source/Runtime/Engine/Classes/Materials/MaterialExpressionMax.h b/Engine/Source/Runtime/Engine/Classes/Materials/MaterialExpressionMax.h
--- a/Engine/Source/Runtime/Engine/Classes/Materials/MaterialExpressionMax.h
+++ b/Engine/Source/Runtime/Engine/Classes/Materials/MaterialExpressionMax.h
@@ -30,6 +30,7 @@ class UMaterialExpressionMax : public UMaterialExpression
 
 	//~ Begin UMaterialExpression Interface
 #if WITH_EDITOR
+	virtual bool GenerateHLSLExpression(FMaterialHLSLGenerator& Generator, UE::HLSLTree::FScope& Scope, int32 OutputIndex, UE::HLSLTree::FExpression*& OutExpression) override;
 	virtual int32 Compile(class FMaterialCompiler* Compiler, int32 OutputIndex) override;
 	virtual void GetCaption(TArray<FString>& OutCaptions) const override;
 #endif
diff --git a/Engine/Source/Runtime/Engine/Classes/Materials/MaterialExpressionMin.h b/Engine/Source/Runtime/Engine/Classes/Materials/MaterialExpressionMin.h
--- a/Engine/Source/Runtime/Engine/Classes/Materials/MaterialExpressionMin.h
+++ b/Engine/Source/Runtime/Engine/Classes/Materials/MaterialExpressionMin.h
@@ -30,6 +30,7 @@ class UMaterialExpressionMin : public UMaterialExpression
 
 	//~ Begin UMaterialExpression Interface
 #if WITH_EDITOR
+	virtual bool GenerateHLSLExpression(FMaterialHLSLGenerator& Generator, UE::HLSLTree::FScope& Scope, int32 OutputIndex, UE::HLSLTree::FExpression*& OutExpression) override;
 	virtual int32 Compile(class FMaterialCompiler* Compiler, int32 OutputIndex) override;
 	virtual void GetCaption(TArray<FString>& OutCaptions) const override;
 #endif
diff --git a/Engine/Source/Runtime/Engine/Classes/Materials/MaterialExpressionOneMinus.h b/Engine/Source/Runtime/Engine/Classes/Materials/MaterialExpressionOneMinus.h
--- a/Engine/Source/Runtime/Engine/Classes/Materials/MaterialExpressionOneMinus.h
+++ b/Engine/Source/Runtime/Engine/Classes/Materials/MaterialExpressionOneMinus.h
@@ -20,6 +20,7 @@ class UMaterialExpressionOneMinus : public UMaterialExpression
 
 	//~ Begin UMaterialExpression Interface
 #if WITH_EDITOR
+	virtual bool GenerateHLSLExpression(FMaterialHLSLGenerator& Generator, UE::HLSLTree::FScope& Scope, int32 OutputIndex, UE::HLSLTree::FExpression*& OutExpression) override;
 	virtual int32 Compile(class FMaterialCompiler* Compiler, int32 OutputIndex) override;
 	virtual void GetCaption(TArray<FString>& OutCaptions) const override;
 	virtual FText GetKeywords() const override {return FText::FromString(TEXT("1-x"));}
diff --git a/Engine/Source/Runtime/Engine/Classes/Materials/MaterialExpressionPixelDepth.h b/Engine/Source/Runtime/Engine/Classes/Materials/MaterialExpressionPixelDepth.h
--- a/Engine/Source/Runtime/Engine/Classes/Materials/MaterialExpressionPixelDepth.h
+++ b/Engine/Source/Runtime/Engine/Classes/Materials/MaterialExpressionPixelDepth.h
@@ -19,6 +19,7 @@ class UMaterialExpressionPixelDepth : public UMaterialExpression
 
 	//~ Begin UMaterialExpression Interface
 #if WITH_EDITOR
+		virtual bool GenerateHLSLExpression(FMaterialHLSLGenerator& Generator, UE::HLSLTree::FScope& Scope, int32 OutputIndex, UE::HLSLTree::FExpression*& OutExpression) override;
 	virtual int32 Compile(class FMaterialCompiler* Compiler, int32 OutputIndex) override;
 	virtual void GetCaption(TArray<FString>& OutCaptions) const override;
 #endif
diff --git a/Engine/Source/Runtime/Engine/Classes/Materials/MaterialExpressionSubtract.h b/Engine/Source/Runtime/Engine/Classes/Materials/MaterialExpressionSubtract.h
--- a/Engine/Source/Runtime/Engine/Classes/Materials/MaterialExpressionSubtract.h
+++ b/Engine/Source/Runtime/Engine/Classes/Materials/MaterialExpressionSubtract.h
@@ -31,6 +31,7 @@ class UMaterialExpressionSubtract : public UMaterialExpression
 
 	//~ Begin UMaterialExpression Interface
 #if WITH_EDITOR
+	virtual bool GenerateHLSLExpression(FMaterialHLSLGenerator& Generator, UE::HLSLTree::FScope& Scope, int32 OutputIndex, UE::HLSLTree::FExpression*& OutExpression) override;
 	virtual int32 Compile(class FMaterialCompiler* Compiler, int32 OutputIndex) override;
 	virtual void GetCaption(TArray<FString>& OutCaptions) const override;
 	virtual FText GetKeywords() const override {return FText::FromString(TEXT("-"));}
diff --git a/Engine/Source/Runtime/Engine/Private/HLSLTree/HLSLTree.cpp b/Engine/Source/Runtime/Engine/Private/HLSLTree/HLSLTree.cpp
--- a/Engine/Source/Runtime/Engine/Private/HLSLTree/HLSLTree.cpp
+++ b/Engine/Source/Runtime/Engine/Private/HLSLTree/HLSLTree.cpp
@@ -118,7 +118,11 @@ FBinaryOpDescription GetBinaryOpDesription(EBinaryOp Op)
 	case EBinaryOp::Mul: return FBinaryOpDescription(TEXT("Multiply"), TEXT("*"), Shader::EPreshaderOpcode::Mul); break;
 	case EBinaryOp::Div: return FBinaryOpDescription(TEXT("Divide"), TEXT("/"), Shader::EPreshaderOpcode::Div); break;
 	case EBinaryOp::Fmod: return FBinaryOpDescription(TEXT("Fmod"), TEXT("%"), Shader::EPreshaderOpcode::Fmod); break;
+	case EBinaryOp::Dot: return FBinaryOpDescription(TEXT("Dot"), TEXT("dot"), Shader::EPreshaderOpcode::Dot); break;
+	case EBinaryOp::Min: return FBinaryOpDescription(TEXT("Min"), TEXT("min"), Shader::EPreshaderOpcode::Min); break;
+	case EBinaryOp::Max: return FBinaryOpDescription(TEXT("Max"), TEXT("max"), Shader::EPreshaderOpcode::Max); break;
 	case EBinaryOp::Less: return FBinaryOpDescription(TEXT("Less"), TEXT("<"), Shader::EPreshaderOpcode::Less); break;
+	case EBinaryOp::Greater: return FBinaryOpDescription(TEXT("Greater"), TEXT(">"), Shader::EPreshaderOpcode::Greater); break;
 	default: checkNoEntry(); return FBinaryOpDescription();
 	}
 }
diff --git a/Engine/Source/Runtime/Engine/Private/HLSLTree/HLSLTreeCommon.cpp b/Engine/Source/Runtime/Engine/Private/HLSLTree/HLSLTreeCommon.cpp
--- a/Engine/Source/Runtime/Engine/Private/HLSLTree/HLSLTreeCommon.cpp
+++ b/Engine/Source/Runtime/Engine/Private/HLSLTree/HLSLTreeCommon.cpp
@@ -59,6 +59,13 @@ FExternalInputDescription GetExternalInputDescription(EExternalInput Input)
 	case EExternalInput::PrevTranslatedWorldPosition: return FExternalInputDescription(TEXT("PrevTranslatedWorldPosition"), Shader::EValueType::Float3, EExternalInput::WorldPosition_Ddx, EExternalInput::WorldPosition_Ddy);
 	case EExternalInput::PrevTranslatedWorldPosition_NoOffsets: return FExternalInputDescription(TEXT("PrevTranslatedWorldPosition_NoOffsets"), Shader::EValueType::Float3, EExternalInput::WorldPosition_Ddx, EExternalInput::WorldPosition_Ddy);
 
+	case EExternalInput::WorldPosition_Ddx: return FExternalInputDescription(TEXT("WorldPosition_Ddx"), Shader::EValueType::Float3);
+	case EExternalInput::WorldPosition_Ddy: return FExternalInputDescription(TEXT("WorldPosition_Ddx"), Shader::EValueType::Float3);
+
+	case EExternalInput::PixelDepth: return FExternalInputDescription(TEXT("PixelDepth"), Shader::EValueType::Float1, EExternalInput::PixelDepth_Ddx, EExternalInput::PixelDepth_Ddy);
+	case EExternalInput::PixelDepth_Ddx: return FExternalInputDescription(TEXT("PixelDepth_Ddx"), Shader::EValueType::Float1);
+	case EExternalInput::PixelDepth_Ddy: return FExternalInputDescription(TEXT("PixelDepth_Ddy"), Shader::EValueType::Float1);
+
 	case EExternalInput::GameTime: return FExternalInputDescription(TEXT("GameTime"), Shader::EValueType::Float1, EExternalInput::None, EExternalInput::None, EExternalInput::PrevGameTime);
 	case EExternalInput::RealTime: return FExternalInputDescription(TEXT("RealTime"), Shader::EValueType::Float1, EExternalInput::None, EExternalInput::None, EExternalInput::PrevRealTime);
 	case EExternalInput::DeltaTime: return FExternalInputDescription(TEXT("DeltaTime"), Shader::EValueType::Float1);
@@ -66,8 +73,7 @@ FExternalInputDescription GetExternalInputDescription(EExternalInput Input)
 	case EExternalInput::PrevGameTime: return FExternalInputDescription(TEXT("PrevGameTime"), Shader::EValueType::Float1);
 	case EExternalInput::PrevRealTime: return FExternalInputDescription(TEXT("PrevRealTime"), Shader::EValueType::Float1);
 
-	case EExternalInput::WorldPosition_Ddx: return FExternalInputDescription(TEXT("WorldPosition_Ddx"), Shader::EValueType::Float3);
-	case EExternalInput::WorldPosition_Ddy: return FExternalInputDescription(TEXT("WorldPosition_Ddx"), Shader::EValueType::Float3);
+	
 	default: checkNoEntry(); return FExternalInputDescription(TEXT("Invalid"), Shader::EValueType::Void);
 	}
 }
@@ -284,13 +290,16 @@ void FExpressionExternalInput::EmitValueShader(FEmitContext& Context, FEmitScope
 		case EExternalInput::PrevWorldPosition_NoOffsets: Code = TEXT("GetPrevWorldPosition_NoMaterialOffsets(Parameters)"); break;
 		case EExternalInput::PrevTranslatedWorldPosition: Code = TEXT("GetPrevTranslatedWorldPosition(Parameters)"); break;
 		case EExternalInput::PrevTranslatedWorldPosition_NoOffsets: Code = TEXT("GetPrevTranslatedWorldPosition_NoMaterialOffsets(Parameters)"); break;
+		case EExternalInput::WorldPosition_Ddx: Code = TEXT("Parameters.WorldPosition_DDX"); break;
+		case EExternalInput::WorldPosition_Ddy: Code = TEXT("Parameters.WorldPosition_DDY"); break;
+		case EExternalInput::PixelDepth: Code = TEXT("GetPixelDepth(Parameters)"); break;
+		case EExternalInput::PixelDepth_Ddx: Code = TEXT("Parameters.ScreenPosition_DDX.w"); break;
+		case EExternalInput::PixelDepth_Ddy: Code = TEXT("Parameters.ScreenPosition_DDY.w"); break;
 		case EExternalInput::GameTime: Code = TEXT("View.GameTime"); break;
 		case EExternalInput::RealTime: Code = TEXT("View.RealTime"); break;
 		case EExternalInput::DeltaTime: Code = TEXT("View.DeltaTime"); break;
 		case EExternalInput::PrevGameTime: Code = TEXT("View.PrevFrameGameTime"); break;
 		case EExternalInput::PrevRealTime: Code = TEXT("View.PrevFrameRealTime"); break;
-		case EExternalInput::WorldPosition_Ddx: Code = TEXT("Parameters.WorldPosition_DDX"); break;
-		case EExternalInput::WorldPosition_Ddy: Code = TEXT("Parameters.WorldPosition_DDY"); break;
 		default:
 			checkNoEntry();
 			break;
@@ -435,16 +444,31 @@ void FExpressionTextureSample::EmitValueShader(FEmitContext& Context, FEmitScope
 	}
 
 	FEmitShaderExpression* TexCoordValue = TexCoordExpression->GetValueShader(Context, Scope, Shader::EValueType::Float2);
-	FEmitShaderExpression* TexCoordValueDdx = TexCoordDerivatives.ExpressionDdx->GetValueShader(Context, Scope, Shader::EValueType::Float2);
-	FEmitShaderExpression* TexCoordValueDdy = TexCoordDerivatives.ExpressionDdy->GetValueShader(Context, Scope, Shader::EValueType::Float2);
-	OutResult.Code = Context.EmitExpression(Scope, Shader::EValueType::Float4, TEXT("%(%Grad(%, %, %, %, %))"),
-		SamplerTypeFunction,
-		SampleFunctionName,
-		*TextureName,
-		*SamplerStateCode,
-		TexCoordValue,
-		TexCoordValueDdx,
-		TexCoordValueDdy);
+	FEmitShaderExpression* TexCoordValueDdx = nullptr;
+	FEmitShaderExpression* TexCoordValueDdy = nullptr;
+	if (TexCoordDerivatives.IsValid())
+	{
+		TexCoordValueDdx = TexCoordDerivatives.ExpressionDdx->GetValueShader(Context, Scope, Shader::EValueType::Float2);
+		TexCoordValueDdy = TexCoordDerivatives.ExpressionDdy->GetValueShader(Context, Scope, Shader::EValueType::Float2);
+		OutResult.Code = Context.EmitExpression(Scope, Shader::EValueType::Float4, TEXT("%(%Grad(%, %, %, %, %))"),
+			SamplerTypeFunction,
+			SampleFunctionName,
+			*TextureName,
+			*SamplerStateCode,
+			TexCoordValue,
+			TexCoordValueDdx,
+			TexCoordValueDdy);
+	}
+	else
+	{
+		OutResult.Code = Context.EmitExpression(Scope, Shader::EValueType::Float4, TEXT("%(%(%, %, %))"),
+			SamplerTypeFunction,
+			SampleFunctionName,
+			*TextureName,
+			*SamplerStateCode,
+			TexCoordValue);
+	}
+	
 }
 
 void FExpressionGetStructField::ComputeAnalyticDerivatives(FTree& Tree, FExpressionDerivatives& OutResult) const
@@ -783,6 +807,7 @@ void FExpressionBinaryOp::ComputeAnalyticDerivatives(FTree& Tree, FExpressionDer
 	switch (Op)
 	{
 	case EBinaryOp::Less:
+	case EBinaryOp::Greater:
 		OutResult.ExpressionDdx = Tree.NewConstant(0.0f);
 		OutResult.ExpressionDdx = OutResult.ExpressionDdy;
 		break;
@@ -827,6 +852,31 @@ void FExpressionBinaryOp::ComputeAnalyticDerivatives(FTree& Tree, FExpressionDer
 			// We can't really do anything meaningful in the non-zero case.
 			OutResult = LhsDerivatives;
 			break;
+		case EBinaryOp::Dot:
+		{
+			// Dot means multiply the values, then sum the resulting components
+			FExpression* MulDdx = Tree.NewAdd(Tree.NewMul(LhsDerivatives.ExpressionDdx, Rhs), Tree.NewMul(RhsDerivatives.ExpressionDdx, Lhs));
+			FExpression* MulDdy = Tree.NewAdd(Tree.NewMul(LhsDerivatives.ExpressionDdy, Rhs), Tree.NewMul(RhsDerivatives.ExpressionDdy, Lhs));
+			// Dot the products with 1 to sum them
+			FExpression* Const1 = Tree.NewConstant(1.0f);
+			OutResult.ExpressionDdx = Tree.NewDot(MulDdx, Const1);
+			OutResult.ExpressionDdy = Tree.NewDot(MulDdy, Const1);
+			break;
+		}
+		case EBinaryOp::Min:
+		{
+			FExpression* Cond = Tree.NewLess(Lhs, Rhs);
+			OutResult.ExpressionDdx = Tree.NewExpression<FExpressionSelect>(Cond, LhsDerivatives.ExpressionDdx, RhsDerivatives.ExpressionDdx);
+			OutResult.ExpressionDdy = Tree.NewExpression<FExpressionSelect>(Cond, LhsDerivatives.ExpressionDdy, RhsDerivatives.ExpressionDdy);
+			break;
+		}
+		case EBinaryOp::Max:
+		{
+			FExpression* Cond = Tree.NewGreater(Lhs, Rhs);
+			OutResult.ExpressionDdx = Tree.NewExpression<FExpressionSelect>(Cond, LhsDerivatives.ExpressionDdx, RhsDerivatives.ExpressionDdx);
+			OutResult.ExpressionDdy = Tree.NewExpression<FExpressionSelect>(Cond, LhsDerivatives.ExpressionDdy, RhsDerivatives.ExpressionDdy);
+			break;
+		}
 		default:
 			checkNoEntry();
 			break;
@@ -859,7 +909,7 @@ bool FExpressionBinaryOp::PrepareValue(FEmitContext& Context, FEmitScope& Scope,
 	}
 
 	FPreparedType ResultType = MergePreparedTypes(LhsType, RhsType);
-	if (Op == EBinaryOp::Less)
+	if (Op == EBinaryOp::Less || Op == EBinaryOp::Greater)
 	{
 		ResultType.ValueComponentType = Shader::EValueComponentType::Bool;
 	}
@@ -896,9 +946,12 @@ FBinaryOpTypes GetBinaryOpTypes(const FRequestedType& RequestedType, EBinaryOp O
 	switch (Op)
 	{
 	case EBinaryOp::Less:
+	case EBinaryOp::Greater:
 		return FBinaryOpTypes(InputType, InputType, Shader::MakeValueType(Shader::EValueComponentType::Bool, NumComponents));
 	case EBinaryOp::Fmod:
 		return FBinaryOpTypes(InputType, NonLWCInputType, NonLWCInputType.GetType());
+	case EBinaryOp::Dot:
+		return FBinaryOpTypes(InputType, InputType, Shader::MakeValueType(InputComponentType, 1));
 	default:
 		return InputType;
 	}
@@ -944,6 +997,16 @@ void FExpressionBinaryOp::EmitValueShader(FEmitContext& Context, FEmitScope& Sco
 			OutResult.Code = Context.EmitExpression(Scope, Types.ResultType, TEXT("(% * %)"), LhsValue, RhsValue);
 		}
 		break;
+	case EBinaryOp::Div:
+		if (bLWC)
+		{
+			OutResult.Code = Context.EmitExpression(Scope, Types.ResultType, TEXT("LWCDivide(%, %)"), LhsValue, RhsValue);
+		}
+		else
+		{
+			OutResult.Code = Context.EmitExpression(Scope, Types.ResultType, TEXT("(% / %)"), LhsValue, RhsValue);
+		}
+		break;
 	case EBinaryOp::Fmod:
 		if (bLWC)
 		{
@@ -954,14 +1017,34 @@ void FExpressionBinaryOp::EmitValueShader(FEmitContext& Context, FEmitScope& Sco
 			OutResult.Code = Context.EmitExpression(Scope, Types.ResultType, TEXT("fmod(%, %)"), LhsValue, RhsValue);
 		}
 		break;
-	case EBinaryOp::Div:
+	case EBinaryOp::Dot:
 		if (bLWC)
 		{
-			OutResult.Code = Context.EmitExpression(Scope, Types.ResultType, TEXT("LWCDivide(%, %)"), LhsValue, RhsValue);
+			OutResult.Code = Context.EmitExpression(Scope, Types.ResultType, TEXT("LWCDot(%, %)"), LhsValue, RhsValue);
 		}
 		else
 		{
-			OutResult.Code = Context.EmitExpression(Scope, Types.ResultType, TEXT("(% / %)"), LhsValue, RhsValue);
+			OutResult.Code = Context.EmitExpression(Scope, Types.ResultType, TEXT("dot(%, %)"), LhsValue, RhsValue);
+		}
+		break;
+	case EBinaryOp::Min:
+		if (bLWC)
+		{
+			OutResult.Code = Context.EmitExpression(Scope, Types.ResultType, TEXT("LWCMin(%, %)"), LhsValue, RhsValue);
+		}
+		else
+		{
+			OutResult.Code = Context.EmitExpression(Scope, Types.ResultType, TEXT("min(%, %)"), LhsValue, RhsValue);
+		}
+		break;
+	case EBinaryOp::Max:
+		if (bLWC)
+		{
+			OutResult.Code = Context.EmitExpression(Scope, Types.ResultType, TEXT("LWCMax(%, %)"), LhsValue, RhsValue);
+		}
+		else
+		{
+			OutResult.Code = Context.EmitExpression(Scope, Types.ResultType, TEXT("max(%, %)"), LhsValue, RhsValue);
 		}
 		break;
 	case EBinaryOp::Less:
@@ -974,6 +1057,16 @@ void FExpressionBinaryOp::EmitValueShader(FEmitContext& Context, FEmitScope& Sco
 			OutResult.Code = Context.EmitExpression(Scope, Types.ResultType, TEXT("(% < %)"), LhsValue, RhsValue);
 		}
 		break;
+	case EBinaryOp::Greater:
+		if (bLWC)
+		{
+			OutResult.Code = Context.EmitExpression(Scope, Types.ResultType, TEXT("LWCGreater(%, %)"), LhsValue, RhsValue);
+		}
+		else
+		{
+			OutResult.Code = Context.EmitExpression(Scope, Types.ResultType, TEXT("(% > %)"), LhsValue, RhsValue);
+		}
+		break;
 	default:
 		checkNoEntry();
 		break;
diff --git a/Engine/Source/Runtime/Engine/Private/Materials/MaterialExpressionHLSL.cpp b/Engine/Source/Runtime/Engine/Private/Materials/MaterialExpressionHLSL.cpp
--- a/Engine/Source/Runtime/Engine/Private/Materials/MaterialExpressionHLSL.cpp
+++ b/Engine/Source/Runtime/Engine/Private/Materials/MaterialExpressionHLSL.cpp
@@ -20,6 +20,7 @@
 #include "Materials/MaterialExpressionDoubleVectorParameter.h"
 #include "Materials/MaterialExpressionScalarParameter.h"
 #include "Materials/MaterialExpressionStaticBoolParameter.h"
+#include "Materials/MaterialExpressionPixelDepth.h"
 #include "Materials/MaterialExpressionWorldPosition.h"
 #include "Materials/MaterialExpressionTime.h"
 #include "Materials/MaterialExpressionDeltaTime.h"
@@ -34,9 +35,16 @@
 #include "Materials/MaterialExpressionMaterialFunctionCall.h"
 #include "Materials/MaterialExpressionStaticSwitch.h"
 #include "Materials/MaterialExpressionGetLocal.h"
+#include "Materials/MaterialExpressionOneMinus.h"
 #include "Materials/MaterialExpressionAdd.h"
+#include "Materials/MaterialExpressionSubtract.h"
 #include "Materials/MaterialExpressionMultiply.h"
 #include "Materials/MaterialExpressionDivide.h"
+#include "Materials/MaterialExpressionDotProduct.h"
+#include "Materials/MaterialExpressionMin.h"
+#include "Materials/MaterialExpressionMax.h"
+#include "Materials/MaterialExpressionClamp.h"
+#include "Materials/MaterialExpressionLinearInterpolate.h"
 #include "Materials/MaterialExpressionBinaryOp.h"
 #include "Materials/MaterialExpressionAppendVector.h"
 #include "Materials/MaterialExpressionComponentMask.h"
@@ -144,20 +152,28 @@ bool UMaterialExpressionStaticBoolParameter::GenerateHLSLExpression(FMaterialHLS
 	return true;
 }
 
+bool UMaterialExpressionPixelDepth::GenerateHLSLExpression(FMaterialHLSLGenerator& Generator, UE::HLSLTree::FScope& Scope, int32 OutputIndex, UE::HLSLTree::FExpression*& OutExpression)
+{
+	using namespace UE::HLSLTree;
+	OutExpression = Generator.GetTree().NewExpression<FExpressionExternalInput>(EExternalInput::PixelDepth);
+	return true;
+}
+
 bool UMaterialExpressionWorldPosition::GenerateHLSLExpression(FMaterialHLSLGenerator& Generator, UE::HLSLTree::FScope& Scope, int32 OutputIndex, UE::HLSLTree::FExpression*& OutExpression)
 {
-	UE::HLSLTree::EExternalInput InputType = UE::HLSLTree::EExternalInput::None;
+	using namespace UE::HLSLTree;
 
+	EExternalInput InputType = EExternalInput::None;
 	switch (WorldPositionShaderOffset)
 	{
-	case WPT_Default: InputType = UE::HLSLTree::EExternalInput::WorldPosition; break;
-	case WPT_ExcludeAllShaderOffsets: InputType = UE::HLSLTree::EExternalInput::WorldPosition_NoOffsets; break;
-	case WPT_CameraRelative: InputType = UE::HLSLTree::EExternalInput::TranslatedWorldPosition; break;
-	case WPT_CameraRelativeNoOffsets: InputType = UE::HLSLTree::EExternalInput::TranslatedWorldPosition_NoOffsets; break;
+	case WPT_Default: InputType = EExternalInput::WorldPosition; break;
+	case WPT_ExcludeAllShaderOffsets: InputType = EExternalInput::WorldPosition_NoOffsets; break;
+	case WPT_CameraRelative: InputType = EExternalInput::TranslatedWorldPosition; break;
+	case WPT_CameraRelativeNoOffsets: InputType = EExternalInput::TranslatedWorldPosition_NoOffsets; break;
 	default: checkNoEntry(); break;
 	}
 
-	OutExpression = Generator.GetTree().NewExpression<UE::HLSLTree::FExpressionExternalInput>(InputType);
+	OutExpression = Generator.GetTree().NewExpression<FExpressionExternalInput>(InputType);
 	return true;
 }
 
@@ -188,36 +204,6 @@ bool UMaterialExpressionDeltaTime::GenerateHLSLExpression(FMaterialHLSLGenerator
 
 bool UMaterialExpressionPanner::GenerateHLSLExpression(FMaterialHLSLGenerator& Generator, UE::HLSLTree::FScope& Scope, int32 OutputIndex, UE::HLSLTree::FExpression*& OutExpression)
 {
-#if 0
-	int32 TimeArg = Time.GetTracedInput().Expression ? Time.Compile(Compiler) : Compiler->GameTime(false, 0.0f);
-	bool bIsSpeedExpressionValid = Speed.GetTracedInput().Expression != nullptr;
-	int32 SpeedVectorArg = bIsSpeedExpressionValid ? Speed.Compile(Compiler) : INDEX_NONE;
-	int32 SpeedXArg = bIsSpeedExpressionValid ? Compiler->ComponentMask(SpeedVectorArg, true, false, false, false) : Compiler->Constant(SpeedX);
-	int32 SpeedYArg = bIsSpeedExpressionValid ? Compiler->ComponentMask(SpeedVectorArg, false, true, false, false) : Compiler->Constant(SpeedY);
-	int32 Arg1;
-	int32 Arg2;
-	if (bFractionalPart)
-	{
-		// Note: this is to avoid (delay) divergent accuracy issues as GameTime increases.
-		// TODO: C++ to calculate its phase via per frame time delta.
-		Arg1 = Compiler->PeriodicHint(Compiler->Frac(Compiler->Mul(TimeArg, SpeedXArg)));
-		Arg2 = Compiler->PeriodicHint(Compiler->Frac(Compiler->Mul(TimeArg, SpeedYArg)));
-	}
-	else
-	{
-		Arg1 = Compiler->PeriodicHint(Compiler->Mul(TimeArg, SpeedXArg));
-		Arg2 = Compiler->PeriodicHint(Compiler->Mul(TimeArg, SpeedYArg));
-	}
-
-	int32 Arg3 = Coordinate.GetTracedInput().Expression ? Coordinate.Compile(Compiler) : Compiler->TextureCoordinate(ConstCoordinate, false, false);
-	return Compiler->Add(
-		Compiler->AppendVector(
-			Arg1,
-			Arg2
-		),
-		Arg3
-	);
-#endif
 	using namespace UE::HLSLTree;
 
 	FExpression* ExpressionTime = Time.TryAcquireHLSLExpression(Generator, Scope);
@@ -243,8 +229,24 @@ bool UMaterialExpressionPanner::GenerateHLSLExpression(FMaterialHLSLGenerator& G
 
 bool UMaterialExpressionTextureCoordinate::GenerateHLSLExpression(FMaterialHLSLGenerator& Generator, UE::HLSLTree::FScope& Scope, int32 OutputIndex, UE::HLSLTree::FExpression*& OutExpression)
 {
-	// TODO - tiling, mirroring
 	OutExpression = Generator.NewTexCoord(CoordinateIndex);
+
+	// TODO - unmirror
+
+	// Depending on whether we have U and V scale values that differ, we can perform a multiply by either
+	// a scalar or a float2.  These tiling values are baked right into the shader node, so they're always
+	// known at compile time.
+	// Avoid emitting the multiply by 1.0f if possible
+	// This should make generated HLSL a bit cleaner, but more importantly will help avoid generating redundant virtual texture stacks
+	if (FMath::Abs(UTiling - VTiling) > SMALL_NUMBER)
+	{
+		OutExpression = Generator.GetTree().NewMul(OutExpression, Generator.NewConstant(FVector2f(UTiling, VTiling)));
+	}
+	else if (FMath::Abs(1.0f - UTiling) > SMALL_NUMBER)
+	{
+		OutExpression = Generator.GetTree().NewMul(OutExpression, Generator.NewConstant(UTiling));
+	}
+
 	return true;
 }
 
@@ -271,7 +273,7 @@ bool UMaterialExpressionTextureSample::GenerateHLSLExpressionBase(FMaterialHLSLG
 	}
 
 	FExpression* TexCoordExpression = Coordinates.GetTracedInput().Expression ? Coordinates.TryAcquireHLSLExpression(Generator, Scope) : Generator.NewTexCoord(ConstCoordinate);
-	const FExpressionDerivatives TexCoordDerivatives = Generator.GetTree().GetAnalyticDerivatives(TexCoordExpression);
+	const FExpressionDerivatives TexCoordDerivatives;// = Generator.GetTree().GetAnalyticDerivatives(TexCoordExpression);
 	OutExpression = Generator.GetTree().NewExpression<FExpressionTextureSample>(TextureDeclaration, TexCoordExpression, TexCoordDerivatives, SamplerSource, MipValueMode);
 	return true;
 }
@@ -308,6 +310,17 @@ bool UMaterialExpressionTextureSampleParameter::GenerateHLSLExpression(FMaterial
 	return GenerateHLSLExpressionBase(Generator, Scope, TextureDeclaration, OutExpression);
 }
 
+bool UMaterialExpressionOneMinus::GenerateHLSLExpression(FMaterialHLSLGenerator& Generator, UE::HLSLTree::FScope& Scope, int32 OutputIndex, UE::HLSLTree::FExpression*& OutExpression)
+{
+	UE::HLSLTree::FExpression* InputExpression = Input.AcquireHLSLExpression(Generator, Scope);
+	if (!InputExpression)
+	{
+		return false;
+	}
+	OutExpression = Generator.GetTree().NewSub(Generator.NewConstant(1.0f), InputExpression);
+	return true;
+}
+
 bool UMaterialExpressionBinaryOp::GenerateHLSLExpression(FMaterialHLSLGenerator& Generator, UE::HLSLTree::FScope& Scope, int32 OutputIndex, UE::HLSLTree::FExpression*& OutExpression)
 {
 	UE::HLSLTree::FExpression* Lhs = A.AcquireHLSLExpressionOrConstant(Generator, Scope, ConstA);
@@ -333,6 +346,18 @@ bool UMaterialExpressionAdd::GenerateHLSLExpression(FMaterialHLSLGenerator& Gene
 	return true;
 }
 
+bool UMaterialExpressionSubtract::GenerateHLSLExpression(FMaterialHLSLGenerator& Generator, UE::HLSLTree::FScope& Scope, int32 OutputIndex, UE::HLSLTree::FExpression*& OutExpression)
+{
+	UE::HLSLTree::FExpression* Lhs = A.AcquireHLSLExpressionOrConstant(Generator, Scope, ConstA);
+	UE::HLSLTree::FExpression* Rhs = B.AcquireHLSLExpressionOrConstant(Generator, Scope, ConstB);
+	if (!Lhs || !Rhs)
+	{
+		return false;
+	}
+	OutExpression = Generator.GetTree().NewBinaryOp(UE::HLSLTree::EBinaryOp::Sub, Lhs, Rhs);
+	return true;
+}
+
 bool UMaterialExpressionMultiply::GenerateHLSLExpression(FMaterialHLSLGenerator& Generator, UE::HLSLTree::FScope& Scope, int32 OutputIndex, UE::HLSLTree::FExpression*& OutExpression)
 {
 	UE::HLSLTree::FExpression* Lhs = A.AcquireHLSLExpressionOrConstant(Generator, Scope, ConstA);
@@ -357,6 +382,80 @@ bool UMaterialExpressionDivide::GenerateHLSLExpression(FMaterialHLSLGenerator& G
 	return true;
 }
 
+bool UMaterialExpressionDotProduct::GenerateHLSLExpression(FMaterialHLSLGenerator& Generator, UE::HLSLTree::FScope& Scope, int32 OutputIndex, UE::HLSLTree::FExpression*& OutExpression)
+{
+	UE::HLSLTree::FExpression* Lhs = A.AcquireHLSLExpression(Generator, Scope);
+	UE::HLSLTree::FExpression* Rhs = B.AcquireHLSLExpression(Generator, Scope);
+	if (!Lhs || !Rhs)
+	{
+		return false;
+	}
+	OutExpression = Generator.GetTree().NewBinaryOp(UE::HLSLTree::EBinaryOp::Dot, Lhs, Rhs);
+	return true;
+}
+
+bool UMaterialExpressionMin::GenerateHLSLExpression(FMaterialHLSLGenerator& Generator, UE::HLSLTree::FScope& Scope, int32 OutputIndex, UE::HLSLTree::FExpression*& OutExpression)
+{
+	UE::HLSLTree::FExpression* Lhs = A.AcquireHLSLExpressionOrConstant(Generator, Scope, ConstA);
+	UE::HLSLTree::FExpression* Rhs = B.AcquireHLSLExpressionOrConstant(Generator, Scope, ConstB);
+	if (!Lhs || !Rhs)
+	{
+		return false;
+	}
+	OutExpression = Generator.GetTree().NewBinaryOp(UE::HLSLTree::EBinaryOp::Min, Lhs, Rhs);
+	return true;
+}
+
+bool UMaterialExpressionMax::GenerateHLSLExpression(FMaterialHLSLGenerator& Generator, UE::HLSLTree::FScope& Scope, int32 OutputIndex, UE::HLSLTree::FExpression*& OutExpression)
+{
+	UE::HLSLTree::FExpression* Lhs = A.AcquireHLSLExpressionOrConstant(Generator, Scope, ConstA);
+	UE::HLSLTree::FExpression* Rhs = B.AcquireHLSLExpressionOrConstant(Generator, Scope, ConstB);
+	if (!Lhs || !Rhs)
+	{
+		return false;
+	}
+	OutExpression = Generator.GetTree().NewBinaryOp(UE::HLSLTree::EBinaryOp::Max, Lhs, Rhs);
+	return true;
+}
+
+bool UMaterialExpressionClamp::GenerateHLSLExpression(FMaterialHLSLGenerator& Generator, UE::HLSLTree::FScope& Scope, int32 OutputIndex, UE::HLSLTree::FExpression*& OutExpression)
+{
+	UE::HLSLTree::FExpression* ExpressionMin = Min.AcquireHLSLExpressionOrConstant(Generator, Scope, MinDefault);
+	UE::HLSLTree::FExpression* ExpressionMax = Min.AcquireHLSLExpressionOrConstant(Generator, Scope, MaxDefault);
+	UE::HLSLTree::FExpression* ExpressionInput = Input.AcquireHLSLExpression(Generator, Scope);
+	if (!ExpressionMin || !ExpressionMax || !ExpressionInput)
+	{
+		return false;
+	}
+
+	if (ClampMode == CMODE_ClampMin || ClampMode == CMODE_Clamp)
+	{
+		ExpressionInput = Generator.GetTree().NewMax(ExpressionInput, ExpressionMin);
+	}
+	if (ClampMode == CMODE_ClampMax || ClampMode == CMODE_Clamp)
+	{
+		ExpressionInput = Generator.GetTree().NewMin(ExpressionInput, ExpressionMax);
+	}
+
+	OutExpression = ExpressionInput;
+	return true;
+}
+
+bool UMaterialExpressionLinearInterpolate::GenerateHLSLExpression(FMaterialHLSLGenerator& Generator, UE::HLSLTree::FScope& Scope, int32 OutputIndex, UE::HLSLTree::FExpression*& OutExpression)
+{
+	UE::HLSLTree::FExpression* ExpressionA = A.AcquireHLSLExpressionOrConstant(Generator, Scope, ConstA);
+	UE::HLSLTree::FExpression* ExpressionB = B.AcquireHLSLExpressionOrConstant(Generator, Scope, ConstB);
+	UE::HLSLTree::FExpression* ExpressionAlpha = Alpha.AcquireHLSLExpressionOrConstant(Generator, Scope, ConstAlpha);
+	if (!ExpressionA || !ExpressionB || !ExpressionAlpha)
+	{
+		return false;
+	}
+
+	// A + (B - A) * Alpha
+	OutExpression = Generator.GetTree().NewAdd(ExpressionA, Generator.GetTree().NewMul(Generator.GetTree().NewSub(ExpressionB, ExpressionA), ExpressionAlpha));
+	return true;
+}
+
 bool UMaterialExpressionAppendVector::GenerateHLSLExpression(FMaterialHLSLGenerator& Generator, UE::HLSLTree::FScope& Scope, int32 OutputIndex, UE::HLSLTree::FExpression*& OutExpression)
 {
 	UE::HLSLTree::FExpression* Lhs = A.AcquireHLSLExpression(Generator, Scope);
diff --git a/Engine/Source/Runtime/Engine/Private/Materials/MaterialHLSLGenerator.cpp b/Engine/Source/Runtime/Engine/Private/Materials/MaterialHLSLGenerator.cpp
--- a/Engine/Source/Runtime/Engine/Private/Materials/MaterialHLSLGenerator.cpp
+++ b/Engine/Source/Runtime/Engine/Private/Materials/MaterialHLSLGenerator.cpp
@@ -221,22 +221,7 @@ static UE::HLSLTree::FExpression* CompileMaterialInput(FMaterialHLSLGenerator& G
 			else
 			{
 				check(InputDescription.Input);
-				if (InputProperty >= MP_CustomizedUVs0 && InputProperty <= MP_CustomizedUVs7)
-				{
-					const int32 TexCoordIndex = (int32)InputProperty - MP_CustomizedUVs0;
-					if (TexCoordIndex < Material->NumCustomizedUVs)
-					{
-						Expression = InputDescription.Input->TryAcquireHLSLExpression(Generator, Scope);
-					}
-					if (!Expression)
-					{
-						Expression = Generator.NewTexCoord(TexCoordIndex);
-					}
-				}
-				else
-				{
-					Expression = InputDescription.Input->TryAcquireHLSLExpression(Generator, Scope);
-				}
+				Expression = InputDescription.Input->TryAcquireHLSLExpression(Generator, Scope);
 			}
 		}
 	}
diff --git a/Engine/Source/Runtime/Engine/Private/Shader/Preshader.cpp b/Engine/Source/Runtime/Engine/Private/Shader/Preshader.cpp
--- a/Engine/Source/Runtime/Engine/Private/Shader/Preshader.cpp
+++ b/Engine/Source/Runtime/Engine/Private/Shader/Preshader.cpp
@@ -600,6 +600,9 @@ FPreshaderValue EvaluatePreshader(const FUniformExpressionSet* UniformExpression
 		case EPreshaderOpcode::Mul: EvaluateBinaryOp(Stack, Mul); break;
 		case EPreshaderOpcode::Div: EvaluateBinaryOp(Stack, Div); break;
 		case EPreshaderOpcode::Less: EvaluateBinaryOp(Stack, Less); break;
+		case EPreshaderOpcode::Greater: EvaluateBinaryOp(Stack, Greater); break;
+		case EPreshaderOpcode::LessEqual: EvaluateBinaryOp(Stack, LessEqual); break;
+		case EPreshaderOpcode::GreaterEqual: EvaluateBinaryOp(Stack, GreaterEqual); break;
 		case EPreshaderOpcode::Fmod: EvaluateBinaryOp(Stack, Fmod); break;
 		case EPreshaderOpcode::Min: EvaluateBinaryOp(Stack, Min); break;
 		case EPreshaderOpcode::Max: EvaluateBinaryOp(Stack, Max); break;
diff --git a/Engine/Source/Runtime/Engine/Private/Shader/ShaderValue.cpp b/Engine/Source/Runtime/Engine/Private/Shader/ShaderValue.cpp
--- a/Engine/Source/Runtime/Engine/Private/Shader/ShaderValue.cpp
+++ b/Engine/Source/Runtime/Engine/Private/Shader/ShaderValue.cpp
@@ -946,6 +946,9 @@ struct FOpMax : public FOpBase { template<typename T> T operator()(T Lhs, T Rhs)
 struct FOpFmod : public FOpBaseNoInt { template<typename T> T operator()(T Lhs, T Rhs) const { return FMath::Fmod(Lhs, Rhs); } };
 struct FOpAtan2 : public FOpBaseNoInt { template<typename T> T operator()(T Lhs, T Rhs) const { return FMath::Atan2(Lhs, Rhs); } };
 struct FOpLess : public FOpBase { template<typename T> bool operator()(T Lhs, T Rhs) const { return Lhs < Rhs; } };
+struct FOpGreater : public FOpBase { template<typename T> bool operator()(T Lhs, T Rhs) const { return Lhs > Rhs; } };
+struct FOpLessEqual : public FOpBase { template<typename T> bool operator()(T Lhs, T Rhs) const { return Lhs <= Rhs; } };
+struct FOpGreaterEqual : public FOpBase { template<typename T> bool operator()(T Lhs, T Rhs) const { return Lhs >= Rhs; } };
 
 template<typename Operation>
 inline FValue UnaryOp(const Operation& Op, const FValue& Value)
@@ -1279,6 +1282,21 @@ FValue Less(const FValue& Lhs, const FValue& Rhs)
 	return Private::CompareOp(Private::FOpLess(), Lhs, Rhs);
 }
 
+FValue Greater(const FValue& Lhs, const FValue& Rhs)
+{
+	return Private::CompareOp(Private::FOpGreater(), Lhs, Rhs);
+}
+
+FValue LessEqual(const FValue& Lhs, const FValue& Rhs)
+{
+	return Private::CompareOp(Private::FOpLessEqual(), Lhs, Rhs);
+}
+
+FValue GreaterEqual(const FValue& Lhs, const FValue& Rhs)
+{
+	return Private::CompareOp(Private::FOpGreaterEqual(), Lhs, Rhs);
+}
+
 FValue Min(const FValue& Lhs, const FValue& Rhs)
 {
 	return Private::BinaryOp(Private::FOpMin(), Lhs, Rhs);
diff --git a/Engine/Source/Runtime/Engine/Public/HLSLTree/HLSLTree.h b/Engine/Source/Runtime/Engine/Public/HLSLTree/HLSLTree.h
--- a/Engine/Source/Runtime/Engine/Public/HLSLTree/HLSLTree.h
+++ b/Engine/Source/Runtime/Engine/Public/HLSLTree/HLSLTree.h
@@ -639,7 +639,11 @@ public:
 	FExpression* NewMul(FExpression* Lhs, FExpression* Rhs) { return NewBinaryOp(EBinaryOp::Mul, Lhs, Rhs); }
 	FExpression* NewDiv(FExpression* Lhs, FExpression* Rhs) { return NewBinaryOp(EBinaryOp::Div, Lhs, Rhs); }
 	FExpression* NewFmod(FExpression* Lhs, FExpression* Rhs) { return NewBinaryOp(EBinaryOp::Fmod, Lhs, Rhs); }
+	FExpression* NewDot(FExpression* Lhs, FExpression* Rhs) { return NewBinaryOp(EBinaryOp::Dot, Lhs, Rhs); }
+	FExpression* NewMin(FExpression* Lhs, FExpression* Rhs) { return NewBinaryOp(EBinaryOp::Min, Lhs, Rhs); }
+	FExpression* NewMax(FExpression* Lhs, FExpression* Rhs) { return NewBinaryOp(EBinaryOp::Max, Lhs, Rhs); }
 	FExpression* NewLess(FExpression* Lhs, FExpression* Rhs) { return NewBinaryOp(EBinaryOp::Less, Lhs, Rhs); }
+	FExpression* NewGreater(FExpression* Lhs, FExpression* Rhs) { return NewBinaryOp(EBinaryOp::Greater, Lhs, Rhs); }
 
 	FTextureParameterDeclaration* NewTextureParameterDeclaration(const FName& Name, const FTextureDescription& DefaultValue);
 
diff --git a/Engine/Source/Runtime/Engine/Public/HLSLTree/HLSLTreeCommon.h b/Engine/Source/Runtime/Engine/Public/HLSLTree/HLSLTreeCommon.h
--- a/Engine/Source/Runtime/Engine/Public/HLSLTree/HLSLTreeCommon.h
+++ b/Engine/Source/Runtime/Engine/Public/HLSLTree/HLSLTreeCommon.h
@@ -83,15 +83,19 @@ enum class EExternalInput : uint8
 	PrevTranslatedWorldPosition,
 	PrevTranslatedWorldPosition_NoOffsets,
 
+	WorldPosition_Ddx,
+	WorldPosition_Ddy,
+
+	PixelDepth,
+	PixelDepth_Ddx,
+	PixelDepth_Ddy,
+
 	GameTime,
 	RealTime,
 	DeltaTime,
 
 	PrevGameTime,
 	PrevRealTime,
-
-	WorldPosition_Ddx,
-	WorldPosition_Ddy,
 };
 static constexpr int32 NumTexCoords = 8;
 
@@ -320,6 +324,15 @@ public:
 	virtual void EmitValueShader(FEmitContext& Context, FEmitScope& Scope, const FRequestedType& RequestedType, FEmitValueShaderResult& OutResult) const override;
 };
 
+class FExpressionCustomHLSL : public FExpression
+{
+public:
+	virtual bool PrepareValue(FEmitContext& Context, FEmitScope& Scope, const FRequestedType& RequestedType, FPrepareValueResult& OutResult) const override;
+	virtual void EmitValueShader(FEmitContext& Context, FEmitScope& Scope, const FRequestedType& RequestedType, FEmitValueShaderResult& OutResult) const override;
+
+
+};
+
 class FStatementReturn : public FStatement
 {
 public:
diff --git a/Engine/Source/Runtime/Engine/Public/HLSLTree/HLSLTreeTypes.h b/Engine/Source/Runtime/Engine/Public/HLSLTree/HLSLTreeTypes.h
--- a/Engine/Source/Runtime/Engine/Public/HLSLTree/HLSLTreeTypes.h
+++ b/Engine/Source/Runtime/Engine/Public/HLSLTree/HLSLTreeTypes.h
@@ -87,7 +87,11 @@ enum class EBinaryOp : uint8
 	Mul,
 	Div,
 	Fmod,
+	Dot,
+	Min,
+	Max,
 	Less,
+	Greater,
 };
 
 struct FBinaryOpDescription
diff --git a/Engine/Source/Runtime/Engine/Public/Shader/Preshader.h b/Engine/Source/Runtime/Engine/Public/Shader/Preshader.h
--- a/Engine/Source/Runtime/Engine/Public/Shader/Preshader.h
+++ b/Engine/Source/Runtime/Engine/Public/Shader/Preshader.h
@@ -69,6 +69,9 @@ enum class EPreshaderOpcode : uint8
 	PushValue,
 	Less,
 	Assign,
+	Greater,
+	LessEqual,
+	GreaterEqual,
 };
 
 struct FPreshaderStructType
diff --git a/Engine/Source/Runtime/Engine/Public/Shader/ShaderTypes.h b/Engine/Source/Runtime/Engine/Public/Shader/ShaderTypes.h
--- a/Engine/Source/Runtime/Engine/Public/Shader/ShaderTypes.h
+++ b/Engine/Source/Runtime/Engine/Public/Shader/ShaderTypes.h
@@ -428,6 +428,9 @@ ENGINE_API FValue Sub(const FValue& Lhs, const FValue& Rhs);
 ENGINE_API FValue Mul(const FValue& Lhs, const FValue& Rhs);
 ENGINE_API FValue Div(const FValue& Lhs, const FValue& Rhs);
 ENGINE_API FValue Less(const FValue& Lhs, const FValue& Rhs);
+ENGINE_API FValue Greater(const FValue& Lhs, const FValue& Rhs);
+ENGINE_API FValue LessEqual(const FValue& Lhs, const FValue& Rhs);
+ENGINE_API FValue GreaterEqual(const FValue& Lhs, const FValue& Rhs);
 ENGINE_API FValue Min(const FValue& Lhs, const FValue& Rhs);
 ENGINE_API FValue Max(const FValue& Lhs, const FValue& Rhs);
 ENGINE_API FValue Clamp(const FValue& Value, const FValue& Low, const FValue& High);
