[Backout] - CL18864159

- backing out this change as it broke Xbox in UE5/Main
 - Example failed Horde tests:
  - https://horde.devtools.epicgames.com/job/61fd7dc8a6165e535689f6e4?step=0856
   - https://horde.devtools.epicgames.com/job/61fd7dc8a6165e535689f6e4?step=0250

Original CL Desc
-----------------------------------------------------------------
Fixed a GameThread stall (waiting for render thread) whenever a CVar is changed during gameplay.

It seems like the intent of the callback sink for r.FreeSkeletalMeshBuffers was to free eligible CPU buffers when the CVar was raised to 1, but unfortunately we're currently relying on this sink to free certain helper structures.

This change allows those structures to be freed after submitting the buffers to RHI so that we can change the callback sink to only do work on the rising edge (which generally only happens when switching game-mode/device profile and not during gameplay).

diff --git a/Engine/Source/Runtime/Engine/Private/PositionVertexBuffer.cpp b/Engine/Source/Runtime/Engine/Private/PositionVertexBuffer.cpp
--- a/Engine/Source/Runtime/Engine/Private/PositionVertexBuffer.cpp
+++ b/Engine/Source/Runtime/Engine/Private/PositionVertexBuffer.cpp
@@ -196,7 +196,22 @@ void FPositionVertexBuffer::operator=(const FPositionVertexBuffer &Other)
 template <bool bRenderThread>
 FBufferRHIRef FPositionVertexBuffer::CreateRHIBuffer_Internal()
 {
-	return CreateRHIBuffer<bRenderThread>(VertexData, NumVertices, BUF_Static | BUF_ShaderResource, TEXT("FPositionVertexBuffer"));
+	if (NumVertices)
+	{
+		FResourceArrayInterface* RESTRICT ResourceArray = VertexData ? VertexData->GetResourceArray() : nullptr;
+		const uint32 SizeInBytes = ResourceArray ? ResourceArray->GetResourceDataSize() : 0;
+		FRHIResourceCreateInfo CreateInfo(TEXT("FPositionVertexBuffer"), ResourceArray);
+		CreateInfo.bWithoutNativeResource = !VertexData;
+		if (bRenderThread)
+		{
+			return RHICreateVertexBuffer(SizeInBytes, BUF_Static | BUF_ShaderResource, CreateInfo);
+		}
+		else
+		{
+			return RHIAsyncCreateVertexBuffer(SizeInBytes, BUF_Static | BUF_ShaderResource, CreateInfo);
+		}
+	}
+	return nullptr;
 }
 
 FBufferRHIRef FPositionVertexBuffer::CreateRHIBuffer_RenderThread()
diff --git a/Engine/Source/Runtime/Engine/Private/Rendering/ColorVertexBuffer.cpp b/Engine/Source/Runtime/Engine/Private/Rendering/ColorVertexBuffer.cpp
--- a/Engine/Source/Runtime/Engine/Private/Rendering/ColorVertexBuffer.cpp
+++ b/Engine/Source/Runtime/Engine/Private/Rendering/ColorVertexBuffer.cpp
@@ -394,7 +394,22 @@ uint32 FColorVertexBuffer::GetAllocatedSize() const
 template <bool bRenderThread>
 FBufferRHIRef FColorVertexBuffer::CreateRHIBuffer_Internal()
 {
-	return CreateRHIBuffer<bRenderThread>(VertexData, NumVertices, BUF_Static | BUF_ShaderResource, TEXT("FColorVertexBuffer"));
+	if (NumVertices)
+	{
+		FResourceArrayInterface* RESTRICT ResourceArray = VertexData ? VertexData->GetResourceArray() : nullptr;
+		const uint32 SizeInBytes = ResourceArray ? ResourceArray->GetResourceDataSize() : 0;
+		FRHIResourceCreateInfo CreateInfo(TEXT("FColorVertexBuffer"), ResourceArray);
+		CreateInfo.bWithoutNativeResource = !VertexData;
+		if (bRenderThread)
+		{
+			return RHICreateVertexBuffer(SizeInBytes, BUF_Static | BUF_ShaderResource, CreateInfo);
+		}
+		else
+		{
+			return RHIAsyncCreateVertexBuffer(SizeInBytes, BUF_Static | BUF_ShaderResource, CreateInfo);
+		}
+	}
+	return nullptr;
 }
 
 FBufferRHIRef FColorVertexBuffer::CreateRHIBuffer_RenderThread()
diff --git a/Engine/Source/Runtime/Engine/Private/Rendering/SkeletalMeshVertexClothBuffer.cpp b/Engine/Source/Runtime/Engine/Private/Rendering/SkeletalMeshVertexClothBuffer.cpp
--- a/Engine/Source/Runtime/Engine/Private/Rendering/SkeletalMeshVertexClothBuffer.cpp
+++ b/Engine/Source/Runtime/Engine/Private/Rendering/SkeletalMeshVertexClothBuffer.cpp
@@ -74,7 +74,24 @@ void FSkeletalMeshVertexClothBuffer::ClearMetaData()
 template <bool bRenderThread>
 FBufferRHIRef FSkeletalMeshVertexClothBuffer::CreateRHIBuffer_Internal()
 {
-	return CreateRHIBuffer<bRenderThread>(VertexData, NumVertices, BUF_Static | BUF_ShaderResource, TEXT("FSkeletalMeshVertexClothBuffer"));
+	if (NumVertices)
+	{
+		FResourceArrayInterface* ResourceArray = VertexData ? VertexData->GetResourceArray() : nullptr;
+		const uint32 SizeInBytes = ResourceArray ? ResourceArray->GetResourceDataSize() : 0;
+		const EBufferUsageFlags BufferFlags = BUF_Static | BUF_ShaderResource;
+		FRHIResourceCreateInfo CreateInfo(TEXT("FSkeletalMeshVertexClothBuffer"), ResourceArray);
+		CreateInfo.bWithoutNativeResource = !VertexData;
+
+		if (bRenderThread)
+		{
+			return RHICreateVertexBuffer(SizeInBytes, BufferFlags, CreateInfo);
+		}
+		else
+		{
+			return RHIAsyncCreateVertexBuffer(SizeInBytes, BufferFlags, CreateInfo);
+		}
+	}
+	return nullptr;
 }
 
 FBufferRHIRef FSkeletalMeshVertexClothBuffer::CreateRHIBuffer_RenderThread()
diff --git a/Engine/Source/Runtime/Engine/Private/Rendering/SkinWeightVertexBuffer.cpp b/Engine/Source/Runtime/Engine/Private/Rendering/SkinWeightVertexBuffer.cpp
--- a/Engine/Source/Runtime/Engine/Private/Rendering/SkinWeightVertexBuffer.cpp
+++ b/Engine/Source/Runtime/Engine/Private/Rendering/SkinWeightVertexBuffer.cpp
@@ -179,7 +179,25 @@ void FSkinWeightLookupVertexBuffer::SetWeightOffsetAndInfluenceCount(uint32 Vert
 template <bool bRenderThread>
 FBufferRHIRef FSkinWeightLookupVertexBuffer::CreateRHIBuffer_Internal()
 {
-	return CreateRHIBuffer<bRenderThread>(LookupData, NumVertices, BUF_Static | BUF_ShaderResource | BUF_SourceCopy, TEXT("FSkinWeightLookupVertexBuffer"));
+	if (NumVertices)
+	{
+		// Create the vertex buffer.
+		FResourceArrayInterface* ResourceArray = LookupData ? LookupData->GetResourceArray() : nullptr;
+		const uint32 SizeInBytes = ResourceArray ? ResourceArray->GetResourceDataSize() : 0;
+		const EBufferUsageFlags BufferFlags = BUF_Static | BUF_ShaderResource | BUF_SourceCopy;
+		FRHIResourceCreateInfo CreateInfo(TEXT("FSkinWeightLookupVertexBuffer"), ResourceArray);
+		CreateInfo.bWithoutNativeResource = !LookupData;
+
+		if (bRenderThread)
+		{
+			return RHICreateVertexBuffer(SizeInBytes, BufferFlags, CreateInfo);
+		}
+		else
+		{
+			return RHIAsyncCreateVertexBuffer(SizeInBytes, BufferFlags, CreateInfo);
+		}
+	}
+	return nullptr;
 }
 
 /*-----------------------------------------------------------------------------
@@ -371,8 +389,26 @@ void FSkinWeightDataVertexBuffer::CopyMetaData(const FSkinWeightDataVertexBuffer
 template <bool bRenderThread>
 FBufferRHIRef FSkinWeightDataVertexBuffer::CreateRHIBuffer_Internal()
 {
-	// BUF_ShaderResource is needed for support of the SkinCache (we could make is dependent on GEnableGPUSkinCacheShaders or are there other users?)
-	return CreateRHIBuffer<bRenderThread>(WeightData, NumBones, BUF_Static | BUF_ShaderResource | BUF_SourceCopy, TEXT("FSkinWeightDataVertexBuffer"));
+	if (NumBones)
+	{
+		// Create the vertex buffer.
+		FResourceArrayInterface* ResourceArray = WeightData ? WeightData->GetResourceArray() : nullptr;
+		const uint32 SizeInBytes = ResourceArray ? ResourceArray->GetResourceDataSize() : 0;
+		const EBufferUsageFlags BufferFlags = BUF_Static | BUF_ShaderResource | BUF_SourceCopy;
+		FRHIResourceCreateInfo CreateInfo(TEXT("FSkinWeightDataVertexBuffer"), ResourceArray);
+		CreateInfo.bWithoutNativeResource = !WeightData;
+
+		// BUF_ShaderResource is needed for support of the SkinCache (we could make is dependent on GEnableGPUSkinCacheShaders or are there other users?)
+		if (bRenderThread)
+		{
+			return RHICreateVertexBuffer(SizeInBytes, BufferFlags, CreateInfo);
+		}
+		else
+		{
+			return RHIAsyncCreateVertexBuffer(SizeInBytes, BufferFlags, CreateInfo);
+		}
+	}
+	return nullptr;
 }
 
 FBufferRHIRef FSkinWeightDataVertexBuffer::CreateRHIBuffer_RenderThread()
diff --git a/Engine/Source/Runtime/Engine/Private/Rendering/StaticMeshVertexBuffer.cpp b/Engine/Source/Runtime/Engine/Private/Rendering/StaticMeshVertexBuffer.cpp
--- a/Engine/Source/Runtime/Engine/Private/Rendering/StaticMeshVertexBuffer.cpp
+++ b/Engine/Source/Runtime/Engine/Private/Rendering/StaticMeshVertexBuffer.cpp
@@ -269,7 +269,22 @@ void FStaticMeshVertexBuffer::operator=(const FStaticMeshVertexBuffer &Other)
 template <bool bRenderThread>
 FBufferRHIRef FStaticMeshVertexBuffer::CreateTangentsRHIBuffer_Internal()
 {
-	return CreateRHIBuffer<bRenderThread>(TangentsData, GetNumVertices(), BUF_Static | BUF_ShaderResource, TEXT("TangentsRHIBuffer"));
+	if (GetNumVertices())
+	{
+ 		FResourceArrayInterface* RESTRICT ResourceArray = TangentsData ? TangentsData->GetResourceArray() : nullptr;
+		const uint32 SizeInBytes = ResourceArray ? ResourceArray->GetResourceDataSize() : 0;
+		FRHIResourceCreateInfo CreateInfo(TEXT("TangentsRHIBuffer"), ResourceArray);
+		CreateInfo.bWithoutNativeResource = !TangentsData;
+		if (bRenderThread)
+		{
+			return RHICreateVertexBuffer(SizeInBytes, BUF_Static | BUF_ShaderResource, CreateInfo);
+		}
+		else
+		{
+			return RHIAsyncCreateVertexBuffer(SizeInBytes, BUF_Static | BUF_ShaderResource, CreateInfo);
+		}
+	}
+	return nullptr;
 }
 
 FBufferRHIRef FStaticMeshVertexBuffer::CreateTangentsRHIBuffer_RenderThread()
@@ -285,7 +300,22 @@ FBufferRHIRef FStaticMeshVertexBuffer::CreateTangentsRHIBuffer_Async()
 template <bool bRenderThread>
 FBufferRHIRef FStaticMeshVertexBuffer::CreateTexCoordRHIBuffer_Internal()
 {
-	return CreateRHIBuffer<bRenderThread>(TexcoordData, GetNumTexCoords(), BUF_Static | BUF_ShaderResource, TEXT("TexCoordRHIBuffer"));
+	if (GetNumTexCoords())
+	{
+		FResourceArrayInterface* RESTRICT ResourceArray = TexcoordData ? TexcoordData->GetResourceArray() : nullptr;
+		const uint32 SizeInBytes = ResourceArray ? ResourceArray->GetResourceDataSize() : 0;
+		FRHIResourceCreateInfo CreateInfo(TEXT("TexCoordRHIBuffer"), ResourceArray);
+		CreateInfo.bWithoutNativeResource = !TexcoordData;
+		if (bRenderThread)
+		{
+			return RHICreateVertexBuffer(SizeInBytes, BUF_Static | BUF_ShaderResource, CreateInfo);
+		}
+		else
+		{
+			return RHIAsyncCreateVertexBuffer(SizeInBytes, BUF_Static | BUF_ShaderResource, CreateInfo);
+		}
+	}
+	return nullptr;
 }
 
 FBufferRHIRef FStaticMeshVertexBuffer::CreateTexCoordRHIBuffer_RenderThread()
diff --git a/Engine/Source/Runtime/Engine/Private/SkeletalMesh.cpp b/Engine/Source/Runtime/Engine/Private/SkeletalMesh.cpp
--- a/Engine/Source/Runtime/Engine/Private/SkeletalMesh.cpp
+++ b/Engine/Source/Runtime/Engine/Private/SkeletalMesh.cpp
@@ -282,28 +282,29 @@ FArchive& operator<<(FArchive& Ar, FClothBufferIndexMapping& ClothBufferIndexMap
 		<< ClothBufferIndexMapping.LODBiasStride;
 }
 
+/*-----------------------------------------------------------------------------
+FreeSkeletalMeshBuffersSinkCallback
+-----------------------------------------------------------------------------*/
+
 void FreeSkeletalMeshBuffersSinkCallback()
 {
-	static const auto CVar = IConsoleManager::Get().FindTConsoleVariableDataInt(TEXT("r.FreeSkeletalMeshBuffers"));
-	static bool bPreviousFreeSkeletalMeshBuffers = false;
-	const bool bFreeSkeletalMeshBuffers = CVar->GetValueOnGameThread() == 1;
-
-	// We previously relied on this sink occurring periodically to free helper structures, so run it if we're not using the new logic that frees on submit.
-	if (!GFreeStructuresOnRHIBufferCreation || (bPreviousFreeSkeletalMeshBuffers != bFreeSkeletalMeshBuffers))
+	// If r.FreeSkeletalMeshBuffers==1 then CPU buffer copies are to be released.
+	static TOptional<bool> bLastFreeSkeletalMeshBuffers;
+	static const auto CVar = IConsoleManager::Get().FindTConsoleVariableDataInt(TEXT("r.FreeSkeletalMeshBuffers"));	
+	bool bFreeSkeletalMeshBuffers = CVar->GetValueOnGameThread() == 1;
+	if(!bLastFreeSkeletalMeshBuffers.IsSet() || (bFreeSkeletalMeshBuffers != bLastFreeSkeletalMeshBuffers.GetValue()))
 	{
-		bPreviousFreeSkeletalMeshBuffers = bFreeSkeletalMeshBuffers;
-		if (bFreeSkeletalMeshBuffers)
+		TRACE_CPUPROFILER_EVENT_SCOPE(FreeSkeletalMeshBuffersSinkCallback);
+
+		FlushRenderingCommands();
+		for (TObjectIterator<USkeletalMesh> It;It;++It)
 		{
-			FlushRenderingCommands();
-			for (TObjectIterator<USkeletalMesh> It; It; ++It)
+			if (!It->HasPendingInitOrStreaming() && !It->GetResourceForRendering()->RequiresCPUSkinning(GMaxRHIFeatureLevel))
 			{
-				USkeletalMesh* SkeletalMesh = *It;
-				if (!SkeletalMesh->HasPendingInitOrStreaming() && !SkeletalMesh->GetResourceForRendering()->RequiresCPUSkinning(GMaxRHIFeatureLevel))
-				{
-					SkeletalMesh->ReleaseCPUResources();
-				}
+				It->ReleaseCPUResources();
 			}
 		}
+		bLastFreeSkeletalMeshBuffers = bFreeSkeletalMeshBuffers;
 	}
 }
 
diff --git a/Engine/Source/Runtime/RenderCore/Private/RenderResource.cpp b/Engine/Source/Runtime/RenderCore/Private/RenderResource.cpp
--- a/Engine/Source/Runtime/RenderCore/Private/RenderResource.cpp
+++ b/Engine/Source/Runtime/RenderCore/Private/RenderResource.cpp
@@ -26,15 +26,9 @@ FAutoConsoleVariableRef CVarMaxVertexBytesAllocatedPerFrame(
 int32 GGlobalBufferNumFramesUnusedThresold = 30;
 FAutoConsoleVariableRef CVarReadBufferNumFramesUnusedThresold(
 	TEXT("r.NumFramesUnusedBeforeReleasingGlobalResourceBuffers"),
-	GGlobalBufferNumFramesUnusedThresold,
+	GGlobalBufferNumFramesUnusedThresold ,
 	TEXT("Number of frames after which unused global resource allocations will be discarded. Set 0 to ignore. (default=30)"));
 
-bool GFreeStructuresOnRHIBufferCreation = true;
-FAutoConsoleVariableRef CVarFreeStructuresOnRHIBufferCreation(
-	TEXT("r.FreeStructuresOnRHIBufferCreation"),
-	GFreeStructuresOnRHIBufferCreation,
-	TEXT("Toggles experimental method for freeing helper structures that own the resource arrays after submitting to RHI instead of in the callback sink."));
-
 int32 GVarDebugForceRuntimeBLAS = 0;
 FAutoConsoleVariableRef CVarDebugForceRuntimeBLAS(
 	TEXT("r.Raytracing.DebugForceRuntimeBLAS"),
diff --git a/Engine/Source/Runtime/RenderCore/Public/RenderResource.h b/Engine/Source/Runtime/RenderCore/Public/RenderResource.h
--- a/Engine/Source/Runtime/RenderCore/Public/RenderResource.h
+++ b/Engine/Source/Runtime/RenderCore/Public/RenderResource.h
@@ -16,9 +16,6 @@
 /** Number of frames after which unused global resource allocations will be discarded. */
 extern int32 GGlobalBufferNumFramesUnusedThresold;
 
-/** Experimental: whether we free helper structures after submitting to RHI. */
-extern RENDERCORE_API bool GFreeStructuresOnRHIBufferCreation;
-
 /**
  * A rendering resource which is owned by the rendering thread.
  * NOTE - Adding new virtual methods to this class may require stubs added to FViewport/FDummyViewport, otherwise certain modules may have link errors
@@ -161,39 +158,6 @@ protected:
 	const FStaticFeatureLevel GetFeatureLevel() const { return FeatureLevel == ERHIFeatureLevel::Num ? FStaticFeatureLevel(GMaxRHIFeatureLevel) : FeatureLevel; }
 	FORCEINLINE bool HasValidFeatureLevel() const { return FeatureLevel < ERHIFeatureLevel::Num; }
 
-	// Helper for submitting a resource array to RHI and freeing eligible CPU memory
-	template<bool bRenderThread, typename T>
-	FBufferRHIRef CreateRHIBuffer(T& InOutResourceObject, const uint32 ResourceCount, EBufferUsageFlags InBufferUsageFlags, const TCHAR* InDebugName)
-	{
-		FBufferRHIRef Buffer;
-		FResourceArrayInterface* RESTRICT ResourceArray = InOutResourceObject ? InOutResourceObject->GetResourceArray() : nullptr;
-		if (ResourceCount != 0)
-		{
-			const uint32 SizeInBytes = ResourceArray ? ResourceArray->GetResourceDataSize() : 0;
-			FRHIResourceCreateInfo CreateInfo(InDebugName, ResourceArray);
-			CreateInfo.bWithoutNativeResource = !InOutResourceObject;
-
-			if (bRenderThread)
-			{
-				Buffer = RHICreateVertexBuffer(SizeInBytes, InBufferUsageFlags, CreateInfo);
-			}
-			else
-			{
-				Buffer = RHIAsyncCreateVertexBuffer(SizeInBytes, InBufferUsageFlags, CreateInfo);
-			}
-		}
-
-		if (GFreeStructuresOnRHIBufferCreation && InOutResourceObject && (!ResourceArray || !ResourceArray->GetAllowCPUAccess()))
-		{
-			// The resource should have been freed as part of the buffer creation
-			check(!ResourceArray || !ResourceArray->GetResourceDataSize());
-			delete InOutResourceObject;
-			InOutResourceObject = nullptr;
-		}
-
-		return Buffer;
-	}
-
 private:
 	TEnumAsByte<ERHIFeatureLevel::Type> FeatureLevel;
 };
@@ -556,7 +520,6 @@ public:
 	{
 		VertexBufferRHI.SafeRelease();
 	}
-
 	virtual FString GetFriendlyName() const override { return TEXT("FVertexBuffer"); }
 };
 
