Fix SetThreadConfig example comment

diff --git a/Engine/Source/Runtime/Engine/Private/UnrealEngine.cpp b/Engine/Source/Runtime/Engine/Private/UnrealEngine.cpp
--- a/Engine/Source/Runtime/Engine/Private/UnrealEngine.cpp
+++ b/Engine/Source/Runtime/Engine/Private/UnrealEngine.cpp
@@ -17174,7 +17174,8 @@ static void SetupThreadConfig(const TArray<FString>& Args)
 	GLog->FlushThreadedLogs();
 }
 
-// example string to pass SetThreadConfig Task 0x3f TPri_Normal TaskBP 0x3f TPri_BelowNormal GT 0x1 TPri_Highest RT 0x10 TPri_Highest RHI TPri_Highest
+// SetThreadConfig console command to change thread configurations/affinities at runtime. Corresponds to ThreadConfigs array in /Script/Engine.Engine
+// Example: SetThreadConfig Task:TPri_Normal:0x3f TaskBP:TPri_SlightlyBelowNormal:0x7f GT:TPri_Highest RT:TPri_Highest RHI:TPri_Highest
 static FAutoConsoleCommand SetupThreadConfigCmd(
 	TEXT("SetThreadConfig"),
 	TEXT("Sets a thread Priority and/or Affinity. A single arg of default resets all the thread Priorities and Affinities, otherwise [GT|RT|RHI|Task|TaskBP] both/either [TPri_type] [0x] sets the Priority and/or Affinity."),
