Exposing to Blueprints the cached modifier-keys state from slate so we can build more functional editor-utility widgets.

diff --git a/Engine/Source/Runtime/Engine/Classes/Kismet/KismetInputLibrary.h b/Engine/Source/Runtime/Engine/Classes/Kismet/KismetInputLibrary.h
--- a/Engine/Source/Runtime/Engine/Classes/Kismet/KismetInputLibrary.h
+++ b/Engine/Source/Runtime/Engine/Classes/Kismet/KismetInputLibrary.h
@@ -10,6 +10,8 @@
 #include "Kismet/BlueprintFunctionLibrary.h"
 #include "KismetInputLibrary.generated.h"
 
+class FModifierKeysState;
+
 UENUM(BlueprintType)
 enum class ESlateGesture : uint8
 {
@@ -21,6 +23,19 @@ enum class ESlateGesture : uint8
 	LongPress
 };
 
+/** A structure which captures the application's modifier key states (shift, alt, ctrl, etc.) */
+USTRUCT(BlueprintType, DisplayName="Modifier Keys State")
+struct FSlateModifierKeysState
+{
+	GENERATED_USTRUCT_BODY()
+
+	UPROPERTY()
+	uint8 ModifierKeysStateMask = 0;
+
+	FSlateModifierKeysState() {}
+	FSlateModifierKeysState(const FModifierKeysState& InModifierKeysState);	
+};
+
 UCLASS(meta=(ScriptName="InputLibrary"))
 class ENGINE_API UKismetInputLibrary : public UBlueprintFunctionLibrary
 {
@@ -220,6 +235,34 @@ class ENGINE_API UKismetInputLibrary : public UBlueprintFunctionLibrary
 	UFUNCTION(BlueprintPure, meta=( DisplayName = "Is Right Command Down" ), Category="Input|InputEvent")
 	static bool InputEvent_IsRightCommandDown(const FInputEvent& Input);
 
+	/**
+	 * Returns true if either shift key was down when the key state was captured 
+	 */
+	UFUNCTION(BlueprintPure, meta=( DisplayName = "Is Shift Down" ), Category="Input|ModifierKeys")
+	static bool ModifierKeysState_IsShiftDown(const FSlateModifierKeysState& KeysState);
+
+	/**
+	 * Returns true if either control key was down when the key state was captured 
+	 */
+	UFUNCTION(BlueprintPure, meta=( DisplayName = "Is Control Down" ), Category="Input|ModifierKeys")
+	static bool ModifierKeysState_IsControlDown(const FSlateModifierKeysState& KeysState);
+
+	/**
+	 * Returns true if either alt key was down when the key state was captured 
+	 */
+	UFUNCTION(BlueprintPure, meta=( DisplayName = "Is Alt Down" ), Category="Input|ModifierKeys")
+	static bool ModifierKeysState_IsAltDown(const FSlateModifierKeysState& KeysState);
+
+	/**
+	 * Returns true if either command key was down when the key state was captured 
+	 */
+	UFUNCTION(BlueprintPure, meta=( DisplayName = "Is Command Down" ), Category="Input|ModifierKeys")
+	static bool ModifierKeysState_IsCommandDown(const FSlateModifierKeysState& KeysState);
+
+	/** Returns a snapshot of the cached modifier-keys state for the application. */
+	UFUNCTION(BlueprintPure, Category = "Input|ModifierKeys")
+	static FSlateModifierKeysState GetModifierKeysState();
+
 	/** @return The display name of the input chord */
 	UFUNCTION(BlueprintPure, meta = (DisplayName = "Get Input Chord Display Name"), Category = "Input|Key")
 	static FText InputChord_GetDisplayName(const FInputChord& Key);
diff --git a/Engine/Source/Runtime/Engine/Private/KismetInputLibrary.cpp b/Engine/Source/Runtime/Engine/Private/KismetInputLibrary.cpp
--- a/Engine/Source/Runtime/Engine/Private/KismetInputLibrary.cpp
+++ b/Engine/Source/Runtime/Engine/Private/KismetInputLibrary.cpp
@@ -4,7 +4,17 @@
 #include "EngineGlobals.h"
 #include "Engine/Engine.h"
 #include "Framework/Application/SlateApplication.h"
+#include "GenericPlatform/GenericApplication.h" // for EModifierKey, FModifierKeysState
 
+FSlateModifierKeysState::FSlateModifierKeysState(const FModifierKeysState& InModifierKeysState)
+	: ModifierKeysStateMask(EModifierKey::FromBools(
+		InModifierKeysState.IsControlDown(),
+		InModifierKeysState.IsAltDown(),
+		InModifierKeysState.IsShiftDown(),
+		InModifierKeysState.IsCommandDown()
+	))
+{
+}
 
 //////////////////////////////////////////////////////////////////////////
 // UKismetInputLibrary
@@ -206,6 +216,35 @@ bool UKismetInputLibrary::InputEvent_IsRightCommandDown(const FInputEvent& Input
 	return Input.IsRightCommandDown();
 }
 
+bool UKismetInputLibrary::ModifierKeysState_IsShiftDown(const FSlateModifierKeysState& KeysState)
+{
+	return (KeysState.ModifierKeysStateMask & EModifierKey::Shift) != 0;
+}
+
+bool UKismetInputLibrary::ModifierKeysState_IsControlDown(const FSlateModifierKeysState& KeysState)
+{
+	return (KeysState.ModifierKeysStateMask & EModifierKey::Control) != 0;
+}
+
+bool UKismetInputLibrary::ModifierKeysState_IsAltDown(const FSlateModifierKeysState& KeysState)
+{
+	return (KeysState.ModifierKeysStateMask & EModifierKey::Alt) != 0;
+}
+
+bool UKismetInputLibrary::ModifierKeysState_IsCommandDown(const FSlateModifierKeysState& KeysState)
+{
+	return (KeysState.ModifierKeysStateMask & EModifierKey::Command) != 0;
+}
+
+FSlateModifierKeysState UKismetInputLibrary::GetModifierKeysState()
+{
+	if (FSlateApplication::IsInitialized())
+	{
+		return FSlateApplication::Get().GetModifierKeys();
+	}
+	return FSlateModifierKeysState();
+}
+
 FText UKismetInputLibrary::InputChord_GetDisplayName(const FInputChord& Key)
 {
 	return Key.GetInputText();
