Update SkyLight CapturePosition when component transform changes.
- Moving the skylight position to test LWC using blueprints was breaking it.

diff --git a/Engine/Source/Runtime/Engine/Classes/Components/SkyLightComponent.h b/Engine/Source/Runtime/Engine/Classes/Components/SkyLightComponent.h
--- a/Engine/Source/Runtime/Engine/Classes/Components/SkyLightComponent.h
+++ b/Engine/Source/Runtime/Engine/Classes/Components/SkyLightComponent.h
@@ -349,6 +349,7 @@ protected:
 	//~ Begin UActorComponent Interface
 	virtual void CreateRenderState_Concurrent(FRegisterComponentContext* Context) override;
 	virtual void DestroyRenderState_Concurrent() override;
+	virtual void SendRenderTransform_Concurrent() override;
 	//~ Begin UActorComponent Interface
 
 	void UpdateLimitedRenderingStateFast();
diff --git a/Engine/Source/Runtime/Engine/Private/Components/SkyLightComponent.cpp b/Engine/Source/Runtime/Engine/Private/Components/SkyLightComponent.cpp
--- a/Engine/Source/Runtime/Engine/Private/Components/SkyLightComponent.cpp
+++ b/Engine/Source/Runtime/Engine/Private/Components/SkyLightComponent.cpp
@@ -518,6 +518,23 @@ void USkyLightComponent::DestroyRenderState_Concurrent()
 	}
 }
 
+void USkyLightComponent::SendRenderTransform_Concurrent()
+{
+	if (SceneProxy)
+	{
+		FSkyLightSceneProxy* InLightSceneProxy = SceneProxy;
+		FVector Position = GetComponentTransform().GetLocation();
+
+		ENQUEUE_RENDER_COMMAND(UpdateSkyLightCapturePosition)(
+			[InLightSceneProxy, Position](FRHICommandListImmediate& RHICmdList)
+			{
+				InLightSceneProxy->CapturePosition = Position;
+			});
+	}
+
+	Super::SendRenderTransform_Concurrent();
+}
+
 #if WITH_EDITOR
 void USkyLightComponent::PreEditChange(FProperty* PropertyAboutToChange)
 {
