Material Dot expressions allows mixing inputs with different number of components

diff --git a/Engine/Source/Runtime/Engine/Private/Materials/HLSLMaterialDerivativeAutogen.cpp b/Engine/Source/Runtime/Engine/Private/Materials/HLSLMaterialDerivativeAutogen.cpp
--- a/Engine/Source/Runtime/Engine/Private/Materials/HLSLMaterialDerivativeAutogen.cpp
+++ b/Engine/Source/Runtime/Engine/Private/Materials/HLSLMaterialDerivativeAutogen.cpp
@@ -396,7 +396,13 @@ FMaterialDerivativeAutogen::FOperationType2 FMaterialDerivativeAutogen::GetFunc2
 		{
 			const EDerivativeType LhsLWCType = MakeLWCType(LhsType);
 			const EDerivativeType RhsLWCType = MakeLWCType(RhsType);
-			if (LhsLWCType == RhsLWCType || LhsLWCType == EDerivativeType::LWCScalar || RhsLWCType == EDerivativeType::LWCScalar)
+			if (LhsLWCType == RhsLWCType ||
+				LhsLWCType == EDerivativeType::LWCScalar ||
+				RhsLWCType == EDerivativeType::LWCScalar ||
+				// Dot allows mixing different numbers of components
+				// Legacy implementation truncated non-scalar types, this will instead promote to the maximum number of components
+				// But because promotions are filled with 0s, this will effectively be the same thing (and simplifies the logic here)
+				Op == EFunc2::Dot)
 			{
 				const EDerivativeType LWCType = (EDerivativeType)FMath::Max<int32>((int32)LhsLWCType, (int32)RhsLWCType);
 				switch (Op)
@@ -425,7 +431,10 @@ FMaterialDerivativeAutogen::FOperationType2 FMaterialDerivativeAutogen::GetFunc2
 	case EFunc2::PowPositiveClamped:
 	case EFunc2::Atan2:
 	case EFunc2::Atan2Fast:
-		if (LhsNonLWCType == RhsNonLWCType || LhsNonLWCType == EDerivativeType::Float1 || RhsNonLWCType == EDerivativeType::Float1)
+		if (LhsNonLWCType == RhsNonLWCType ||
+			LhsNonLWCType == EDerivativeType::Float1 ||
+			RhsNonLWCType == EDerivativeType::Float1 ||
+			Op == EFunc2::Dot)
 		{
 			// if the initial type is different from the output type, then it's only valid if type is 0 (float). We can convert
 			// a float to a type with more components, but for example, we can't implicitly convert a float2 to a float3/float4.
