Support LWC texcoords for external textures

diff --git a/Engine/Shaders/Private/MaterialTemplate.ush b/Engine/Shaders/Private/MaterialTemplate.ush
--- a/Engine/Shaders/Private/MaterialTemplate.ush
+++ b/Engine/Shaders/Private/MaterialTemplate.ush
@@ -1895,6 +1895,11 @@ float2 RotateScaleOffsetTexCoords(float2 InTexCoords, float4 InRotationScale, fl
 	return float2(dot(InTexCoords, InRotationScale.xy), dot(InTexCoords, InRotationScale.zw)) + InOffset;
 }
 
+FLWCVector2 RotateScaleOffsetTexCoords(FLWCVector2 InTexCoords, float4 InRotationScale, float2 InOffset)
+{
+	return LWCAdd(MakeLWCVector(LWCDot(InTexCoords, InRotationScale.xy), LWCDot(InTexCoords, InRotationScale.zw)), InOffset);
+}
+
 #if USES_SPEEDTREE
 
 /** Vertex offset for SpeedTree wind and LOD */
diff --git a/Engine/Source/Runtime/Engine/Private/Materials/HLSLMaterialDerivativeAutogen.cpp b/Engine/Source/Runtime/Engine/Private/Materials/HLSLMaterialDerivativeAutogen.cpp
--- a/Engine/Source/Runtime/Engine/Private/Materials/HLSLMaterialDerivativeAutogen.cpp
+++ b/Engine/Source/Runtime/Engine/Private/Materials/HLSLMaterialDerivativeAutogen.cpp
@@ -1116,13 +1116,13 @@ int32 FMaterialDerivativeAutogen::GenerateRotateScaleOffsetTexCoordsFunc(FHLSLMa
 	const FDerivInfo RotationScaleDerivInfo	= Translator.GetDerivInfo(RotationScale);
 	const FDerivInfo OffsetDerivInfo		= Translator.GetDerivInfo(Offset);
 
-	const EMaterialValueType ResultType = MCT_Float2;
+	const EMaterialValueType ResultType = IsLWCType(TexCoordDerivInfo.Type) ? MCT_LWCVector2 : MCT_Float2;
 	const uint32 NumResultComponents = 2;
 
 	const bool bAllZeroDeriv = (TexCoordDerivInfo.DerivativeStatus == EDerivativeStatus::Zero && RotationScaleDerivInfo.DerivativeStatus == EDerivativeStatus::Zero && OffsetDerivInfo.DerivativeStatus == EDerivativeStatus::Zero);
-	FString FiniteString = FString::Printf(TEXT("RotateScaleOffsetTexCoords(%s, %s, %s.xy)"),	*Translator.GetParameterCode(TexCoord), 
-																								*Translator.GetParameterCode(RotationScale),
-																								*Translator.GetParameterCode(Offset));
+	FString FiniteString = FString::Printf(TEXT("RotateScaleOffsetTexCoords(%s, %s, %s.xy)"),	*Translator.CoerceParameter(TexCoord, ResultType),
+																								*Translator.CoerceParameter(RotationScale, MCT_Float4),
+																								*Translator.CoerceParameter(Offset, MCT_Float2));
 
 	if (!bAllZeroDeriv && IsDerivativeValid(TexCoordDerivInfo.DerivativeStatus) && IsDerivativeValid(RotationScaleDerivInfo.DerivativeStatus) && IsDerivativeValid(OffsetDerivInfo.DerivativeStatus))
 	{
diff --git a/Engine/Source/Runtime/Engine/Private/Materials/HLSLMaterialTranslator.cpp b/Engine/Source/Runtime/Engine/Private/Materials/HLSLMaterialTranslator.cpp
--- a/Engine/Source/Runtime/Engine/Private/Materials/HLSLMaterialTranslator.cpp
+++ b/Engine/Source/Runtime/Engine/Private/Materials/HLSLMaterialTranslator.cpp
@@ -6800,11 +6800,11 @@ int32 FHLSLMaterialTranslator::ExternalTextureCoordinateScaleRotation(const FGui
 }
 int32 FHLSLMaterialTranslator::ExternalTextureCoordinateOffset(int32 TextureReferenceIndex, TOptional<FName> ParameterName)
 {
-	return AddUniformExpression(new FMaterialUniformExpressionExternalTextureCoordinateOffset(TextureReferenceIndex, ParameterName), MCT_Float4, TEXT(""));
+	return AddUniformExpression(new FMaterialUniformExpressionExternalTextureCoordinateOffset(TextureReferenceIndex, ParameterName), MCT_Float2, TEXT(""));
 }
 int32 FHLSLMaterialTranslator::ExternalTextureCoordinateOffset(const FGuid& ExternalTextureGuid)
 {
-	return AddUniformExpression(new FMaterialUniformExpressionExternalTextureCoordinateOffset(ExternalTextureGuid), MCT_Float4, TEXT(""));
+	return AddUniformExpression(new FMaterialUniformExpressionExternalTextureCoordinateOffset(ExternalTextureGuid), MCT_Float2, TEXT(""));
 }
 
 UObject* FHLSLMaterialTranslator::GetReferencedTexture(int32 Index)
