Exposed montage section info to blueprints

diff --git a/Engine/Source/Runtime/Engine/Classes/Animation/AnimMontage.h b/Engine/Source/Runtime/Engine/Classes/Animation/AnimMontage.h
--- a/Engine/Source/Runtime/Engine/Classes/Animation/AnimMontage.h
+++ b/Engine/Source/Runtime/Engine/Classes/Animation/AnimMontage.h
@@ -769,12 +769,18 @@ public:
 	// @todo document
 	ENGINE_API float GetSectionLength(int32 SectionIndex) const;
 	
-	/** Get SectionIndex from SectionName */
+	/** Get SectionIndex from SectionName. Returns INDEX_None if not found */
+	UFUNCTION(BlueprintPure, Category = "Montage")
 	ENGINE_API int32 GetSectionIndex(FName InSectionName) const;
 	
-	/** Get SectionName from SectionIndex in TArray */
+	/** Get SectionName from SectionIndex. Returns NAME_None if not found */
+	UFUNCTION(BlueprintPure, Category = "Montage")
 	ENGINE_API FName GetSectionName(int32 SectionIndex) const;
 
+	/** Returns the number of sections this montage has */
+	UFUNCTION(BlueprintPure, Category = "Montage")
+	ENGINE_API int32 GetNumSections() const { return CompositeSections.Num(); }
+
 	/** @return true if valid section */
 	UFUNCTION(BlueprintCallable, Category = "Montage")
 	ENGINE_API bool IsValidSectionName(FName InSectionName) const;
