Fix a thread race that can occur when a cube map stored in the map build registry is upgraded to a new format while it is still caching it's platform data in an async job.

#rnx

- GenerateEncodedHDRTextureCube should call ::PreEditChange/::PostEditChange on the texture before modifying it so that we wait for any outstanding async work to complete first.
- UTexture::Modify now checks if the async cache is complete and if not waits on it, in addition to potentially call FTextureCompilingManager::FinishCompilation, as not all cases are covered by the latter.

diff --git a/Engine/Source/Runtime/Engine/Private/Components/ReflectionCaptureComponent.cpp b/Engine/Source/Runtime/Engine/Private/Components/ReflectionCaptureComponent.cpp
--- a/Engine/Source/Runtime/Engine/Private/Components/ReflectionCaptureComponent.cpp
+++ b/Engine/Source/Runtime/Engine/Private/Components/ReflectionCaptureComponent.cpp
@@ -686,6 +686,8 @@ void GenerateEncodedHDRTextureCube(UMapBuildDataRegistry* Registry, FReflectionC
 
 	if (TextureCube)
 	{
+		TextureCube->PreEditChange(nullptr);
+
 		TArray<uint8> TemporaryEncodedHDRCapturedData;
 
 		GenerateEncodedHDRData(ReflectionCaptureData.FullHDRCapturedData, ReflectionCaptureData.CubemapSize, MaxValueRGBM, TemporaryEncodedHDRCapturedData);
@@ -719,6 +721,8 @@ void GenerateEncodedHDRTextureCube(UMapBuildDataRegistry* Registry, FReflectionC
 
 		TextureCube->UpdateResource();
 		TextureCube->MarkPackageDirty();
+
+		TextureCube->PostEditChange();
 	}
 	ReflectionCaptureData.EncodedCaptureData = TextureCube;
 	ReflectionCaptureData.bBrightnessBakedInEncodedHDRCubemap = false;
diff --git a/Engine/Source/Runtime/Engine/Private/Texture.cpp b/Engine/Source/Runtime/Engine/Private/Texture.cpp
--- a/Engine/Source/Runtime/Engine/Private/Texture.cpp
+++ b/Engine/Source/Runtime/Engine/Private/Texture.cpp
@@ -303,6 +303,11 @@ bool UTexture::Modify(bool bAlwaysMarkDirty)
 {
 	// Before applying any modification to the texture
 	// make sure no compilation is still ongoing.
+	if (!IsAsyncCacheComplete())
+	{
+		FinishCachePlatformData();
+	}	
+	
 	if (IsDefaultTexture())
 	{
 		FTextureCompilingManager::Get().FinishCompilation({this});
