Build HISM instances for foliage in a translated space and compensate in the render matrix.
This is to address large world coordinates issues where foliage HISM is typically constructed in world space.

diff --git a/Engine/Source/Runtime/Engine/Classes/Components/HierarchicalInstancedStaticMeshComponent.h b/Engine/Source/Runtime/Engine/Classes/Components/HierarchicalInstancedStaticMeshComponent.h
--- a/Engine/Source/Runtime/Engine/Classes/Components/HierarchicalInstancedStaticMeshComponent.h
+++ b/Engine/Source/Runtime/Engine/Classes/Components/HierarchicalInstancedStaticMeshComponent.h
@@ -273,7 +273,7 @@ protected:
 	void ApplyEmpty();
 	void SetPerInstanceLightMapAndEditorData(FStaticMeshInstanceData& PerInstanceData, const TArray<TRefCountPtr<HHitProxy>>& HitProxies);
 
-	void GetInstanceTransforms(TArray<FMatrix>& InstanceTransforms) const;
+	void GetInstanceTransforms(TArray<FMatrix>& InstanceTransforms, FVector const& Offset) const;
 	void InitializeInstancingRandomSeed();
 
 	/** Removes specified instances */ 
diff --git a/Engine/Source/Runtime/Engine/Classes/Components/InstancedStaticMeshComponent.h b/Engine/Source/Runtime/Engine/Classes/Components/InstancedStaticMeshComponent.h
--- a/Engine/Source/Runtime/Engine/Classes/Components/InstancedStaticMeshComponent.h
+++ b/Engine/Source/Runtime/Engine/Classes/Components/InstancedStaticMeshComponent.h
@@ -320,6 +320,7 @@ public:
 
 	//~ Begin UPrimitiveComponent Interface
 	virtual FPrimitiveSceneProxy* CreateSceneProxy() override;
+	virtual FMatrix GetRenderMatrix() const override;
 	virtual FBodyInstance* GetBodyInstance(FName BoneName = NAME_None, bool bGetWelded = true, int32 Index = INDEX_NONE) const override;
 protected:
 	virtual void OnCreatePhysicsState() override;
@@ -409,6 +410,13 @@ protected:
 	/** Terminate all body instances owned by this component. */
 	void ClearAllInstanceBodies();
 
+	/**
+	 * Get the translated space for instance transforms to be passed to the renderer.
+	 * In the renderer data structures we only have floating point precision for instance transforms relative to their owning primitive. The primitive transform itself has double precision.
+	 * Some ISM that are authored in world space need to adjust the local space to keep instance transforms within precision limits.
+	 */
+	virtual FVector GetTranslatedInstanceSpaceOrigin() const { return FVector::Zero(); }
+
 	/** Request to navigation system to update only part of navmesh occupied by specified instance. */
 	virtual void PartialNavigationUpdate(int32 InstanceIdx);
 
diff --git a/Engine/Source/Runtime/Engine/Private/HierarchicalInstancedStaticMesh.cpp b/Engine/Source/Runtime/Engine/Private/HierarchicalInstancedStaticMesh.cpp
--- a/Engine/Source/Runtime/Engine/Private/HierarchicalInstancedStaticMesh.cpp
+++ b/Engine/Source/Runtime/Engine/Private/HierarchicalInstancedStaticMesh.cpp
@@ -345,7 +345,8 @@ protected:
 				float RandomID = RandomStream.GetFraction();
 				if (RenderIndex >= 0)
 				{
-					BuiltInstanceData->SetInstance(RenderIndex, FMatrix44f(Transforms[i]), RandomID, LightmapUVBias, ShadowmapUVBias);		// LWC_TODO: Precision loss
+					// LWC_TODO: Precision loss here has been compensated for by use of TranslatedInstanceSpaceOrigin.
+					BuiltInstanceData->SetInstance(RenderIndex, FMatrix44f(Transforms[i]), RandomID, LightmapUVBias, ShadowmapUVBias);
 					for (int32 DataIndex = 0; DataIndex < NumCustomDataFloats; ++DataIndex)
 					{
 						BuiltInstanceData->SetInstanceCustomData(RenderIndex, DataIndex, CustomDataFloats[NumCustomDataFloats * i + DataIndex]);
@@ -2005,6 +2006,12 @@ FBoxSphereBounds UHierarchicalInstancedStaticMeshComponent::CalcBounds(const FTr
 	}
 }
 
+static FBox GetClusterTreeBounds(TArray<FClusterNode> const& InClusterTree, FVector InOffset)
+{
+	// Return top node of cluster tree. Apply offset on node bounds.
+	return (InClusterTree.Num() > 0 ? FBox(InOffset + InClusterTree[0].BoundMin, InOffset + InClusterTree[0].BoundMax) : FBox(ForceInit));
+}
+
 UHierarchicalInstancedStaticMeshComponent::UHierarchicalInstancedStaticMeshComponent(const FObjectInitializer& ObjectInitializer)
 	: Super(ObjectInitializer)
 	, ClusterTreePtr(MakeShareable(new TArray<FClusterNode>))
@@ -2117,8 +2124,7 @@ void UHierarchicalInstancedStaticMeshComponent::Serialize(FArchive& Ar)
 	if (Ar.IsLoading() && !BuiltInstanceBounds.IsValid)
 	{
 		TArray<FClusterNode>& ClusterTree = *ClusterTreePtr.Get();
-
-		BuiltInstanceBounds = (ClusterTree.Num() > 0 ? FBox(ClusterTree[0].BoundMin, ClusterTree[0].BoundMax) : FBox(ForceInit));
+		BuiltInstanceBounds = GetClusterTreeBounds(ClusterTree, GetTranslatedInstanceSpaceOrigin());
 	}
 }
 
@@ -2572,7 +2578,7 @@ void UHierarchicalInstancedStaticMeshComponent::ClearInstances()
 	SortedInstances.Empty();
 	UnbuiltInstanceBounds.Init();
 	UnbuiltInstanceBoundsList.Empty();
-	
+
 	if (ProxySize)
 	{
 		DEC_DWORD_STAT_BY(STAT_FoliageInstanceBuffers, ProxySize);
@@ -2712,9 +2718,12 @@ void UHierarchicalInstancedStaticMeshComponent::BuildTree()
 	// The tree will be fully rebuilt once the static mesh compilation is finished, no need for incremental update in that case.
 	if (PerInstanceSMData.Num() > 0 && GetStaticMesh() && !GetStaticMesh()->IsCompiling() && GetStaticMesh()->HasValidRenderData())
 	{
+		// Build the tree in translated space to maintain precision.
+		const FVector TranslatedInstanceSpaceOrigin = GetTranslatedInstanceSpaceOrigin();
+
 		InitializeInstancingRandomSeed();
 		TArray<FMatrix> InstanceTransforms;
-		GetInstanceTransforms(InstanceTransforms);
+		GetInstanceTransforms(InstanceTransforms, -TranslatedInstanceSpaceOrigin);
 
 		FClusterBuilder Builder(InstanceTransforms, PerInstanceSMCustomData, NumCustomDataFloats, GetStaticMesh()->GetBounds().GetBox(), DesiredInstancesPerLeaf(), CurrentDensityScaling, InstancingRandomSeed, PerInstanceSMData.Num() > 0);
 		Builder.BuildTreeAndBuffer();
@@ -2763,7 +2772,8 @@ void UHierarchicalInstancedStaticMeshComponent::AcceptPrebuiltTree(TArray<FClust
 	// this is only for prebuild data, already in the correct order
 	check(!PerInstanceSMData.Num());
 	NumBuiltInstances = 0;
-	check(PerInstanceRenderData.IsValid());	
+	check(GetTranslatedInstanceSpaceOrigin() == FVector::Zero());
+	check(PerInstanceRenderData.IsValid());
 	NumBuiltRenderInstances = InNumBuiltRenderInstances;
 	check(NumBuiltRenderInstances);
 	UnbuiltInstanceBounds.Init();
@@ -2772,7 +2782,7 @@ void UHierarchicalInstancedStaticMeshComponent::AcceptPrebuiltTree(TArray<FClust
 	InstanceReorderTable.Empty();
 	SortedInstances.Empty();
 	OcclusionLayerNumNodes = InOcclusionLayerNumNodes;
-	BuiltInstanceBounds = (InClusterTree.Num() > 0 ? FBox(InClusterTree[0].BoundMin, InClusterTree[0].BoundMax) : FBox(ForceInit));
+	BuiltInstanceBounds = GetClusterTreeBounds(InClusterTree, FVector::Zero());
 	InstanceCountToRender = InNumBuiltRenderInstances;
 
 	// Verify that the mesh is valid before using it.
@@ -2865,8 +2875,10 @@ void UHierarchicalInstancedStaticMeshComponent::ApplyBuildTree(FClusterBuilder&
 	TUniquePtr<FStaticMeshInstanceData> BuiltInstanceData = MoveTemp(Builder.BuiltInstanceData);
 
 	OcclusionLayerNumNodes = Builder.Result->OutOcclusionLayerNum;
+
+	// Get the new bounds taking into account the translated space used when building the tree.
 	const TArray<FClusterNode>& ClusterTree = *ClusterTreePtr;
-	BuiltInstanceBounds = (ClusterTree.Num() > 0 ? FBox(ClusterTree[0].BoundMin, ClusterTree[0].BoundMax) : FBox(ForceInit));
+	BuiltInstanceBounds = GetClusterTreeBounds(ClusterTree, GetTranslatedInstanceSpaceOrigin());
 
 	UnbuiltInstanceBounds.Init();
 	UnbuiltInstanceBoundsList.Empty();
@@ -2957,7 +2969,7 @@ bool UHierarchicalInstancedStaticMeshComponent::BuildTreeIfOutdated(bool Async,
 	return false;
 }
 
-void UHierarchicalInstancedStaticMeshComponent::GetInstanceTransforms(TArray<FMatrix>& InstanceTransforms) const
+void UHierarchicalInstancedStaticMeshComponent::GetInstanceTransforms(TArray<FMatrix>& InstanceTransforms, FVector const& Offset) const
 {
 	double StartTime = FPlatformTime::Seconds();
 	int32 Num = PerInstanceSMData.Num();
@@ -2965,7 +2977,7 @@ void UHierarchicalInstancedStaticMeshComponent::GetInstanceTransforms(TArray<FMa
 	InstanceTransforms.SetNumUninitialized(Num);
 	for (int32 Index = 0; Index < Num; Index++)
 	{
-		InstanceTransforms[Index] = PerInstanceSMData[Index].Transform;
+		InstanceTransforms[Index] = PerInstanceSMData[Index].Transform.ConcatTranslation(Offset);
 	}
 
 	UE_LOG(LogStaticMesh, Verbose, TEXT("Copied %d transforms in %.3fs."), Num, float(FPlatformTime::Seconds() - StartTime));
@@ -3006,9 +3018,12 @@ void UHierarchicalInstancedStaticMeshComponent::BuildTreeAsync()
 	{
 		double StartTime = FPlatformTime::Seconds();
 		
+		// Build the tree in translated space to maintain precision.
+		const FVector TranslatedInstanceSpaceOrigin = GetTranslatedInstanceSpaceOrigin();
+
 		InitializeInstancingRandomSeed();
 		TArray<FMatrix> InstanceTransforms;
-		GetInstanceTransforms(InstanceTransforms);
+		GetInstanceTransforms(InstanceTransforms, -TranslatedInstanceSpaceOrigin);
 		
 		TSharedRef<FClusterBuilder, ESPMode::ThreadSafe> Builder(new FClusterBuilder(InstanceTransforms, PerInstanceSMCustomData, NumCustomDataFloats, GetStaticMesh()->GetBounds().GetBox(), DesiredInstancesPerLeaf(), CurrentDensityScaling, InstancingRandomSeed, PerInstanceSMData.Num() > 0));
 
@@ -3442,7 +3457,7 @@ void UHierarchicalInstancedStaticMeshComponent::FlushAccumulatedNavigationUpdate
 		const TArray<FClusterNode>& ClusterTree = *ClusterTreePtr;
 		if (ClusterTree.Num())
 		{
-			FBox NewBounds = FBox(ClusterTree[0].BoundMin, ClusterTree[0].BoundMax).TransformBy(GetComponentTransform());
+			const FBox NewBounds = GetClusterTreeBounds(ClusterTree, GetTranslatedInstanceSpaceOrigin()).TransformBy(GetComponentTransform());
 			FNavigationSystem::OnComponentBoundsChanged(*this, NewBounds, AccumulatedNavigationDirtyArea);
 		}
 			
diff --git a/Engine/Source/Runtime/Engine/Private/InstancedStaticMesh.cpp b/Engine/Source/Runtime/Engine/Private/InstancedStaticMesh.cpp
--- a/Engine/Source/Runtime/Engine/Private/InstancedStaticMesh.cpp
+++ b/Engine/Source/Runtime/Engine/Private/InstancedStaticMesh.cpp
@@ -2297,6 +2297,13 @@ FPrimitiveSceneProxy* UInstancedStaticMeshComponent::CreateSceneProxy()
 	}
 }
 
+FMatrix UInstancedStaticMeshComponent::GetRenderMatrix() const
+{
+	// Apply the translated space to the render matrix.
+	const FVector TranslatedInstanceSpaceOrigin = GetTranslatedInstanceSpaceOrigin();
+	return GetComponentTransform().ToMatrixWithScale().ConcatTranslation(TranslatedInstanceSpaceOrigin);
+}
+
 void UInstancedStaticMeshComponent::CreateHitProxyData(TArray<TRefCountPtr<HHitProxy>>& HitProxies)
 {
 	if (GIsEditor && bHasPerInstanceHitProxies)
@@ -2349,6 +2356,8 @@ void UInstancedStaticMeshComponent::BuildRenderData(FStaticMeshInstanceData& Out
 	auto AdditionalRandomSeedsIt = AdditionalRandomSeeds.CreateIterator();
 	int32 SeedResetIndex = AdditionalRandomSeedsIt ? AdditionalRandomSeedsIt->StartInstanceIndex : INDEX_NONE;
 	
+	const FVector TranslatedInstanceSpaceOffset = -GetTranslatedInstanceSpaceOrigin();
+
 	for (int32 Index = 0; Index < NumInstances; ++Index)
 	{
 		const int32 RenderIndex = GetRenderIndex(Index);
@@ -2376,7 +2385,8 @@ void UInstancedStaticMeshComponent::BuildRenderData(FStaticMeshInstanceData& Out
 			ShadowmapUVBias = FVector2D(MeshMapBuildData->PerInstanceLightmapData[Index].ShadowmapUVBias);
 		}
 	
-		OutData.SetInstance(RenderIndex, FMatrix44f(InstanceData.Transform), RandomStream.GetFraction(), LightmapUVBias, ShadowmapUVBias);		// LWC_TODO: Precision loss
+		// LWC_TODO: Precision loss here is compensated for by use of TranslatedInstanceSpaceOffset.
+		OutData.SetInstance(RenderIndex, FMatrix44f(InstanceData.Transform.ConcatTranslation(TranslatedInstanceSpaceOffset)), RandomStream.GetFraction(), LightmapUVBias, ShadowmapUVBias);
 
 		for (int32 CustomDataIndex = 0; CustomDataIndex < NumCustomDataFloats; ++CustomDataIndex)
 		{
diff --git a/Engine/Source/Runtime/Foliage/Private/InstancedFoliage.cpp b/Engine/Source/Runtime/Foliage/Private/InstancedFoliage.cpp
--- a/Engine/Source/Runtime/Foliage/Private/InstancedFoliage.cpp
+++ b/Engine/Source/Runtime/Foliage/Private/InstancedFoliage.cpp
@@ -5630,6 +5630,18 @@ UFoliageInstancedStaticMeshComponent::UFoliageInstancedStaticMeshComponent(const
 #endif
 }
 
+FVector UFoliageInstancedStaticMeshComponent::GetTranslatedInstanceSpaceOrigin() const
+{
+	// Foliage is often built in world space which can cause problems with large world coordinates because
+	// the instance transforms in the renderer are single precision, and the HISM culling is also single precision.
+	// We should fix the HISM culling to be double precision.
+	// But the instance transforms (relative to the owner primitive) will probably stay single precision to not bloat memory.
+	// A fix for that is to have authoring tools set sensible primitive transforms (instead of identity).
+	// But until that happens we set a translated instance space here.
+	// For simplicity we use the first instance as the origin of the translated space.
+	return PerInstanceSMData.Num() ? PerInstanceSMData[0].Transform.GetOrigin() : FVector::Zero();
+}
+
 void UFoliageInstancedStaticMeshComponent::ReceiveComponentDamage(float DamageAmount, FDamageEvent const& DamageEvent, AController* EventInstigator, AActor* DamageCauser)
 {
 	Super::ReceiveComponentDamage(DamageAmount, DamageEvent, EventInstigator, DamageCauser);
diff --git a/Engine/Source/Runtime/Foliage/Public/FoliageInstancedStaticMeshComponent.h b/Engine/Source/Runtime/Foliage/Public/FoliageInstancedStaticMeshComponent.h
--- a/Engine/Source/Runtime/Foliage/Public/FoliageInstancedStaticMeshComponent.h
+++ b/Engine/Source/Runtime/Foliage/Public/FoliageInstancedStaticMeshComponent.h
@@ -43,6 +43,9 @@ public:
 	virtual uint64 GetHiddenEditorViews() const override;
 #endif// WITH_EDITOR
 
+	// UInstancedStaticMesh interface
+	virtual FVector GetTranslatedInstanceSpaceOrigin() const override;
+
 	/** Used by procedural generation to link generated component with its creator */
 	void SetGenerationGuid(const FGuid& InGuid) { GenerationGuid = InGuid; }
 	const FGuid& GetGenerationGuid() const { return GenerationGuid; }
