Fixed issues in CharacterMovement where NetConnection wasn't going through the character owner
Added possibility for rotation correction when a server correction is sent to clients

diff --git a/Engine/Source/Runtime/Engine/Classes/GameFramework/CharacterMovementComponent.h b/Engine/Source/Runtime/Engine/Classes/GameFramework/CharacterMovementComponent.h
--- a/Engine/Source/Runtime/Engine/Classes/GameFramework/CharacterMovementComponent.h
+++ b/Engine/Source/Runtime/Engine/Classes/GameFramework/CharacterMovementComponent.h
@@ -1058,6 +1058,9 @@ protected:
 
 public:
 
+	/** Returns if the character rotation should be corrected on clients when sending a server move response correction. */
+	virtual bool ShouldCorrectRotation() const { return false; }
+
 	UPROPERTY(Category="Character Movement: Avoidance", EditAnywhere, BlueprintReadOnly, meta=(ForceUnits=cm))
 	float AvoidanceConsiderationRadius;
 
@@ -2364,7 +2367,7 @@ public:
 	virtual void ClientAckGoodMove_Implementation(float TimeStamp);
 
 	/* Replicate position correction to client, associated with a timestamped servermove.  Client will replay subsequent moves after applying adjustment. */
-	virtual void ClientAdjustPosition_Implementation(float TimeStamp, FVector NewLoc, FVector NewVel, UPrimitiveComponent* NewBase, FName NewBaseBoneName, bool bHasBase, bool bBaseRelativePosition, uint8 ServerMovementMode);
+	virtual void ClientAdjustPosition_Implementation(float TimeStamp, FVector NewLoc, FVector NewVel, UPrimitiveComponent* NewBase, FName NewBaseBoneName, bool bHasBase, bool bBaseRelativePosition, uint8 ServerMovementMode, TOptional<FRotator> OptionalRotation = TOptional<FRotator>());
 
 	/* Bandwidth saving version, when velocity is zeroed */
 	virtual void ClientVeryShortAdjustPosition_Implementation(float TimeStamp, FVector NewLoc, UPrimitiveComponent* NewBase, FName NewBaseBoneName, bool bHasBase, bool bBaseRelativePosition, uint8 ServerMovementMode);
diff --git a/Engine/Source/Runtime/Engine/Private/Components/CharacterMovementComponent.cpp b/Engine/Source/Runtime/Engine/Private/Components/CharacterMovementComponent.cpp
--- a/Engine/Source/Runtime/Engine/Private/Components/CharacterMovementComponent.cpp
+++ b/Engine/Source/Runtime/Engine/Private/Components/CharacterMovementComponent.cpp
@@ -8478,8 +8478,7 @@ void UCharacterMovementComponent::CallServerMovePacked(const FSavedMove_Characte
 	BitWriterReset.Pop(ServerMoveBitWriter);
 
 	// Extract the net package map used for serializing object references.
-	APlayerController* PC = Cast<APlayerController>(CharacterOwner->GetController());
-	UNetConnection* NetConnection = PC ? PC->GetNetConnection() : nullptr;
+	UNetConnection* NetConnection = CharacterOwner->GetNetConnection();
 	ServerMoveBitWriter.PackageMap = NetConnection ? ToRawPtr(NetConnection->PackageMap) : nullptr;
 	if (ServerMoveBitWriter.PackageMap == nullptr)
 	{
@@ -9998,8 +9997,7 @@ void UCharacterMovementComponent::ServerSendMoveResponse(const FClientAdjustment
 	BitWriterReset.Pop(MoveResponseBitWriter);
 
 	// Extract the net package map used for serializing object references.
-	APlayerController* PC = Cast<APlayerController>(CharacterOwner->GetController());
-	UNetConnection* NetConnection = PC ? PC->GetNetConnection() : nullptr;
+	UNetConnection* NetConnection = CharacterOwner->GetNetConnection();
 	MoveResponseBitWriter.PackageMap = NetConnection ? ToRawPtr(NetConnection->PackageMap) : nullptr;
 	if (MoveResponseBitWriter.PackageMap == nullptr)
 	{
@@ -10039,6 +10037,7 @@ void FCharacterMoveResponseDataContainer::ServerFillResponseData(const UCharacte
 
 	if (!PendingAdjustment.bAckGoodMove)
 	{
+		bHasRotation = CharacterMovement.ShouldCorrectRotation();
 		bHasBase = (PendingAdjustment.NewBase != nullptr);
 		if (const ACharacter* CharacterOwner = CharacterMovement.GetCharacterOwner())
 		{
@@ -10167,7 +10166,9 @@ void UCharacterMovementComponent::ClientHandleMoveResponse(const FCharacterMoveR
 				MoveResponse.ClientAdjustment.NewBaseBoneName,
 				MoveResponse.bHasBase,
 				MoveResponse.ClientAdjustment.bBaseRelativePosition,
-				MoveResponse.ClientAdjustment.MovementMode);
+				MoveResponse.ClientAdjustment.MovementMode,
+				MoveResponse.bHasRotation ? MoveResponse.ClientAdjustment.NewRot : TOptional<FRotator>()
+				);
 		}
 	}
 }
@@ -10358,7 +10359,8 @@ void UCharacterMovementComponent::ClientAdjustPosition_Implementation
 	FName NewBaseBoneName,
 	bool bHasBase,
 	bool bBaseRelativePosition,
-	uint8 ServerMovementMode
+	uint8 ServerMovementMode,
+	TOptional<FRotator> OptionalRotation /* = TOptional<FRotator>()*/
 	)
 {
 	if (!HasValidData() || !IsActive())
@@ -10420,8 +10422,16 @@ void UCharacterMovementComponent::ClientAdjustPosition_Implementation
 	// Trust the server's positioning.
 	if (UpdatedComponent)
 	{
-		UpdatedComponent->SetWorldLocation(WorldShiftedNewLocation, false, nullptr, ETeleportType::TeleportPhysics);
+		if (OptionalRotation.IsSet())
+		{
+			UpdatedComponent->SetWorldLocationAndRotation(WorldShiftedNewLocation, OptionalRotation.GetValue(), false, nullptr, ETeleportType::TeleportPhysics);
+		}
+		else
+		{
+			UpdatedComponent->SetWorldLocation(WorldShiftedNewLocation, false, nullptr, ETeleportType::TeleportPhysics);
+		}
 	}
+
 	Velocity = NewVelocity;
 
 	// Trust the server's movement mode
