Fix cloth rendering broken in skin cache mode on consoles, and delete deprecated ClothLocalToWorld accessors

diff --git a/Engine/Shaders/Private/GpuSkinCacheComputeShader.usf b/Engine/Shaders/Private/GpuSkinCacheComputeShader.usf
--- a/Engine/Shaders/Private/GpuSkinCacheComputeShader.usf
+++ b/Engine/Shaders/Private/GpuSkinCacheComputeShader.usf
@@ -431,7 +431,7 @@ void SkinCacheUpdateBatchCS(uint3 GroupID : SV_GroupID,
 #if GPUSKIN_APEX_CLOTH
 	// Transform simulated position from cloth space into local space and blend with skinned position
     float4 TransformedSimulatedPos = mul(float4(Intermediates.SimulatedPosition.xyz, 1), ClothToLocal);
-    SPos = lerp(SPos.xyz, TransformedSimulatedPos, ClothWeight);
+    SPos = lerp(SPos.xyz, TransformedSimulatedPos.xyz, ClothWeight);
 #endif
 
 	float3 TangentX = normalize(mul((float3x3)Intermediates.BlendMatrix, LocalTangentX));
diff --git a/Engine/Source/Runtime/Engine/Private/SkeletalRenderGPUSkin.cpp b/Engine/Source/Runtime/Engine/Private/SkeletalRenderGPUSkin.cpp
--- a/Engine/Source/Runtime/Engine/Private/SkeletalRenderGPUSkin.cpp
+++ b/Engine/Source/Runtime/Engine/Private/SkeletalRenderGPUSkin.cpp
@@ -753,7 +753,6 @@ void FSkeletalMeshObjectGPUSkin::ProcessUpdatedDynamicData(EGPUSkinCacheEntryMod
 				bNeedFence = ClothShaderData.UpdateClothSimulData(RHICmdList, SimData->Positions, SimData->Normals, FrameNumberToPrepare, FeatureLevel) || bNeedFence;
 				// Transform from cloth space to local space. Cloth space is relative to cloth root bone, local space is component space.
 				ClothShaderData.GetClothToLocalForWriting(FrameNumberToPrepare) = FMatrix44f(SimData->ComponentRelativeTransform.ToMatrixWithScale());
-				ClothShaderData.GetClothLocalToWorldForWriting(FrameNumberToPrepare) = FMatrix44f(SimData->ComponentRelativeTransform.ToMatrixWithScale() * DynamicData->ClothObjectLocalToWorld);
 			}
 #endif // WITH_APEX_CLOTHING || WITH_CHAOS_CLOTHING
 
@@ -1475,7 +1474,6 @@ void FSkeletalMeshObjectGPUSkin::RefreshClothingTransforms(const FMatrix& InNewL
 					if(FClothSimulData* SimData = DynamicData->ClothingSimData.Find(ActorIdx))
 					{
 						ClothShaderData.GetClothToLocalForWriting(FrameNumber) = FMatrix44f(SimData->ComponentRelativeTransform.ToMatrixWithScale());
-						ClothShaderData.GetClothLocalToWorldForWriting(FrameNumber) = FMatrix44f(SimData->ComponentRelativeTransform.ToMatrixWithScale() * InNewLocalToWorld);
 					}
 				}
 			}
diff --git a/Engine/Source/Runtime/Engine/Public/GPUSkinVertexFactory.h b/Engine/Source/Runtime/Engine/Public/GPUSkinVertexFactory.h
--- a/Engine/Source/Runtime/Engine/Public/GPUSkinVertexFactory.h
+++ b/Engine/Source/Runtime/Engine/Public/GPUSkinVertexFactory.h
@@ -742,27 +742,6 @@ public:
 			return ClothToLocal[Index];
 		}
 
-		FMatrix44f& GetClothLocalToWorldForWriting(uint32 FrameNumber)
-		{
-			uint32 Index = GetOldestIndex(FrameNumber);
-			Index = (BufferFrameNumber[0] == FrameNumber) ? 0 : Index;
-			Index = (BufferFrameNumber[1] == FrameNumber) ? 1 : Index;
-
-			return ClothLocalToWorld[Index];
-		}
-
-		const FMatrix44f& GetClothLocalToWorldForReading(bool bPrevious, uint32 FrameNumber) const
-		{
-			int32 Index = GetMostRecentIndex(FrameNumber);
-
-			if(bPrevious && DoWeHavePreviousData())
-			{
-				Index = 1 - Index;
-			}
-
-			return ClothLocalToWorld[Index];
-		}
-
 		/**
 		 * weight to blend between simulated positions and key-framed poses
 		 * if ClothBlendWeight is 1.0, it shows only simulated positions and if it is 0.0, it shows only key-framed animation
@@ -780,7 +759,6 @@ public:
 		/**
 		 * Matrix to apply to positions/normals
 		 */
-		FMatrix44f ClothLocalToWorld[2];
 		FMatrix44f ClothToLocal[2];
 
 		/** Whether to double buffer. */
@@ -859,8 +837,6 @@ public:
 			BufferFrameNumber[0] = -1;
 			BufferFrameNumber[1] = -1;
 
-			ClothLocalToWorld[0] = FMatrix44f::Identity;
-			ClothLocalToWorld[1] = FMatrix44f::Identity;
 			ClothToLocal[0] = FMatrix44f::Identity;
 			ClothToLocal[1] = FMatrix44f::Identity;
 
