Adding TRACE_CPUPROFILER_EVENT_SCOPE to the RecompileShaders exec command.

diff --git a/Engine/Source/Runtime/Engine/Private/ShaderCompiler/ShaderCompiler.cpp b/Engine/Source/Runtime/Engine/Private/ShaderCompiler/ShaderCompiler.cpp
--- a/Engine/Source/Runtime/Engine/Private/ShaderCompiler/ShaderCompiler.cpp
+++ b/Engine/Source/Runtime/Engine/Private/ShaderCompiler/ShaderCompiler.cpp
@@ -6126,6 +6126,7 @@ bool RecompileShaders(const TCHAR* Cmd, FOutputDevice& Ar)
 	FString FlagStr(FParse::Token(Cmd, 0));
 	if( FlagStr.Len() > 0 )
 	{
+		TRACE_CPUPROFILER_EVENT_SCOPE(RecompileShaders);
 		GWarn->BeginSlowTask( NSLOCTEXT("ShaderCompilingManager", "BeginRecompilingShadersTask", "Recompiling shaders"), true );
 
 		// Flush the shader file cache so that any changes to shader source files will be detected
