Scene queries automatically use PT data when run as fixed tick

diff --git a/Engine/Source/Runtime/Engine/Private/Collision/CollisionConversions.h b/Engine/Source/Runtime/Engine/Private/Collision/CollisionConversions.h
--- a/Engine/Source/Runtime/Engine/Private/Collision/CollisionConversions.h
+++ b/Engine/Source/Runtime/Engine/Private/Collision/CollisionConversions.h
@@ -83,6 +83,7 @@ bool IsBlocking(const FPhysicsShape& PShape, const FCollisionFilterData& QueryFi
  * @return	OutOverlaps	Converted data
  */
 bool ConvertOverlapResults(int32 NumOverlaps, FHitOverlap* POverlapResults, const FCollisionFilterData& QueryFilter, TArray<FOverlapResult>& OutOverlaps);
+bool ConvertOverlapResults(int32 NumOverlaps, FPTOverlapHit* POverlapResults, const FCollisionFilterData& QueryFilter, TArray<FOverlapResult>& OutOverlaps);
 
 
 struct FCompareFHitResultTime
diff --git a/Engine/Source/Runtime/Engine/Private/Collision/SceneQuery.cpp b/Engine/Source/Runtime/Engine/Private/Collision/SceneQuery.cpp
--- a/Engine/Source/Runtime/Engine/Private/Collision/SceneQuery.cpp
+++ b/Engine/Source/Runtime/Engine/Private/Collision/SceneQuery.cpp
@@ -309,8 +309,39 @@ struct TSQTraits
 	constexpr static bool IsSweep() { return GeometryQuery == ESweepOrRay::Sweep;  }
 };
 
+enum class EThreadQueryContext
+{
+	GTData,		//use interpolated GT data
+	PTDataWithGTObjects,	//use pt data, but convert back to GT when possible
+	PTOnlyData,	//use only the PT data and don't try to convert anything back to GT
+};
+
+EThreadQueryContext GetThreadQueryContext(const Chaos::FPhysicsSolver& Solver)
+{
+	if (Solver.IsGameThreadFrozen())
+	{
+		//If the game thread is frozen the solver is currently in fixed tick mode (i.e. fixed tick callbacks are being executed on GT)
+		if (IsInGameThread() || IsInParallelGameThread())
+		{
+			//Since we are on GT or parallel GT we must be in fixed tick, so use PT data and convert back to GT where possible
+			return EThreadQueryContext::PTDataWithGTObjects;
+		}
+		else
+		{
+			//The solver can't be running since it's calling fixed tick callbacks on gt, so it must be an unrelated thread task (audio, animation, etc...) so just use interpolated gt data
+			return EThreadQueryContext::GTData;
+		}
+	}
+	else
+	{
+		//TODO: need a way to know we are on a physics thread task (this isn't supported yet)
+		//For now just use interpolated data
+		return EThreadQueryContext::GTData;
+	}
+}
+
 template <typename Traits, typename TGeomInputs>
-bool TSceneCastCommon(const UWorld* World, typename Traits::TOutHits& OutHits, const TGeomInputs& GeomInputs, const FVector Start, const FVector End, ECollisionChannel TraceChannel, const struct FCollisionQueryParams& Params, const struct FCollisionResponseParams& ResponseParams, const struct FCollisionObjectQueryParams& ObjectParams)
+bool TSceneCastCommonImp(const UWorld* World, typename Traits::TOutHits& OutHits, const TGeomInputs& GeomInputs, const FVector Start, const FVector End, ECollisionChannel TraceChannel, const struct FCollisionQueryParams& Params, const struct FCollisionResponseParams& ResponseParams, const struct FCollisionObjectQueryParams& ObjectParams)
 {
 	FScopeCycleCounter Counter(Params.StatId);
 	STARTQUERYTIMER();
@@ -320,11 +351,6 @@ bool TSceneCastCommon(const UWorld* World, typename Traits::TOutHits& OutHits, c
 		Traits::ResetOutHits(OutHits, Start, End);
 	}
 
-	if ((World == NULL) || (World->GetPhysicsScene() == NULL))
-	{
-		return false;
-	}
-
 	// Track if we get any 'blocking' hits
 	bool bHaveBlockingHit = false;
 
@@ -413,6 +439,25 @@ bool TSceneCastCommon(const UWorld* World, typename Traits::TOutHits& OutHits, c
 	return bHaveBlockingHit;
 }
 
+template <typename Traits, typename PTTraits, typename TGeomInputs>
+bool TSceneCastCommon(const UWorld* World, typename Traits::TOutHits& OutHits, const TGeomInputs& GeomInputs, const FVector Start, const FVector End, ECollisionChannel TraceChannel, const struct FCollisionQueryParams& Params, const struct FCollisionResponseParams& ResponseParams, const struct FCollisionObjectQueryParams& ObjectParams)
+{
+	if ((World == NULL) || (World->GetPhysicsScene() == NULL))
+	{
+		return false;
+	}
+
+	const EThreadQueryContext ThreadContext = GetThreadQueryContext(*World->GetPhysicsScene()->GetSolver());
+	if(ThreadContext == EThreadQueryContext::GTData)
+	{
+		return TSceneCastCommonImp<Traits, TGeomInputs>(World, OutHits, GeomInputs, Start, End, TraceChannel, Params, ResponseParams, ObjectParams);
+	}
+	else
+	{
+		return TSceneCastCommonImp<PTTraits, TGeomInputs>(World, OutHits, GeomInputs, Start, End, TraceChannel, Params, ResponseParams, ObjectParams);
+	}
+}
+
 //////////////////////////////////////////////////////////////////////////
 // RAYCAST
 
@@ -423,8 +468,10 @@ bool FGenericPhysicsInterface::RaycastTest(const UWorld* World, const FVector St
 	CSV_SCOPED_TIMING_STAT(SceneQuery, RaycastTest);
 
 	using TCastTraits = TSQTraits<FHitRaycast, ESweepOrRay::Raycast, ESingleMultiOrTest::Test>;
+	using TPTCastTraits = TSQTraits<FPTRaycastHit, ESweepOrRay::Raycast, ESingleMultiOrTest::Test>;
 	FHitResult DummyHit(NoInit);
-	return TSceneCastCommon<TCastTraits>(World, DummyHit, FRaycastSQAdditionalInputs(), Start, End, TraceChannel, Params, ResponseParams, ObjectParams);
+
+	return TSceneCastCommon<TCastTraits, TPTCastTraits>(World, DummyHit, FRaycastSQAdditionalInputs(), Start, End, TraceChannel, Params, ResponseParams, ObjectParams);
 }
 
 bool FGenericPhysicsInterface::RaycastSingle(const UWorld* World, struct FHitResult& OutHit, const FVector Start, const FVector End, ECollisionChannel TraceChannel, const struct FCollisionQueryParams& Params, const struct FCollisionResponseParams& ResponseParams, const struct FCollisionObjectQueryParams& ObjectParams)
@@ -433,8 +480,9 @@ bool FGenericPhysicsInterface::RaycastSingle(const UWorld* World, struct FHitRes
 	SCOPE_CYCLE_COUNTER(STAT_Collision_RaycastSingle);
 	CSV_SCOPED_TIMING_STAT(SceneQuery, RaycastSingle);
 
-	using TCastTraits = TSQTraits<FHitRaycast, ESweepOrRay::Raycast, ESingleMultiOrTest::Single>;
-	return TSceneCastCommon<TCastTraits>(World, OutHit, FRaycastSQAdditionalInputs(), Start, End, TraceChannel, Params, ResponseParams, ObjectParams);
+	using TCastTraits = TSQTraits<FRaycastHit, ESweepOrRay::Raycast, ESingleMultiOrTest::Single>;
+	using TPTCastTraits = TSQTraits<FPTRaycastHit, ESweepOrRay::Raycast, ESingleMultiOrTest::Single>;
+	return TSceneCastCommon<TCastTraits, TPTCastTraits>(World, OutHit, FRaycastSQAdditionalInputs(), Start, End, TraceChannel, Params, ResponseParams, ObjectParams);
 }
 
 
@@ -446,7 +494,8 @@ bool FGenericPhysicsInterface::RaycastMulti(const UWorld* World, TArray<struct F
 	CSV_SCOPED_TIMING_STAT(SceneQuery, RaycastMultiple);
 
 	using TCastTraits = TSQTraits<FHitRaycast, ESweepOrRay::Raycast, ESingleMultiOrTest::Multi>;
-	return TSceneCastCommon<TCastTraits> (World, OutHits, FRaycastSQAdditionalInputs(), Start, End, TraceChannel, Params, ResponseParams, ObjectParams);
+	using TPTCastTraits = TSQTraits<FPTRaycastHit, ESweepOrRay::Raycast, ESingleMultiOrTest::Multi>;
+	return TSceneCastCommon<TCastTraits, TPTCastTraits> (World, OutHits, FRaycastSQAdditionalInputs(), Start, End, TraceChannel, Params, ResponseParams, ObjectParams);
 }
 
 //////////////////////////////////////////////////////////////////////////
@@ -459,8 +508,10 @@ bool FGenericPhysicsInterface::GeomSweepTest(const UWorld* World, const struct F
 	CSV_SCOPED_TIMING_STAT(SceneQuery, GeomSweepTest);
 
 	using TCastTraits = TSQTraits<FHitSweep, ESweepOrRay::Sweep, ESingleMultiOrTest::Test>;
+	using TPTCastTraits = TSQTraits<FPTSweepHit, ESweepOrRay::Sweep, ESingleMultiOrTest::Test>;
 	FHitResult DummyHit(NoInit);
-	return TSceneCastCommon<TCastTraits>(World, DummyHit, FGeomSQAdditionalInputs(CollisionShape, Rot), Start, End, TraceChannel, Params, ResponseParams, ObjectParams);
+
+	return TSceneCastCommon<TCastTraits, TPTCastTraits>(World, DummyHit, FGeomSQAdditionalInputs(CollisionShape, Rot), Start, End, TraceChannel, Params, ResponseParams, ObjectParams);
 }
 
 bool FGenericPhysicsInterface::GeomSweepSingle(const UWorld* World, const struct FCollisionShape& CollisionShape, const FQuat& Rot, FHitResult& OutHit, FVector Start, FVector End, ECollisionChannel TraceChannel, const struct FCollisionQueryParams& Params, const struct FCollisionResponseParams& ResponseParams, const struct FCollisionObjectQueryParams& ObjectParams)
@@ -469,8 +520,9 @@ bool FGenericPhysicsInterface::GeomSweepSingle(const UWorld* World, const struct
 	SCOPE_CYCLE_COUNTER(STAT_Collision_GeomSweepSingle);
 	CSV_SCOPED_TIMING_STAT(SceneQuery, GeomSweepSingle);
 
-	using TCastTraits = TSQTraits<FHitSweep, ESweepOrRay::Sweep, ESingleMultiOrTest::Single>;
-	return TSceneCastCommon<TCastTraits>(World, OutHit, FGeomSQAdditionalInputs(CollisionShape, Rot), Start, End, TraceChannel, Params, ResponseParams, ObjectParams);
+	using TCastTraits = TSQTraits<FSweepHit, ESweepOrRay::Sweep, ESingleMultiOrTest::Single>;
+	using TPTCastTraits = TSQTraits<FPTSweepHit, ESweepOrRay::Sweep, ESingleMultiOrTest::Single>;
+	return TSceneCastCommon<TCastTraits, TPTCastTraits>(World, OutHit, FGeomSQAdditionalInputs(CollisionShape, Rot), Start, End, TraceChannel, Params, ResponseParams, ObjectParams);
 }
 
 template<>
@@ -481,7 +533,8 @@ bool FGenericPhysicsInterface::GeomSweepMulti(const UWorld* World, const FPhysic
 	CSV_SCOPED_TIMING_STAT(SceneQuery, GeomSweepMultiple);
 
 	using TCastTraits = TSQTraits<FHitSweep, ESweepOrRay::Sweep, ESingleMultiOrTest::Multi>;
-	return TSceneCastCommon<TCastTraits>(World, OutHits, FGeomCollectionSQAdditionalInputs(InGeom, InGeomRot), Start, End, TraceChannel, Params, ResponseParams, ObjectParams);
+	using TPTCastTraits = TSQTraits<FPTSweepHit, ESweepOrRay::Sweep, ESingleMultiOrTest::Multi>;
+	return TSceneCastCommon<TCastTraits, TPTCastTraits>(World, OutHits, FGeomCollectionSQAdditionalInputs(InGeom, InGeomRot), Start, End, TraceChannel, Params, ResponseParams, ObjectParams);
 }
 
 template<>
@@ -492,7 +545,8 @@ bool FGenericPhysicsInterface::GeomSweepMulti(const UWorld* World, const FCollis
 	CSV_SCOPED_TIMING_STAT(SceneQuery, GeomSweepMultiple);
 
 	using TCastTraits = TSQTraits<FHitSweep, ESweepOrRay::Sweep, ESingleMultiOrTest::Multi>;
-	return TSceneCastCommon<TCastTraits>(World, OutHits, FGeomSQAdditionalInputs(InGeom, InGeomRot), Start, End, TraceChannel, Params, ResponseParams, ObjectParams);
+	using TPTCastTraits = TSQTraits<FPTSweepHit, ESweepOrRay::Sweep, ESingleMultiOrTest::Multi>;
+	return TSceneCastCommon<TCastTraits, TPTCastTraits>(World, OutHits, FGeomSQAdditionalInputs(InGeom, InGeomRot), Start, End, TraceChannel, Params, ResponseParams, ObjectParams);
 }
 
 
@@ -510,7 +564,7 @@ namespace EQueryInfo
 	};
 }
 
-template <EQueryInfo::Type InfoType, typename TCollisionAnalyzerType>
+template <typename TOverlapHit, EQueryInfo::Type InfoType, typename TCollisionAnalyzerType>
 bool GeomOverlapMultiImp(const UWorld* World, const FPhysicsGeometry& Geom, const TCollisionAnalyzerType& CollisionAnalyzerType, const FTransform& GeomPose, TArray<FOverlapResult>& OutOverlaps, ECollisionChannel TraceChannel, const struct FCollisionQueryParams& Params, const struct FCollisionResponseParams& ResponseParams, const struct FCollisionObjectQueryParams& ObjectParams)
 {
 	FScopeCycleCounter Counter(Params.StatId);
@@ -539,7 +593,7 @@ bool GeomOverlapMultiImp(const UWorld* World, const FPhysicsGeometry& Geom, cons
 		{
 			QueryFlags = QueryFlags | EQueryFlags::SkipNarrowPhase;
 		}
-		FDynamicHitBuffer<FHitOverlap> OverlapBuffer;
+		FDynamicHitBuffer<TOverlapHit> OverlapBuffer;
 
 		FQueryDebugParams DebugParams;
 #if !(UE_BUILD_TEST || UE_BUILD_SHIPPING) && WITH_CHAOS
@@ -552,7 +606,7 @@ bool GeomOverlapMultiImp(const UWorld* World, const FPhysicsGeometry& Geom, cons
 		FPhysicsCommand::ExecuteRead(&PhysScene, [&]()
 		{
 			{
-				FScopedSQHitchRepeater<TRemoveReference<decltype(OverlapBuffer)>::Type> HitchRepeater(OverlapBuffer, QueryCallback, FHitchDetectionInfo(GeomPose, TraceChannel, Params));
+				FScopedSQHitchRepeater<typename TRemoveReference<decltype(OverlapBuffer)>::Type> HitchRepeater(OverlapBuffer, QueryCallback, FHitchDetectionInfo(GeomPose, TraceChannel, Params));
 				do
 				{
 					LowLevelOverlap(PhysScene, Geom, GeomPose, HitchRepeater.GetBuffer(), QueryFlags, Filter, MakeQueryFilterData(Filter, QueryFlags, Params), &QueryCallback, DebugParams);
@@ -601,6 +655,25 @@ bool GeomOverlapMultiImp(const UWorld* World, const FPhysicsGeometry& Geom, cons
 	return bHaveBlockingHit;
 }
 
+template <EQueryInfo::Type InfoType, typename TCollisionAnalyzerType>
+bool GeomOverlapMultiHelper(const UWorld* World, const FPhysicsGeometry& Geom, const TCollisionAnalyzerType& CollisionAnalyzerType, const FTransform& GeomPose, TArray<FOverlapResult>& OutOverlaps, ECollisionChannel TraceChannel, const struct FCollisionQueryParams& Params, const struct FCollisionResponseParams& ResponseParams, const struct FCollisionObjectQueryParams& ObjectParams)
+{
+	if ((World == NULL) || (World->GetPhysicsScene() == NULL))
+	{
+		return false;
+	}
+
+	const EThreadQueryContext ThreadContext = GetThreadQueryContext(*World->GetPhysicsScene()->GetSolver());
+	if (ThreadContext == EThreadQueryContext::GTData)
+	{
+		return GeomOverlapMultiImp<FOverlapHit, InfoType, TCollisionAnalyzerType>(World, Geom, CollisionAnalyzerType, GeomPose, OutOverlaps, TraceChannel, Params, ResponseParams, ObjectParams);
+	}
+	else
+	{
+		return GeomOverlapMultiImp<FPTOverlapHit, InfoType, TCollisionAnalyzerType>(World, Geom, CollisionAnalyzerType, GeomPose, OutOverlaps, TraceChannel, Params, ResponseParams, ObjectParams);
+	}
+}
+
 bool FGenericPhysicsInterface::GeomOverlapBlockingTest(const UWorld* World, const struct FCollisionShape& CollisionShape, const FVector& Pos, const FQuat& Rot, ECollisionChannel TraceChannel, const struct FCollisionQueryParams& Params, const struct FCollisionResponseParams& ResponseParams, const struct FCollisionObjectQueryParams& ObjectParams)
 {
 	SCOPE_CYCLE_COUNTER(STAT_Collision_SceneQueryTotal);
@@ -610,7 +683,7 @@ bool FGenericPhysicsInterface::GeomOverlapBlockingTest(const UWorld* World, cons
 	TArray<FOverlapResult> Overlaps;	//needed only for template shared code
 	FTransform GeomTransform(Rot, Pos);
 	FPhysicsShapeAdapter Adaptor(GeomTransform.GetRotation(), CollisionShape);
-	return GeomOverlapMultiImp<EQueryInfo::IsBlocking>(World, Adaptor.GetGeometry(), CollisionShape, Adaptor.GetGeomPose(GeomTransform.GetTranslation()), Overlaps, TraceChannel, Params, ResponseParams, ObjectParams);
+	return GeomOverlapMultiHelper<EQueryInfo::IsBlocking>(World, Adaptor.GetGeometry(), CollisionShape, Adaptor.GetGeomPose(GeomTransform.GetTranslation()), Overlaps, TraceChannel, Params, ResponseParams, ObjectParams);
 }
 
 bool FGenericPhysicsInterface::GeomOverlapAnyTest(const UWorld* World, const struct FCollisionShape& CollisionShape, const FVector& Pos, const FQuat& Rot, ECollisionChannel TraceChannel, const struct FCollisionQueryParams& Params, const struct FCollisionResponseParams& ResponseParams, const struct FCollisionObjectQueryParams& ObjectParams)
@@ -622,7 +695,7 @@ bool FGenericPhysicsInterface::GeomOverlapAnyTest(const UWorld* World, const str
 	TArray<FOverlapResult> Overlaps;	//needed only for template shared code
 	FTransform GeomTransform(Rot, Pos);
 	FPhysicsShapeAdapter Adaptor(GeomTransform.GetRotation(), CollisionShape);
-	return GeomOverlapMultiImp<EQueryInfo::IsAnything>(World, Adaptor.GetGeometry(), CollisionShape, Adaptor.GetGeomPose(GeomTransform.GetTranslation()), Overlaps, TraceChannel, Params, ResponseParams, ObjectParams);
+	return GeomOverlapMultiHelper<EQueryInfo::IsAnything>(World, Adaptor.GetGeometry(), CollisionShape, Adaptor.GetGeomPose(GeomTransform.GetTranslation()), Overlaps, TraceChannel, Params, ResponseParams, ObjectParams);
 }
 
 template<>
@@ -633,7 +706,7 @@ bool FGenericPhysicsInterface::GeomOverlapMulti(const UWorld* World, const FPhys
 	CSV_SCOPED_TIMING_STAT(SceneQuery, GeomOverlapMultiple);
 
 	FTransform GeomTransform(InRotation, InPosition);
-	return GeomOverlapMultiImp<EQueryInfo::GatherAll>(World, InGeom.GetGeometry(), InGeom, GeomTransform, OutOverlaps, TraceChannel, Params, ResponseParams, ObjectParams);
+	return GeomOverlapMultiHelper<EQueryInfo::GatherAll>(World, InGeom.GetGeometry(), InGeom, GeomTransform, OutOverlaps, TraceChannel, Params, ResponseParams, ObjectParams);
 }
 
 template<>
@@ -645,5 +718,5 @@ bool FGenericPhysicsInterface::GeomOverlapMulti(const UWorld* World, const FColl
 
 	FTransform GeomTransform(InRotation, InPosition);
 	FPhysicsShapeAdapter Adaptor(GeomTransform.GetRotation(), InGeom);
-	return GeomOverlapMultiImp<EQueryInfo::GatherAll>(World, Adaptor.GetGeometry(), InGeom, Adaptor.GetGeomPose(GeomTransform.GetTranslation()), OutOverlaps, TraceChannel, Params, ResponseParams, ObjectParams);
+	return GeomOverlapMultiHelper<EQueryInfo::GatherAll>(World, Adaptor.GetGeometry(), InGeom, Adaptor.GetGeomPose(GeomTransform.GetTranslation()), OutOverlaps, TraceChannel, Params, ResponseParams, ObjectParams);
 }
\ No newline at end of file
diff --git a/Engine/Source/Runtime/Engine/Private/Collision/SceneQueryLowLevel.cpp b/Engine/Source/Runtime/Engine/Private/Collision/SceneQueryLowLevel.cpp
--- a/Engine/Source/Runtime/Engine/Private/Collision/SceneQueryLowLevel.cpp
+++ b/Engine/Source/Runtime/Engine/Private/Collision/SceneQueryLowLevel.cpp
@@ -187,7 +187,9 @@ template <typename THitRaycast>
 void LowLevelRaycastImp(FPhysScene& Scene, const FVector& Start, const FVector& Dir, float DeltaMag, FPhysicsHitCallback<THitRaycast>& HitBuffer, EHitFlags OutputFlags, FQueryFlags QueryFlags, const FCollisionFilterData& Filter, const FQueryFilterData& QueryFilterData, ICollisionQueryFilterCallbackBase* QueryCallback, const FQueryDebugParams& DebugParams)
 {
 #if !defined(PHYSICS_INTERFACE_PHYSX) || !PHYSICS_INTERFACE_PHYSX
-	if (const auto& SolverAccelerationStructure = Scene.GetSpacialAcceleration())
+	constexpr bool bGTData = std::is_same<THitRaycast, FHitRaycast>::value;
+	const auto SolverAccelerationStructure = bGTData ? Scene.GetSpacialAcceleration() : Scene.GetSolver()->GetInternalAccelerationStructure_Internal();
+	if (SolverAccelerationStructure)
 	{
 		FChaosSQAccelerator SQAccelerator(*SolverAccelerationStructure);
 		double Time = 0.0;
@@ -196,7 +198,7 @@ void LowLevelRaycastImp(FPhysScene& Scene, const FVector& Start, const FVector&
 			SQAccelerator.Raycast(Start, Dir, DeltaMag, HitBuffer, OutputFlags, QueryFilterData, *QueryCallback, DebugParams);
 		}
 
-		if constexpr(std::is_same<THitRaycast, FHitRaycast>::value)
+		if constexpr(bGTData)
 		{
 			if (!!SerializeSQs && !!EnableRaycastSQCapture)
 			{
@@ -237,7 +239,9 @@ template <typename THitSweep>
 void LowLevelSweepImp(FPhysScene& Scene, const FPhysicsGeometry& QueryGeom, const FTransform& StartTM, const FVector& Dir, float DeltaMag, FPhysicsHitCallback<THitSweep>& HitBuffer, EHitFlags OutputFlags, FQueryFlags QueryFlags, const FCollisionFilterData& Filter, const FQueryFilterData& QueryFilterData, ICollisionQueryFilterCallbackBase* QueryCallback, const FQueryDebugParams& DebugParams)
 {
 #if !defined(PHYSICS_INTERFACE_PHYSX) || !PHYSICS_INTERFACE_PHYSX
-	if (const auto& SolverAccelerationStructure = Scene.GetSpacialAcceleration())
+	constexpr bool bGTData = std::is_same<THitSweep, FHitSweep>::value;
+	const auto SolverAccelerationStructure = bGTData ? Scene.GetSpacialAcceleration() : Scene.GetSolver()->GetInternalAccelerationStructure_Internal();
+	if (SolverAccelerationStructure)
 	{
 		FChaosSQAccelerator SQAccelerator(*SolverAccelerationStructure);
 		{
@@ -248,7 +252,7 @@ void LowLevelSweepImp(FPhysScene& Scene, const FPhysicsGeometry& QueryGeom, cons
 				SQAccelerator.Sweep(QueryGeom, StartTM, Dir, DeltaMag, HitBuffer, OutputFlags, QueryFilterData, *QueryCallback, DebugParams);
 			}
 
-			if constexpr (std::is_same<THitSweep, FHitSweep>::value)
+			if constexpr (bGTData)
 			{
 				if (!!SerializeSQs && !!EnableSweepSQCapture)
 				{
@@ -290,7 +294,9 @@ template <typename THitOverlap>
 void LowLevelOverlapImp(FPhysScene& Scene, const FPhysicsGeometry& QueryGeom, const FTransform& GeomPose, FPhysicsHitCallback<THitOverlap>& HitBuffer, FQueryFlags QueryFlags, const FCollisionFilterData& Filter, const FQueryFilterData& QueryFilterData, ICollisionQueryFilterCallbackBase* QueryCallback, const FQueryDebugParams& DebugParams)
 {
 #if !defined(PHYSICS_INTERFACE_PHYSX) || !PHYSICS_INTERFACE_PHYSX
-	if (const auto& SolverAccelerationStructure = Scene.GetSpacialAcceleration())
+	constexpr bool bGTData = std::is_same<THitOverlap, FHitOverlap>::value;
+	const auto SolverAccelerationStructure = bGTData ? Scene.GetSpacialAcceleration() : Scene.GetSolver()->GetInternalAccelerationStructure_Internal();
+	if (SolverAccelerationStructure)
 	{
 		FChaosSQAccelerator SQAccelerator(*SolverAccelerationStructure);
 		double Time = 0.0;
@@ -299,7 +305,7 @@ void LowLevelOverlapImp(FPhysScene& Scene, const FPhysicsGeometry& QueryGeom, co
 			SQAccelerator.Overlap(QueryGeom, GeomPose, HitBuffer, QueryFilterData, *QueryCallback, DebugParams);
 		}
 
-		if constexpr (std::is_same<THitOverlap, FHitOverlap>::value)
+		if constexpr (bGTData)
 		{
 			if (!!SerializeSQs && !!EnableOverlapSQCapture)
 			{
diff --git a/Engine/Source/Runtime/Experimental/Chaos/Private/Chaos/Framework/PhysicsSolverBase.cpp b/Engine/Source/Runtime/Experimental/Chaos/Private/Chaos/Framework/PhysicsSolverBase.cpp
--- a/Engine/Source/Runtime/Experimental/Chaos/Private/Chaos/Framework/PhysicsSolverBase.cpp
+++ b/Engine/Source/Runtime/Experimental/Chaos/Private/Chaos/Framework/PhysicsSolverBase.cpp
@@ -110,9 +110,9 @@ namespace Chaos
 		FFrozenGameThreadContextScope FrozenScope;	//Make sure we fire ensures if any physics GT data is used
 #endif
 
-		Solver.SetOnFrozenGameThread(true);
+		Solver.SetGameThreadFrozen(true);
 		Solver.ApplyCallbacks_Internal();
-		Solver.SetOnFrozenGameThread(false);
+		Solver.SetGameThreadFrozen(false);
 		
 	}
 
diff --git a/Engine/Source/Runtime/Experimental/Chaos/Public/Chaos/Framework/PhysicsSolverBase.h b/Engine/Source/Runtime/Experimental/Chaos/Public/Chaos/Framework/PhysicsSolverBase.h
--- a/Engine/Source/Runtime/Experimental/Chaos/Public/Chaos/Framework/PhysicsSolverBase.h
+++ b/Engine/Source/Runtime/Experimental/Chaos/Public/Chaos/Framework/PhysicsSolverBase.h
@@ -375,11 +375,11 @@ namespace Chaos
 
 		//Tells us if we're on the frozen game thread. This is needed for knowing which data to read/write to
 		//The IsInGameThread check is so that other threads (e.g audio thread) which might be running queries in parallel will continue to use the correct interpolated GT data
-		bool IsOnFrozenGameThread() const { return bOnFrozenGameThread && (IsInGameThread() || IsInParallelGameThread()); }
+		bool IsGameThreadFrozen() const { return bGameThreadFrozen; }
 
-		void SetOnFrozenGameThread(bool InOnFrozenGameThread)
+		void SetGameThreadFrozen(bool InGameThreadFrozen)
 		{
-			bOnFrozenGameThread = InOnFrozenGameThread;
+			bGameThreadFrozen = InGameThreadFrozen;
 		}
 
 		void ApplyCallbacks_Internal()
@@ -392,7 +392,7 @@ namespace Chaos
 				const FReal SimTime = GetSolverTime();
 				for (ISimCallbackObject* Callback : SimCallbackObjects)
 				{
-					if (Callback->RunOnFrozenGameThread() == bOnFrozenGameThread)
+					if (Callback->RunOnFrozenGameThread() == bGameThreadFrozen)
 					{
 						Callback->SetSimAndDeltaTime_Internal(SimTime, MLastDt);
 						Callback->PreSimulate_Internal();
@@ -547,7 +547,7 @@ namespace Chaos
 	FGraphEventRef PendingTasks;
 
 	bool bSolverHasFrozenGameThreadCallbacks = false;
-	bool bOnFrozenGameThread = false;
+	bool bGameThreadFrozen = false;
 	FReal MLastDt = FReal(0);
 	FReal MTime = FReal(0);
 
diff --git a/Engine/Source/Runtime/Experimental/Chaos/Public/PBDRigidsSolver.h b/Engine/Source/Runtime/Experimental/Chaos/Public/PBDRigidsSolver.h
--- a/Engine/Source/Runtime/Experimental/Chaos/Public/PBDRigidsSolver.h
+++ b/Engine/Source/Runtime/Experimental/Chaos/Public/PBDRigidsSolver.h
@@ -280,6 +280,10 @@ namespace Chaos
 		const FPerSolverFieldSystem& GetPerSolverField() const { return *PerSolverField; }
 
 		void UpdateExternalAccelerationStructure_External(ISpatialAccelerationCollection<FAccelerationStructureHandle,FReal,3>*& ExternalStructure);
+		const ISpatialAccelerationCollection<FAccelerationStructureHandle, FReal, 3>* GetInternalAccelerationStructure_Internal() const
+		{
+			return MEvolution->GetSpatialAcceleration();
+		}
 
 		/** Apply a solver configuration to this solver, set externally by the owner of a solver (see UPhysicsSettings for world solver settings) */
 		void ApplyConfig(const FChaosSolverConfiguration& InConfig);
