Fixing `recompileshaders changed` command which I broke in CL# 18848560.

- The call I added to early out of shader compilation was preventing `recompileshaders` to detect if there were outdated shaders and the command ended up not compiling any shaders.
- Since we are now caching uniform expressions again we need to ensure draw commands are re-created where necessary to avoid stale data and crashes.
- SubmitRemainingJobsForWorld now uses a FMaterialUpdateContext to ensure all the materials we attempt to call CacheShaders on will re-create any state needed as necessary.

diff --git a/Engine/Source/Runtime/Engine/Private/Materials/Material.cpp b/Engine/Source/Runtime/Engine/Private/Materials/Material.cpp
--- a/Engine/Source/Runtime/Engine/Private/Materials/Material.cpp
+++ b/Engine/Source/Runtime/Engine/Private/Materials/Material.cpp
@@ -2136,8 +2136,6 @@ void UMaterial::CacheResourceShadersForRendering(bool bRegenerateId, EMaterialSh
 
 	if (FApp::CanEverRender())
 	{
-		bool bCachedShaders{ false };
-
 		const EMaterialQualityLevel::Type ActiveQualityLevel = GetCachedScalabilityCVars().MaterialQualityLevel;
 		uint32 FeatureLevelsToCompile = GetFeatureLevelsToCompileForRendering();
 
@@ -2153,12 +2151,6 @@ void UMaterial::CacheResourceShadersForRendering(bool bRegenerateId, EMaterialSh
 			FMaterialResource* CurrentResource = FindOrCreateMaterialResource(MaterialResources, this, nullptr, FeatureLevel, ActiveQualityLevel);
 			check(CurrentResource);
 
-			// If this resource is complete there is nothing to do here.
-			if (CurrentResource->IsGameThreadShaderMapComplete())
-			{
-				continue;
-			}
-
 #if STORE_ONLY_ACTIVE_SHADERMAPS
 			if (CurrentResource && !CurrentResource->GetGameThreadShaderMap())
 			{
@@ -2179,7 +2171,6 @@ void UMaterial::CacheResourceShadersForRendering(bool bRegenerateId, EMaterialSh
 			ResourcesToCache.Reset();
 			ResourcesToCache.Add(CurrentResource);
 			CacheShadersForResources(ShaderPlatform, ResourcesToCache, PrecompileMode);
-			bCachedShaders = true;
 		}
 
 		FString AdditionalFormatToCache = GCompileMaterialsForShaderFormatCVar->GetString();
@@ -2199,10 +2190,7 @@ void UMaterial::CacheResourceShadersForRendering(bool bRegenerateId, EMaterialSh
 			}
 		}
 
-		if (bCachedShaders)
-		{
-			RecacheUniformExpressions(true);
-		}
+		RecacheUniformExpressions(true);
 	}
 
 	FMaterial::DeferredDeleteArray(ResourcesToFree);
diff --git a/Engine/Source/Runtime/Engine/Private/Materials/MaterialInstance.cpp b/Engine/Source/Runtime/Engine/Private/Materials/MaterialInstance.cpp
--- a/Engine/Source/Runtime/Engine/Private/Materials/MaterialInstance.cpp
+++ b/Engine/Source/Runtime/Engine/Private/Materials/MaterialInstance.cpp
@@ -1985,8 +1985,6 @@ void UMaterialInstance::CacheResourceShadersForRendering(EMaterialShaderPrecompi
 		check(IsA(UMaterialInstanceConstant::StaticClass()));
 		UMaterial* BaseMaterial = GetMaterial();
 
-		bool bCachedShaders{ false };
-
 		uint32 FeatureLevelsToCompile = GetFeatureLevelsToCompileForRendering();
 		const EMaterialQualityLevel::Type ActiveQualityLevel = GetCachedScalabilityCVars().MaterialQualityLevel;
 
@@ -2002,12 +2000,6 @@ void UMaterialInstance::CacheResourceShadersForRendering(EMaterialShaderPrecompi
 			FMaterialResource* CurrentResource = FindOrCreateMaterialResource(StaticPermutationMaterialResources, BaseMaterial, this, FeatureLevel, ActiveQualityLevel);
 			check(CurrentResource);
 
-			// If this resource is complete there is nothing to do here.
-			if (CurrentResource->IsGameThreadShaderMapComplete())
-			{
-				continue;
-			}
-
 #if STORE_ONLY_ACTIVE_SHADERMAPS
 			if (!CurrentResource->GetGameThreadShaderMap())
 			{
@@ -2028,11 +2020,6 @@ void UMaterialInstance::CacheResourceShadersForRendering(EMaterialShaderPrecompi
 			ResourcesToCache.Reset();
 			ResourcesToCache.Add(CurrentResource);
 			CacheShadersForResources(ShaderPlatform, ResourcesToCache, PrecompileMode);
-			bCachedShaders = true;
-		}
-
-		if (bCachedShaders)
-		{
 			RecacheUniformExpressions(true);
 		}
 	}
diff --git a/Engine/Source/Runtime/Engine/Private/Materials/MaterialInterface.cpp b/Engine/Source/Runtime/Engine/Private/Materials/MaterialInterface.cpp
--- a/Engine/Source/Runtime/Engine/Private/Materials/MaterialInterface.cpp
+++ b/Engine/Source/Runtime/Engine/Private/Materials/MaterialInterface.cpp
@@ -453,6 +453,8 @@ void UMaterialInterface::SubmitRemainingJobsForWorld(UWorld* World, EMaterialSha
 	TSet<UMaterialInterface*> MaterialsToCache;
 	FObjectCacheContextScope ObjectCacheScope;
 
+	FMaterialUpdateContext UpdateContext(FMaterialUpdateContext::EOptions::SyncWithRenderingThread);
+
 	for (UPrimitiveComponent* PrimitiveComponent : ObjectCacheScope.GetContext().GetPrimitiveComponents())
 	{
 		if (World && !World->ContainsActor(PrimitiveComponent->GetOwner()))
@@ -486,6 +488,9 @@ void UMaterialInterface::SubmitRemainingJobsForWorld(UWorld* World, EMaterialSha
 
 	for (UMaterialInterface* Material : MaterialsToCache)
 	{
+		// This is needed because CacheShaders blindly recreates uniform buffers
+		// which can only be done if the draw command is going to be re-cached.
+		UpdateContext.AddMaterialInterface(Material);
 		Material->CacheShaders(CompileMode);
 	}
 }
