Rework the deferred running of construction scripts when async asset compilation is running to used a throttling asset compiling manager
This should fix the case where the post asset compile callback to run the deferred script was getting triggered from undesirable points

diff --git a/Engine/Source/Runtime/Engine/Classes/Engine/Level.h b/Engine/Source/Runtime/Engine/Classes/Engine/Level.h
--- a/Engine/Source/Runtime/Engine/Classes/Engine/Level.h
+++ b/Engine/Source/Runtime/Engine/Classes/Engine/Level.h
@@ -38,7 +38,6 @@ class FLevelPartitionOperationScope;
 class FRegisterComponentContext;
 class SNotificationItem;
 class UActorFolder;
-struct FAssetCompileData;
 
 UINTERFACE()
 class ULevelPartitionInterface : public UInterface
@@ -828,9 +827,6 @@ private:
 	/** Temporary array containing actor folder objects manually loaded from their external packages (only used while loading the level). */
 	UPROPERTY(Transient)
 	TArray<TObjectPtr<UActorFolder>> LoadedExternalActorFolders;
-
-	/** Actors awaiting outstanding asset compilation prior to running their non trivial construction scripts after loading this level. */
-	TArray<TWeakObjectPtr<AActor>> PendingConstructionScriptActors;
 #endif // #if WITH_EDITORONLY_DATA
 
 	enum class ERouteActorInitializationState : uint8
@@ -1326,9 +1322,6 @@ private:
 	 */
 	bool DeferRunningConstructionScripts(AActor* InActor);
 
-	/** Callback for asset compilation to determine if we can run deferred construction scripts. */
-	void OnAssetPostCompile(const TArray<FAssetCompileData>&);
-
 	/** Attempts to detect and fix any issues with the level script blueprint and associated objects */
 	void RepairLevelScript();
 
diff --git a/Engine/Source/Runtime/Engine/Private/ActorDeferredScriptManager.cpp b/Engine/Source/Runtime/Engine/Private/ActorDeferredScriptManager.cpp
new file mode 100644
--- /dev/null
+++ b/Engine/Source/Runtime/Engine/Private/ActorDeferredScriptManager.cpp
@@ -0,0 +1,99 @@
+// Copyright Epic Games, Inc. All Rights Reserved.
+
+#include "ActorDeferredScriptManager.h"
+#include "CoreGlobals.h"
+#include "GameFramework/Actor.h"
+#include "StaticMeshCompiler.h"
+
+#if WITH_EDITOR
+
+#define LOCTEXT_NAMESPACE "ActorDeferredScriptManager"
+
+
+FActorDeferredScriptManager& FActorDeferredScriptManager::Get()
+{
+	static FActorDeferredScriptManager Singleton;
+	return Singleton;
+}
+
+void FActorDeferredScriptManager::Shutdown()
+{
+}
+
+FActorDeferredScriptManager::FActorDeferredScriptManager()
+	: Notification(GetAssetNameFormat())
+{
+}
+
+FName FActorDeferredScriptManager::GetStaticAssetTypeName()
+{
+	return TEXT("UE-ActorConstructionScripts");
+}
+
+void FActorDeferredScriptManager::AddActor(AActor* InActor)
+{
+	PendingConstructionScriptActors.PushFirst(InActor);
+}
+
+FName FActorDeferredScriptManager::GetAssetTypeName() const
+{
+	return GetStaticAssetTypeName();
+}
+
+FTextFormat FActorDeferredScriptManager::GetAssetNameFormat() const
+{
+	return LOCTEXT("ScriptNameFormat", "{0}|plural(one=Construction Script,other=Construction Scripts)");
+}
+
+TArrayView<FName> FActorDeferredScriptManager::GetDependentTypeNames() const
+{
+	static FName DependentTypeNames[] =
+	{
+		// Construction scripts needs to wait until static meshes are done prior to running
+		FStaticMeshCompilingManager::GetStaticAssetTypeName()
+	};
+	return TArrayView<FName>(DependentTypeNames);
+}
+
+int32 FActorDeferredScriptManager::GetNumRemainingAssets() const
+{
+	return PendingConstructionScriptActors.Num();
+}
+
+void FActorDeferredScriptManager::ProcessAsyncTasks(bool bLimitExecutionTime)
+{
+	if (!PendingConstructionScriptActors.IsEmpty() && FStaticMeshCompilingManager::Get().GetNumRemainingMeshes() == 0)
+	{
+		double TickStartTime = FPlatformTime::Seconds();
+		const double MaxSecondsPerFrame = 0.016;
+		bool bHasTimeLeft = true;
+
+		// since this deferred run of construction script was supposed to be done during level load
+		// temporarily set the global flag to prevent dirtying the level package.
+		// @note: it would quite preferable if we would have a scoped context than touching a global variable... 
+		GIsEditorLoadingPackage = true;
+		while (!PendingConstructionScriptActors.IsEmpty() && bHasTimeLeft)
+		{
+			TWeakObjectPtr<AActor> WeakActor = PendingConstructionScriptActors.First();
+			PendingConstructionScriptActors.PopFirst();
+			if (AActor* Actor = WeakActor.Get())
+			{
+				Actor->RerunConstructionScripts();
+			}
+			bHasTimeLeft = bLimitExecutionTime ? ((FPlatformTime::Seconds() - TickStartTime) < MaxSecondsPerFrame) : true;
+		}
+		GIsEditorLoadingPackage = false;
+	}
+
+	UpdateCompilationNotification();
+}
+
+void FActorDeferredScriptManager::UpdateCompilationNotification()
+{
+	Notification.Update(GetNumRemainingAssets());
+}
+
+
+#undef LOCTEXT_NAMESPACE
+
+#endif // #if WITH_EDITOR
diff --git a/Engine/Source/Runtime/Engine/Private/AssetCompilingManager.cpp b/Engine/Source/Runtime/Engine/Private/AssetCompilingManager.cpp
--- a/Engine/Source/Runtime/Engine/Private/AssetCompilingManager.cpp
+++ b/Engine/Source/Runtime/Engine/Private/AssetCompilingManager.cpp
@@ -13,6 +13,7 @@ LLM_DEFINE_TAG(AssetCompilation, NAME_None, NAME_None, GET_STATFNAME(STAT_AssetC
 #include "Misc/CommandLine.h"
 #include "UObject/UObjectIterator.h"
 #include "HAL/IConsoleManager.h"
+#include "ActorDeferredScriptManager.h"
 #include "StaticMeshCompiler.h"
 #include "TextureCompiler.h"
 #include "Components/StaticMeshComponent.h"
@@ -201,6 +202,7 @@ FAssetCompilingManager::FAssetCompilingManager()
 	RegisterManager(&FStaticMeshCompilingManager::Get());
 	RegisterManager(&FSkeletalMeshCompilingManager::Get());
 	RegisterManager(&FTextureCompilingManager::Get());
+	RegisterManager(&FActorDeferredScriptManager::Get());
 #endif
 }
 
diff --git a/Engine/Source/Runtime/Engine/Private/Level.cpp b/Engine/Source/Runtime/Engine/Private/Level.cpp
--- a/Engine/Source/Runtime/Engine/Private/Level.cpp
+++ b/Engine/Source/Runtime/Engine/Private/Level.cpp
@@ -55,6 +55,7 @@ Level.cpp: Level-related functions
 #if WITH_EDITOR
 #include "AssetCompilingManager.h"
 #include "StaticMeshCompiler.h"
+#include "ActorDeferredScriptManager.h"
 #include "Components/InstancedStaticMeshComponent.h"
 #include "Kismet2/KismetEditorUtilities.h"
 #include "Kismet2/BlueprintEditorUtils.h"
@@ -73,6 +74,7 @@ Level.cpp: Level-related functions
 #include "Misc/MessageDialog.h"
 #include "ScopedTransaction.h"
 #include "EditorActorFolders.h"
+#include "UObject/MetaData.h"
 #endif
 #include "WorldPartition/WorldPartition.h"
 #include "Engine/LevelStreaming.h"
@@ -944,10 +946,6 @@ void ULevel::PostLoad()
 			}
 		}
 	}
-
-	// Register to the asset async compilation delegate to manage the deferring of non trivial construction scripts on actors
-	FAssetCompilingManager::Get().OnAssetPostCompileEvent().AddUObject(this, &ULevel::OnAssetPostCompile);
-
 #endif
 
 	// Ensure that the level is pointed to the owning world.  For streamed levels, this will be the world of the P map
@@ -1608,33 +1606,12 @@ bool ULevel::DeferRunningConstructionScripts(AActor* InActor)
 	if (FStaticMeshCompilingManager::Get().GetNumRemainingMeshes() &&
 		InActor->HasNonTrivialUserConstructionScript())
 	{
-		PendingConstructionScriptActors.Add(InActor);
+		FActorDeferredScriptManager::Get().AddActor(InActor);
 		return true;
 	}
 	return false;
 }
 
-void ULevel::OnAssetPostCompile(const TArray<FAssetCompileData>&)
-{
-	// Once there are no more outstanding asset (static meshes), run actors non trivial construction scripts
-	if (FStaticMeshCompilingManager::Get().GetNumRemainingMeshes() == 0)
-	{
-		// since this deferred run of construction script was supposed to be done during level load
-		// temporarily set the global flag to prevent dirtying the level package.
-		// @note: it would quite preferable if we would have a scoped context than touching a global variable... 
-		GIsEditorLoadingPackage = true;
-		for (TWeakObjectPtr<AActor> WeakActor : PendingConstructionScriptActors)
-		{
-			if (AActor* Actor = WeakActor.Get())
-			{
-				Actor->RerunConstructionScripts();
-			}
-		}
-		GIsEditorLoadingPackage = false;
-		PendingConstructionScriptActors.Empty();
-	}
-}
-
 void ULevel::CreateModelComponents()
 {
 	TRACE_CPUPROFILER_EVENT_SCOPE(ULevel::CreateModelComponents)
diff --git a/Engine/Source/Runtime/Engine/Public/ActorDeferredScriptManager.h b/Engine/Source/Runtime/Engine/Public/ActorDeferredScriptManager.h
new file mode 100644
--- /dev/null
+++ b/Engine/Source/Runtime/Engine/Public/ActorDeferredScriptManager.h
@@ -0,0 +1,57 @@
+// Copyright Epic Games, Inc. All Rights Reserved.
+#pragma once
+
+#include "UObject/WeakObjectPtr.h"
+#include "Containers/Deque.h"
+#include "AssetCompilingManager.h"
+#include "AsyncCompilationHelpers.h"
+
+class AActor;
+
+#if WITH_EDITOR
+
+class FActorDeferredScriptManager : IAssetCompilingManager
+{
+public:
+	ENGINE_API static FActorDeferredScriptManager& Get();
+
+	/**
+	 * Cancel any pending work and blocks until it is safe to shut down.
+	 */
+	ENGINE_API void Shutdown() override;
+
+	/** Get the name of the asset type this compiler handles */
+	ENGINE_API static FName GetStaticAssetTypeName();
+
+	/**
+	 * empty implementation since deferred actor scripts are not actually async
+	 */
+	ENGINE_API void FinishAllCompilation() override {}
+
+
+	/** 
+	 * Add an actor for the manager to run the construction of once dependent asset compilation are done (i.e Static Mesh)
+	 */
+	void AddActor(AActor* InActor);
+
+private:
+	friend class FAssetCompilingManager;
+
+	FActorDeferredScriptManager();
+
+	FName GetAssetTypeName() const override;
+	FTextFormat GetAssetNameFormat() const override;
+	TArrayView<FName> GetDependentTypeNames() const override;
+	int32 GetNumRemainingAssets() const override;
+	void ProcessAsyncTasks(bool bLimitExecutionTime = false) override;
+
+	void UpdateCompilationNotification();
+
+	/** Actors awaiting outstanding asset compilation prior to running their non trivial construction scripts after loading this level. */
+	TDeque<TWeakObjectPtr<AActor>> PendingConstructionScriptActors;
+
+	/** Notification for the amount of pending construction scripts to run */
+	FAsyncCompilationNotification Notification;
+};
+
+#endif
\ No newline at end of file
