Add an error when loading an ExternalActor package on a non-partionned world and when the package has no actor

diff --git a/Engine/Source/Runtime/Engine/Private/Level.cpp b/Engine/Source/Runtime/Engine/Private/Level.cpp
--- a/Engine/Source/Runtime/Engine/Private/Level.cpp
+++ b/Engine/Source/Runtime/Engine/Private/Level.cpp
@@ -934,6 +934,7 @@ void ULevel::PostLoad()
 
 				ActorPackage = LoadPackage(ActorPackage, *ActorPackageName, bPackageForPIE ? LOAD_PackageForPIE : LOAD_None, nullptr, &InstancingContext);
 
+				int32 PreviousActorCount = Actors.Num();
 				ForEachObjectWithPackage(ActorPackage, [this](UObject* PackageObject)
 				{
 					// There might be multiple actors per package in the case where an actor as a child actor component as we put child actor in the same package as their parent
@@ -943,6 +944,12 @@ void ULevel::PostLoad()
 					}
 					return true;
 				}, false);
+
+				bool bAddedPackageActors = Actors.Num() > PreviousActorCount;
+				if (!bAddedPackageActors)
+				{
+					UE_LOG(LogLevel, Error, TEXT("Failed to load Actor for External Actor Package %s"), *ActorPackageName);
+				}
 			}
 		}
 	}
