Support matrix types in new HLSL translator
Other various fixes/additions

diff --git a/Engine/Shaders/Private/LargeWorldCoordinates.ush b/Engine/Shaders/Private/LargeWorldCoordinates.ush
--- a/Engine/Shaders/Private/LargeWorldCoordinates.ush
+++ b/Engine/Shaders/Private/LargeWorldCoordinates.ush
@@ -176,6 +176,26 @@ FLWCScalar LWCGetComponent(FLWCVector4 V, int C) { return MakeLWCScalar(LWCGetTi
 #define LWCGetZ(V) LWCGetComponent(V, 2)
 #define LWCGetW(V) LWCGetComponent(V, 3)
 
+FLWCScalar LWCSwizzle(FLWCScalar V, int C0) { return V; }
+FLWCScalar LWCSwizzle(FLWCVector2 V, int C0) { return LWCGetComponent(V, C0); }
+FLWCScalar LWCSwizzle(FLWCVector3 V, int C0) { return LWCGetComponent(V, C0); }
+FLWCScalar LWCSwizzle(FLWCVector4 V, int C0) { return LWCGetComponent(V, C0); }
+
+FLWCVector2 LWCSwizzle(FLWCScalar V, int C0, int C1) { return MakeLWCVector(V, V); }
+FLWCVector2 LWCSwizzle(FLWCVector2 V, int C0, int C1) { return MakeLWCVector(LWCGetComponent(V, C0), LWCGetComponent(V, C1)); }
+FLWCVector2 LWCSwizzle(FLWCVector3 V, int C0, int C1) { return MakeLWCVector(LWCGetComponent(V, C0), LWCGetComponent(V, C1)); }
+FLWCVector2 LWCSwizzle(FLWCVector4 V, int C0, int C1) { return MakeLWCVector(LWCGetComponent(V, C0), LWCGetComponent(V, C1)); }
+
+FLWCVector3 LWCSwizzle(FLWCScalar V, int C0, int C1, int C2) { return MakeLWCVector(V, V, V); }
+FLWCVector3 LWCSwizzle(FLWCVector2 V, int C0, int C1, int C2) { return MakeLWCVector(LWCGetComponent(V, C0), LWCGetComponent(V, C1), LWCGetComponent(V, C2)); }
+FLWCVector3 LWCSwizzle(FLWCVector3 V, int C0, int C1, int C2) { return MakeLWCVector(LWCGetComponent(V, C0), LWCGetComponent(V, C1), LWCGetComponent(V, C2)); }
+FLWCVector3 LWCSwizzle(FLWCVector4 V, int C0, int C1, int C2) { return MakeLWCVector(LWCGetComponent(V, C0), LWCGetComponent(V, C1), LWCGetComponent(V, C2)); }
+
+FLWCVector4 LWCSwizzle(FLWCScalar V, int C0, int C1, int C2, int C3) { return MakeLWCVector(V, V, V, V); }
+FLWCVector4 LWCSwizzle(FLWCVector2 V, int C0, int C1, int C2, int C3) { return MakeLWCVector(LWCGetComponent(V, C0), LWCGetComponent(V, C1), LWCGetComponent(V, C2), LWCGetComponent(V, C3)); }
+FLWCVector4 LWCSwizzle(FLWCVector3 V, int C0, int C1, int C2, int C3) { return MakeLWCVector(LWCGetComponent(V, C0), LWCGetComponent(V, C1), LWCGetComponent(V, C2), LWCGetComponent(V, C3)); }
+FLWCVector4 LWCSwizzle(FLWCVector4 V, int C0, int C1, int C2, int C3) { return MakeLWCVector(LWCGetComponent(V, C0), LWCGetComponent(V, C1), LWCGetComponent(V, C2), LWCGetComponent(V, C3)); }
+
 float LWCToFloat(FLWCScalar Value)   { return LWCGetTileOffset(Value) + Value.Offset; }
 float2 LWCToFloat(FLWCVector2 Value) { return LWCGetTileOffset(Value) + Value.Offset; }
 float3 LWCToFloat(FLWCVector3 Value) { return LWCGetTileOffset(Value) + Value.Offset; }
diff --git a/Engine/Shaders/Private/MaterialTemplate.ush b/Engine/Shaders/Private/MaterialTemplate.ush
--- a/Engine/Shaders/Private/MaterialTemplate.ush
+++ b/Engine/Shaders/Private/MaterialTemplate.ush
@@ -1090,6 +1090,25 @@ MaterialFloat3 TransformTangentVectorToView(FMaterialPixelParameters Parameters,
 	return mul(mul(InTangentVector, Parameters.TangentToWorld), (MaterialFloat3x3)ResolvedView.TranslatedWorldToView);
 }
 
+FLWCMatrix GetLocalToWorld(FMaterialVertexParameters Parameters)
+{
+#if USE_INSTANCING || USE_INSTANCE_CULLING || IS_MESHPARTICLE_FACTORY
+	return Parameters.InstanceLocalToWorld;
+#else
+	return GetPrimitiveData(Parameters).LocalToWorld;
+#endif
+}
+
+FLWCMatrix GetLocalToWorld(FMaterialPixelParameters Parameters)
+{
+	return GetPrimitiveData(Parameters).LocalToWorld;
+}
+
+FLWCMatrix GetPrevLocalToWorld(FMaterialVertexParameters Parameters)
+{
+	return Parameters.PrevFrameLocalToWorld;
+}
+
 /** Transforms a vector from local space to world space (VS version) */
 MaterialFloat3 TransformLocalVectorToWorld(FMaterialVertexParameters Parameters,MaterialFloat3 InLocalVector)
 {
diff --git a/Engine/Source/Runtime/Engine/Classes/Materials/MaterialExpressionCameraPositionWS.h b/Engine/Source/Runtime/Engine/Classes/Materials/MaterialExpressionCameraPositionWS.h
--- a/Engine/Source/Runtime/Engine/Classes/Materials/MaterialExpressionCameraPositionWS.h
+++ b/Engine/Source/Runtime/Engine/Classes/Materials/MaterialExpressionCameraPositionWS.h
@@ -16,6 +16,7 @@ class UMaterialExpressionCameraPositionWS : public UMaterialExpression
 
 	//~ Begin UMaterialExpression Interface
 #if WITH_EDITOR
+	virtual bool GenerateHLSLExpression(FMaterialHLSLGenerator& Generator, UE::HLSLTree::FScope& Scope, int32 OutputIndex, UE::HLSLTree::FExpression*& OutExpression) override;
 	virtual int32 Compile(class FMaterialCompiler* Compiler, int32 OutputIndex) override;
 	virtual void GetCaption(TArray<FString>& OutCaptions) const override;
 #endif
diff --git a/Engine/Source/Runtime/Engine/Classes/Materials/MaterialExpressionDistance.h b/Engine/Source/Runtime/Engine/Classes/Materials/MaterialExpressionDistance.h
--- a/Engine/Source/Runtime/Engine/Classes/Materials/MaterialExpressionDistance.h
+++ b/Engine/Source/Runtime/Engine/Classes/Materials/MaterialExpressionDistance.h
@@ -23,6 +23,7 @@ class UMaterialExpressionDistance : public UMaterialExpression
 
 	//~ Begin UMaterialExpression Interface
 #if WITH_EDITOR
+	virtual bool GenerateHLSLExpression(FMaterialHLSLGenerator& Generator, UE::HLSLTree::FScope& Scope, int32 OutputIndex, UE::HLSLTree::FExpression*& OutExpression) override;
 	virtual int32 Compile(class FMaterialCompiler* Compiler, int32 OutputIndex) override;
 	virtual void GetCaption(TArray<FString>& OutCaptions) const override;
 #endif
diff --git a/Engine/Source/Runtime/Engine/Classes/Materials/MaterialExpressionNormalize.h b/Engine/Source/Runtime/Engine/Classes/Materials/MaterialExpressionNormalize.h
--- a/Engine/Source/Runtime/Engine/Classes/Materials/MaterialExpressionNormalize.h
+++ b/Engine/Source/Runtime/Engine/Classes/Materials/MaterialExpressionNormalize.h
@@ -20,6 +20,7 @@ class UMaterialExpressionNormalize : public UMaterialExpression
 
 	//~ Begin UMaterialExpression Interface
 #if WITH_EDITOR
+	virtual bool GenerateHLSLExpression(FMaterialHLSLGenerator& Generator, UE::HLSLTree::FScope& Scope, int32 OutputIndex, UE::HLSLTree::FExpression*& OutExpression) override;
 	virtual int32 Compile(class FMaterialCompiler* Compiler, int32 OutputIndex) override;
 	virtual void GetCaption(TArray<FString>& OutCaptions) const override { OutCaptions.Add(TEXT("Normalize")); }
 	virtual FText GetKeywords() const override { return FText::FromString(TEXT("normal")); }
diff --git a/Engine/Source/Runtime/Engine/Classes/Materials/MaterialExpressionTransform.h b/Engine/Source/Runtime/Engine/Classes/Materials/MaterialExpressionTransform.h
--- a/Engine/Source/Runtime/Engine/Classes/Materials/MaterialExpressionTransform.h
+++ b/Engine/Source/Runtime/Engine/Classes/Materials/MaterialExpressionTransform.h
@@ -82,6 +82,7 @@ class UMaterialExpressionTransform : public UMaterialExpression
 
 	//~ Begin UMaterialExpression Interface
 #if WITH_EDITOR
+	virtual bool GenerateHLSLExpression(FMaterialHLSLGenerator& Generator, UE::HLSLTree::FScope& Scope, int32 OutputIndex, UE::HLSLTree::FExpression*& OutExpression) override;
 	virtual int32 Compile(class FMaterialCompiler* Compiler, int32 OutputIndex) override;
 	virtual void GetCaption(TArray<FString>& OutCaptions) const override;
 #endif
diff --git a/Engine/Source/Runtime/Engine/Classes/Materials/MaterialExpressionTransformPosition.h b/Engine/Source/Runtime/Engine/Classes/Materials/MaterialExpressionTransformPosition.h
--- a/Engine/Source/Runtime/Engine/Classes/Materials/MaterialExpressionTransformPosition.h
+++ b/Engine/Source/Runtime/Engine/Classes/Materials/MaterialExpressionTransformPosition.h
@@ -56,6 +56,7 @@ class UMaterialExpressionTransformPosition : public UMaterialExpression
 
 	//~ Begin UMaterialExpression Interface
 #if WITH_EDITOR
+	virtual bool GenerateHLSLExpression(FMaterialHLSLGenerator& Generator, UE::HLSLTree::FScope& Scope, int32 OutputIndex, UE::HLSLTree::FExpression*& OutExpression) override;
 	virtual int32 Compile(class FMaterialCompiler* Compiler, int32 OutputIndex) override;
 	virtual void GetCaption(TArray<FString>& OutCaptions) const override;
 #endif
diff --git a/Engine/Source/Runtime/Engine/Private/HLSLTree/HLSLTree.cpp b/Engine/Source/Runtime/Engine/Private/HLSLTree/HLSLTree.cpp
--- a/Engine/Source/Runtime/Engine/Private/HLSLTree/HLSLTree.cpp
+++ b/Engine/Source/Runtime/Engine/Private/HLSLTree/HLSLTree.cpp
@@ -54,7 +54,7 @@ public:
 	virtual FExpression* ComputePreviousFrame(FTree& Tree, const FRequestedType& RequestedType) const override;
 	virtual bool PrepareValue(FEmitContext& Context, FEmitScope& Scope, const FRequestedType& RequestedType, FPrepareValueResult& OutResult) const override;
 	virtual void EmitValueShader(FEmitContext& Context, FEmitScope& Scope, const FRequestedType& RequestedType, FEmitValueShaderResult& OutResult) const override;
-	virtual void EmitValuePreshader(FEmitContext& Context, FEmitScope& Scope, const FRequestedType& RequestedType, Shader::FPreshaderData& OutPreshader) const override;
+	virtual void EmitValuePreshader(FEmitContext& Context, FEmitScope& Scope, const FRequestedType& RequestedType, FEmitValuePreshaderResult& OutResult) const override;
 
 	TArray<FLocalPHIChainEntry, TInlineAllocator<8>> Chain;
 	FName LocalName;
@@ -74,7 +74,7 @@ public:
 
 	virtual bool PrepareValue(FEmitContext& Context, FEmitScope& Scope, const FRequestedType& RequestedType, FPrepareValueResult& OutResult) const override;
 	virtual void EmitValueShader(FEmitContext& Context, FEmitScope& Scope, const FRequestedType& RequestedType, FEmitValueShaderResult& OutResult) const override;
-	virtual void EmitValuePreshader(FEmitContext& Context, FEmitScope& Scope, const FRequestedType& RequestedType, Shader::FPreshaderData& OutPreshader) const override;
+	virtual void EmitValuePreshader(FEmitContext& Context, FEmitScope& Scope, const FRequestedType& RequestedType, FEmitValuePreshaderResult& OutResult) const override;
 
 	FFunction* Function;
 	int32 OutputIndex;
@@ -104,6 +104,8 @@ FUnaryOpDescription GetUnaryOpDesription(EUnaryOp Op)
 	case EUnaryOp::Neg: return FUnaryOpDescription(TEXT("Neg"), TEXT("-"), Shader::EPreshaderOpcode::Neg); break;
 	case EUnaryOp::Rcp: return FUnaryOpDescription(TEXT("Rcp"), TEXT("/"), Shader::EPreshaderOpcode::Rcp); break;
 	case EUnaryOp::Frac: return FUnaryOpDescription(TEXT("Frac"), TEXT("frac"), Shader::EPreshaderOpcode::Frac); break;
+	case EUnaryOp::Length: return FUnaryOpDescription(TEXT("Length"), TEXT("length"), Shader::EPreshaderOpcode::Length); break;
+	case EUnaryOp::Normalize: return FUnaryOpDescription(TEXT("Normalize"), TEXT("normalize"), Shader::EPreshaderOpcode::Normalize); break;
 	default: checkNoEntry(); return FUnaryOpDescription();
 	}
 }
@@ -123,6 +125,10 @@ FBinaryOpDescription GetBinaryOpDesription(EBinaryOp Op)
 	case EBinaryOp::Max: return FBinaryOpDescription(TEXT("Max"), TEXT("max"), Shader::EPreshaderOpcode::Max); break;
 	case EBinaryOp::Less: return FBinaryOpDescription(TEXT("Less"), TEXT("<"), Shader::EPreshaderOpcode::Less); break;
 	case EBinaryOp::Greater: return FBinaryOpDescription(TEXT("Greater"), TEXT(">"), Shader::EPreshaderOpcode::Greater); break;
+	case EBinaryOp::VecMulMatrix3: return FBinaryOpDescription(TEXT("VecMulMatrix3"), TEXT("mul"), Shader::EPreshaderOpcode::Nop); break;
+	case EBinaryOp::VecMulMatrix4: return FBinaryOpDescription(TEXT("VecMulMatrix4"), TEXT("mul"), Shader::EPreshaderOpcode::Nop); break;
+	case EBinaryOp::Matrix3MulVec: return FBinaryOpDescription(TEXT("Matrix3MulVec"), TEXT("mul"), Shader::EPreshaderOpcode::Nop); break;
+	case EBinaryOp::Matrix4MulVec: return FBinaryOpDescription(TEXT("Matrix4MulVec"), TEXT("mul"), Shader::EPreshaderOpcode::Nop); break;
 	default: checkNoEntry(); return FBinaryOpDescription();
 	}
 }
@@ -391,7 +397,7 @@ void FExpressionLocalPHI::EmitValueShader(FEmitContext& Context, FEmitScope& Sco
 	if (!EmitExpression)
 	{
 		const int32 LocalPHIIndex = Context.NumExpressionLocalPHIs++;
-		const FRequestedType LocalType = GetRequestedType();
+		const FPreparedType& LocalType = GetPreparedType();
 
 		// This is the first time we've emitted shader code for this PHI
 		// Create an expression and add it to the map first, so if this is called recursively this path will only be taken the first time
@@ -457,7 +463,7 @@ void FExpressionLocalPHI::EmitValueShader(FEmitContext& Context, FEmitScope& Sco
 	OutResult.Code = EmitExpression;
 }
 
-void FExpressionLocalPHI::EmitValuePreshader(FEmitContext& Context, FEmitScope& Scope, const FRequestedType& RequestedType, Shader::FPreshaderData& OutPreshader) const
+void FExpressionLocalPHI::EmitValuePreshader(FEmitContext& Context, FEmitScope& Scope, const FRequestedType& RequestedType, FEmitValuePreshaderResult& OutResult) const
 {
 	int32 ValueStackPosition = INDEX_NONE;
 	for (int32 Index = Context.PreshaderLocalPHIScopes.Num() - 1; Index >= 0; --Index)
@@ -470,12 +476,12 @@ void FExpressionLocalPHI::EmitValuePreshader(FEmitContext& Context, FEmitScope&
 		}
 	}
 
+	OutResult.Type = GetType();
 	if (ValueStackPosition == INDEX_NONE)
 	{
 		// Assign the initial value
-		const FRequestedType LocalType = GetRequestedType();
 		Context.PreshaderStackPosition++;
-		OutPreshader.WriteOpcode(Shader::EPreshaderOpcode::ConstantZero).Write(LocalType.GetType());
+		OutResult.Preshader.WriteOpcode(Shader::EPreshaderOpcode::ConstantZero).Write(OutResult.Type);
 
 		ValueStackPosition = Context.PreshaderStackPosition;
 		const FPreshaderLocalPHIScope LocalPHIScope(this, ValueStackPosition);
@@ -491,7 +497,7 @@ void FExpressionLocalPHI::EmitValuePreshader(FEmitContext& Context, FEmitScope&
 			EmitRootScope = FEmitScope::FindSharedParent(PreshaderScopes[i].Scope, EmitRootScope);
 		}
 
-		Context.EmitPreshaderScope(*EmitRootScope, RequestedType, MakeArrayView(PreshaderScopes, NumValues), OutPreshader);
+		Context.EmitPreshaderScope(*EmitRootScope, RequestedType, MakeArrayView(PreshaderScopes, NumValues), OutResult.Preshader);
 		verify(Context.PreshaderLocalPHIScopes.Pop(false) == &LocalPHIScope);
 		check(Context.PreshaderStackPosition == ValueStackPosition);
 	}
@@ -501,7 +507,7 @@ void FExpressionLocalPHI::EmitValuePreshader(FEmitContext& Context, FEmitScope&
 		check(PreshaderStackOffset >= 0 && PreshaderStackOffset <= 0xffff);
 
 		Context.PreshaderStackPosition++;
-		OutPreshader.WriteOpcode(Shader::EPreshaderOpcode::PushValue).Write((uint16)PreshaderStackOffset);
+		OutResult.Preshader.WriteOpcode(Shader::EPreshaderOpcode::PushValue).Write((uint16)PreshaderStackOffset);
 	}
 }
 
@@ -533,9 +539,9 @@ void FExpressionFunctionCall::EmitValueShader(FEmitContext& Context, FEmitScope&
 	OutResult.Code = Context.EmitInlineExpressionWithDependency(Scope, Dependency, EmitFunctionOutput->Type, TEXT("%"), EmitFunctionOutput);
 }
 
-void FExpressionFunctionCall::EmitValuePreshader(FEmitContext& Context, FEmitScope& Scope, const FRequestedType& RequestedType, Shader::FPreshaderData& OutPreshader) const
+void FExpressionFunctionCall::EmitValuePreshader(FEmitContext& Context, FEmitScope& Scope, const FRequestedType& RequestedType, FEmitValuePreshaderResult& OutResult) const
 {
-	Function->OutputExpressions[OutputIndex]->GetValuePreshader(Context, Scope, RequestedType, OutPreshader);
+	OutResult.Type = Function->OutputExpressions[OutputIndex]->GetValuePreshader(Context, Scope, RequestedType, OutResult.Preshader);
 }
 
 void FExpression::Reset()
@@ -543,58 +549,24 @@ void FExpression::Reset()
 	PrepareValueResult = FPrepareValueResult();
 }
 
-FRequestedType::FRequestedType(int32 NumComponents, bool bDefaultRequest)
+FRequestedType::FRequestedType(ERequestedType InType)
 {
-	RequestedComponents.Init(bDefaultRequest, NumComponents);
+	const int32 NumComponents = (uint8)InType;
+	check(NumComponents >= 0 && NumComponents <= 16);
+	RequestedComponents.Init(true, NumComponents);
 }
 
-FRequestedType::FRequestedType(const Shader::FType& InType, bool bDefaultRequest)
+FRequestedType::FRequestedType(int32 NumComponents)
 {
-	int32 NumComponents = 0;
-	if (InType.IsStruct())
-	{
-		StructType = InType.StructType;
-		NumComponents = InType.StructType->ComponentTypes.Num();
-	}
-	else
-	{
-		const Shader::FValueTypeDescription TypeDesc = Shader::GetValueTypeDescription(InType);
-		ValueComponentType = TypeDesc.ComponentType;
-		NumComponents = TypeDesc.NumComponents;
-	}
-
-	RequestedComponents.Init(bDefaultRequest, NumComponents);
-}
-
-FRequestedType::FRequestedType(const Shader::EValueType& InType, bool bDefaultRequest)
-{
-	const Shader::FValueTypeDescription TypeDesc = Shader::GetValueTypeDescription(InType);
-	ValueComponentType = TypeDesc.ComponentType;
-	RequestedComponents.Init(bDefaultRequest, TypeDesc.NumComponents);
-}
-
-const Shader::FType FRequestedType::GetType() const
-{
-	if (IsStruct())
-	{
-		return StructType;
-	}
-	return Shader::MakeValueType(ValueComponentType, GetNumComponents());
+	RequestedComponents.Init(true, NumComponents);
 }
 
 int32 FRequestedType::GetNumComponents() const
 {
-	if (StructType)
+	const int32 MaxComponentIndex = RequestedComponents.FindLast(true);
+	if (MaxComponentIndex != INDEX_NONE)
 	{
-		return StructType->ComponentTypes.Num();
-	}
-	else
-	{
-		const int32 MaxComponentIndex = RequestedComponents.FindLast(true);
-		if (MaxComponentIndex != INDEX_NONE)
-		{
-			return MaxComponentIndex + 1;
-		}
+		return MaxComponentIndex + 1;
 	}
 	return 0;
 }
@@ -631,7 +603,7 @@ void FRequestedType::SetField(const Shader::FStructField* Field, const FRequeste
 
 FRequestedType FRequestedType::GetField(const Shader::FStructField* Field) const
 {
-	FRequestedType Result(Field->Type);
+	FRequestedType Result;
 	const int32 NumComponents = Field->GetNumComponents();
 	for (int32 Index = 0; Index < NumComponents; ++Index)
 	{
@@ -688,7 +660,7 @@ FPreparedComponent CombineComponents(const FPreparedComponent& Lhs, const FPrepa
 	}
 }
 
-FPreparedType::FPreparedType(const Shader::FType& InType)
+FPreparedType::FPreparedType(const Shader::FType& InType, EExpressionEvaluation InEvalution)
 {
 	if (InType.IsStruct())
 	{
@@ -698,6 +670,16 @@ FPreparedType::FPreparedType(const Shader::FType& InType)
 	{
 		ValueComponentType = Shader::GetValueTypeDescription(InType).ComponentType;
 	}
+
+	if (InEvalution != EExpressionEvaluation::None)
+	{
+		const int32 NumComponents = InType.GetNumComponents();
+		PreparedComponents.Reserve(NumComponents);
+		for (int32 Index = 0; Index < NumComponents; ++Index)
+		{
+			SetComponent(Index, InEvalution);
+		}
+	}
 }
 
 int32 FPreparedType::GetNumComponents() const
@@ -717,15 +699,6 @@ int32 FPreparedType::GetNumComponents() const
 	return 0;
 }
 
-FRequestedType MakeRequestedType(Shader::EValueComponentType ComponentType, const FRequestedType& RequestedComponents)
-{
-	check(!RequestedComponents.IsStruct());
-	FRequestedType Result;
-	Result.ValueComponentType = ComponentType;
-	Result.RequestedComponents = RequestedComponents.RequestedComponents;
-	return Result;
-}
-
 bool FPreparedType::IsVoid() const
 {
 	return GetNumComponents() == 0;
@@ -746,14 +719,6 @@ FRequestedType FPreparedType::GetRequestedType() const
 	FRequestedType Result;
 	if (NumComponents > 0)
 	{
-		if (StructType)
-		{
-			Result.StructType = StructType;
-		}
-		else
-		{
-			Result.ValueComponentType = ValueComponentType;
-		}
 		for (int32 Index = 0; Index < NumComponents; ++Index)
 		{
 			const FPreparedComponent Component = GetComponent(Index);
@@ -766,6 +731,24 @@ FRequestedType FPreparedType::GetRequestedType() const
 	return Result;
 }
 
+FPreparedType FPreparedType::GetTypeForRequest(const FRequestedType& RequestedType) const
+{
+	const int32 NumComponents = FMath::Min(PreparedComponents.Num(), RequestedType.RequestedComponents.Num());
+
+	FPreparedType Result;
+	Result.StructType = StructType;
+	Result.ValueComponentType = ValueComponentType;
+	Result.PreparedComponents.Reserve(NumComponents);
+	for (int32 Index = 0; Index < NumComponents; ++Index)
+	{
+		if (RequestedType.IsComponentRequested(Index))
+		{
+			Result.SetComponent(Index, GetComponent(Index));
+		}
+	}
+	return Result;
+}
+
 EExpressionEvaluation FPreparedType::GetEvaluation(const FEmitScope& Scope) const
 {
 	EExpressionEvaluation Result = EExpressionEvaluation::None;
@@ -879,7 +862,7 @@ void FPreparedType::SetLoopEvaluation(FEmitScope& Scope, const FRequestedType& R
 	}
 }
 
-void FPreparedType::SetForwardValue(FExpression* ForwardValue, const FRequestedType& RequestedType)
+void FPreparedType::SetForwardValue(const FRequestedType& RequestedType, FExpression* ForwardValue)
 {
 	for (int32 Index = 0; Index < PreparedComponents.Num(); ++Index)
 	{
@@ -940,6 +923,7 @@ FPreparedType MergePreparedTypes(const FPreparedType& Lhs, const FPreparedType&
 		NumComponents = FMath::Max(Lhs.GetNumComponents(), Rhs.GetNumComponents());
 	}
 
+	Result.PreparedComponents.Reset(NumComponents);
 	for (int32 Index = 0; Index < NumComponents; ++Index)
 	{
 		const FPreparedComponent LhsComponent = Lhs.GetComponent(Index);
@@ -950,6 +934,31 @@ FPreparedType MergePreparedTypes(const FPreparedType& Lhs, const FPreparedType&
 	return Result;
 }
 
+FRequestedValueType::FRequestedValueType(Shader::EValueType InType)
+{
+	const Shader::FValueTypeDescription TypeDesc = Shader::GetValueTypeDescription(InType);
+	ComponentType = TypeDesc.ComponentType;
+	RequestedType.RequestedComponents.Init(true, TypeDesc.NumComponents);
+}
+
+FRequestedValueType::FRequestedValueType(const FPreparedType& InType)
+{
+	ComponentType = InType.ValueComponentType;
+	if (ComponentType != Shader::EValueComponentType::Void)
+	{
+		RequestedType = InType.GetRequestedType();
+	}
+}
+
+Shader::EValueType FRequestedValueType::GetType() const
+{
+	if (ComponentType != Shader::EValueComponentType::Void)
+	{
+		return Shader::MakeValueType(ComponentType, GetNumComponents());
+	}
+	return Shader::EValueType::Void;
+}
+
 bool FPrepareValueResult::TryMergePreparedType(FEmitContext& Context, const Shader::FStructType* StructType, Shader::EValueComponentType ComponentType)
 {
 	// If we previously had a forwarded value set, reset that and start over
@@ -1028,9 +1037,9 @@ bool FPrepareValueResult::SetType(FEmitContext& Context, const FRequestedType& R
 
 bool FPrepareValueResult::SetForwardValue(FEmitContext& Context, FEmitScope& Scope, const FRequestedType& RequestedType, FExpression* InForwardValue)
 {
-	const FPreparedType ForwardPreparedType = Context.PrepareExpression(InForwardValue, Scope, RequestedType);
-	PreparedType = MergePreparedTypes(PreparedType, ForwardPreparedType);
-	return !PreparedType.IsVoid();
+	FPreparedType ForwardPreparedType = Context.PrepareExpression(InForwardValue, Scope, RequestedType);
+	ForwardPreparedType.SetForwardValue(RequestedType, InForwardValue);
+	return SetType(Context, RequestedType, ForwardPreparedType);
 }
 
 void FStatement::EmitPreshader(FEmitContext& Context, FEmitScope& Scope, const FRequestedType& RequestedType, TArrayView<const FEmitPreshaderScope> Scopes, Shader::FPreshaderData& OutPreshader) const
@@ -1053,7 +1062,7 @@ void FExpression::EmitValueShader(FEmitContext& Context, FEmitScope& Scope, cons
 	check(false);
 }
 
-void FExpression::EmitValuePreshader(FEmitContext& Context, FEmitScope& Scope, const FRequestedType& RequestedType, Shader::FPreshaderData& OutPreshader) const
+void FExpression::EmitValuePreshader(FEmitContext& Context, FEmitScope& Scope, const FRequestedType& RequestedType, FEmitValuePreshaderResult& OutResult) const
 {
 	check(false);
 }
@@ -1124,13 +1133,13 @@ FExpression* FExpression::GetForwardValue(const FRequestedType& RequestedType)
 	return nullptr;
 }
 
-FEmitShaderExpression* FExpression::GetValueShader(FEmitContext& Context, FEmitScope& Scope, const FRequestedType& RequestedType, const Shader::FType& ResultType)
+FEmitShaderExpression* FExpression::GetValueShader(FEmitContext& Context, FEmitScope& Scope, const FRequestedType& RequestedType)
 {
 	FOwnerScope OwnerScope(*Context.Errors, GetOwner());
 	FExpression* ForwardValue = GetForwardValue(RequestedType);
 	if (ForwardValue)
 	{
-		return ForwardValue->GetValueShader(Context, Scope, RequestedType, ResultType);
+		return ForwardValue->GetValueShader(Context, Scope, RequestedType);
 	}
 
 	const EExpressionEvaluation Evaluation = PrepareValueResult.PreparedType.GetEvaluation(Scope, RequestedType);
@@ -1149,13 +1158,23 @@ FEmitShaderExpression* FExpression::GetValueShader(FEmitContext& Context, FEmitS
 		Value = Result.Code;
 	}
 
-	Value = Context.EmitCast(Scope, Value, ResultType);
 	return Value;
 }
 
-FEmitShaderExpression* FExpression::GetValueShader(FEmitContext& Context, FEmitScope& Scope, const FRequestedType& RequestedType)
+FEmitShaderExpression* FExpression::GetValueShader(FEmitContext& Context, FEmitScope& Scope, const FRequestedType& RequestedType, const Shader::FType& ResultType)
+{
+	FEmitShaderExpression* Value = GetValueShader(Context, Scope, RequestedType);
+	return Context.EmitCast(Scope, Value, ResultType);
+}
+
+FEmitShaderExpression* FExpression::GetValueShader(FEmitContext& Context, FEmitScope& Scope, const FRequestedValueType& ResultType)
+{
+	return GetValueShader(Context, Scope, ResultType.RequestedType, ResultType.GetType());
+}
+
+FEmitShaderExpression* FExpression::GetValueShader(FEmitContext& Context, FEmitScope& Scope, const FPreparedType& ResultType)
 {
-	return GetValueShader(Context, Scope, RequestedType, RequestedType.GetType());
+	return GetValueShader(Context, Scope, ResultType.GetRequestedType(), ResultType.GetType());
 }
 
 FEmitShaderExpression* FExpression::GetValueShader(FEmitContext& Context, FEmitScope& Scope)
@@ -1163,7 +1182,7 @@ FEmitShaderExpression* FExpression::GetValueShader(FEmitContext& Context, FEmitS
 	return GetValueShader(Context, Scope, GetRequestedType());
 }
 
-void FExpression::GetValuePreshader(FEmitContext& Context, FEmitScope& Scope, const FRequestedType& RequestedType, Shader::FPreshaderData& OutPreshader)
+Shader::FType FExpression::GetValuePreshader(FEmitContext& Context, FEmitScope& Scope, const FRequestedType& RequestedType, Shader::FPreshaderData& OutPreshader)
 {
 	FOwnerScope OwnerScope(*Context.Errors, GetOwner());
 	FExpression* ForwardValue = GetForwardValue(RequestedType);
@@ -1175,18 +1194,21 @@ void FExpression::GetValuePreshader(FEmitContext& Context, FEmitScope& Scope, co
 	const int32 PrevStackPosition = Context.PreshaderStackPosition;
 	const EExpressionEvaluation Evaluation = PrepareValueResult.PreparedType.GetEvaluation(Scope, RequestedType);
 
+	FEmitValuePreshaderResult Result(OutPreshader);
 	if (Evaluation == EExpressionEvaluation::Constant)
 	{
 		const Shader::FValue ConstantValue = GetValueConstant(Context, Scope, RequestedType);
 		Context.PreshaderStackPosition++;
 		OutPreshader.WriteOpcode(Shader::EPreshaderOpcode::Constant).Write(ConstantValue);
+		Result.Type = ConstantValue.Type;
 	}
 	else
 	{
 		check(Evaluation != EExpressionEvaluation::None && Evaluation != EExpressionEvaluation::Unknown && Evaluation != EExpressionEvaluation::Shader);
-		EmitValuePreshader(Context, Scope, RequestedType, OutPreshader);
+		EmitValuePreshader(Context, Scope, RequestedType, Result);
 	}
 	check(Context.PreshaderStackPosition == PrevStackPosition + 1);
+	return Result.Type;
 }
 
 Shader::FValue FExpression::GetValueConstant(FEmitContext& Context, FEmitScope& Scope, const FRequestedType& RequestedType)
@@ -1207,7 +1229,8 @@ Shader::FValue FExpression::GetValueConstant(FEmitContext& Context, FEmitScope&
 	{
 		FExpressionReentryScope ReentryScope(this);
 		const int32 PrevPreshaderStackPosition = Context.PreshaderStackPosition;
-		EmitValuePreshader(Context, Scope, RequestedType, ConstantPreshader);
+		FEmitValuePreshaderResult PreshaderResult(ConstantPreshader);
+		EmitValuePreshader(Context, Scope, RequestedType, PreshaderResult);
 		check(Context.PreshaderStackPosition == PrevPreshaderStackPosition + 1);
 		Context.PreshaderStackPosition--;
 	}
@@ -1216,17 +1239,25 @@ Shader::FValue FExpression::GetValueConstant(FEmitContext& Context, FEmitScope&
 	Shader::FPreshaderStack Stack;
 	const Shader::FPreshaderValue PreshaderValue = ConstantPreshader.EvaluateConstant(*Context.Material, Stack);
 	Shader::FValue Result = PreshaderValue.AsShaderValue(Context.TypeRegistry);
+	return Result;
+}
 
-	const Shader::FType RequestedConstantType = RequestedType.GetType();
-	if (Result.Type.IsNumeric() && RequestedConstantType.IsNumeric())
+Shader::FValue FExpression::GetValueConstant(FEmitContext& Context, FEmitScope& Scope, const FRequestedType& RequestedType, const Shader::FType& ResultType)
+{
+	Shader::FValue Result = GetValueConstant(Context, Scope, RequestedType);
+	if (Result.Type.IsNumeric() && ResultType.IsNumeric())
 	{
-		Result = Shader::Cast(Result, RequestedConstantType.ValueType);
+		Result = Shader::Cast(Result, ResultType.ValueType);
 	}
-
-	check(Result.Type == RequestedConstantType);
+	check(Result.Type == ResultType);
 	return Result;
 }
 
+Shader::FValue FExpression::GetValueConstant(FEmitContext& Context, FEmitScope& Scope, const Shader::FType& ResultType)
+{
+	return GetValueConstant(Context, Scope, FRequestedType(ResultType.GetNumComponents()), ResultType);
+}
+
 bool FScope::HasParentScope(const FScope& InParentScope) const
 {
 	const FScope* CurrentScope = this;
diff --git a/Engine/Source/Runtime/Engine/Private/HLSLTree/HLSLTreeCommon.cpp b/Engine/Source/Runtime/Engine/Private/HLSLTree/HLSLTreeCommon.cpp
--- a/Engine/Source/Runtime/Engine/Private/HLSLTree/HLSLTreeCommon.cpp
+++ b/Engine/Source/Runtime/Engine/Private/HLSLTree/HLSLTreeCommon.cpp
@@ -62,6 +62,29 @@ FExternalInputDescription GetExternalInputDescription(EExternalInput Input)
 	case EExternalInput::WorldPosition_Ddx: return FExternalInputDescription(TEXT("WorldPosition_Ddx"), Shader::EValueType::Float3);
 	case EExternalInput::WorldPosition_Ddy: return FExternalInputDescription(TEXT("WorldPosition_Ddx"), Shader::EValueType::Float3);
 
+	case EExternalInput::CameraWorldPosition: return FExternalInputDescription(TEXT("CameraWorldPosition"), Shader::EValueType::Double3, EExternalInput::None, EExternalInput::None, EExternalInput::PrevCameraWorldPosition);
+	case EExternalInput::PreViewTranslation: return FExternalInputDescription(TEXT("PreViewTranslation"), Shader::EValueType::Double3, EExternalInput::None, EExternalInput::None, EExternalInput::PrevPreViewTranslation);
+	case EExternalInput::TangentToWorld: return FExternalInputDescription(TEXT("TangentToWorld"), Shader::EValueType::Float4x4);
+	case EExternalInput::LocalToWorld: return FExternalInputDescription(TEXT("LocalToWorld"), Shader::EValueType::Double4x4, EExternalInput::None, EExternalInput::None, EExternalInput::PrevLocalToWorld);
+	case EExternalInput::WorldToLocal: return FExternalInputDescription(TEXT("WorldToLocal"), Shader::EValueType::DoubleInverse4x4, EExternalInput::None, EExternalInput::None, EExternalInput::PrevWorldToLocal);
+	case EExternalInput::TranslatedWorldToCameraView: return FExternalInputDescription(TEXT("TranslatedWorldToCameraView"), Shader::EValueType::Float4x4, EExternalInput::None, EExternalInput::None, EExternalInput::PrevTranslatedWorldToCameraView);
+	case EExternalInput::TranslatedWorldToView: return FExternalInputDescription(TEXT("TranslatedWorldToView"), Shader::EValueType::Float4x4, EExternalInput::None, EExternalInput::None, EExternalInput::PrevTranslatedWorldToView);
+	case EExternalInput::CameraViewToTranslatedWorld: return FExternalInputDescription(TEXT("CameraViewToTranslatedWorld"), Shader::EValueType::Float4x4, EExternalInput::None, EExternalInput::None, EExternalInput::PrevCameraViewToTranslatedWorld);
+	case EExternalInput::ViewToTranslatedWorld: return FExternalInputDescription(TEXT("ViewToTranslatedWorld"), Shader::EValueType::Float4x4, EExternalInput::None, EExternalInput::None, EExternalInput::PrevViewToTranslatedWorld);
+	case EExternalInput::WorldToParticle: return FExternalInputDescription(TEXT("WorldToParticle"), Shader::EValueType::DoubleInverse4x4);
+	case EExternalInput::WorldToInstance: return FExternalInputDescription(TEXT("WorldToInstance"), Shader::EValueType::DoubleInverse4x4);
+	case EExternalInput::ParticleToWorld: return FExternalInputDescription(TEXT("ParticleToWorld"), Shader::EValueType::Double4x4);
+	case EExternalInput::InstanceToWorld: return FExternalInputDescription(TEXT("InstanceToWorld"), Shader::EValueType::Double4x4);
+
+	case EExternalInput::PrevCameraWorldPosition: return FExternalInputDescription(TEXT("PrevCameraWorldPosition"), Shader::EValueType::Double3);
+	case EExternalInput::PrevPreViewTranslation: return FExternalInputDescription(TEXT("PrevPreViewTranslation"), Shader::EValueType::Double3);
+	case EExternalInput::PrevLocalToWorld: return FExternalInputDescription(TEXT("PrevLocalToWorld"), Shader::EValueType::Double4x4);
+	case EExternalInput::PrevWorldToLocal: return FExternalInputDescription(TEXT("PrevWorldToLocal"), Shader::EValueType::DoubleInverse4x4);
+	case EExternalInput::PrevTranslatedWorldToCameraView: return FExternalInputDescription(TEXT("PrevTranslatedWorldToCameraView"), Shader::EValueType::Float4x4);
+	case EExternalInput::PrevTranslatedWorldToView: return FExternalInputDescription(TEXT("PrevTranslatedWorldToView"), Shader::EValueType::Float4x4);
+	case EExternalInput::PrevCameraViewToTranslatedWorld: return FExternalInputDescription(TEXT("PrevCameraViewToTranslatedWorld"), Shader::EValueType::Float4x4);
+	case EExternalInput::PrevViewToTranslatedWorld: return FExternalInputDescription(TEXT("PrevViewToTranslatedWorld"), Shader::EValueType::Float4x4);
+
 	case EExternalInput::PixelDepth: return FExternalInputDescription(TEXT("PixelDepth"), Shader::EValueType::Float1, EExternalInput::PixelDepth_Ddx, EExternalInput::PixelDepth_Ddy);
 	case EExternalInput::PixelDepth_Ddx: return FExternalInputDescription(TEXT("PixelDepth_Ddx"), Shader::EValueType::Float1);
 	case EExternalInput::PixelDepth_Ddy: return FExternalInputDescription(TEXT("PixelDepth_Ddy"), Shader::EValueType::Float1);
@@ -73,7 +96,6 @@ FExternalInputDescription GetExternalInputDescription(EExternalInput Input)
 	case EExternalInput::PrevGameTime: return FExternalInputDescription(TEXT("PrevGameTime"), Shader::EValueType::Float1);
 	case EExternalInput::PrevRealTime: return FExternalInputDescription(TEXT("PrevRealTime"), Shader::EValueType::Float1);
 
-	
 	default: checkNoEntry(); return FExternalInputDescription(TEXT("Invalid"), Shader::EValueType::Void);
 	}
 }
@@ -112,8 +134,6 @@ FSwizzleParameters::FSwizzleParameters(int8 InR, int8 InG, int8 InB, int8 InA) :
 FRequestedType FSwizzleParameters::GetRequestedInputType(const FRequestedType& RequestedType) const
 {
 	FRequestedType RequestedInputType;
-	RequestedInputType.ValueComponentType = RequestedType.ValueComponentType;
-
 	for (int32 Index = 0; Index < NumComponents; ++Index)
 	{
 		if (RequestedType.IsComponentRequested(Index))
@@ -122,7 +142,6 @@ FRequestedType FSwizzleParameters::GetRequestedInputType(const FRequestedType& R
 			RequestedInputType.SetComponentRequest(SwizzledComponentIndex);
 		}
 	}
-	
 	return RequestedInputType;
 }
 
@@ -170,10 +189,11 @@ bool FExpressionConstant::PrepareValue(FEmitContext& Context, FEmitScope& Scope,
 	return OutResult.SetType(Context, RequestedType, EExpressionEvaluation::Constant, Value.Type);
 }
 
-void FExpressionConstant::EmitValuePreshader(FEmitContext& Context, FEmitScope& Scope, const FRequestedType& RequestedType, Shader::FPreshaderData& OutPreshader) const
+void FExpressionConstant::EmitValuePreshader(FEmitContext& Context, FEmitScope& Scope, const FRequestedType& RequestedType, FEmitValuePreshaderResult& OutResult) const
 {
 	Context.PreshaderStackPosition++;
-	OutPreshader.WriteOpcode(Shader::EPreshaderOpcode::Constant).Write(Value);
+	OutResult.Type = Value.Type;
+	OutResult.Preshader.WriteOpcode(Shader::EPreshaderOpcode::Constant).Write(Value);
 }
 
 void FExpressionMaterialParameter::ComputeAnalyticDerivatives(FTree& Tree, FExpressionDerivatives& OutResult) const
@@ -194,9 +214,10 @@ bool FExpressionMaterialParameter::PrepareValue(FEmitContext& Context, FEmitScop
 	return OutResult.SetType(Context, RequestedType, Evaluation, GetShaderValueType(ParameterType));
 }
 
-void FExpressionMaterialParameter::EmitValuePreshader(FEmitContext& Context, FEmitScope& Scope, const FRequestedType& RequestedType, Shader::FPreshaderData& OutPreshader) const
+void FExpressionMaterialParameter::EmitValuePreshader(FEmitContext& Context, FEmitScope& Scope, const FRequestedType& RequestedType, FEmitValuePreshaderResult& OutResult) const
 {
 	Context.PreshaderStackPosition++;
+	OutResult.Type = GetShaderValueType(ParameterType);
 	if (ParameterType == EMaterialParameterType::StaticSwitch)
 	{
 		const FMaterialParameterInfo ParameterInfo(ParameterName);
@@ -209,7 +230,7 @@ void FExpressionMaterialParameter::EmitValuePreshader(FEmitContext& Context, FEm
 				break;
 			}
 		}
-		OutPreshader.WriteOpcode(Shader::EPreshaderOpcode::Constant).Write(Value);
+		OutResult.Preshader.WriteOpcode(Shader::EPreshaderOpcode::Constant).Write(Value);
 	}
 	else
 	{
@@ -226,7 +247,7 @@ void FExpressionMaterialParameter::EmitValuePreshader(FEmitContext& Context, FEm
 		}
 		const int32 ParameterIndex = Context.MaterialCompilationOutput->UniformExpressionSet.FindOrAddNumericParameter(ParameterType, ParameterName, DefaultOffset);
 		check(ParameterIndex >= 0 && ParameterIndex <= 0xffff);
-		OutPreshader.WriteOpcode(Shader::EPreshaderOpcode::Parameter).Write((uint16)ParameterIndex);
+		OutResult.Preshader.WriteOpcode(Shader::EPreshaderOpcode::Parameter).Write((uint16)ParameterIndex);
 	}
 }
 
@@ -292,6 +313,30 @@ void FExpressionExternalInput::EmitValueShader(FEmitContext& Context, FEmitScope
 		case EExternalInput::PrevTranslatedWorldPosition_NoOffsets: Code = TEXT("GetPrevTranslatedWorldPosition_NoMaterialOffsets(Parameters)"); break;
 		case EExternalInput::WorldPosition_Ddx: Code = TEXT("Parameters.WorldPosition_DDX"); break;
 		case EExternalInput::WorldPosition_Ddy: Code = TEXT("Parameters.WorldPosition_DDY"); break;
+
+		case EExternalInput::CameraWorldPosition: Code = TEXT("ResolvedView.WorldCameraOrigin"); break;
+		case EExternalInput::PreViewTranslation: Code = TEXT("ResolvedView.PreViewTranslation"); break;
+		case EExternalInput::TangentToWorld: Code = TEXT("Parameters.TangentToWorld"); break;
+		case EExternalInput::LocalToWorld: Code = TEXT("GetLocalToWorld(Parameters)"); break;
+		case EExternalInput::WorldToLocal: Code = TEXT("GetPrimitiveData(Parameters).WorldToLocal"); break;
+		case EExternalInput::TranslatedWorldToCameraView: Code = TEXT("ResolvedView.TranslatedWorldToCameraView"); break;
+		case EExternalInput::TranslatedWorldToView: Code = TEXT("ResolvedView.TranslatedWorldToView"); break;
+		case EExternalInput::CameraViewToTranslatedWorld: Code = TEXT("ResolvedView.CameraViewToTranslatedWorld"); break;
+		case EExternalInput::ViewToTranslatedWorld: Code = TEXT("ResolvedView.ViewToTranslatedWorld"); break;
+		case EExternalInput::WorldToParticle: Code = TEXT("Parameters.Particle.WorldToParticle"); break;
+		case EExternalInput::WorldToInstance: Code = TEXT("GetWorldToInstance(Parameters)"); break;
+		case EExternalInput::ParticleToWorld: Code = TEXT("Parameters.Particle.ParticleToWorld"); break;
+		case EExternalInput::InstanceToWorld: Code = TEXT("GetInstanceToWorld(Parameters)"); break;
+
+		case EExternalInput::PrevCameraWorldPosition: Code = TEXT("ResolvedView.PrevWorldCameraOrigin"); break;
+		case EExternalInput::PrevPreViewTranslation: Code = TEXT("ResolvedView.PrevPreViewTranslation"); break;
+		case EExternalInput::PrevLocalToWorld: Code = TEXT("GetPrevLocalToWorld(Parameters)"); break;
+		case EExternalInput::PrevWorldToLocal: Code = TEXT("GetPrimitiveData(Parameters).PreviousWorldToLocal"); break;
+		case EExternalInput::PrevTranslatedWorldToCameraView: Code = TEXT("ResolvedView.PrevTranslatedWorldToCameraView"); break;
+		case EExternalInput::PrevTranslatedWorldToView: Code = TEXT("ResolvedView.PrevTranslatedWorldToView"); break;
+		case EExternalInput::PrevCameraViewToTranslatedWorld: Code = TEXT("ResolvedView.PrevCameraViewToTranslatedWorld"); break;
+		case EExternalInput::PrevViewToTranslatedWorld: Code = TEXT("ResolvedView.PrevViewToTranslatedWorld"); break;
+
 		case EExternalInput::PixelDepth: Code = TEXT("GetPixelDepth(Parameters)"); break;
 		case EExternalInput::PixelDepth_Ddx: Code = TEXT("Parameters.ScreenPosition_DDX.w"); break;
 		case EExternalInput::PixelDepth_Ddy: Code = TEXT("Parameters.ScreenPosition_DDY.w"); break;
@@ -310,9 +355,9 @@ void FExpressionExternalInput::EmitValueShader(FEmitContext& Context, FEmitScope
 
 bool FExpressionTextureSample::PrepareValue(FEmitContext& Context, FEmitScope& Scope, const FRequestedType& RequestedType, FPrepareValueResult& OutResult) const
 {
-	Context.PrepareExpression(TexCoordExpression, Scope, Shader::EValueType::Float2);
-	Context.PrepareExpression(TexCoordDerivatives.ExpressionDdx, Scope, Shader::EValueType::Float2);
-	Context.PrepareExpression(TexCoordDerivatives.ExpressionDdy, Scope, Shader::EValueType::Float2);
+	Context.PrepareExpression(TexCoordExpression, Scope, ERequestedType::Vector2);
+	Context.PrepareExpression(TexCoordDerivatives.ExpressionDdx, Scope, ERequestedType::Vector2);
+	Context.PrepareExpression(TexCoordDerivatives.ExpressionDdy, Scope, ERequestedType::Vector2);
 
 	return OutResult.SetType(Context, RequestedType, EExpressionEvaluation::Shader, Shader::EValueType::Float4);
 }
@@ -488,28 +533,29 @@ void FExpressionGetStructField::ComputeAnalyticDerivatives(FTree& Tree, FExpress
 
 FExpression* FExpressionGetStructField::ComputePreviousFrame(FTree& Tree, const FRequestedType& RequestedType) const
 {
-	FRequestedType RequestedStructType(StructType, false);
+	FRequestedType RequestedStructType;
 	RequestedStructType.SetField(Field, RequestedType);
 	return Tree.NewExpression<FExpressionGetStructField>(StructType, Field, Tree.GetPreviousFrame(StructExpression, RequestedStructType));
 }
 
 bool FExpressionGetStructField::PrepareValue(FEmitContext& Context, FEmitScope& Scope, const FRequestedType& RequestedType, FPrepareValueResult& OutResult) const
 {
-	FRequestedType RequestedStructType(StructType, false);
+	FRequestedType RequestedStructType;
 	RequestedStructType.SetField(Field, RequestedType);
 
-	const FPreparedType& StructPreparedType = Context.PrepareExpression(StructExpression, Scope, RequestedStructType);
+	FPreparedType StructPreparedType = Context.PrepareExpression(StructExpression, Scope, RequestedStructType);
 	if (!StructPreparedType.IsVoid() && StructPreparedType.StructType != StructType)
 	{
 		return Context.Errors->AddErrorf(TEXT("Expected type %s"), StructType->Name);
 	}
 
+	StructPreparedType.SetForwardValue(RequestedStructType, StructExpression);
 	return OutResult.SetType(Context, RequestedType, StructPreparedType.GetFieldType(Field));
 }
 
 void FExpressionGetStructField::EmitValueShader(FEmitContext& Context, FEmitScope& Scope, const FRequestedType& RequestedType, FEmitValueShaderResult& OutResult) const
 {
-	FRequestedType RequestedStructType(StructType, false);
+	FRequestedType RequestedStructType;
 	RequestedStructType.SetField(Field, RequestedType);
 
 	FEmitShaderExpression* StructValue = StructExpression->GetValueShader(Context, Scope, RequestedStructType);
@@ -519,13 +565,14 @@ void FExpressionGetStructField::EmitValueShader(FEmitContext& Context, FEmitScop
 		Field->Name);
 }
 
-void FExpressionGetStructField::EmitValuePreshader(FEmitContext& Context, FEmitScope& Scope, const FRequestedType& RequestedType, Shader::FPreshaderData& OutPreshader) const
+void FExpressionGetStructField::EmitValuePreshader(FEmitContext& Context, FEmitScope& Scope, const FRequestedType& RequestedType, FEmitValuePreshaderResult& OutResult) const
 {
-	FRequestedType RequestedStructType(StructType, false);
+	FRequestedType RequestedStructType;
 	RequestedStructType.SetField(Field, RequestedType);
 
-	StructExpression->GetValuePreshader(Context, Scope, RequestedStructType, OutPreshader);
-	OutPreshader.WriteOpcode(Shader::EPreshaderOpcode::GetField).Write(Field->Type).Write(Field->ComponentIndex);
+	StructExpression->GetValuePreshader(Context, Scope, RequestedStructType, OutResult.Preshader);
+	OutResult.Preshader.WriteOpcode(Shader::EPreshaderOpcode::GetField).Write(Field->Type).Write(Field->ComponentIndex);
+	OutResult.Type = Field->Type;
 }
 
 void FExpressionSetStructField::ComputeAnalyticDerivatives(FTree& Tree, FExpressionDerivatives& OutResult) const
@@ -562,13 +609,18 @@ bool FExpressionSetStructField::PrepareValue(FEmitContext& Context, FEmitScope&
 	FRequestedType RequestedStructType(RequestedType);
 	RequestedStructType.ClearFieldRequested(Field);
 
-	const FPreparedType& StructPreparedType = Context.PrepareExpression(StructExpression, Scope, RequestedStructType);
+	FPreparedType StructPreparedType = Context.PrepareExpression(StructExpression, Scope, RequestedStructType);
 	if (!StructPreparedType.IsVoid() && StructPreparedType.StructType != StructType)
 	{
 		return Context.Errors->AddErrorf(TEXT("Expected type %s"), StructType->Name);
 	}
 
-	const FPreparedType& FieldPreparedType = Context.PrepareExpression(FieldExpression, Scope, RequestedType.GetField(Field));
+	const FRequestedType RequestedFieldType = RequestedType.GetField(Field);
+	FPreparedType FieldPreparedType = Context.PrepareExpression(FieldExpression, Scope, RequestedFieldType);
+
+	StructPreparedType.SetForwardValue(RequestedStructType, StructExpression);
+	FieldPreparedType.SetForwardValue(RequestedFieldType, FieldExpression);
+
 	FPreparedType ResultType(StructPreparedType);
 	if (ResultType.IsVoid())
 	{
@@ -611,7 +663,7 @@ void FExpressionSetStructField::EmitValueShader(FEmitContext& Context, FEmitScop
 	}
 }
 
-void FExpressionSetStructField::EmitValuePreshader(FEmitContext& Context, FEmitScope& Scope, const FRequestedType& RequestedType, Shader::FPreshaderData& OutPreshader) const
+void FExpressionSetStructField::EmitValuePreshader(FEmitContext& Context, FEmitScope& Scope, const FRequestedType& RequestedType, FEmitValuePreshaderResult& OutResult) const
 {
 	FRequestedType RequestedStructType(RequestedType);
 	RequestedStructType.ClearFieldRequested(Field);
@@ -622,29 +674,30 @@ void FExpressionSetStructField::EmitValuePreshader(FEmitContext& Context, FEmitS
 
 	if (StructEvaluation != EExpressionEvaluation::None)
 	{
-		StructExpression->GetValuePreshader(Context, Scope, RequestedStructType, OutPreshader);
+		StructExpression->GetValuePreshader(Context, Scope, RequestedStructType, OutResult.Preshader);
 	}
 	else
 	{
 		Context.PreshaderStackPosition++;
-		OutPreshader.WriteOpcode(Shader::EPreshaderOpcode::ConstantZero).Write(Shader::FType(StructType));
+		OutResult.Preshader.WriteOpcode(Shader::EPreshaderOpcode::ConstantZero).Write(Shader::FType(StructType));
 	}
 
 	if (FieldEvaluation != EExpressionEvaluation::None)
 	{
-		FieldExpression->GetValuePreshader(Context, Scope, RequestedFieldType, OutPreshader);
+		FieldExpression->GetValuePreshader(Context, Scope, RequestedFieldType, OutResult.Preshader);
 
 		check(Context.PreshaderStackPosition > 0);
 		Context.PreshaderStackPosition--;
 
-		OutPreshader.WriteOpcode(Shader::EPreshaderOpcode::SetField).Write(Field->ComponentIndex).Write(Field->GetNumComponents());
+		OutResult.Preshader.WriteOpcode(Shader::EPreshaderOpcode::SetField).Write(Field->ComponentIndex).Write(Field->GetNumComponents());
 	}
+	OutResult.Type = StructType;
 }
 
 bool FExpressionSelect::PrepareValue(FEmitContext& Context, FEmitScope& Scope, const FRequestedType& RequestedType, FPrepareValueResult& OutResult) const
 {
-	const FPreparedType& ConditionType = Context.PrepareExpression(ConditionExpression, Scope, Shader::EValueType::Bool1);
-	if (ConditionType.GetEvaluation(Scope, Shader::EValueType::Bool1) == EExpressionEvaluation::Constant)
+	const FPreparedType& ConditionType = Context.PrepareExpression(ConditionExpression, Scope, ERequestedType::Scalar);
+	if (ConditionType.GetEvaluation(Scope, ERequestedType::Scalar) == EExpressionEvaluation::Constant)
 	{
 		const bool bCondition = ConditionExpression->GetValueConstant(Context, Scope, Shader::EValueType::Bool1).AsBoolScalar();
 		return OutResult.SetForwardValue(Context, Scope, RequestedType, bCondition ? TrueExpression : FalseExpression);
@@ -664,16 +717,17 @@ bool FExpressionSelect::PrepareValue(FEmitContext& Context, FEmitScope& Scope, c
 
 void FExpressionSelect::EmitValueShader(FEmitContext& Context, FEmitScope& Scope, const FRequestedType& RequestedType, FEmitValueShaderResult& OutResult) const
 {
-	FEmitShaderExpression* TrueValue = TrueExpression->GetValueShader(Context, Scope, RequestedType);
-	FEmitShaderExpression* FalseValue = FalseExpression->GetValueShader(Context, Scope, RequestedType);
+	const Shader::FType LocalType = GetType();
+	FEmitShaderExpression* TrueValue = TrueExpression->GetValueShader(Context, Scope, RequestedType, LocalType);
+	FEmitShaderExpression* FalseValue = FalseExpression->GetValueShader(Context, Scope, RequestedType, LocalType);
 
-	OutResult.Code = Context.EmitExpression(Scope, RequestedType.GetType(), TEXT("(% ? % : %)"),
+	OutResult.Code = Context.EmitExpression(Scope, LocalType, TEXT("(% ? % : %)"),
 		ConditionExpression->GetValueShader(Context, Scope, Shader::EValueType::Bool1),
 		TrueValue,
 		FalseValue);
 }
 
-void FExpressionSelect::EmitValuePreshader(FEmitContext& Context, FEmitScope& Scope, const FRequestedType& RequestedType, Shader::FPreshaderData& OutPreshader) const
+void FExpressionSelect::EmitValuePreshader(FEmitContext& Context, FEmitScope& Scope, const FRequestedType& RequestedType, FEmitValuePreshaderResult& OutResult) const
 {
 	check(false); // TODO
 }
@@ -705,6 +759,10 @@ void FExpressionUnaryOp::ComputeAnalyticDerivatives(FTree& Tree, FExpressionDeri
 		case EUnaryOp::Frac:
 			OutResult = InputDerivatives;
 			break;
+		case EUnaryOp::Length:
+		case EUnaryOp::Normalize:
+			// TODO
+			break;
 		default:
 			checkNoEntry();
 			break;
@@ -719,7 +777,20 @@ FExpression* FExpressionUnaryOp::ComputePreviousFrame(FTree& Tree, const FReques
 
 bool FExpressionUnaryOp::PrepareValue(FEmitContext& Context, FEmitScope& Scope, const FRequestedType& RequestedType, FPrepareValueResult& OutResult) const
 {
-	const FPreparedType& InputType = Context.PrepareExpression(Input, Scope, RequestedType);
+	FRequestedType InputRequestedType;
+	switch (Op)
+	{
+	case EUnaryOp::Length:
+	case EUnaryOp::Normalize:
+		// Each component of result is influenced by all components of input
+		InputRequestedType = FRequestedType(4);
+		break;
+	default:
+		InputRequestedType = RequestedType;
+		break;
+	}
+
+	const FPreparedType& InputType = Context.PrepareExpression(Input, Scope, InputRequestedType);
 	if (InputType.IsVoid())
 	{
 		return false;
@@ -730,8 +801,22 @@ bool FExpressionUnaryOp::PrepareValue(FEmitContext& Context, FEmitScope& Scope,
 		return Context.Errors->AddError(TEXT("Invalid operation on non-numeric type"));
 	}
 
+	FPreparedType ResultType;
+	switch (Op)
+	{
+	case EUnaryOp::Length:
+		ResultType = FPreparedType(Shader::MakeValueType(InputType.ValueComponentType, 1), InputType.GetEvaluation(Scope));
+		break;
+	default:
+		ResultType = InputType;
+		if (Op == EUnaryOp::Normalize)
+		{
+			ResultType.ValueComponentType = Shader::MakeNonLWCType(ResultType.ValueComponentType);
+		}
+		break;
+	}
+
 	const FUnaryOpDescription OpDesc = GetUnaryOpDesription(Op);
-	FPreparedType ResultType = InputType;
 	if (OpDesc.PreshaderOpcode == Shader::EPreshaderOpcode::Nop)
 	{
 		// No preshader support
@@ -741,44 +826,96 @@ bool FExpressionUnaryOp::PrepareValue(FEmitContext& Context, FEmitScope& Scope,
 	return OutResult.SetType(Context, RequestedType, ResultType);
 }
 
+namespace Private
+{
+struct FUnaryOpTypes
+{
+	FUnaryOpTypes(const FRequestedValueType& InType) : InputType(InType), ResultType(InType.GetType()) {}
+	FUnaryOpTypes(const FRequestedValueType& InInputType, const Shader::FType& InResultType) : InputType(InInputType), ResultType(InResultType) {}
+
+	FRequestedValueType InputType;
+	Shader::FType ResultType;
+};
+FUnaryOpTypes GetUnaryyOpTypes(const FRequestedType& RequestedType, EUnaryOp Op, const FPreparedType& InputPreparedType)
+{
+	const Shader::EValueComponentType InputComponentType = InputPreparedType.ValueComponentType;
+	const Shader::EValueComponentType NonLWCInputComponentType = Shader::MakeNonLWCType(InputComponentType);
+	const FRequestedValueType InputType(InputComponentType, RequestedType);
+	const FRequestedValueType NonLWCInputType(NonLWCInputComponentType, RequestedType);
+	const int32 NumInputComponents = InputPreparedType.GetNumComponents();
+	const int32 NumRequestedComponents = RequestedType.GetNumComponents();
+
+	switch (Op)
+	{
+	case EUnaryOp::Length:
+		return FUnaryOpTypes(InputPreparedType, Shader::MakeValueType(InputComponentType, 1));
+	case EUnaryOp::Normalize:
+		return FUnaryOpTypes(InputPreparedType, Shader::MakeValueType(NonLWCInputComponentType, NumInputComponents));
+	case EUnaryOp::Frac:
+	case EUnaryOp::Rcp:
+		return FUnaryOpTypes(InputType, Shader::MakeValueType(NonLWCInputComponentType, NumRequestedComponents));
+	default:
+		return InputType;
+	}
+}
+} // namespace Private
+
 void FExpressionUnaryOp::EmitValueShader(FEmitContext& Context, FEmitScope& Scope, const FRequestedType& RequestedType, FEmitValueShaderResult& OutResult) const
 {
-	FEmitShaderExpression* InputValue = Input->GetValueShader(Context, Scope, RequestedType);
-	Shader::FType ResultType = InputValue->Type;
-	const Shader::FValueTypeDescription InputTypeDesc = Shader::GetValueTypeDescription(ResultType);
+	const Private::FUnaryOpTypes Types = Private::GetUnaryyOpTypes(RequestedType, Op, Input->GetPreparedType());
+	const bool bLWC = Types.InputType.ComponentType == Shader::EValueComponentType::Double;
+	FEmitShaderExpression* InputValue = Input->GetValueShader(Context, Scope, Types.InputType);
 
 	switch (Op)
 	{
 	case EUnaryOp::Neg:
-		if (InputTypeDesc.ComponentType == Shader::EValueComponentType::Double)
+		if (bLWC)
 		{
-			OutResult.Code = Context.EmitExpression(Scope, ResultType, TEXT("LWCNegate(%)"), InputValue);
+			OutResult.Code = Context.EmitExpression(Scope, Types.ResultType, TEXT("LWCNegate(%)"), InputValue);
 		}
 		else
 		{
-			OutResult.Code = Context.EmitInlineExpression(Scope, ResultType, TEXT("(-%)"), InputValue);
+			OutResult.Code = Context.EmitInlineExpression(Scope, Types.ResultType, TEXT("(-%)"), InputValue);
 		}
 		break;
 	case EUnaryOp::Rcp:
-		if (InputTypeDesc.ComponentType == Shader::EValueComponentType::Double)
+		if (bLWC)
 		{
-			ResultType = Shader::MakeValueType(Shader::EValueComponentType::Float, ResultType.GetNumComponents());
-			OutResult.Code = Context.EmitExpression(Scope, ResultType, TEXT("LWCRcp(%)"), InputValue);
+			OutResult.Code = Context.EmitExpression(Scope, Types.ResultType, TEXT("LWCRcp(%)"), InputValue);
 		}
 		else
 		{
-			OutResult.Code = Context.EmitExpression(Scope, ResultType, TEXT("rcp(%)"), InputValue);
+			OutResult.Code = Context.EmitExpression(Scope, Types.ResultType, TEXT("rcp(%)"), InputValue);
 		}
 		break;
 	case EUnaryOp::Frac:
-		if (InputTypeDesc.ComponentType == Shader::EValueComponentType::Double)
+		if (bLWC)
+		{
+			OutResult.Code = Context.EmitExpression(Scope, Types.ResultType, TEXT("LWCFrac(%)"), InputValue);
+		}
+		else
+		{
+			OutResult.Code = Context.EmitExpression(Scope, Types.ResultType, TEXT("frac(%)"), InputValue);
+		}
+		break;
+	case EUnaryOp::Length:
+		if (bLWC)
 		{
-			ResultType = Shader::MakeValueType(Shader::EValueComponentType::Float, ResultType.GetNumComponents());
-			OutResult.Code = Context.EmitExpression(Scope, ResultType, TEXT("LWCFrac(%)"), InputValue);
+			OutResult.Code = Context.EmitExpression(Scope, Types.ResultType, TEXT("LWCLength(%)"), InputValue);
 		}
 		else
 		{
-			OutResult.Code = Context.EmitExpression(Scope, ResultType, TEXT("frac(%)"), InputValue);
+			OutResult.Code = Context.EmitExpression(Scope, Types.ResultType, TEXT("length(%)"), InputValue);
+		}
+		break;
+	case EUnaryOp::Normalize:
+		if (bLWC)
+		{
+			OutResult.Code = Context.EmitExpression(Scope, Types.ResultType, TEXT("LWCNormalize(%)"), InputValue);
+		}
+		else
+		{
+			OutResult.Code = Context.EmitExpression(Scope, Types.ResultType, TEXT("normalize(%)"), InputValue);
 		}
 		break;
 	default:
@@ -787,13 +924,15 @@ void FExpressionUnaryOp::EmitValueShader(FEmitContext& Context, FEmitScope& Scop
 	}
 }
 
-void FExpressionUnaryOp::EmitValuePreshader(FEmitContext& Context, FEmitScope& Scope, const FRequestedType& RequestedType, Shader::FPreshaderData& OutPreshader) const
+void FExpressionUnaryOp::EmitValuePreshader(FEmitContext& Context, FEmitScope& Scope, const FRequestedType& RequestedType, FEmitValuePreshaderResult& OutResult) const
 {
+	const Private::FUnaryOpTypes Types = Private::GetUnaryyOpTypes(RequestedType, Op, Input->GetPreparedType());
 	const FUnaryOpDescription OpDesc = GetUnaryOpDesription(Op);
 	check(OpDesc.PreshaderOpcode != Shader::EPreshaderOpcode::Nop);
 
-	Input->GetValuePreshader(Context, Scope, RequestedType, OutPreshader);
-	OutPreshader.WriteOpcode(OpDesc.PreshaderOpcode);
+	Input->GetValuePreshader(Context, Scope, Types.InputType.RequestedType, OutResult.Preshader);
+	OutResult.Preshader.WriteOpcode(OpDesc.PreshaderOpcode);
+	OutResult.Type = Types.ResultType;
 }
 
 FExpression* FTree::NewBinaryOp(EBinaryOp Op, FExpression* Lhs, FExpression* Rhs)
@@ -877,6 +1016,12 @@ void FExpressionBinaryOp::ComputeAnalyticDerivatives(FTree& Tree, FExpressionDer
 			OutResult.ExpressionDdy = Tree.NewExpression<FExpressionSelect>(Cond, LhsDerivatives.ExpressionDdy, RhsDerivatives.ExpressionDdy);
 			break;
 		}
+		case EBinaryOp::VecMulMatrix3:
+		case EBinaryOp::VecMulMatrix4:
+		case EBinaryOp::Matrix3MulVec:
+		case EBinaryOp::Matrix4MulVec:
+			// TODO
+			break;
 		default:
 			checkNoEntry();
 			break;
@@ -891,8 +1036,35 @@ FExpression* FExpressionBinaryOp::ComputePreviousFrame(FTree& Tree, const FReque
 
 bool FExpressionBinaryOp::PrepareValue(FEmitContext& Context, FEmitScope& Scope, const FRequestedType& RequestedType, FPrepareValueResult& OutResult) const
 {
-	const FPreparedType& LhsType = Context.PrepareExpression(Lhs, Scope, RequestedType);
-	const FPreparedType& RhsType = Context.PrepareExpression(Rhs, Scope, RequestedType);
+	FRequestedType LhsRequestedType;
+	FRequestedType RhsRequestedType;
+	bool bIsMatrixOperation = false;
+	switch (Op)
+	{
+	case EBinaryOp::VecMulMatrix3:
+	case EBinaryOp::VecMulMatrix4:
+		LhsRequestedType = FRequestedType(3);
+		RhsRequestedType = FRequestedType(16);
+		bIsMatrixOperation = true;
+		break;
+	case EBinaryOp::Matrix3MulVec:
+	case EBinaryOp::Matrix4MulVec:
+		LhsRequestedType = FRequestedType(16);
+		RhsRequestedType = FRequestedType(3);
+		bIsMatrixOperation = true;
+		break;
+	case EBinaryOp::Dot:
+		LhsRequestedType = FRequestedType(4);
+		RhsRequestedType = FRequestedType(4);
+		break;
+	default:
+		LhsRequestedType = RequestedType;
+		RhsRequestedType = RequestedType;
+		break;
+	}
+
+	const FPreparedType& LhsType = Context.PrepareExpression(Lhs, Scope, LhsRequestedType);
+	const FPreparedType& RhsType = Context.PrepareExpression(Rhs, Scope, RhsRequestedType);
 	if (LhsType.IsVoid() || RhsType.IsVoid())
 	{
 		return false;
@@ -903,18 +1075,27 @@ bool FExpressionBinaryOp::PrepareValue(FEmitContext& Context, FEmitScope& Scope,
 		return Context.Errors->AddError(TEXT("Invalid arithmetic between non-numeric types"));
 	}
 
-	if (LhsType.GetNumComponents() != 1 && RhsType.GetNumComponents() != 1 && LhsType.GetNumComponents() != RhsType.GetNumComponents())
+	const FBinaryOpDescription OpDesc = GetBinaryOpDesription(Op);
+	FPreparedType ResultType;
+	if (bIsMatrixOperation)
 	{
-		return Context.Errors->AddErrorf(TEXT("Invalid arithmetic between %s and %s"), LhsType.GetType().GetName(), RhsType.GetType().GetName());
+		const EExpressionEvaluation Evaluation = CombineEvaluations(LhsType.GetEvaluation(Scope), RhsType.GetEvaluation(Scope));
+		const Shader::EValueComponentType ComponentType = Shader::CombineComponentTypes(LhsType.ValueComponentType, RhsType.ValueComponentType);
+		ResultType = FPreparedType(Shader::MakeValueType(ComponentType, 3), Evaluation);
 	}
-
-	FPreparedType ResultType = MergePreparedTypes(LhsType, RhsType);
-	if (Op == EBinaryOp::Less || Op == EBinaryOp::Greater)
+	else
 	{
-		ResultType.ValueComponentType = Shader::EValueComponentType::Bool;
+		if (LhsType.GetNumComponents() != 1 && RhsType.GetNumComponents() != 1 && LhsType.GetNumComponents() != RhsType.GetNumComponents())
+		{
+			return Context.Errors->AddErrorf(TEXT("Invalid arithmetic between %s and %s"), LhsType.GetType().GetName(), RhsType.GetType().GetName());
+		}
+		ResultType = MergePreparedTypes(LhsType, RhsType);
+		if (Op == EBinaryOp::Less || Op == EBinaryOp::Greater)
+		{
+			ResultType.ValueComponentType = Shader::EValueComponentType::Bool;
+		}
 	}
 
-	const FBinaryOpDescription OpDesc = GetBinaryOpDesription(Op);
 	if (OpDesc.PreshaderOpcode == Shader::EPreshaderOpcode::Nop)
 	{
 		// No preshader support
@@ -928,20 +1109,20 @@ namespace Private
 {
 struct FBinaryOpTypes
 {
-	FBinaryOpTypes(const FRequestedType& InType) : LhsType(InType), RhsType(LhsType), ResultType(InType.GetType()) {}
-	FBinaryOpTypes(const FRequestedType& InLhsType, const FRequestedType& InRhsType, const Shader::FType& InResultType) : LhsType(InLhsType), RhsType(InRhsType), ResultType(InResultType) {}
+	FBinaryOpTypes(const FRequestedValueType& InType) : LhsType(InType), RhsType(LhsType), ResultType(InType.GetType()) {}
+	FBinaryOpTypes(const FRequestedValueType& InLhsType, const FRequestedValueType& InRhsType, const Shader::FType& InResultType) : LhsType(InLhsType), RhsType(InRhsType), ResultType(InResultType) {}
 
-	FRequestedType LhsType;
-	FRequestedType RhsType;
+	FRequestedValueType LhsType;
+	FRequestedValueType RhsType;
 	Shader::FType ResultType;
 };
-FBinaryOpTypes GetBinaryOpTypes(const FRequestedType& RequestedType, EBinaryOp Op, const FRequestedType& LhsType, const FRequestedType& RhsType)
+FBinaryOpTypes GetBinaryOpTypes(const FRequestedType& RequestedType, EBinaryOp Op, const FPreparedType& LhsType, const FPreparedType& RhsType)
 {
-	const int32 NumComponents = RequestedType.GetNumComponents();
 	const Shader::EValueComponentType InputComponentType = Shader::CombineComponentTypes(LhsType.ValueComponentType, RhsType.ValueComponentType);
 	const Shader::EValueComponentType NonLWCInputComponentType = Shader::MakeNonLWCType(InputComponentType);
-	const FRequestedType InputType = MakeRequestedType(InputComponentType, RequestedType);
-	const FRequestedType NonLWCInputType = MakeRequestedType(NonLWCInputComponentType, RequestedType);
+	const FRequestedValueType InputType(InputComponentType, RequestedType);
+	const FRequestedValueType NonLWCInputType(NonLWCInputComponentType, RequestedType);
+	const int32 NumComponents = RequestedType.GetNumComponents();
 
 	switch (Op)
 	{
@@ -952,6 +1133,15 @@ FBinaryOpTypes GetBinaryOpTypes(const FRequestedType& RequestedType, EBinaryOp O
 		return FBinaryOpTypes(InputType, NonLWCInputType, NonLWCInputType.GetType());
 	case EBinaryOp::Dot:
 		return FBinaryOpTypes(InputType, InputType, Shader::MakeValueType(InputComponentType, 1));
+	case EBinaryOp::VecMulMatrix3:
+		// No LWC for matrix3
+		return FBinaryOpTypes(Shader::EValueType::Float3, Shader::EValueType::Float4x4, Shader::EValueType::Float3);
+	case EBinaryOp::VecMulMatrix4:
+		return FBinaryOpTypes(Shader::EValueType::Float3, Shader::EValueType::Float4x4, Shader::EValueType::Float3);
+	case EBinaryOp::Matrix3MulVec:
+	case EBinaryOp::Matrix4MulVec:
+		// No LWC for transpose matrices
+		return FBinaryOpTypes(Shader::EValueType::Float4x4, Shader::EValueType::Float3, Shader::EValueType::Float3);
 	default:
 		return InputType;
 	}
@@ -960,8 +1150,8 @@ FBinaryOpTypes GetBinaryOpTypes(const FRequestedType& RequestedType, EBinaryOp O
 
 void FExpressionBinaryOp::EmitValueShader(FEmitContext& Context, FEmitScope& Scope, const FRequestedType& RequestedType, FEmitValueShaderResult& OutResult) const
 {
-	const Private::FBinaryOpTypes Types = Private::GetBinaryOpTypes(RequestedType, Op, Lhs->GetRequestedType(), Rhs->GetRequestedType());
-	const bool bLWC = Types.LhsType.ValueComponentType == Shader::EValueComponentType::Double;
+	const Private::FBinaryOpTypes Types = Private::GetBinaryOpTypes(RequestedType, Op, Lhs->GetPreparedType(), Rhs->GetPreparedType());
+	const bool bLWC = Types.LhsType.ComponentType == Shader::EValueComponentType::Double;
 	FEmitShaderExpression* LhsValue = Lhs->GetValueShader(Context, Scope, Types.LhsType);
 	FEmitShaderExpression* RhsValue = Rhs->GetValueShader(Context, Scope, Types.RhsType);
 
@@ -1067,25 +1257,52 @@ void FExpressionBinaryOp::EmitValueShader(FEmitContext& Context, FEmitScope& Sco
 			OutResult.Code = Context.EmitExpression(Scope, Types.ResultType, TEXT("(% > %)"), LhsValue, RhsValue);
 		}
 		break;
+	case EBinaryOp::VecMulMatrix3:
+		if (bLWC)
+		{
+			OutResult.Code = Context.EmitExpression(Scope, Types.ResultType, TEXT("LWCMlutiplyVector(%, %)"), LhsValue, RhsValue);
+		}
+		else
+		{
+			OutResult.Code = Context.EmitExpression(Scope, Types.ResultType, TEXT("mul(%, (float3x3)%)"), LhsValue, RhsValue);
+		}
+		break;
+	case EBinaryOp::VecMulMatrix4:
+		if (bLWC)
+		{
+			OutResult.Code = Context.EmitExpression(Scope, Types.ResultType, TEXT("LWCMlutiply(%, %)"), LhsValue, RhsValue);
+		}
+		else
+		{
+			OutResult.Code = Context.EmitExpression(Scope, Types.ResultType, TEXT("mul(%, %)"), LhsValue, RhsValue);
+		}
+		break;
+	case EBinaryOp::Matrix3MulVec:
+		OutResult.Code = Context.EmitExpression(Scope, Types.ResultType, TEXT("mul((float3x3)%, %)"), LhsValue, RhsValue);
+		break;
+	case EBinaryOp::Matrix4MulVec:
+		OutResult.Code = Context.EmitExpression(Scope, Types.ResultType, TEXT("mul(%, %)"), LhsValue, RhsValue);
+		break;
 	default:
 		checkNoEntry();
 		break;
 	}
 }
 
-void FExpressionBinaryOp::EmitValuePreshader(FEmitContext& Context, FEmitScope& Scope, const FRequestedType& RequestedType, Shader::FPreshaderData& OutPreshader) const
+void FExpressionBinaryOp::EmitValuePreshader(FEmitContext& Context, FEmitScope& Scope, const FRequestedType& RequestedType, FEmitValuePreshaderResult& OutResult) const
 {
-	const Private::FBinaryOpTypes Types = Private::GetBinaryOpTypes(RequestedType, Op, Lhs->GetRequestedType(), Rhs->GetRequestedType());
+	const Private::FBinaryOpTypes Types = Private::GetBinaryOpTypes(RequestedType, Op, Lhs->GetPreparedType(), Rhs->GetPreparedType());
 	const FBinaryOpDescription OpDesc = GetBinaryOpDesription(Op);
 	check(OpDesc.PreshaderOpcode != Shader::EPreshaderOpcode::Nop);
 
-	Lhs->GetValuePreshader(Context, Scope, Types.LhsType, OutPreshader);
-	Rhs->GetValuePreshader(Context, Scope, Types.RhsType, OutPreshader);
+	Lhs->GetValuePreshader(Context, Scope, Types.LhsType.RequestedType, OutResult.Preshader);
+	Rhs->GetValuePreshader(Context, Scope, Types.RhsType.RequestedType, OutResult.Preshader);
 
 	check(Context.PreshaderStackPosition > 0);
 	Context.PreshaderStackPosition--;
 
-	OutPreshader.WriteOpcode(OpDesc.PreshaderOpcode);
+	OutResult.Preshader.WriteOpcode(OpDesc.PreshaderOpcode);
+	OutResult.Type = Types.ResultType;
 }
 
 void FExpressionSwizzle::ComputeAnalyticDerivatives(FTree& Tree, FExpressionDerivatives& OutResult) const
@@ -1125,12 +1342,16 @@ void FExpressionSwizzle::EmitValueShader(FEmitContext& Context, FEmitScope& Scop
 
 	static const TCHAR ComponentName[] = { 'x', 'y', 'z', 'w' };
 	TCHAR Swizzle[5] = TEXT("");
+	TCHAR LWCSwizzle[10] = TEXT("");
 	bool bHasSwizzleReorder = false;
 
 	for (int32 ComponentIndex = 0; ComponentIndex < Parameters.NumComponents; ++ComponentIndex)
 	{
 		const int32 SwizzledComponentIndex = Parameters.ComponentIndex[ComponentIndex];
 		Swizzle[ComponentIndex] = ComponentName[SwizzledComponentIndex];
+		LWCSwizzle[ComponentIndex * 2 + 0] = TEXT(',');
+		LWCSwizzle[ComponentIndex * 2 + 1] = TEXT('0') + SwizzledComponentIndex;
+
 		if (SwizzledComponentIndex != ComponentIndex)
 		{
 			bHasSwizzleReorder = true;
@@ -1138,11 +1359,22 @@ void FExpressionSwizzle::EmitValueShader(FEmitContext& Context, FEmitScope& Scop
 	}
 
 	FEmitShaderExpression* InputValue = Input->GetValueShader(Context, Scope, RequestedInputType);
-	if (bHasSwizzleReorder)
+	if (bHasSwizzleReorder || Parameters.NumComponents != InputValue->Type.GetNumComponents())
 	{
-		OutResult.Code = Context.EmitInlineExpression(Scope, RequestedType.GetType(), TEXT("%.%"),
-			InputValue,
-			Swizzle);
+		const Shader::FValueTypeDescription InputTypeDesc = Shader::GetValueTypeDescription(InputValue->Type);
+		const Shader::EValueType ResultType = Shader::MakeValueType(InputTypeDesc.ComponentType, Parameters.NumComponents);
+		if (InputTypeDesc.ComponentType == Shader::EValueComponentType::Double)
+		{
+			OutResult.Code = Context.EmitInlineExpression(Scope, ResultType, TEXT("LWCSwizzle(%%)"),
+				InputValue,
+				LWCSwizzle);
+		}
+		else
+		{
+			OutResult.Code = Context.EmitInlineExpression(Scope, ResultType, TEXT("%.%"),
+				InputValue,
+				Swizzle);
+		}
 	}
 	else
 	{
@@ -1150,17 +1382,20 @@ void FExpressionSwizzle::EmitValueShader(FEmitContext& Context, FEmitScope& Scop
 	}
 }
 
-void FExpressionSwizzle::EmitValuePreshader(FEmitContext& Context, FEmitScope& Scope, const FRequestedType& RequestedType, Shader::FPreshaderData& OutPreshader) const
+void FExpressionSwizzle::EmitValuePreshader(FEmitContext& Context, FEmitScope& Scope, const FRequestedType& RequestedType, FEmitValuePreshaderResult& OutResult) const
 {
 	const FRequestedType RequestedInputType = Parameters.GetRequestedInputType(RequestedType);
 
-	Input->GetValuePreshader(Context, Scope, RequestedInputType, OutPreshader);
-	OutPreshader.WriteOpcode(Shader::EPreshaderOpcode::ComponentSwizzle)
+	const Shader::FType InputType = Input->GetValuePreshader(Context, Scope, RequestedInputType, OutResult.Preshader);
+	const Shader::FValueTypeDescription InputTypeDesc = Shader::GetValueTypeDescription(InputType);
+
+	OutResult.Preshader.WriteOpcode(Shader::EPreshaderOpcode::ComponentSwizzle)
 		.Write((uint8)Parameters.NumComponents)
 		.Write((uint8)Parameters.ComponentIndex[0])
 		.Write((uint8)Parameters.ComponentIndex[1])
 		.Write((uint8)Parameters.ComponentIndex[2])
 		.Write((uint8)Parameters.ComponentIndex[3]);
+	OutResult.Type = Shader::MakeValueType(InputTypeDesc.ComponentType, Parameters.NumComponents);
 }
 
 void FExpressionAppend::ComputeAnalyticDerivatives(FTree& Tree, FExpressionDerivatives& OutResult) const
@@ -1219,27 +1454,28 @@ struct FAppendTypes
 {
 	Shader::FType ResultType;
 	Shader::EValueComponentType ResultComponentType;
-	FRequestedType LhsType;
-	FRequestedType RhsType;
+	FRequestedValueType LhsType;
+	FRequestedValueType RhsType;
 };
-FAppendTypes GetAppendTypes(const FRequestedType& RequestedType, const FRequestedType& LhsType, const FRequestedType& RhsType)
+FAppendTypes GetAppendTypes(const FRequestedType& RequestedType, const FPreparedType& LhsType, const FPreparedType& RhsType)
 {
 	const int32 NumResultComponents = RequestedType.GetNumComponents();
 	const int32 NumLhsComponents = FMath::Min<int32>(NumResultComponents, LhsType.GetNumComponents());
 	const int32 NumRhsComponents = FMath::Min<int32>(NumResultComponents - NumLhsComponents, RhsType.GetNumComponents());
+	const Shader::EValueComponentType ComponentType = Shader::CombineComponentTypes(LhsType.ValueComponentType, RhsType.ValueComponentType);
 
 	FAppendTypes Types;
-	Types.ResultComponentType = RequestedType.ValueComponentType;
-	Types.ResultType = Shader::MakeValueType(RequestedType.ValueComponentType, NumLhsComponents + NumRhsComponents);
-	Types.LhsType.ValueComponentType = RequestedType.ValueComponentType;
-	Types.RhsType.ValueComponentType = RequestedType.ValueComponentType;
+	Types.ResultComponentType = ComponentType;
+	Types.ResultType = Shader::MakeValueType(ComponentType, NumLhsComponents + NumRhsComponents);
+	Types.LhsType.ComponentType = ComponentType;
+	Types.RhsType.ComponentType = ComponentType;
 	for (int32 Index = 0; Index < NumLhsComponents; ++Index)
 	{
-		Types.LhsType.SetComponentRequest(Index, RequestedType.IsComponentRequested(Index));
+		Types.LhsType.RequestedType.SetComponentRequest(Index, RequestedType.IsComponentRequested(Index));
 	}
 	for (int32 Index = 0; Index < NumRhsComponents; ++Index)
 	{
-		Types.RhsType.SetComponentRequest(Index, RequestedType.IsComponentRequested(NumLhsComponents + Index));
+		Types.RhsType.RequestedType.SetComponentRequest(Index, RequestedType.IsComponentRequested(NumLhsComponents + Index));
 	}
 	return Types;
 }
@@ -1247,7 +1483,7 @@ FAppendTypes GetAppendTypes(const FRequestedType& RequestedType, const FRequeste
 
 void FExpressionAppend::EmitValueShader(FEmitContext& Context, FEmitScope& Scope, const FRequestedType& RequestedType, FEmitValueShaderResult& OutResult) const
 {
-	const Private::FAppendTypes Types = Private::GetAppendTypes(RequestedType, Lhs->GetRequestedType(), Rhs->GetRequestedType());
+	const Private::FAppendTypes Types = Private::GetAppendTypes(RequestedType, Lhs->GetPreparedType(), Rhs->GetPreparedType());
 	FEmitShaderExpression* LhsValue = Lhs->GetValueShader(Context, Scope, Types.LhsType);
 
 	if (Types.RhsType.IsVoid())
@@ -1273,19 +1509,20 @@ void FExpressionAppend::EmitValueShader(FEmitContext& Context, FEmitScope& Scope
 	}
 }
 
-void FExpressionAppend::EmitValuePreshader(FEmitContext& Context, FEmitScope& Scope, const FRequestedType& RequestedType, Shader::FPreshaderData& OutPreshader) const
+void FExpressionAppend::EmitValuePreshader(FEmitContext& Context, FEmitScope& Scope, const FRequestedType& RequestedType, FEmitValuePreshaderResult& OutResult) const
 {
-	const Private::FAppendTypes Types = Private::GetAppendTypes(RequestedType, Lhs->GetRequestedType(), Rhs->GetRequestedType());
-	Lhs->GetValuePreshader(Context, Scope, Types.LhsType, OutPreshader);
+	const Private::FAppendTypes Types = Private::GetAppendTypes(RequestedType, Lhs->GetPreparedType(), Rhs->GetPreparedType());
+	Lhs->GetValuePreshader(Context, Scope, Types.LhsType.RequestedType, OutResult.Preshader);
 	if (Types.RhsType.IsVoid())
 	{
-		Rhs->GetValuePreshader(Context, Scope, Types.RhsType, OutPreshader);
+		Rhs->GetValuePreshader(Context, Scope, Types.RhsType.RequestedType, OutResult.Preshader);
 
 		check(Context.PreshaderStackPosition > 0);
 		Context.PreshaderStackPosition--;
 
-		OutPreshader.WriteOpcode(Shader::EPreshaderOpcode::AppendVector);
+		OutResult.Preshader.WriteOpcode(Shader::EPreshaderOpcode::AppendVector);
 	}
+	OutResult.Type = Types.ResultType;
 }
 
 bool FExpressionReflectionVector::PrepareValue(FEmitContext& Context, FEmitScope& Scope, const FRequestedType& RequestedType, FPrepareValueResult& OutResult) const
@@ -1329,13 +1566,13 @@ void FStatementReturn::EmitShader(FEmitContext& Context, FEmitScope& Scope) cons
 
 bool FStatementIf::Prepare(FEmitContext& Context, FEmitScope& Scope) const
 {
-	const FPreparedType& ConditionType = Context.PrepareExpression(ConditionExpression, Scope, Shader::EValueType::Bool1);
+	const FPreparedType& ConditionType = Context.PrepareExpression(ConditionExpression, Scope, ERequestedType::Scalar);
 	if (ConditionType.IsVoid())
 	{
 		return false;
 	}
 
-	const EExpressionEvaluation ConditionEvaluation = ConditionType.GetEvaluation(Scope, Shader::EValueType::Bool1);
+	const EExpressionEvaluation ConditionEvaluation = ConditionType.GetEvaluation(Scope, ERequestedType::Scalar);
 	check(ConditionEvaluation != EExpressionEvaluation::None);
 	if (ConditionEvaluation == EExpressionEvaluation::Constant)
 	{
@@ -1363,7 +1600,7 @@ bool FStatementIf::Prepare(FEmitContext& Context, FEmitScope& Scope) const
 void FStatementIf::EmitShader(FEmitContext& Context, FEmitScope& Scope) const
 {
 	FEmitShaderNode* Dependency = nullptr;
-	const EExpressionEvaluation ConditionEvaluation = ConditionExpression->GetEvaluation(Scope, Shader::EValueType::Bool1);
+	const EExpressionEvaluation ConditionEvaluation = ConditionExpression->GetEvaluation(Scope, ERequestedType::Scalar);
 	if (ConditionEvaluation == EExpressionEvaluation::Constant)
 	{
 		const bool bCondition = ConditionExpression->GetValueConstant(Context, Scope, Shader::EValueType::Bool1).AsBoolScalar();
@@ -1387,7 +1624,7 @@ void FStatementIf::EmitShader(FEmitContext& Context, FEmitScope& Scope) const
 
 void FStatementIf::EmitPreshader(FEmitContext& Context, FEmitScope& Scope, const FRequestedType& RequestedType, TArrayView<const FEmitPreshaderScope> Scopes, Shader::FPreshaderData& OutPreshader) const
 {
-	ConditionExpression->GetValuePreshader(Context, Scope, Shader::EValueType::Bool1, OutPreshader);
+	ConditionExpression->GetValuePreshader(Context, Scope, FRequestedType(1), OutPreshader);
 
 	check(Context.PreshaderStackPosition > 0);
 	Context.PreshaderStackPosition--;
diff --git a/Engine/Source/Runtime/Engine/Private/HLSLTree/HLSLTreeEmit.cpp b/Engine/Source/Runtime/Engine/Private/HLSLTree/HLSLTreeEmit.cpp
--- a/Engine/Source/Runtime/Engine/Private/HLSLTree/HLSLTreeEmit.cpp
+++ b/Engine/Source/Runtime/Engine/Private/HLSLTree/HLSLTreeEmit.cpp
@@ -236,9 +236,7 @@ FPreparedType FEmitContext::PrepareExpression(FExpression* InExpression, FEmitSc
 	{
 		ResultType = Result.PreparedType;
 		check(!ResultType.IsVoid());
-		ResultType.SetForwardValue(InExpression, RequestedType);
 	}
-
 	return ResultType;
 }
 
@@ -661,9 +659,12 @@ void WriteMaterialUniformAccess(Shader::EValueComponentType ComponentType, uint3
 FEmitShaderExpression* FEmitContext::EmitPreshaderOrConstant(FEmitScope& Scope, const FRequestedType& RequestedType, FExpression* Expression)
 {
 	Shader::FPreshaderData LocalPreshader;
-	Expression->EmitValuePreshader(*this, Scope, RequestedType, LocalPreshader);
-
-	const Shader::FType Type = RequestedType.GetType();
+	Shader::FType Type;
+	{
+		FEmitValuePreshaderResult PreshaderResult(LocalPreshader);
+		Expression->EmitValuePreshader(*this, Scope, RequestedType, PreshaderResult);
+		Type = PreshaderResult.Type;
+	}
 
 	FXxHash64Builder Hasher;
 	Hasher.Update(&Type, sizeof(Type));
@@ -714,7 +715,8 @@ FEmitShaderExpression* FEmitContext::EmitPreshaderOrConstant(FEmitScope& Scope,
 				PreshaderHeader->FieldIndex = UniformExpressionSet.UniformPreshaderFields.Num();
 				PreshaderHeader->NumFields = 0u;
 				PreshaderHeader->OpcodeOffset = UniformExpressionSet.UniformPreshaderData.Num();
-				Expression->EmitValuePreshader(*this, Scope, RequestedType, UniformExpressionSet.UniformPreshaderData);
+				FEmitValuePreshaderResult PreshaderResult(UniformExpressionSet.UniformPreshaderData);
+				Expression->EmitValuePreshader(*this, Scope, RequestedType, PreshaderResult);
 				PreshaderHeader->OpcodeSize = UniformExpressionSet.UniformPreshaderData.Num() - PreshaderHeader->OpcodeOffset;
 			}
 
diff --git a/Engine/Source/Runtime/Engine/Private/Materials/MaterialExpressionHLSL.cpp b/Engine/Source/Runtime/Engine/Private/Materials/MaterialExpressionHLSL.cpp
--- a/Engine/Source/Runtime/Engine/Private/Materials/MaterialExpressionHLSL.cpp
+++ b/Engine/Source/Runtime/Engine/Private/Materials/MaterialExpressionHLSL.cpp
@@ -22,6 +22,7 @@
 #include "Materials/MaterialExpressionStaticBoolParameter.h"
 #include "Materials/MaterialExpressionPixelDepth.h"
 #include "Materials/MaterialExpressionWorldPosition.h"
+#include "Materials/MaterialExpressionCameraPositionWS.h"
 #include "Materials/MaterialExpressionTime.h"
 #include "Materials/MaterialExpressionDeltaTime.h"
 #include "Materials/MaterialExpressionPanner.h"
@@ -45,12 +46,16 @@
 #include "Materials/MaterialExpressionMax.h"
 #include "Materials/MaterialExpressionClamp.h"
 #include "Materials/MaterialExpressionLinearInterpolate.h"
+#include "Materials/MaterialExpressionDistance.h"
+#include "Materials/MaterialExpressionNormalize.h"
 #include "Materials/MaterialExpressionBinaryOp.h"
 #include "Materials/MaterialExpressionAppendVector.h"
 #include "Materials/MaterialExpressionComponentMask.h"
 #include "Materials/MaterialExpressionGetMaterialAttributes.h"
 #include "Materials/MaterialExpressionSetMaterialAttributes.h"
 #include "Materials/MaterialExpressionReflectionVectorWS.h"
+#include "Materials/MaterialExpressionTransform.h"
+#include "Materials/MaterialExpressionTransformPosition.h"
 #include "Materials/MaterialExpressionSetLocal.h"
 #include "Materials/MaterialExpressionIfThenElse.h"
 #include "Materials/MaterialExpressionForLoop.h"
@@ -177,6 +182,13 @@ bool UMaterialExpressionWorldPosition::GenerateHLSLExpression(FMaterialHLSLGener
 	return true;
 }
 
+bool UMaterialExpressionCameraPositionWS::GenerateHLSLExpression(FMaterialHLSLGenerator& Generator, UE::HLSLTree::FScope& Scope, int32 OutputIndex, UE::HLSLTree::FExpression*& OutExpression)
+{
+	using namespace UE::HLSLTree;
+	OutExpression = Generator.GetTree().NewExpression<FExpressionExternalInput>(EExternalInput::CameraWorldPosition);
+	return true;
+}
+
 bool UMaterialExpressionTime::GenerateHLSLExpression(FMaterialHLSLGenerator& Generator, UE::HLSLTree::FScope& Scope, int32 OutputIndex, UE::HLSLTree::FExpression*& OutExpression)
 {
 	using namespace UE::HLSLTree;
@@ -456,6 +468,31 @@ bool UMaterialExpressionLinearInterpolate::GenerateHLSLExpression(FMaterialHLSLG
 	return true;
 }
 
+bool UMaterialExpressionDistance::GenerateHLSLExpression(FMaterialHLSLGenerator& Generator, UE::HLSLTree::FScope& Scope, int32 OutputIndex, UE::HLSLTree::FExpression*& OutExpression)
+{
+	UE::HLSLTree::FExpression* ExpressionA = A.AcquireHLSLExpression(Generator, Scope);
+	UE::HLSLTree::FExpression* ExpressionB = B.AcquireHLSLExpression(Generator, Scope);
+	if (!ExpressionA || !ExpressionB)
+	{
+		return false;
+	}
+
+	OutExpression = Generator.GetTree().NewLength(Generator.GetTree().NewSub(ExpressionA, ExpressionB));
+	return true;
+}
+
+bool UMaterialExpressionNormalize::GenerateHLSLExpression(FMaterialHLSLGenerator& Generator, UE::HLSLTree::FScope& Scope, int32 OutputIndex, UE::HLSLTree::FExpression*& OutExpression)
+{
+	UE::HLSLTree::FExpression* ExpressionInput = VectorInput.AcquireHLSLExpression(Generator, Scope);
+	if (!ExpressionInput)
+	{
+		return false;
+	}
+
+	OutExpression = Generator.GetTree().NewNormalize(ExpressionInput);
+	return true;
+}
+
 bool UMaterialExpressionAppendVector::GenerateHLSLExpression(FMaterialHLSLGenerator& Generator, UE::HLSLTree::FScope& Scope, int32 OutputIndex, UE::HLSLTree::FExpression*& OutExpression)
 {
 	UE::HLSLTree::FExpression* Lhs = A.AcquireHLSLExpression(Generator, Scope);
@@ -564,6 +601,245 @@ bool UMaterialExpressionMaterialFunctionCall::GenerateHLSLExpression(FMaterialHL
 	return true;
 }
 
+static UE::HLSLTree::FExpression* TransformBase(UE::HLSLTree::FTree& Tree,
+	EMaterialCommonBasis SourceCoordBasis,
+	EMaterialCommonBasis DestCoordBasis,
+	UE::HLSLTree::FExpression* Input,
+	bool bWComponent)
+{
+	using namespace UE::HLSLTree;
+	if (!Input)
+	{
+		// unable to compile
+		return nullptr;
+	}
+
+	if (SourceCoordBasis == DestCoordBasis)
+	{
+		// no transformation needed
+		return Input;
+	}
+
+	FExpression* Result = nullptr;
+	EMaterialCommonBasis IntermediaryBasis = MCB_World;
+	EBinaryOp Op = bWComponent ? EBinaryOp::VecMulMatrix4 : EBinaryOp::VecMulMatrix3;
+	switch (SourceCoordBasis)
+	{
+	case MCB_Tangent:
+	{
+		check(!bWComponent);
+		if (DestCoordBasis == MCB_World)
+		{
+			Result = Tree.NewBinaryOp(Op, Input, Tree.NewExpression<FExpressionExternalInput>(EExternalInput::TangentToWorld));
+			//CodeStr = TEXT("mul(<A>, Parameters.TangentToWorld)");
+		}
+		// else use MCB_World as intermediary basis
+		break;
+	}
+	case MCB_Local:
+	{
+		if (DestCoordBasis == MCB_World)
+		{
+			Result = Tree.NewBinaryOp(Op, Input, Tree.NewExpression<FExpressionExternalInput>(EExternalInput::LocalToWorld));
+			//CodeStr = TEXT("TransformLocal<TO><PREV>World(Parameters, <A>)");
+		}
+		// else use MCB_World as intermediary basis
+		break;
+	}
+	case MCB_TranslatedWorld:
+	{
+		if (DestCoordBasis == MCB_World)
+		{
+			if (bWComponent)
+			{
+				Result = Tree.NewSub(Input, Tree.NewExpression<FExpressionExternalInput>(EExternalInput::PreViewTranslation));
+				//CodeStr = TEXT("LWCSubtract(<A>, ResolvedView.<PREV>PreViewTranslation)");
+			}
+			else
+			{
+				Result = Input;
+				//CodeStr = TEXT("<A>");
+			}
+		}
+		else if (DestCoordBasis == MCB_Camera)
+		{
+			Result = Tree.NewBinaryOp(Op, Input, Tree.NewExpression<FExpressionExternalInput>(EExternalInput::TranslatedWorldToCameraView));
+			//CodeStr = MultiplyMatrix(TEXT("<A>"), TEXT("ResolvedView.<PREV>TranslatedWorldToCameraView"), AWComponent);
+		}
+		else if (DestCoordBasis == MCB_View)
+		{
+			Result = Tree.NewBinaryOp(Op, Input, Tree.NewExpression<FExpressionExternalInput>(EExternalInput::TranslatedWorldToView));
+			//CodeStr = MultiplyMatrix(TEXT("<A>"), TEXT("ResolvedView.<PREV>TranslatedWorldToView"), AWComponent);
+		}
+		// else use MCB_World as intermediary basis
+		break;
+	}
+	case MCB_World:
+	{
+		if (DestCoordBasis == MCB_Tangent)
+		{
+			Result = Tree.NewBinaryOp(Op, Tree.NewExpression<FExpressionExternalInput>(EExternalInput::TangentToWorld), Input);
+			//CodeStr = MultiplyTransposeMatrix(TEXT("Parameters.TangentToWorld"), TEXT("<A>"), AWComponent);
+		}
+		else if (DestCoordBasis == MCB_Local)
+		{
+			Result = Tree.NewBinaryOp(Op, Input, Tree.NewExpression<FExpressionExternalInput>(EExternalInput::WorldToLocal));
+			/*const EMaterialDomain Domain = (const EMaterialDomain)Material->GetMaterialDomain();
+
+			if (Domain != MD_Surface && Domain != MD_Volume)
+			{
+				// TODO: for decals we could support it
+				Errorf(TEXT("This transformation is only supported in the 'Surface' material domain."));
+				return INDEX_NONE;
+			}
+
+			// TODO: inconsistent with TransformLocal<TO>World with instancing
+			CodeStr = LWCMultiplyMatrix(TEXT("<A>"), TEXT("GetPrimitiveData(Parameters).<PREVIOUS>WorldToLocal"), AWComponent);*/
+		}
+		else if (DestCoordBasis == MCB_TranslatedWorld)
+		{
+			if (bWComponent)
+			{
+				// TODO - explicit cast to float
+				Result = Tree.NewAdd(Input, Tree.NewExpression<FExpressionExternalInput>(EExternalInput::PreViewTranslation));
+				//CodeStr = TEXT("LWCToFloat(LWCAdd(<A>, ResolvedView.<PREV>PreViewTranslation))");
+			}
+			else
+			{
+				Result = Input;
+				//CodeStr = TEXT("<A>");
+			}
+		}
+		else if (DestCoordBasis == MCB_MeshParticle)
+		{
+			//CodeStr = LWCMultiplyMatrix(TEXT("<A>"), TEXT("Parameters.Particle.WorldToParticle"), AWComponent);
+			Result = Tree.NewBinaryOp(Op, Input, Tree.NewExpression<FExpressionExternalInput>(EExternalInput::WorldToParticle));
+			//bUsesParticleWorldToLocal = true;
+		}
+		else if (DestCoordBasis == MCB_Instance)
+		{
+			Result = Tree.NewBinaryOp(Op, Input, Tree.NewExpression<FExpressionExternalInput>(EExternalInput::WorldToInstance));
+			//CodeStr = LWCMultiplyMatrix(TEXT("<A>"), TEXT("GetWorldToInstance(Parameters)"), AWComponent);
+			//bUsesInstanceWorldToLocalPS = ShaderFrequency == SF_Pixel;
+		}
+
+		// else use MCB_TranslatedWorld as intermediary basis
+		IntermediaryBasis = MCB_TranslatedWorld;
+		break;
+	}
+	case MCB_Camera:
+	{
+		if (DestCoordBasis == MCB_TranslatedWorld)
+		{
+			Result = Tree.NewBinaryOp(Op, Input, Tree.NewExpression<FExpressionExternalInput>(EExternalInput::CameraViewToTranslatedWorld));
+			//CodeStr = MultiplyMatrix(TEXT("<A>"), TEXT("ResolvedView.<PREV>CameraViewToTranslatedWorld"), AWComponent);
+		}
+		// else use MCB_TranslatedWorld as intermediary basis
+		IntermediaryBasis = MCB_TranslatedWorld;
+		break;
+	}
+	case MCB_View:
+	{
+		if (DestCoordBasis == MCB_TranslatedWorld)
+		{
+			Result = Tree.NewBinaryOp(Op, Input, Tree.NewExpression<FExpressionExternalInput>(EExternalInput::ViewToTranslatedWorld));
+			//CodeStr = MultiplyMatrix(TEXT("<A>"), TEXT("ResolvedView.<PREV>ViewToTranslatedWorld"), AWComponent);
+		}
+		// else use MCB_TranslatedWorld as intermediary basis
+		IntermediaryBasis = MCB_TranslatedWorld;
+		break;
+	}
+	case MCB_MeshParticle:
+	{
+		if (DestCoordBasis == MCB_World)
+		{
+			Result = Tree.NewBinaryOp(Op, Input, Tree.NewExpression<FExpressionExternalInput>(EExternalInput::ParticleToWorld));
+			//CodeStr = LWCMultiplyMatrix(TEXT("<A>"), TEXT("Parameters.Particle.ParticleToWorld"), AWComponent);
+			//bUsesParticleLocalToWorld = true;
+		}
+		// use World as an intermediary base
+		break;
+	}
+	case MCB_Instance:
+	{
+		if (DestCoordBasis == MCB_World)
+		{
+			Result = Tree.NewBinaryOp(Op, Input, Tree.NewExpression<FExpressionExternalInput>(EExternalInput::InstanceToWorld));
+			//CodeStr = LWCMultiplyMatrix(TEXT("<A>"), TEXT("GetInstanceToWorld(Parameters)"), AWComponent);
+			//bUsesInstanceLocalToWorldPS = ShaderFrequency == SF_Pixel;
+		}
+		// use World as an intermediary base
+		break;
+	}
+
+	default:
+		check(0);
+		break;
+	}
+
+	if (!Result)
+	{
+		// check intermediary basis so we don't have infinite recursion
+		check(IntermediaryBasis != SourceCoordBasis);
+		check(IntermediaryBasis != DestCoordBasis);
+
+		// use intermediary basis
+		FExpression* IntermediaryExpression = TransformBase(Tree, SourceCoordBasis, IntermediaryBasis, Input, bWComponent);
+		return TransformBase(Tree, IntermediaryBasis, DestCoordBasis, IntermediaryExpression, bWComponent);
+	}
+
+	return Result;
+}
+
+bool UMaterialExpressionTransform::GenerateHLSLExpression(FMaterialHLSLGenerator& Generator, UE::HLSLTree::FScope& Scope, int32 OutputIndex, UE::HLSLTree::FExpression*& OutExpression)
+{
+	using namespace UE::HLSLTree;
+
+	static const EMaterialCommonBasis kTable[TRANSFORM_MAX] = {
+		MCB_Tangent,					// TRANSFORM_Tangent
+		MCB_Local,						// TRANSFORM_Local
+		MCB_World,						// TRANSFORM_World
+		MCB_View,						// TRANSFORM_View
+		MCB_Camera,						// TRANSFORM_Camera
+		MCB_MeshParticle,				// TRANSFORM_Particle
+		MCB_Instance,					// TRANSFORM_Instance
+	};
+
+	FExpression* ExpressionInput = Input.AcquireHLSLExpression(Generator, Scope);
+	if (!ExpressionInput)
+	{
+		return false;
+	}
+
+	OutExpression = TransformBase(Generator.GetTree(), kTable[TransformSourceType], kTable[TransformType], ExpressionInput, false);
+	return true;
+}
+
+
+bool UMaterialExpressionTransformPosition::GenerateHLSLExpression(FMaterialHLSLGenerator& Generator, UE::HLSLTree::FScope& Scope, int32 OutputIndex, UE::HLSLTree::FExpression*& OutExpression)
+{
+	using namespace UE::HLSLTree;
+
+	static const EMaterialCommonBasis kTable[TRANSFORMPOSSOURCE_MAX] = {
+		MCB_Local,						// TRANSFORMPOSSOURCE_Local
+		MCB_World,						// TRANSFORMPOSSOURCE_World
+		MCB_TranslatedWorld,			// TRANSFORMPOSSOURCE_TranslatedWorld
+		MCB_View,						// TRANSFORMPOSSOURCE_View
+		MCB_Camera,						// TRANSFORMPOSSOURCE_Camera
+		MCB_MeshParticle,				// TRANSFORMPOSSOURCE_Particle
+		MCB_Instance,					// TRANSFORMPOSSOURCE_Instance
+	};
+
+	FExpression* ExpressionInput = Input.AcquireHLSLExpression(Generator, Scope);
+	if (!ExpressionInput)
+	{
+		return false;
+	}
+
+	OutExpression = TransformBase(Generator.GetTree(), kTable[TransformSourceType], kTable[TransformType], ExpressionInput, true);
+	return true;
+}
+
 bool UMaterialExpressionExecBegin::GenerateHLSLStatements(FMaterialHLSLGenerator& Generator, UE::HLSLTree::FScope& Scope)
 {
 	Exec.GenerateHLSLStatements(Generator, Scope);
diff --git a/Engine/Source/Runtime/Engine/Private/Materials/MaterialHLSLEmitter.cpp b/Engine/Source/Runtime/Engine/Private/Materials/MaterialHLSLEmitter.cpp
--- a/Engine/Source/Runtime/Engine/Private/Materials/MaterialHLSLEmitter.cpp
+++ b/Engine/Source/Runtime/Engine/Private/Materials/MaterialHLSLEmitter.cpp
@@ -732,7 +732,7 @@ static bool IsAttributeUsed(FMaterialHLSLGenerator& Generator,
 	}
 	else if (Evaluation == EExpressionEvaluation::Constant)
 	{
-		FRequestedType RequestedType(Generator.GetMaterialAttributesType(), false);
+		FRequestedType RequestedType;
 		RequestedType.SetFieldRequested(PropertyField);
 
 		const FValue& DefaultValue = Generator.GetMaterialAttributesDefaultValue();
@@ -815,7 +815,7 @@ bool MaterialEmitHLSL(const FMaterialCompileTargetParameters& InCompilerTarget,
 		FEmitScope* EmitResultScope = EmitContext.PrepareScope(&Generator.GetResultStatement()->GetParentScope());
 
 		// Prepare all fields *except* normal
-		FRequestedType RequestedPixelAttributesType(Generator.GetMaterialAttributesType(), false);
+		FRequestedType RequestedPixelAttributesType;
 		for (const FGuid& AttributeID : OrderedVisibleAttributes)
 		{
 			if (FMaterialAttributeDefinitionMap::GetShaderFrequency(AttributeID) == SF_Pixel)
@@ -860,7 +860,7 @@ bool MaterialEmitHLSL(const FMaterialCompileTargetParameters& InCompilerTarget,
 			EmitResultScope = EmitContext.PrepareScope(&Generator.GetResultStatement()->GetParentScope());
 
 			// Prepare code for just the normal
-			FRequestedType RequestedMaterialNormal(Generator.GetMaterialAttributesType(), false);
+			FRequestedType RequestedMaterialNormal;
 			RequestedMaterialNormal.SetFieldRequested(NormalField);
 			const FPreparedType PixelResultType1 = EmitContext.PrepareExpression(Generator.GetResultExpression(), *EmitResultScope, RequestedMaterialNormal);
 			if (PixelResultType1.IsVoid())
@@ -876,7 +876,7 @@ bool MaterialEmitHLSL(const FMaterialCompileTargetParameters& InCompilerTarget,
 	// Prepare vertex shader code
 	FStringBuilderMemstack VertexCode(Allocator, 1024 * 1024);
 	{
-		FRequestedType RequestedVertexAttributesType(Generator.GetMaterialAttributesType(), false);
+		FRequestedType RequestedVertexAttributesType;
 		for (const FGuid& AttributeID : OrderedVisibleAttributes)
 		{
 			if (FMaterialAttributeDefinitionMap::GetShaderFrequency(AttributeID) == SF_Vertex)
diff --git a/Engine/Source/Runtime/Engine/Private/Materials/MaterialHLSLGenerator.cpp b/Engine/Source/Runtime/Engine/Private/Materials/MaterialHLSLGenerator.cpp
--- a/Engine/Source/Runtime/Engine/Private/Materials/MaterialHLSLGenerator.cpp
+++ b/Engine/Source/Runtime/Engine/Private/Materials/MaterialHLSLGenerator.cpp
@@ -283,7 +283,7 @@ bool FMaterialHLSLGenerator::GenerateResult(UE::HLSLTree::FScope& Scope)
 					const FStructField* WPOField = GetMaterialAttributesType()->FindFieldByName(*WPOName);
 					const FStructField* PrevWPOField = GetMaterialAttributesType()->FindFieldByName(TEXT("PrevWorldPositionOffset"));
 
-					FRequestedType PrevRequestedType(GetMaterialAttributesType(), false);
+					FRequestedType PrevRequestedType;
 					PrevRequestedType.SetFieldRequested(WPOField);
 
 					FExpression* PrevAttributesExpression = HLSLTree->GetPreviousFrame(AttributesExpression, PrevRequestedType);
@@ -439,7 +439,7 @@ UE::HLSLTree::FExpression* FMaterialHLSLGenerator::AcquireFunctionInputExpressio
 		InputExpression = MaterialExpression->Preview.TryAcquireHLSLExpression(*this, Scope);
 		if (!InputExpression)
 		{
-			const FVector4& PreviewValue = MaterialExpression->PreviewValue;
+			const FVector4f PreviewValue(MaterialExpression->PreviewValue);
 			UE::Shader::FValue DefaultValue;
 			switch (MaterialExpression->InputType)
 			{
diff --git a/Engine/Source/Runtime/Engine/Private/Shader/ShaderValue.cpp b/Engine/Source/Runtime/Engine/Private/Shader/ShaderValue.cpp
--- a/Engine/Source/Runtime/Engine/Private/Shader/ShaderValue.cpp
+++ b/Engine/Source/Runtime/Engine/Private/Shader/ShaderValue.cpp
@@ -207,11 +207,12 @@ void AsType(const Operation& Op, const FValue& Value, ResultType& OutResult)
 	}
 	else
 	{
-		for (int32 i = 0; i < TypeDesc.NumComponents; ++i)
+		const int32 NumComponents = FMath::Min<int32>(TypeDesc.NumComponents, 4);
+		for (int32 i = 0; i < NumComponents; ++i)
 		{
 			OutResult[i] = Op(TypeDesc.ComponentType, Value.Component[i]);
 		}
-		for (int32 i = TypeDesc.NumComponents; i < 4; ++i)
+		for (int32 i = NumComponents; i < 4; ++i)
 		{
 			OutResult[i] = (FComponentType)0;
 		}
@@ -585,6 +586,9 @@ FValueTypeDescription GetValueTypeDescription(EValueType Type)
 	case EValueType::Bool2: return FValueTypeDescription(TEXT("bool2"), EValueComponentType::Bool, 2);
 	case EValueType::Bool3: return FValueTypeDescription(TEXT("bool3"), EValueComponentType::Bool, 3);
 	case EValueType::Bool4: return FValueTypeDescription(TEXT("bool4"), EValueComponentType::Bool, 4);
+	case EValueType::Float4x4: return FValueTypeDescription(TEXT("float4x4"), EValueComponentType::Float, 16);
+	case EValueType::Double4x4: return FValueTypeDescription(TEXT("FLWCMatrix"), EValueComponentType::Double, 16);
+	case EValueType::DoubleInverse4x4: return FValueTypeDescription(TEXT("FLWCInverseMatrix"), EValueComponentType::Double, 16);
 	case EValueType::Struct: return FValueTypeDescription(TEXT("struct"), EValueComponentType::Void, 0);
 	default: checkNoEntry(); return FValueTypeDescription(TEXT("<INVALID>"), EValueComponentType::Void, 0);
 	}
@@ -606,6 +610,7 @@ EValueType MakeValueType(EValueComponentType ComponentType, int32 NumComponents)
 		case 2: return EValueType::Float2;
 		case 3: return EValueType::Float3;
 		case 4: return EValueType::Float4;
+		case 16: return EValueType::Float4x4;
 		default: break;
 		}
 	case EValueComponentType::Double:
@@ -615,6 +620,7 @@ EValueType MakeValueType(EValueComponentType ComponentType, int32 NumComponents)
 		case 2: return EValueType::Double2;
 		case 3: return EValueType::Double3;
 		case 4: return EValueType::Double4;
+		case 16: return EValueType::Double4x4;
 		default: break;
 		}
 	case EValueComponentType::Int:
diff --git a/Engine/Source/Runtime/Engine/Public/HLSLTree/HLSLTree.h b/Engine/Source/Runtime/Engine/Public/HLSLTree/HLSLTree.h
--- a/Engine/Source/Runtime/Engine/Public/HLSLTree/HLSLTree.h
+++ b/Engine/Source/Runtime/Engine/Public/HLSLTree/HLSLTree.h
@@ -230,20 +230,28 @@ private:
 };
 
 /**
- * Like Shader::FType, but tracks which individual components are needed
+ * Used to initialize FRequestedType with a fixed number of components
+ * enum value is the number of parameters
+ */
+enum class ERequestedType : uint8
+{
+	Void = 0u,
+	Scalar = 1u,
+	Vector2 = 2u,
+	Vector3 = 3u,
+	Vector4 = 4u,
+	Matrix4x4 = 16u,
+};
+
+/**
+ * Represents a request for an arbitrary set of untyped components
  */
 class FRequestedType
 {
 public:
 	FRequestedType() = default;
-	FRequestedType(int32 NumComponents, bool bDefaultRequest = true);
-	FRequestedType(const Shader::FType& InType, bool bDefaultRequest = true);
-	FRequestedType(const Shader::EValueType& InType, bool bDefaultRequest = true);
-	
-	const TCHAR* GetName() const { return GetType().GetName(); }
-	bool IsStruct() const { return StructType != nullptr; }
-	const Shader::FStructType* GetStructType() const { return StructType; }
-	const Shader::FType GetType() const;
+	FRequestedType(ERequestedType InType);
+	explicit FRequestedType(int32 NumComponents);
 
 	int32 GetNumComponents() const;
 	bool IsComponentRequested(int32 Index) const { return RequestedComponents.IsValidIndex(Index) ? (bool)RequestedComponents[Index] : false; }
@@ -253,8 +261,6 @@ public:
 
 	void Reset()
 	{
-		StructType = nullptr;
-		ValueComponentType = Shader::EValueComponentType::Void;
 		RequestedComponents.Reset();
 	}
 
@@ -271,21 +277,12 @@ public:
 	/** Returns the requested type of the given field */
 	FRequestedType GetField(const Shader::FStructField* Field) const;
 
-	/**
-	 * If either StructType or ValueComponentType are set, then the request is for an explicit type
-	 * Otherwise, the request is for any type with the given components
-	 */
-	const Shader::FStructType* StructType = nullptr;
-	Shader::EValueComponentType ValueComponentType = Shader::EValueComponentType::Void;
-
 	/** 1 bit per component, a value of 'true' means the specified component is requsted */
 	TBitArray<> RequestedComponents;
 };
 inline bool operator==(const FRequestedType& Lhs, const FRequestedType& Rhs)
 {
-	return Lhs.ValueComponentType == Rhs.ValueComponentType &&
-		Lhs.StructType == Rhs.StructType &&
-		Lhs.RequestedComponents == Rhs.RequestedComponents;
+	return Lhs.RequestedComponents == Rhs.RequestedComponents;
 }
 inline bool operator!=(const FRequestedType& Lhs, const FRequestedType& Rhs)
 {
@@ -294,13 +291,9 @@ inline bool operator!=(const FRequestedType& Lhs, const FRequestedType& Rhs)
 
 inline void AppendHash(FHasher& Hasher, const FRequestedType& Value)
 {
-	AppendHash(Hasher, Value.StructType);
-	AppendHash(Hasher, Value.ValueComponentType);
 	AppendHash(Hasher, Value.RequestedComponents);
 }
 
-FRequestedType MakeRequestedType(Shader::EValueComponentType ComponentType, const FRequestedType& RequestedComponents);
-
 struct FPreparedComponent
 {
 	FPreparedComponent() = default;
@@ -342,19 +335,21 @@ public:
 	FPreparedType() = default;
 	FPreparedType(Shader::EValueComponentType InComponentType) : ValueComponentType(InComponentType) {}
 	FPreparedType(const Shader::FStructType* InStructType) : StructType(InStructType) {}
-	FPreparedType(const Shader::FType& InType);
+	FPreparedType(const Shader::FType& InType, EExpressionEvaluation InEvalution = EExpressionEvaluation::None);
 
 	void SetEvaluation(EExpressionEvaluation Evaluation);
 	void MergeEvaluation(EExpressionEvaluation Evaluation);
 	void SetLoopEvaluation(FEmitScope& Scope, const FRequestedType& RequestedType);
-	void SetForwardValue(FExpression* ForwardValue, const FRequestedType& RequestedType);
+	void SetForwardValue(const FRequestedType& RequestedType, FExpression* ForwardValue);
 
 	void SetField(const Shader::FStructField* Field, const FPreparedType& FieldType);
 	FPreparedType GetFieldType(const Shader::FStructField* Field) const;
 
 	int32 GetNumComponents() const;
 	FRequestedType GetRequestedType() const;
+	FPreparedType GetTypeForRequest(const FRequestedType& RequestedType) const;
 	Shader::FType GetType() const;
+	const TCHAR* GetName() const { return GetType().GetName(); }
 	bool IsStruct() const { return !IsVoid() && StructType != nullptr; }
 	bool IsNumeric() const { return !IsVoid() && ValueComponentType != Shader::EValueComponentType::Void; }
 	bool IsInitialized() const { return StructType != nullptr || ValueComponentType != Shader::EValueComponentType::Void; }
@@ -390,6 +385,21 @@ inline bool operator!=(const FPreparedType& Lhs, const FPreparedType& Rhs)
 
 FPreparedType MergePreparedTypes(const FPreparedType& Lhs, const FPreparedType& Rhs);
 
+struct FRequestedValueType
+{
+	FRequestedValueType() = default;
+	FRequestedValueType(Shader::EValueComponentType InComponentType, const FRequestedType& InRequestedType) : ComponentType(InComponentType), RequestedType(InRequestedType) {}
+	FRequestedValueType(Shader::EValueType InType);
+	FRequestedValueType(const FPreparedType& InType);
+
+	int32 GetNumComponents() const { return RequestedType.GetNumComponents(); }
+	Shader::EValueType GetType() const;
+	inline bool IsVoid() const { return ComponentType == Shader::EValueComponentType::Void || GetNumComponents() == 0; }
+
+	Shader::EValueComponentType ComponentType = Shader::EValueComponentType::Void;
+	FRequestedType RequestedType;
+};
+
 class FPrepareValueResult
 {
 public:
@@ -416,6 +426,14 @@ struct FEmitValueShaderResult
 	FEmitShaderExpression* Code = nullptr;
 };
 
+struct FEmitValuePreshaderResult
+{
+	explicit FEmitValuePreshaderResult(Shader::FPreshaderData& InPreshader) : Preshader(InPreshader) {}
+
+	Shader::FPreshaderData& Preshader;
+	Shader::FType Type;
+};
+
 enum class EDerivativeCoordinate : uint8
 {
 	Ddx,
@@ -449,19 +467,24 @@ public:
 
 	virtual void Reset() override;
 
-	FEmitShaderExpression* GetValueShader(FEmitContext& Context, FEmitScope& Scope, const FRequestedType& RequestedType, const Shader::FType& ResultType);
-	void GetValuePreshader(FEmitContext& Context, FEmitScope& Scope, const FRequestedType& RequestedType, Shader::FPreshaderData& OutPreshader);
+	FEmitShaderExpression* GetValueShader(FEmitContext& Context, FEmitScope& Scope, const FRequestedType& RequestedType);
+	Shader::FType GetValuePreshader(FEmitContext& Context, FEmitScope& Scope, const FRequestedType& RequestedType, Shader::FPreshaderData& OutPreshader);
 	Shader::FValue GetValueConstant(FEmitContext& Context, FEmitScope& Scope, const FRequestedType& RequestedType);
 
-	FEmitShaderExpression* GetValueShader(FEmitContext& Context, FEmitScope& Scope, const FRequestedType& RequestedType);
+	FEmitShaderExpression* GetValueShader(FEmitContext& Context, FEmitScope& Scope, const FRequestedType& RequestedType, const Shader::FType& ResultType);
+	FEmitShaderExpression* GetValueShader(FEmitContext& Context, FEmitScope& Scope, const FRequestedValueType& ResultType);
+	FEmitShaderExpression* GetValueShader(FEmitContext& Context, FEmitScope& Scope, const FPreparedType& ResultType);
 	FEmitShaderExpression* GetValueShader(FEmitContext& Context, FEmitScope& Scope);
 
+	Shader::FValue GetValueConstant(FEmitContext& Context, FEmitScope& Scope, const FRequestedType& RequestedType, const Shader::FType& ResultType);
+	Shader::FValue GetValueConstant(FEmitContext& Context, FEmitScope& Scope, const Shader::FType& ResultType);
+
 protected:
 	virtual void ComputeAnalyticDerivatives(FTree& Tree, FExpressionDerivatives& OutResult) const;
 	virtual FExpression* ComputePreviousFrame(FTree& Tree, const FRequestedType& RequestedType) const;
 	virtual bool PrepareValue(FEmitContext& Context, FEmitScope& Scope, const FRequestedType& RequestedType, FPrepareValueResult& OutResult) const = 0;
 	virtual void EmitValueShader(FEmitContext& Context, FEmitScope& Scope, const FRequestedType& RequestedType, FEmitValueShaderResult& OutResult) const;
-	virtual void EmitValuePreshader(FEmitContext& Context, FEmitScope& Scope, const FRequestedType& RequestedType, Shader::FPreshaderData& OutPreshader) const;
+	virtual void EmitValuePreshader(FEmitContext& Context, FEmitScope& Scope, const FRequestedType& RequestedType, FEmitValuePreshaderResult& OutResult) const;
 
 private:
 	FExpressionDerivatives Derivatives;
@@ -633,6 +656,8 @@ public:
 	FExpression* NewNeg(FExpression* Input) { return NewUnaryOp(EUnaryOp::Neg, Input); }
 	FExpression* NewRcp(FExpression* Input) { return NewUnaryOp(EUnaryOp::Rcp, Input); }
 	FExpression* NewFrac(FExpression* Input) { return NewUnaryOp(EUnaryOp::Frac, Input); }
+	FExpression* NewLength(FExpression* Input) { return NewUnaryOp(EUnaryOp::Length, Input); }
+	FExpression* NewNormalize(FExpression* Input) { return NewUnaryOp(EUnaryOp::Normalize, Input); }
 
 	FExpression* NewAdd(FExpression* Lhs, FExpression* Rhs) { return NewBinaryOp(EBinaryOp::Add, Lhs, Rhs); }
 	FExpression* NewSub(FExpression* Lhs, FExpression* Rhs) { return NewBinaryOp(EBinaryOp::Sub, Lhs, Rhs); }
diff --git a/Engine/Source/Runtime/Engine/Public/HLSLTree/HLSLTreeCommon.h b/Engine/Source/Runtime/Engine/Public/HLSLTree/HLSLTreeCommon.h
--- a/Engine/Source/Runtime/Engine/Public/HLSLTree/HLSLTreeCommon.h
+++ b/Engine/Source/Runtime/Engine/Public/HLSLTree/HLSLTreeCommon.h
@@ -22,7 +22,7 @@ public:
 
 	virtual void ComputeAnalyticDerivatives(FTree& Tree, FExpressionDerivatives& OutResult) const override;
 	virtual bool PrepareValue(FEmitContext& Context, FEmitScope& Scope, const FRequestedType& RequestedType, FPrepareValueResult& OutResult) const override;
-	virtual void EmitValuePreshader(FEmitContext& Context, FEmitScope& Scope, const FRequestedType& RequestedType, Shader::FPreshaderData& OutPreshader) const override;
+	virtual void EmitValuePreshader(FEmitContext& Context, FEmitScope& Scope, const FRequestedType& RequestedType, FEmitValuePreshaderResult& OutResult) const override;
 };
 
 class FExpressionMaterialParameter : public FExpression
@@ -39,7 +39,7 @@ public:
 
 	virtual void ComputeAnalyticDerivatives(FTree& Tree, FExpressionDerivatives& OutResult) const override;
 	virtual bool PrepareValue(FEmitContext& Context, FEmitScope& Scope, const FRequestedType& RequestedType, FPrepareValueResult& OutResult) const override;
-	virtual void EmitValuePreshader(FEmitContext& Context, FEmitScope& Scope, const FRequestedType& RequestedType, Shader::FPreshaderData& OutPreshader) const override;
+	virtual void EmitValuePreshader(FEmitContext& Context, FEmitScope& Scope, const FRequestedType& RequestedType, FEmitValuePreshaderResult& OutResult) const override;
 };
 
 enum class EExternalInput : uint8
@@ -86,6 +86,29 @@ enum class EExternalInput : uint8
 	WorldPosition_Ddx,
 	WorldPosition_Ddy,
 
+	CameraWorldPosition,
+	PreViewTranslation,
+	TangentToWorld,
+	LocalToWorld,
+	WorldToLocal,
+	TranslatedWorldToCameraView,
+	TranslatedWorldToView,
+	CameraViewToTranslatedWorld,
+	ViewToTranslatedWorld,
+	WorldToParticle,
+	WorldToInstance,
+	ParticleToWorld,
+	InstanceToWorld,
+
+	PrevCameraWorldPosition,
+	PrevPreViewTranslation,
+	PrevLocalToWorld,
+	PrevWorldToLocal,
+	PrevTranslatedWorldToCameraView,
+	PrevTranslatedWorldToView,
+	PrevCameraViewToTranslatedWorld,
+	PrevViewToTranslatedWorld,
+
 	PixelDepth,
 	PixelDepth_Ddx,
 	PixelDepth_Ddy,
@@ -184,7 +207,7 @@ public:
 	virtual FExpression* ComputePreviousFrame(FTree& Tree, const FRequestedType& RequestedType) const override;
 	virtual bool PrepareValue(FEmitContext& Context, FEmitScope& Scope, const FRequestedType& RequestedType, FPrepareValueResult& OutResult) const override;
 	virtual void EmitValueShader(FEmitContext& Context, FEmitScope& Scope, const FRequestedType& RequestedType, FEmitValueShaderResult& OutResult) const override;
-	virtual void EmitValuePreshader(FEmitContext& Context, FEmitScope& Scope, const FRequestedType& RequestedType, Shader::FPreshaderData& OutPreshader) const override;
+	virtual void EmitValuePreshader(FEmitContext& Context, FEmitScope& Scope, const FRequestedType& RequestedType, FEmitValuePreshaderResult& OutResult) const override;
 };
 
 class FExpressionSetStructField : public FExpression
@@ -209,7 +232,7 @@ public:
 	virtual FExpression* ComputePreviousFrame(FTree& Tree, const FRequestedType& RequestedType) const override;
 	virtual bool PrepareValue(FEmitContext& Context, FEmitScope& Scope, const FRequestedType& RequestedType, FPrepareValueResult& OutResult) const override;
 	virtual void EmitValueShader(FEmitContext& Context, FEmitScope& Scope, const FRequestedType& RequestedType, FEmitValueShaderResult& OutResult) const override;
-	virtual void EmitValuePreshader(FEmitContext& Context, FEmitScope& Scope, const FRequestedType& RequestedType, Shader::FPreshaderData& OutPreshader) const override;
+	virtual void EmitValuePreshader(FEmitContext& Context, FEmitScope& Scope, const FRequestedType& RequestedType, FEmitValuePreshaderResult& OutResult) const override;
 };
 
 class FExpressionSelect : public FExpression
@@ -227,7 +250,7 @@ public:
 
 	virtual bool PrepareValue(FEmitContext& Context, FEmitScope& Scope, const FRequestedType& RequestedType, FPrepareValueResult& OutResult) const override;
 	virtual void EmitValueShader(FEmitContext& Context, FEmitScope& Scope, const FRequestedType& RequestedType, FEmitValueShaderResult& OutResult) const override;
-	virtual void EmitValuePreshader(FEmitContext& Context, FEmitScope& Scope, const FRequestedType& RequestedType, Shader::FPreshaderData& OutPreshader) const override;
+	virtual void EmitValuePreshader(FEmitContext& Context, FEmitScope& Scope, const FRequestedType& RequestedType, FEmitValuePreshaderResult& OutResult) const override;
 };
 
 class FExpressionUnaryOp : public FExpression
@@ -245,7 +268,7 @@ public:
 	virtual FExpression* ComputePreviousFrame(FTree& Tree, const FRequestedType& RequestedType) const override;
 	virtual bool PrepareValue(FEmitContext& Context, FEmitScope& Scope, const FRequestedType& RequestedType, FPrepareValueResult& OutResult) const override;
 	virtual void EmitValueShader(FEmitContext& Context, FEmitScope& Scope, const FRequestedType& RequestedType, FEmitValueShaderResult& OutResult) const override;
-	virtual void EmitValuePreshader(FEmitContext& Context, FEmitScope& Scope, const FRequestedType& RequestedType, Shader::FPreshaderData& OutPreshader) const override;
+	virtual void EmitValuePreshader(FEmitContext& Context, FEmitScope& Scope, const FRequestedType& RequestedType, FEmitValuePreshaderResult& OutResult) const override;
 };
 
 class FExpressionBinaryOp : public FExpression
@@ -265,7 +288,7 @@ public:
 	virtual FExpression* ComputePreviousFrame(FTree& Tree, const FRequestedType& RequestedType) const override;
 	virtual bool PrepareValue(FEmitContext& Context, FEmitScope& Scope, const FRequestedType& RequestedType, FPrepareValueResult& OutResult) const override;
 	virtual void EmitValueShader(FEmitContext& Context, FEmitScope& Scope, const FRequestedType& RequestedType, FEmitValueShaderResult& OutResult) const override;
-	virtual void EmitValuePreshader(FEmitContext& Context, FEmitScope& Scope, const FRequestedType& RequestedType, Shader::FPreshaderData& OutPreshader) const override;
+	virtual void EmitValuePreshader(FEmitContext& Context, FEmitScope& Scope, const FRequestedType& RequestedType, FEmitValuePreshaderResult& OutResult) const override;
 };
 
 struct FSwizzleParameters
@@ -296,7 +319,7 @@ public:
 	virtual FExpression* ComputePreviousFrame(FTree& Tree, const FRequestedType& RequestedType) const override;
 	virtual bool PrepareValue(FEmitContext& Context, FEmitScope& Scope, const FRequestedType& RequestedType, FPrepareValueResult& OutResult) const override;
 	virtual void EmitValueShader(FEmitContext& Context, FEmitScope& Scope, const FRequestedType& RequestedType, FEmitValueShaderResult& OutResult) const override;
-	virtual void EmitValuePreshader(FEmitContext& Context, FEmitScope& Scope, const FRequestedType& RequestedType, Shader::FPreshaderData& OutPreshader) const override;
+	virtual void EmitValuePreshader(FEmitContext& Context, FEmitScope& Scope, const FRequestedType& RequestedType, FEmitValuePreshaderResult& OutResult) const override;
 };
 
 class FExpressionAppend : public FExpression
@@ -314,7 +337,7 @@ public:
 	virtual FExpression* ComputePreviousFrame(FTree& Tree, const FRequestedType& RequestedType) const override;
 	virtual bool PrepareValue(FEmitContext& Context, FEmitScope& Scope, const FRequestedType& RequestedType, FPrepareValueResult& OutResult) const override;
 	virtual void EmitValueShader(FEmitContext& Context, FEmitScope& Scope, const FRequestedType& RequestedType, FEmitValueShaderResult& OutResult) const override;
-	virtual void EmitValuePreshader(FEmitContext& Context, FEmitScope& Scope, const FRequestedType& RequestedType, Shader::FPreshaderData& OutPreshader) const override;
+	virtual void EmitValuePreshader(FEmitContext& Context, FEmitScope& Scope, const FRequestedType& RequestedType, FEmitValuePreshaderResult& OutResult) const override;
 };
 
 class FExpressionReflectionVector : public FExpression
diff --git a/Engine/Source/Runtime/Engine/Public/HLSLTree/HLSLTreeTypes.h b/Engine/Source/Runtime/Engine/Public/HLSLTree/HLSLTreeTypes.h
--- a/Engine/Source/Runtime/Engine/Public/HLSLTree/HLSLTreeTypes.h
+++ b/Engine/Source/Runtime/Engine/Public/HLSLTree/HLSLTreeTypes.h
@@ -65,6 +65,8 @@ enum class EUnaryOp : uint8
 	Neg,
 	Rcp,
 	Frac,
+	Length,
+	Normalize,
 };
 
 struct FUnaryOpDescription
@@ -92,6 +94,11 @@ enum class EBinaryOp : uint8
 	Max,
 	Less,
 	Greater,
+
+	VecMulMatrix3,
+	VecMulMatrix4,
+	Matrix3MulVec,
+	Matrix4MulVec,
 };
 
 struct FBinaryOpDescription
diff --git a/Engine/Source/Runtime/Engine/Public/Shader/ShaderTypes.h b/Engine/Source/Runtime/Engine/Public/Shader/ShaderTypes.h
--- a/Engine/Source/Runtime/Engine/Public/Shader/ShaderTypes.h
+++ b/Engine/Source/Runtime/Engine/Public/Shader/ShaderTypes.h
@@ -64,6 +64,14 @@ enum class EValueType : uint8
 	Bool3,
 	Bool4,
 
+	// float4x4
+	Float4x4,
+
+	// Both of these are double4x4 on CPU
+	// On GPU, they map to FLWCMatrix and FLWCInverseMatrix
+	Double4x4,
+	DoubleInverse4x4,
+
 	Struct,
 };
 
@@ -156,6 +164,8 @@ struct FStructType
 	 */
 	TArrayView<const EValueType> FlatFieldTypes;
 
+	int32 GetNumComponents() const { return ComponentTypes.Num(); }
+
 	const FStructField* FindFieldByName(const TCHAR* InName) const;
 };
 
@@ -211,7 +221,7 @@ enum class EValueStringFormat
 
 struct FMemoryImageValue
 {
-	static const uint32 MaxSize = sizeof(double) * 4;
+	static const uint32 MaxSize = sizeof(double) * 16;
 	uint8 Bytes[MaxSize];
 	uint32 Size;
 };
@@ -394,7 +404,7 @@ struct FValue
 	const TCHAR* ToString(EValueStringFormat Format, FStringBuilderBase& OutString) const;
 
 	FType Type;
-	TArray<FValueComponent, TInlineAllocator<8>> Component;
+	TArray<FValueComponent, TInlineAllocator<16>> Component;
 };
 
 ENGINE_API bool operator==(const FValue& Lhs, const FValue& Rhs);
