Fixed:  NonUnity compile

diff --git a/Engine/Source/Runtime/Engine/Private/SkeletalMesh.cpp b/Engine/Source/Runtime/Engine/Private/SkeletalMesh.cpp
--- a/Engine/Source/Runtime/Engine/Private/SkeletalMesh.cpp
+++ b/Engine/Source/Runtime/Engine/Private/SkeletalMesh.cpp
@@ -79,6 +79,7 @@
 #include "IMeshReductionManagerModule.h"
 #include "SkeletalMeshReductionSettings.h"
 #include "Engine/RendererSettings.h"
+#include "UnrealEngine.h"
 
 #endif // #if WITH_EDITOR
 
