Correct instances of FindOrAddNetworkObjectInfo that should have been FindNetworkObjectInfo

diff --git a/Engine/Source/Runtime/Engine/Private/Components/CharacterMovementComponent.cpp b/Engine/Source/Runtime/Engine/Private/Components/CharacterMovementComponent.cpp
--- a/Engine/Source/Runtime/Engine/Private/Components/CharacterMovementComponent.cpp
+++ b/Engine/Source/Runtime/Engine/Private/Components/CharacterMovementComponent.cpp
@@ -2651,7 +2651,7 @@ void UCharacterMovementComponent::PerformMovement(float DeltaSeconds)
 		UNetDriver* NetDriver = MyWorld->GetNetDriver();
 		if (NetDriver && NetDriver->IsServer())
 		{
-			FNetworkObjectInfo* NetActor = NetDriver->FindOrAddNetworkObjectInfo(CharacterOwner);
+			FNetworkObjectInfo* NetActor = NetDriver->FindNetworkObjectInfo(CharacterOwner);
 				
 			if (NetActor && MyWorld->GetTimeSeconds() <= NetActor->NextUpdateTime && NetDriver->IsNetworkActorUpdateFrequencyThrottled(*NetActor))
 			{
