Adding missing include for non-unity compiles.

diff --git a/Engine/Source/Runtime/Engine/Private/GPUSort.cpp b/Engine/Source/Runtime/Engine/Private/GPUSort.cpp
--- a/Engine/Source/Runtime/Engine/Private/GPUSort.cpp
+++ b/Engine/Source/Runtime/Engine/Private/GPUSort.cpp
@@ -15,6 +15,7 @@
 #include "SceneUtils.h"
 #include "ShaderParameterUtils.h"
 #include "GlobalShader.h"
+#include "PipelineStateCache.h"
 
 DEFINE_LOG_CATEGORY_STATIC(LogGPUSort, Log, All);
 
