Distance Crossfade Node - only add to our crossfade input if we're adding a child node outside of the range of existing values

diff --git a/Engine/Source/Runtime/Engine/Private/SoundNodeDistanceCrossFade.cpp b/Engine/Source/Runtime/Engine/Private/SoundNodeDistanceCrossFade.cpp
--- a/Engine/Source/Runtime/Engine/Private/SoundNodeDistanceCrossFade.cpp
+++ b/Engine/Source/Runtime/Engine/Private/SoundNodeDistanceCrossFade.cpp
@@ -147,9 +147,12 @@ void USoundNodeDistanceCrossFade::CreateStartingConnectors()
 void USoundNodeDistanceCrossFade::InsertChildNode( int32 Index )
 {
 	Super::InsertChildNode( Index );
-	CrossFadeInput.InsertZeroed( Index );
 
-	CrossFadeInput[ Index ].Volume = 1.0f;
+	if (Index >= CrossFadeInput.Num())
+	{
+		CrossFadeInput.InsertZeroed(Index);
+		CrossFadeInput[Index].Volume = 1.0f;
+	}
 }
 
 
