Audio: read timecode rate and drop frame-ness from RIFF iXML chunk of wave audio

When available, this data is used to populate the TimecodeSource property of audio
track sections in Sequencer.

diff --git a/Engine/Source/Runtime/Engine/Classes/Sound/SoundWaveTimecodeInfo.h b/Engine/Source/Runtime/Engine/Classes/Sound/SoundWaveTimecodeInfo.h
--- a/Engine/Source/Runtime/Engine/Classes/Sound/SoundWaveTimecodeInfo.h
+++ b/Engine/Source/Runtime/Engine/Classes/Sound/SoundWaveTimecodeInfo.h
@@ -1,6 +1,8 @@
 // Copyright Epic Games, Inc. All Rights Reserved.
 #pragma once
 
+#include "Misc/FrameRate.h"
+
 #include "SoundWaveTimecodeInfo.generated.h"
 
 USTRUCT()
@@ -29,11 +31,19 @@ struct FSoundWaveTimecodeInfo
 	UPROPERTY(Category = Timecode, VisibleAnywhere)
 	FString OriginatorReference;
 
+	UPROPERTY(Category = Timecode, VisibleAnywhere)
+	FFrameRate TimecodeRate;
+
+	UPROPERTY(Category = Timecode, VisibleAnywhere)
+	bool bTimecodeIsDropFrame = false;
+
 	inline bool operator==(const FSoundWaveTimecodeInfo& InRhs) const
 	{
 		// Note, we don't compare the strings.
 		return NumSamplesSinceMidnight == InRhs.NumSamplesSinceMidnight &&
-			NumSamplesPerSecond == InRhs.NumSamplesPerSecond;
+			NumSamplesPerSecond == InRhs.NumSamplesPerSecond &&
+			TimecodeRate == InRhs.TimecodeRate &&
+			bTimecodeIsDropFrame == InRhs.bTimecodeIsDropFrame;
 	}
 
 	inline double GetNumSecondsSinceMidnight() const
diff --git a/Engine/Source/Runtime/Engine/Engine.Build.cs b/Engine/Source/Runtime/Engine/Engine.Build.cs
--- a/Engine/Source/Runtime/Engine/Engine.Build.cs
+++ b/Engine/Source/Runtime/Engine/Engine.Build.cs
@@ -130,7 +130,8 @@ public class Engine : ModuleRules
 				"CrunchCompression",
 				"IntelISPC",
 				"TraceLog",
-				"ColorManagement"
+				"ColorManagement",
+				"XmlParser"
 			}
 		);
 
diff --git a/Engine/Source/Runtime/Engine/Private/Audio.cpp b/Engine/Source/Runtime/Engine/Private/Audio.cpp
--- a/Engine/Source/Runtime/Engine/Private/Audio.cpp
+++ b/Engine/Source/Runtime/Engine/Private/Audio.cpp
@@ -17,6 +17,7 @@
 #include "DrawDebugHelpers.h"
 #include "EngineAnalytics.h"
 #include "Interfaces/IAnalyticsProvider.h"
+#include "Misc/FrameRate.h"
 #include "Misc/Paths.h"
 #include "Sound/SoundBase.h"
 #include "Sound/SoundCue.h"
@@ -27,6 +28,8 @@
 #include "Sound/QuartzQuantizationUtilities.h"
 #include "UObject/UObjectHash.h"
 #include "UObject/UObjectIterator.h"
+#include "XmlFile.h"
+#include "XmlNode.h"
 
 DEFINE_LOG_CATEGORY(LogAudio);
 
@@ -1691,17 +1694,64 @@ bool FWaveModInfo::ReadWaveInfo( const uint8* WaveData, int32 WaveDataSize, FStr
 			TimecodeInfo->OriginatorDate			= StringCast<TCHAR>((ANSICHAR*)BextChunk->OriginationDate).Get();
 			TimecodeInfo->OriginatorReference		= StringCast<TCHAR>((ANSICHAR*)BextChunk->OriginatorReference).Get();
 			TimecodeInfo->OriginatorTime			= StringCast<TCHAR>((ANSICHAR*)BextChunk->OriginationTime).Get();
+
+			// Use the sample rate as the timecode rate for now. We'll replace
+			// it with a more accurate rate if possible below.
+			TimecodeInfo->TimecodeRate = FFrameRate(TimecodeInfo->NumSamplesPerSecond, 1u);
 		}
 	}
 
 	// Look for the cue chunks
 	RiffChunk = FindRiffChunk(RiffChunkStart, WaveDataEnd, UE_mmioFOURCC('i', 'X', 'M', 'L'));
 
-	if (RiffChunk != nullptr)
+	// If we got timecode info above, extend it with the timecode rate and drop
+	// frame flag if they are present in the XML.
+	if (TimecodeInfo.IsValid() && RiffChunk != nullptr)
 	{
 		const FRiffIXmlChunk* IXmlChunk = (const FRiffIXmlChunk*)((uint8*)RiffChunk);
-		FString XmlString = StringCast<TCHAR>((ANSICHAR*)IXmlChunk->XmlText).Get();
-		// TODO.
+		const FString XmlString = StringCast<TCHAR>((ANSICHAR*)IXmlChunk->XmlText).Get();
+
+		// Detail on the iXML specification here: http://www.gallery.co.uk/ixml/
+		const FXmlFile RiffXml(XmlString, EConstructMethod::ConstructFromBuffer);
+		if (RiffXml.IsValid())
+		{
+			const FString SpeedTag(TEXT("SPEED"));
+			const FString TimecodeRateTag(TEXT("TIMECODE_RATE"));
+			const FString TimecodeRateDelimiter(TEXT("/"));
+			const FString TimecodeFlagTag(TEXT("TIMECODE_FLAG"));
+			const FString DropFrameFlag(TEXT("DF"));
+
+			const FXmlNode* RootNode = RiffXml.GetRootNode();
+			const FXmlNode* SpeedNode = RootNode->FindChildNode(SpeedTag);
+			const FXmlNode* TimecodeRateNode = SpeedNode ? SpeedNode->FindChildNode(TimecodeRateTag) : nullptr;
+			if (TimecodeRateNode)
+			{
+				const FString TimecodeRateContent = TimecodeRateNode->GetContent();
+				TArray<FString> TimecodeRateParts;
+				TimecodeRateContent.ParseIntoArray(TimecodeRateParts, *TimecodeRateDelimiter);
+				if (TimecodeRateParts.Num() == 2)
+				{
+					uint32 TimecodeRateNumerator = 0u;
+					uint32 TimecodeRateDenominator = 0u;
+					LexFromString(TimecodeRateNumerator, *TimecodeRateParts[0]);
+					LexFromString(TimecodeRateDenominator, *TimecodeRateParts[1]);
+					if (TimecodeRateNumerator != 0u && TimecodeRateDenominator != 0u)
+					{
+						TimecodeInfo->TimecodeRate = FFrameRate(TimecodeRateNumerator, TimecodeRateDenominator);
+					}
+				}
+			}
+
+			const FXmlNode* TimecodeFlagNode = SpeedNode ? SpeedNode->FindChildNode(TimecodeFlagTag) : nullptr;
+			if (TimecodeFlagNode)
+			{
+				const FString TimecodeFlagContent = TimecodeFlagNode->GetContent();
+				if (TimecodeFlagContent == DropFrameFlag)
+				{
+					TimecodeInfo->bTimecodeIsDropFrame = true;
+				}
+			}
+		}
 	}
 
 	return true;
diff --git a/Engine/Source/Runtime/MovieSceneTracks/Private/Tracks/MovieSceneAudioTrack.cpp b/Engine/Source/Runtime/MovieSceneTracks/Private/Tracks/MovieSceneAudioTrack.cpp
--- a/Engine/Source/Runtime/MovieSceneTracks/Private/Tracks/MovieSceneAudioTrack.cpp
+++ b/Engine/Source/Runtime/MovieSceneTracks/Private/Tracks/MovieSceneAudioTrack.cpp
@@ -100,10 +100,44 @@ UMovieSceneSection* UMovieSceneAudioTrack::AddNewSoundOnRow(USoundBase* Sound, F
 	}
 
 	// add the section
-	UMovieSceneAudioSection* NewSection = NewObject<UMovieSceneAudioSection>(this, NAME_None, RF_Transactional);
+	UMovieSceneAudioSection* NewSection = Cast<UMovieSceneAudioSection>(CreateNewSection());
 	NewSection->InitialPlacementOnRow( AudioSections, Time, DurationToUse.FrameNumber.Value, RowIndex );
 	NewSection->SetSound(Sound);
 
+#if WITH_EDITORONLY_DATA
+	// Use the timecode info from the sound wave if it's available to populate
+	// the section's TimecodeSource property. Otherwise try the base sound's
+	// timecode offset.
+	TOptional<FSoundWaveTimecodeInfo> TimecodeInfo;
+	if (const USoundWave* SoundWave = Cast<USoundWave>(Sound))
+	{
+		TimecodeInfo = SoundWave->GetTimecodeInfo();
+	}
+
+	if (TimecodeInfo.IsSet())
+	{
+		const double NumSecondsSinceMidnight = TimecodeInfo->GetNumSecondsSinceMidnight();
+		const FTimecode Timecode(NumSecondsSinceMidnight, TimecodeInfo->TimecodeRate, TimecodeInfo->bTimecodeIsDropFrame, /* InbRollover = */ true);
+
+		NewSection->TimecodeSource = FMovieSceneTimecodeSource(Timecode);
+	}
+	else
+	{
+		const TOptional<FSoundTimecodeOffset> TimecodeOffset = Sound->GetTimecodeOffset();
+		if (TimecodeOffset.IsSet())
+		{
+			const double NumSecondsSinceMidnight = TimecodeOffset->NumOfSecondsSinceMidnight;
+
+			// The timecode offset does not carry a rate with it, so just use the
+			// display rate for this movie scene.
+			const FFrameRate DisplayRate = GetTypedOuter<UMovieScene>()->GetDisplayRate();
+			const FTimecode Timecode(NumSecondsSinceMidnight, DisplayRate, /* InbRollover = */ true);
+
+			NewSection->TimecodeSource = FMovieSceneTimecodeSource(Timecode);
+		}
+	}
+#endif
+
 	AudioSections.Add(NewSection);
 
 	return NewSection;
