Added anim state machine set/get state functions

diff --git a/Engine/Source/Runtime/AnimGraphRuntime/Private/AnimationStateMachineLibrary.cpp b/Engine/Source/Runtime/AnimGraphRuntime/Private/AnimationStateMachineLibrary.cpp
--- a/Engine/Source/Runtime/AnimGraphRuntime/Private/AnimationStateMachineLibrary.cpp
+++ b/Engine/Source/Runtime/AnimGraphRuntime/Private/AnimationStateMachineLibrary.cpp
@@ -13,6 +13,11 @@ void UAnimationStateMachineLibrary::ConvertToAnimationStateResult(const FAnimNod
 	AnimationState = FAnimNodeReference::ConvertToType<FAnimationStateResultReference>(Node, Result);
 }
 
+void UAnimationStateMachineLibrary::ConvertToAnimationStateMachine(const FAnimNodeReference& Node, FAnimationStateMachineReference& AnimationStateMachine, EAnimNodeReferenceConversionResult& Result)
+{
+	AnimationStateMachine = FAnimNodeReference::ConvertToType<FAnimationStateMachineReference>(Node, Result);
+}
+
 bool UAnimationStateMachineLibrary::IsStateBlendingIn(const FAnimUpdateContext& UpdateContext, const FAnimationStateResultReference& Node)
 {
 	bool bResult = false;
@@ -78,3 +83,49 @@ bool UAnimationStateMachineLibrary::IsStateBlendingOut(const FAnimUpdateContext&
 
 	return bResult;
 }
+
+void UAnimationStateMachineLibrary::SetState(const FAnimUpdateContext& UpdateContext, const FAnimationStateMachineReference& Node, FName TargetState, float Duration
+		, TEnumAsByte<ETransitionLogicType::Type> BlendType, UBlendProfile* BlendProfile, EAlphaBlendOption AlphaBlendOption, UCurveFloat* CustomBlendCurve)
+{
+	Node.CallAnimNodeFunction<FAnimNode_StateMachine>(
+		TEXT("SetState"),
+		[&UpdateContext, TargetState, Duration, BlendType, BlendProfile, AlphaBlendOption, CustomBlendCurve](FAnimNode_StateMachine& StateMachineNode)
+		{
+			if (const FAnimationUpdateContext* AnimationUpdateContext = UpdateContext.GetContext())
+			{
+				const int32 CurrentStateIndex = StateMachineNode.GetCurrentState();
+				const int32 TargetStateIndex = StateMachineNode.GetStateIndex(TargetState);
+					
+				if(CurrentStateIndex != INDEX_NONE && TargetStateIndex != INDEX_NONE)
+				{
+					FAnimationTransitionBetweenStates TransitionInfo;
+					TransitionInfo.PreviousState = CurrentStateIndex;
+					TransitionInfo.NextState = TargetStateIndex;
+					TransitionInfo.CrossfadeDuration = Duration;
+					TransitionInfo.BlendMode = AlphaBlendOption;
+					TransitionInfo.CustomCurve = CustomBlendCurve;
+					TransitionInfo.BlendProfile = BlendProfile;
+					TransitionInfo.LogicType = BlendType;
+
+					StateMachineNode.TransitionToState(*AnimationUpdateContext, TransitionInfo);
+				}
+			}
+		});
+}
+
+FName UAnimationStateMachineLibrary::GetState(const FAnimUpdateContext& UpdateContext, const FAnimationStateMachineReference& Node)
+{
+	FName OutName = NAME_None;
+	Node.CallAnimNodeFunction<FAnimNode_StateMachine>(
+		TEXT("GetState"),
+		[&UpdateContext, &OutName](FAnimNode_StateMachine& StateMachineNode)
+		{
+			if (const FAnimationUpdateContext* AnimationUpdateContext = UpdateContext.GetContext())
+			{
+				const int32 CurrentStateIndex = StateMachineNode.GetCurrentState();
+				OutName = StateMachineNode.GetStateInfo(CurrentStateIndex).StateName;
+			}
+		});
+
+	return OutName;
+}
diff --git a/Engine/Source/Runtime/AnimGraphRuntime/Public/AnimationStateMachineLibrary.h b/Engine/Source/Runtime/AnimGraphRuntime/Public/AnimationStateMachineLibrary.h
--- a/Engine/Source/Runtime/AnimGraphRuntime/Public/AnimationStateMachineLibrary.h
+++ b/Engine/Source/Runtime/AnimGraphRuntime/Public/AnimationStateMachineLibrary.h
@@ -9,7 +9,7 @@
 #include "Animation/AnimNodeReference.h"
 #include "AnimationStateMachineLibrary.generated.h"
 
-struct FAnimNode_AnimationStateMachine;
+struct FAnimNode_StateMachine;
 struct FAnimNode_StateResult;
 
 USTRUCT(BlueprintType, DisplayName = "Animation State Reference")
@@ -20,6 +20,14 @@ struct FAnimationStateResultReference : public FAnimNodeReference
 	typedef FAnimNode_StateResult FInternalNodeType;
 };
 
+USTRUCT(BlueprintType, DisplayName = "Animation State Machine")
+struct FAnimationStateMachineReference : public FAnimNodeReference
+{
+	GENERATED_BODY()
+
+	typedef FAnimNode_StateMachine FInternalNodeType;
+};
+
 // Exposes operations to be performed on anim state machine node contexts
 UCLASS(Experimental)
 class ANIMGRAPHRUNTIME_API UAnimationStateMachineLibrary : public UBlueprintFunctionLibrary
@@ -36,9 +44,23 @@ public:
 	static void ConvertToAnimationStateResultPure(const FAnimNodeReference& Node, FAnimationStateResultReference& AnimationState, bool& Result)
 	{
 		EAnimNodeReferenceConversionResult ConversionResult;
-		 ConvertToAnimationStateResult(Node, AnimationState, ConversionResult);
+		ConvertToAnimationStateResult(Node, AnimationState, ConversionResult);
 		Result = (ConversionResult == EAnimNodeReferenceConversionResult::Succeeded);	
 	}
+
+	/** Get an anim state machine from an anim node reference */
+	UFUNCTION(BlueprintCallable, Category = "State Machine", meta=(BlueprintThreadSafe, DisplayName = "Convert to Animation State Machine", ExpandEnumAsExecs = "Result"))
+	static void ConvertToAnimationStateMachine(const FAnimNodeReference& Node, FAnimationStateMachineReference& AnimationState, EAnimNodeReferenceConversionResult& Result);
+
+
+	/** Get an anim state machine from an anim node reference (pure) */
+	UFUNCTION(BlueprintPure, Category = "State Machine", meta = (BlueprintThreadSafe, DisplayName = "Convert to Animation State Machine"))
+	static void ConvertToAnimationStateMachinePure(const FAnimNodeReference& Node, FAnimationStateMachineReference& AnimationState, bool& Result)
+	{
+		EAnimNodeReferenceConversionResult ConversionResult;
+		ConvertToAnimationStateMachine(Node, AnimationState, ConversionResult);
+		Result = (ConversionResult == EAnimNodeReferenceConversionResult::Succeeded);
+	}
 	
 	/** Returns whether the state the node belongs to is blending in */
 	UFUNCTION(BlueprintPure, Category = "State Machine", meta=(BlueprintThreadSafe))
@@ -47,4 +69,14 @@ public:
 	/** Returns whether the state the node belongs to is blending out */
 	UFUNCTION(BlueprintPure, Category = "State Machine", meta=(BlueprintThreadSafe))
 	static bool IsStateBlendingOut(const FAnimUpdateContext& UpdateContext, const FAnimationStateResultReference& Node);
+
+	/** Manually set the current state of the state machine
+		NOTE: Custom blend type is not supported */
+	UFUNCTION(BlueprintCallable, Category = "State Machine", meta=(BlueprintThreadSafe, AdvancedDisplay = "4"))
+	static void SetState(const FAnimUpdateContext& UpdateContext, const FAnimationStateMachineReference& Node, FName TargetState, float Duration
+		, TEnumAsByte<ETransitionLogicType::Type> BlendType, UBlendProfile* BlendProfile, EAlphaBlendOption AlphaBlendOption, UCurveFloat* CustomBlendCurve);
+
+	/** Returns the name of the current state of this state machine */
+	UFUNCTION(BlueprintPure, Category = "State Machine", meta=(BlueprintThreadSafe))
+	static FName GetState(const FAnimUpdateContext& UpdateContext, const FAnimationStateMachineReference& Node);
 };
\ No newline at end of file
diff --git a/Engine/Source/Runtime/Engine/Classes/Animation/AnimNode_StateMachine.h b/Engine/Source/Runtime/Engine/Classes/Animation/AnimNode_StateMachine.h
--- a/Engine/Source/Runtime/Engine/Classes/Animation/AnimNode_StateMachine.h
+++ b/Engine/Source/Runtime/Engine/Classes/Animation/AnimNode_StateMachine.h
@@ -77,6 +77,9 @@ public:
 
 public:
 	FAnimationActiveTransitionEntry();
+	FAnimationActiveTransitionEntry(int32 NextStateID, float ExistingWeightOfNextState, int32 PreviousStateID, const FAnimationTransitionBetweenStates& ReferenceTransitionInfo, float CrossfadeTimeAdjustment);
+	
+	UE_DEPRECATED(5.1, "Please use FAnimationActiveTransitionEntry constructor with different signature")
 	FAnimationActiveTransitionEntry(int32 NextStateID, float ExistingWeightOfNextState, FAnimationActiveTransitionEntry* ExistingTransitionForNextState, int32 PreviousStateID, const FAnimationTransitionBetweenStates& ReferenceTransitionInfo, const FAnimationPotentialTransition& PotentialTransition);
 
 	void InitializeCustomGraphLinks(const FAnimationUpdateContext& Context, const FBakedStateExitTransition& TransitionRule);
@@ -241,11 +244,14 @@ public:
 	/** Cache the internal machine description */
 	void CacheMachineDescription(IAnimClassInterface* AnimBlueprintClass);
 
+	void SetState(const FAnimationBaseContext& Context, int32 NewStateIndex);
+	void TransitionToState(const FAnimationUpdateContext& Context, const FAnimationTransitionBetweenStates& TransitionInfo, const FAnimationPotentialTransition* BakedTransitionInfo = nullptr);
+	const int32 GetStateIndex(FName StateName) const;
+
 protected:
 	// Tries to get the instance information for the state machine
 	const FBakedAnimationStateMachine* GetMachineDescription() const;
 
-	void SetState(const FAnimationBaseContext& Context, int32 NewStateIndex);
 	void SetStateInternal(int32 NewStateIndex);
 
 	const FBakedAnimationState& GetStateInfo() const;
diff --git a/Engine/Source/Runtime/Engine/Private/Animation/AnimNode_StateMachine.cpp b/Engine/Source/Runtime/Engine/Private/Animation/AnimNode_StateMachine.cpp
--- a/Engine/Source/Runtime/Engine/Private/Animation/AnimNode_StateMachine.cpp
+++ b/Engine/Source/Runtime/Engine/Private/Animation/AnimNode_StateMachine.cpp
@@ -43,6 +43,11 @@ FAnimationActiveTransitionEntry::FAnimationActiveTransitionEntry()
 }
 
 FAnimationActiveTransitionEntry::FAnimationActiveTransitionEntry(int32 NextStateID, float ExistingWeightOfNextState, FAnimationActiveTransitionEntry* ExistingTransitionForNextState, int32 PreviousStateID, const FAnimationTransitionBetweenStates& ReferenceTransitionInfo, const FAnimationPotentialTransition& PotentialTransition)
+	: FAnimationActiveTransitionEntry(NextStateID, ExistingWeightOfNextState, PreviousStateID, ReferenceTransitionInfo, PotentialTransition.CrossfadeTimeAdjustment)
+{
+}
+
+FAnimationActiveTransitionEntry::FAnimationActiveTransitionEntry(int32 NextStateID, float ExistingWeightOfNextState, int32 PreviousStateID, const FAnimationTransitionBetweenStates& ReferenceTransitionInfo, float CrossfadeTimeAdjustment)
 	: ElapsedTime(0.0f)
 	, Alpha(0.0f)
 	, NextState(NextStateID)
@@ -56,7 +61,7 @@ FAnimationActiveTransitionEntry::FAnimationActiveTransitionEntry(int32 NextState
 	, bActive(true)
 {
 	const float Scaler = 1.0f - ExistingWeightOfNextState;
-	CrossfadeDuration = (LogicType == ETransitionLogicType::TLT_Inertialization) ? 0.0f : FMath::Max(ReferenceTransitionInfo.CrossfadeDuration - PotentialTransition.CrossfadeTimeAdjustment, 0.f) * CalculateInverseAlpha(BlendOption, Scaler);
+	CrossfadeDuration = (LogicType == ETransitionLogicType::TLT_Inertialization) ? 0.0f : FMath::Max(ReferenceTransitionInfo.CrossfadeDuration - CrossfadeTimeAdjustment, 0.f) * CalculateInverseAlpha(BlendOption, Scaler);
 
 	Blend.SetBlendTime(CrossfadeDuration);
 	Blend.SetBlendOption(BlendOption);
@@ -104,14 +109,6 @@ void FAnimationActiveTransitionEntry::InitializeCustomGraphLinks(const FAnimatio
 			}
 		}
 	}
-
-	// Initialize blend data if necessary
-	if(BlendProfile)
-	{
-		StateBlendData.AddZeroed(2);
-		StateBlendData[0].PerBoneBlendData.AddZeroed(BlendProfile->GetNumBlendEntries());
-		StateBlendData[1].PerBoneBlendData.AddZeroed(BlendProfile->GetNumBlendEntries());
-	}
 }
 
 void FAnimationActiveTransitionEntry::Update(const FAnimationUpdateContext& Context, int32 CurrentStateIndex, bool& bOutFinished)
@@ -327,6 +324,19 @@ const int32 FAnimNode_StateMachine::GetStateIndex( const FBakedAnimationState& S
 }
 
 
+const int32 FAnimNode_StateMachine::GetStateIndex(FName StateName) const
+{
+	for (int32 Index = 0; Index < PRIVATE_MachineDescription->States.Num(); ++Index)
+	{
+		if (PRIVATE_MachineDescription->States[Index].StateName == StateName)
+		{
+			return Index;
+		}
+	}
+
+	return INDEX_NONE;
+}
+
 const FAnimationTransitionBetweenStates& FAnimNode_StateMachine::GetTransitionInfo(int32 TransIndex) const
 {
 	return PRIVATE_MachineDescription->Transitions[TransIndex];
@@ -423,67 +433,8 @@ void FAnimNode_StateMachine::Update_AnyThread(const FAnimationUpdateContext& Con
 		{
 			bFoundValidTransition = true;
 
-			// let the latest transition know it has been interrupted
-			if ((ActiveTransitionArray.Num() > 0) && ActiveTransitionArray[ActiveTransitionArray.Num()-1].bActive)
-			{
-				Context.AnimInstanceProxy->AddAnimNotifyFromGeneratedClass(ActiveTransitionArray[ActiveTransitionArray.Num()-1].InterruptNotify);
-			}
-
-			const int32 PreviousState = CurrentState;
-			const int32 NextState = PotentialTransition.TargetState;
-
-			// Fire off Notifies for state transition
-			Context.AnimInstanceProxy->AddAnimNotifyFromGeneratedClass(GetStateInfo(PreviousState).EndNotify);
-			Context.AnimInstanceProxy->AddAnimNotifyFromGeneratedClass(GetStateInfo(NextState).StartNotify);
-			
-			// Get the current weight of the next state, which may be non-zero
-			const float ExistingWeightOfNextState = GetStateWeight(NextState);
-
-			FAnimationActiveTransitionEntry* PreviousTransitionForNextState = nullptr;
-			for(int32 i = ActiveTransitionArray.Num() - 1 ;  i >= 0 ; --i)
-			{
-				FAnimationActiveTransitionEntry& TransitionEntry = ActiveTransitionArray[i];
-				if(TransitionEntry.PreviousState == NextState)
-				{
-					PreviousTransitionForNextState = &TransitionEntry;
-					break;
-				}
-			}
-
-			// Don't add a transition if the previous state is a conduit (likely means we're finding the best entry state)
-			if (!IsAConduitState(PreviousState))
-			{
-				// Push the transition onto the stack
-				const FAnimationTransitionBetweenStates& ReferenceTransition = GetTransitionInfo(PotentialTransition.TransitionRule->TransitionIndex); //-V595
-				FAnimationActiveTransitionEntry* NewTransition = new (ActiveTransitionArray) FAnimationActiveTransitionEntry(NextState, ExistingWeightOfNextState, PreviousTransitionForNextState, PreviousState, ReferenceTransition, PotentialTransition);
-				if (NewTransition && PotentialTransition.TransitionRule)
-				{
-					NewTransition->InitializeCustomGraphLinks(Context, *(PotentialTransition.TransitionRule));
-
-					if (ReferenceTransition.LogicType == ETransitionLogicType::TLT_Inertialization)
-					{
-						UE::Anim::IInertializationRequester* InertializationRequester = Context.GetMessage<UE::Anim::IInertializationRequester>();
-						if (InertializationRequester)
-						{
-							InertializationRequester->RequestInertialization(ReferenceTransition.CrossfadeDuration, ReferenceTransition.BlendProfile);
-							InertializationRequester->AddDebugRecord(*Context.AnimInstanceProxy, Context.GetCurrentNodeId());
-						}
-						else
-						{
-							LogInertializationRequestError(Context, PreviousState, NextState);
-						}
-					}
-
-					NewTransition->SourceTransitionIndices = PotentialTransition.SourceTransitionIndices;
-
-					if (!bFirstUpdate)
-					{
-						Context.AnimInstanceProxy->AddAnimNotifyFromGeneratedClass(NewTransition->StartNotify);
-					}
-				}
-			}
-
-			SetState(Context, NextState);
+			const FAnimationTransitionBetweenStates& ReferenceTransition = GetTransitionInfo(PotentialTransition.TransitionRule->TransitionIndex); //-V595
+			TransitionToState(Context, ReferenceTransition, &PotentialTransition);
 
 			TransitionCountThisFrame++;
 		}
@@ -1084,6 +1035,93 @@ void FAnimNode_StateMachine::SetState(const FAnimationBaseContext& Context, int3
 	}
 }
 
+void FAnimNode_StateMachine::TransitionToState(const FAnimationUpdateContext& Context, const FAnimationTransitionBetweenStates& TransitionInfo, const FAnimationPotentialTransition* BakedTransitionInfo)
+{
+	// let the latest transition know it has been interrupted
+	if ((ActiveTransitionArray.Num() > 0) && ActiveTransitionArray[ActiveTransitionArray.Num() - 1].bActive)
+	{
+		Context.AnimInstanceProxy->AddAnimNotifyFromGeneratedClass(ActiveTransitionArray[ActiveTransitionArray.Num() - 1].InterruptNotify);
+	}
+
+	const int32 PreviousState = TransitionInfo.PreviousState;
+	const int32 NextState = TransitionInfo.NextState;
+
+	// Fire off Notifies for state transition
+	Context.AnimInstanceProxy->AddAnimNotifyFromGeneratedClass(GetStateInfo(PreviousState).EndNotify);
+	Context.AnimInstanceProxy->AddAnimNotifyFromGeneratedClass(GetStateInfo(NextState).StartNotify);
+
+	FAnimationActiveTransitionEntry* PreviousTransitionForNextState = nullptr;
+	for (int32 i = ActiveTransitionArray.Num() - 1; i >= 0; --i)
+	{
+		FAnimationActiveTransitionEntry& TransitionEntry = ActiveTransitionArray[i];
+		if (TransitionEntry.PreviousState == NextState)
+		{
+			PreviousTransitionForNextState = &TransitionEntry;
+			break;
+		}
+	}
+
+	// Don't add a transition if the previous state is a conduit (likely means we're finding the best entry state)
+	if (!IsAConduitState(PreviousState))
+	{
+
+		// If we have baked transition rule info available, use it
+		float CrossFadeTimeAdjustment = 0.0f;
+		const FBakedStateExitTransition* BakedExitTransition = nullptr;
+		TArray<int32, TInlineAllocator<3>> SourceTransitionIndices;
+		if (BakedTransitionInfo)
+		{
+			CrossFadeTimeAdjustment = BakedTransitionInfo->CrossfadeTimeAdjustment;
+			BakedExitTransition = BakedTransitionInfo->TransitionRule;
+			SourceTransitionIndices = BakedTransitionInfo->SourceTransitionIndices;
+		}
+
+		// Get the current weight of the next state, which may be non-zero
+		const float ExistingWeightOfNextState = GetStateWeight(NextState);
+
+		// Push the transition onto the stack
+		FAnimationActiveTransitionEntry* NewTransition = new (ActiveTransitionArray) FAnimationActiveTransitionEntry(NextState, ExistingWeightOfNextState, PreviousState, TransitionInfo, CrossFadeTimeAdjustment);
+		if (NewTransition)
+		{
+			if ((TransitionInfo.LogicType == ETransitionLogicType::TLT_Custom) && BakedExitTransition)
+			{
+				NewTransition->InitializeCustomGraphLinks(Context, *BakedExitTransition);
+			}
+
+			// Initialize blend data if necessary
+			if (NewTransition->BlendProfile)
+			{
+				NewTransition->StateBlendData.AddZeroed(2);
+				NewTransition->StateBlendData[0].PerBoneBlendData.AddZeroed(NewTransition->BlendProfile->GetNumBlendEntries());
+				NewTransition->StateBlendData[1].PerBoneBlendData.AddZeroed(NewTransition->BlendProfile->GetNumBlendEntries());
+			}
+
+			if (TransitionInfo.LogicType == ETransitionLogicType::TLT_Inertialization)
+			{
+				UE::Anim::IInertializationRequester* InertializationRequester = Context.GetMessage<UE::Anim::IInertializationRequester>();
+				if (InertializationRequester)
+				{
+					InertializationRequester->RequestInertialization(TransitionInfo.CrossfadeDuration, TransitionInfo.BlendProfile);
+					InertializationRequester->AddDebugRecord(*Context.AnimInstanceProxy, Context.GetCurrentNodeId());
+				}
+				else
+				{
+					LogInertializationRequestError(Context, PreviousState, NextState);
+				}
+			}
+
+			NewTransition->SourceTransitionIndices = SourceTransitionIndices;
+
+			if (!bFirstUpdate)
+			{
+				Context.AnimInstanceProxy->AddAnimNotifyFromGeneratedClass(NewTransition->StartNotify);
+			}
+		}
+	}
+
+	SetState(Context, NextState);
+}
+
 float FAnimNode_StateMachine::GetStateWeight(int32 StateIndex) const
 {
 	const int32 NumTransitions = ActiveTransitionArray.Num();
