Additional granular network memory tracking

diff --git a/Engine/Plugins/Runtime/ReplicationGraph/Source/Private/ReplicationGraph.cpp b/Engine/Plugins/Runtime/ReplicationGraph/Source/Private/ReplicationGraph.cpp
--- a/Engine/Plugins/Runtime/ReplicationGraph/Source/Private/ReplicationGraph.cpp
+++ b/Engine/Plugins/Runtime/ReplicationGraph/Source/Private/ReplicationGraph.cpp
@@ -2413,7 +2413,7 @@ void UNetReplicationGraphConnection::Serialize(FArchive& Ar)
 		GRANULAR_NETWORK_MEMORY_TRACKING_INIT(Ar, "UNetReplicationGraphConnection::Serialize");
 
 		GRANULAR_NETWORK_MEMORY_TRACKING_TRACK("ActorInfoMap", ActorInfoMap.CountBytes(Ar));
-		GRANULAR_NETWORK_MEMORY_TRACKING_TRACK("ActorInfoMap", OnClientVisibleLevelNameAddMap.CountBytes(Ar));
+		GRANULAR_NETWORK_MEMORY_TRACKING_TRACK("OnClientVisibleLevelNameAddMap", OnClientVisibleLevelNameAddMap.CountBytes(Ar));
 
 		GRANULAR_NETWORK_MEMORY_TRACKING_TRACK("PendingDestructionInfoList",
 			PendingDestructInfoList.CountBytes(Ar);
@@ -3101,6 +3101,19 @@ void FStreamingLevelActorListCollection::TearDown()
 
 // --------------------------------------------------------------------------------------------------------------------------------------------
 
+void UReplicationGraphNode_ActorList::Serialize(FArchive& Ar)
+{
+	Super::Serialize(Ar);
+
+	if (Ar.IsCountingMemory())
+	{
+		GRANULAR_NETWORK_MEMORY_TRACKING_INIT(Ar, "UReplicationGraphNode_ActorList::Serialize");
+
+		GRANULAR_NETWORK_MEMORY_TRACKING_TRACK("ReplicationActorList", ReplicationActorList.CountBytes(Ar));
+		GRANULAR_NETWORK_MEMORY_TRACKING_TRACK("StreamingLevelCollection", StreamingLevelCollection.CountBytes(Ar));
+	}
+}
+
 void UReplicationGraphNode_ActorList::NotifyAddNetworkActor(const FNewReplicatedActorInfo& ActorInfo)
 {
 	UE_CLOG(CVar_RepGraph_LogActorAdd>0, LogReplicationGraph, Display, TEXT("UReplicationGraphNode_ActorList::NotifyAddNetworkActor %s on %s."), *ActorInfo.Actor->GetFullName(), *GetPathName());
diff --git a/Engine/Plugins/Runtime/ReplicationGraph/Source/Public/ReplicationGraph.h b/Engine/Plugins/Runtime/ReplicationGraph/Source/Public/ReplicationGraph.h
--- a/Engine/Plugins/Runtime/ReplicationGraph/Source/Public/ReplicationGraph.h
+++ b/Engine/Plugins/Runtime/ReplicationGraph/Source/Public/ReplicationGraph.h
@@ -197,9 +197,14 @@ struct FStreamingLevelActorListCollection
 	static const int32 NumInlineAllocations = 4;
 	TArray<FStreamingLevelActors, TInlineAllocator<NumInlineAllocations>> StreamingLevelLists;
 
-	void CountBytes(FArchive& Ar)
+	void CountBytes(FArchive& Ar) const
 	{
 		StreamingLevelLists.CountBytes(Ar);
+
+		for (const FStreamingLevelActors& List : StreamingLevelLists)
+		{
+			List.ReplicationActorList.CountBytes(Ar);
+		}
 	}
 };
 
@@ -212,6 +217,7 @@ class REPLICATIONGRAPH_API UReplicationGraphNode_ActorList : public UReplication
 	GENERATED_BODY()
 
 public:
+	virtual void Serialize(FArchive& Ar) override;
 
 	virtual void NotifyAddNetworkActor(const FNewReplicatedActorInfo& ActorInfo) override;
 	
diff --git a/Engine/Plugins/Runtime/ReplicationGraph/Source/Public/ReplicationGraphTypes.h b/Engine/Plugins/Runtime/ReplicationGraph/Source/Public/ReplicationGraphTypes.h
--- a/Engine/Plugins/Runtime/ReplicationGraph/Source/Public/ReplicationGraphTypes.h
+++ b/Engine/Plugins/Runtime/ReplicationGraph/Source/Public/ReplicationGraphTypes.h
@@ -370,6 +370,11 @@ struct REPLICATIONGRAPH_API FActorRepListRefView
 		return RepList.Contains(Value);
 	}
 
+	void CountBytes(FArchive& Ar) const
+	{
+		RepList.CountBytes(Ar);
+	}
+
 private:
 
 	friend struct FActorRepListStatCollector;
@@ -572,7 +577,7 @@ struct FGlobalActorReplicationInfo
 
 	void LogDebugString(FOutputDevice& Ar) const;
 
-	void CountBytes(FArchive& Ar)
+	void CountBytes(FArchive& Ar) const
 	{
 		// Note, we don't count DependentActorList because it's memory will be cached by the allocator / pooling stuff.
 		if (FastSharedReplicationInfo.IsValid())
@@ -580,6 +585,9 @@ struct FGlobalActorReplicationInfo
 			Ar.CountBytes(sizeof(FFastSharedReplicationInfo), sizeof(FFastSharedReplicationInfo));
 			FastSharedReplicationInfo->CountBytes(Ar);
 		}
+
+		DependentActorList.CountBytes(Ar);
+		ParentActorList.CountBytes(Ar);
 	}
 
 	typedef TArray<FActorRepListType> FDependantListType;
@@ -774,9 +782,9 @@ struct FGlobalActorReplicationInfoMap
 
 
 	/** Finds data associated with the actor but does not create if its not there yet. */
-	FORCEINLINE FGlobalActorReplicationInfo* Find(const FActorRepListType& Actor)
+	FORCEINLINE FGlobalActorReplicationInfo* Find(const FActorRepListType& Actor) const 
 	{
-		if (TUniquePtr<FGlobalActorReplicationInfo>* Ptr = ActorMap.Find(Actor))
+		if (const TUniquePtr<FGlobalActorReplicationInfo>* Ptr = ActorMap.Find(Actor))
 		{
 			return Ptr->Get();
 		}
diff --git a/Engine/Source/Runtime/Engine/Classes/Engine/NetDriver.h b/Engine/Source/Runtime/Engine/Classes/Engine/NetDriver.h
--- a/Engine/Source/Runtime/Engine/Classes/Engine/NetDriver.h
+++ b/Engine/Source/Runtime/Engine/Classes/Engine/NetDriver.h
@@ -381,6 +381,9 @@ DECLARE_DELEGATE_RetVal(bool, FShouldSkipRepNotifies);
 //
 #define DO_ENABLE_NET_TEST !(UE_BUILD_SHIPPING)
 
+#ifndef NET_DEBUG_RELEVANT_ACTORS
+#define NET_DEBUG_RELEVANT_ACTORS !(UE_BUILD_SHIPPING || UE_BUILD_TEST)
+#endif 
 
 /** Holds the packet simulation settings in one place */
 USTRUCT()
@@ -1007,6 +1010,8 @@ public:
 	float						JoinInProgressStandbyWaitTime;
 	/** Used to track whether a given actor was replicated by the net driver recently */
 	int32						NetTag;
+
+#if NET_DEBUG_RELEVANT_ACTORS
 	/** Dumps next net update's relevant actors when true*/
 	bool						DebugRelevantActors;
 
@@ -1017,6 +1022,7 @@ public:
 	TArray< TWeakObjectPtr<AActor> >	LastNonRelevantActors;
 
 	void						PrintDebugRelevantActors();
+#endif // NET_DEBUG_RELEVANT_ACTORS
 	
 	/** The server adds an entry into this map for every actor that is destroyed that join-in-progress
 	 *  clients need to know about, that is, startup actors. Also, individual UNetConnections
diff --git a/Engine/Source/Runtime/Engine/Private/DataChannel.cpp b/Engine/Source/Runtime/Engine/Private/DataChannel.cpp
--- a/Engine/Source/Runtime/Engine/Private/DataChannel.cpp
+++ b/Engine/Source/Runtime/Engine/Private/DataChannel.cpp
@@ -3444,6 +3444,7 @@ void UActorChannel::Serialize(FArchive& Ar)
 		);
 
 		GRANULAR_NETWORK_MEMORY_TRACKING_TRACK("PendingGuidResolves", PendingGuidResolves.CountBytes(Ar));
+		GRANULAR_NETWORK_MEMORY_TRACKING_TRACK("CreateSubObjects", CreateSubObjects.CountBytes(Ar));
 		GRANULAR_NETWORK_MEMORY_TRACKING_TRACK("QueuedMustBeMappedGuidsInLastBunch", QueuedMustBeMappedGuidsInLastBunch.CountBytes(Ar));
 
 		GRANULAR_NETWORK_MEMORY_TRACKING_TRACK("QueuedExportBunches",
@@ -3468,6 +3469,8 @@ void UActorChannel::Serialize(FArchive& Ar)
 		);
 
 		GRANULAR_NETWORK_MEMORY_TRACKING_TRACK("PendingObjKeys", PendingObjKeys.CountBytes(Ar));
+
+		GRANULAR_NETWORK_MEMORY_TRACKING_TRACK("QueuedBunchObjectReferences", QueuedBunchObjectReferences.CountBytes(Ar));
 	}
 }
 
diff --git a/Engine/Source/Runtime/Engine/Private/NetDriver.cpp b/Engine/Source/Runtime/Engine/Private/NetDriver.cpp
--- a/Engine/Source/Runtime/Engine/Private/NetDriver.cpp
+++ b/Engine/Source/Runtime/Engine/Private/NetDriver.cpp
@@ -399,7 +399,9 @@ UNetDriver::UNetDriver(const FObjectInitializer& ObjectInitializer)
 ,	bCollectNetStats(false)
 ,	LastCleanupTime(0.0)
 ,	NetTag(0)
+#if NET_DEBUG_RELEVANT_ACTORS
 ,	DebugRelevantActors(false)
+#endif // NET_DEBUG_RELEVANT_ACTORS
 #if !UE_BUILD_SHIPPING
 ,	SendRPCDel()
 #endif
@@ -2644,10 +2646,12 @@ void UNetDriver::Serialize( FArchive& Ar )
 			}
 		);
 
+#if NET_DEBUG_RELEVANT_ACTORS
 		GRANULAR_NETWORK_MEMORY_TRACKING_TRACK("LastPrioritizedActors", LastPrioritizedActors.CountBytes(Ar));
 		GRANULAR_NETWORK_MEMORY_TRACKING_TRACK("LastRelevantActors", LastRelevantActors.CountBytes(Ar));
 		GRANULAR_NETWORK_MEMORY_TRACKING_TRACK("LastSentActors", LastSentActors.CountBytes(Ar));
 		GRANULAR_NETWORK_MEMORY_TRACKING_TRACK("LastNonRelevantActors", LastNonRelevantActors.CountBytes(Ar));
+#endif // NET_DEBUG_RELEVANT_ACTORS
 
 		GRANULAR_NETWORK_MEMORY_TRACKING_TRACK("DestroyedStartupOrDormantActors",
 			DestroyedStartupOrDormantActors.CountBytes(Ar);
@@ -4577,10 +4581,12 @@ int32 UNetDriver::ServerReplicateActors_PrioritizeActors( UNetConnection* Connec
 
 				FinalSortedCount++;
 
+#if NET_DEBUG_RELEVANT_ACTORS
 				if ( DebugRelevantActors )
 				{
 					LastPrioritizedActors.Add( Actor );
 				}
+#endif // NET_DEBUG_RELEVANT_ACTORS
 			}
 		}
 
@@ -4676,10 +4682,12 @@ int32 UNetDriver::ServerReplicateActors_ProcessPrioritizedActors( UNetConnection
 					{
 						bIsRelevant = true;
 					}
+#if NET_DEBUG_RELEVANT_ACTORS
 					else if ( DebugRelevantActors )
 					{
 						LastNonRelevantActors.Add( Actor );
 					}
+#endif // NET_DEBUG_RELEVANT_ACTORS
 				}
 			}
 			else
@@ -4737,10 +4745,13 @@ int32 UNetDriver::ServerReplicateActors_ProcessPrioritizedActors( UNetConnection
 					{
 						// replicate the actor
 						UE_LOG( LogNetTraffic, Log, TEXT( "- Replicate %s. %d" ), *Actor->GetName(), PriorityActors[j]->Priority );
+
+#if NET_DEBUG_RELEVANT_ACTORS
 						if ( DebugRelevantActors )
 						{
 							LastRelevantActors.Add( Actor );
 						}
+#endif // NET_DEBUG_RELEVANT_ACTORS
 
 						double ChannelLastNetUpdateTime = Channel->LastUpdateTime;
 
@@ -4755,10 +4766,13 @@ int32 UNetDriver::ServerReplicateActors_ProcessPrioritizedActors( UNetConnection
 #endif
 
 							ActorUpdatesThisConnectionSent++;
+
+#if NET_DEBUG_RELEVANT_ACTORS
 							if ( DebugRelevantActors )
 							{
 								LastSentActors.Add( Actor );
 							}
+#endif // NET_DEBUG_RELEVANT_ACTORS
 
 							// Calculate min delta (max rate actor will upate), and max delta (slowest rate actor will update)
 							const float MinOptimalDelta				= 1.0f / Actor->NetUpdateFrequency;
@@ -5253,6 +5267,7 @@ int32 UNetDriver::ServerReplicateActors(float DeltaSeconds)
 	}
 	Mark.Pop();
 
+#if NET_DEBUG_RELEVANT_ACTORS
 	if (DebugRelevantActors)
 	{
 		PrintDebugRelevantActors();
@@ -5263,6 +5278,7 @@ int32 UNetDriver::ServerReplicateActors(float DeltaSeconds)
 
 		DebugRelevantActors  = false;
 	}
+#endif // NET_DEBUG_RELEVANT_ACTORS
 
 	for (UNetConnection* ConnectionToClose : ConnectionsToClose)
 	{
@@ -5330,6 +5346,7 @@ void UNetDriver::SetNetDriverName(FName NewNetDriverNamed)
 	InitPacketSimulationSettings();
 }
 
+#if NET_DEBUG_RELEVANT_ACTORS
 void UNetDriver::PrintDebugRelevantActors()
 {
 	struct SLocal
@@ -5397,8 +5414,8 @@ void UNetDriver::PrintDebugRelevantActors()
 	UE_LOG(LogNet, Warning, TEXT(" Num Connections: %d"), ClientConnections.Num() );
 
 	UE_LOG(LogNet, Warning, TEXT("---------------------------------") );
-
 }
+#endif // NET_DEBUG_RELEVANT_ACTORS
 
 void UNetDriver::DrawNetDriverDebug()
 {
@@ -5951,6 +5968,7 @@ UNetConnection* UNetDriver::GetConnectionById(uint32 ConnectionId) const
 	return nullptr;
 }
 
+#if NET_DEBUG_RELEVANT_ACTORS
 static void	DumpRelevantActors( UWorld* InWorld )
 {
 	UNetDriver *NetDriver = InWorld->NetDriver;
@@ -5961,6 +5979,7 @@ static void	DumpRelevantActors( UWorld* InWorld )
 
 	NetDriver->DebugRelevantActors = true;
 }
+#endif // NET_DEBUG_RELEVANT_ACTORS
 
 TSharedPtr<FRepChangedPropertyTracker> UNetDriver::FindOrCreateRepChangedPropertyTracker(UObject* Obj)
 {
@@ -6460,12 +6479,13 @@ void UNetDriver::MoveMappedObjectToUnmapped(const UObject* Object)
 	}
 }
 
+#if NET_DEBUG_RELEVANT_ACTORS
 FAutoConsoleCommandWithWorld	DumpRelevantActorsCommand(
 	TEXT("net.DumpRelevantActors"), 
 	TEXT( "Dumps information on relevant actors during next network update" ), 
 	FConsoleCommandWithWorldDelegate::CreateStatic(DumpRelevantActors)
 	);
-
+#endif // NET_DEBUG_RELEVANT_ACTORS
 
 #if DO_ENABLE_NET_TEST
 
diff --git a/Engine/Source/Runtime/Engine/Private/NetworkObjectList.cpp b/Engine/Source/Runtime/Engine/Private/NetworkObjectList.cpp
--- a/Engine/Source/Runtime/Engine/Private/NetworkObjectList.cpp
+++ b/Engine/Source/Runtime/Engine/Private/NetworkObjectList.cpp
@@ -5,6 +5,7 @@
 #include "Engine/Level.h"
 #include "EngineUtils.h"
 #include "Serialization/Archive.h"
+#include "Net/NetworkGranularMemoryLogging.h"
 
 void FNetworkObjectList::AddInitialObjects(UWorld* const World, UNetDriver* NetDriver)
 {
@@ -303,26 +304,32 @@ void FNetworkObjectList::Reset()
 
 void FNetworkObjectInfo::CountBytes(FArchive& Ar) const
 {
-	DormantConnections.CountBytes(Ar);
-	RecentlyDormantConnections.CountBytes(Ar);
+	GRANULAR_NETWORK_MEMORY_TRACKING_INIT(Ar, "FNetworkObjectInfo::CountBytes");
+
+	GRANULAR_NETWORK_MEMORY_TRACKING_TRACK("DormantConnections", DormantConnections.CountBytes(Ar));
+	GRANULAR_NETWORK_MEMORY_TRACKING_TRACK("RecentlyDormantConnections", RecentlyDormantConnections.CountBytes(Ar));
 }
 
 void FNetworkObjectList::CountBytes(FArchive& Ar) const
 {
-	AllNetworkObjects.CountBytes(Ar);
-	ActiveNetworkObjects.CountBytes(Ar);
-	ObjectsDormantOnAllConnections.CountBytes(Ar);
-	NumDormantObjectsPerConnection.CountBytes(Ar);
- 
+	GRANULAR_NETWORK_MEMORY_TRACKING_INIT(Ar, "FNetworkObjectList::CountBytes");
+
+	GRANULAR_NETWORK_MEMORY_TRACKING_TRACK("ActiveNetworkObjects", ActiveNetworkObjects.CountBytes(Ar));
+	GRANULAR_NETWORK_MEMORY_TRACKING_TRACK("ObjectsDormantOnAllConnections", ObjectsDormantOnAllConnections.CountBytes(Ar));
+	GRANULAR_NETWORK_MEMORY_TRACKING_TRACK("NumDormantObjectsPerConnection", NumDormantObjectsPerConnection.CountBytes(Ar));
+
 	// ObjectsDormantOnAllConnections and ActiveNetworkObjects are both sub sets of AllNetworkObjects
 	// and only have pointers back to the data there.
 	// So, to avoid double (or triple) counting, only explicit count the elements from AllNetworkObjects.
-	for (const TSharedPtr<FNetworkObjectInfo>& SharedInfo : AllNetworkObjects)
-	{
-		if (FNetworkObjectInfo const * const Info = SharedInfo.Get())
+	GRANULAR_NETWORK_MEMORY_TRACKING_TRACK("AllNetworkObjects",
+		AllNetworkObjects.CountBytes(Ar);
+		for (const TSharedPtr<FNetworkObjectInfo>& SharedInfo : AllNetworkObjects)
 		{
-			Ar.CountBytes(sizeof(FNetworkObjectInfo), sizeof(FNetworkObjectInfo));
-			Info->CountBytes(Ar);
+			if (FNetworkObjectInfo const* const Info = SharedInfo.Get())
+			{
+				Ar.CountBytes(sizeof(FNetworkObjectInfo), sizeof(FNetworkObjectInfo));
+				Info->CountBytes(Ar);
+			}
 		}
-	}
+	);
 }
\ No newline at end of file
diff --git a/Engine/Source/Runtime/Engine/Private/RepLayout.cpp b/Engine/Source/Runtime/Engine/Private/RepLayout.cpp
--- a/Engine/Source/Runtime/Engine/Private/RepLayout.cpp
+++ b/Engine/Source/Runtime/Engine/Private/RepLayout.cpp
@@ -375,8 +375,10 @@ public:
 
 	void CountBytes(FArchive& Ar) const
 	{
-		LifetimeCustomDeltaProperties.CountBytes(Ar);
-		LifetimeCustomDeltaIndexLookup.CountBytes(Ar);
+		GRANULAR_NETWORK_MEMORY_TRACKING_INIT(Ar, "FLifetimeCustomDeltaState::CountBytes");
+
+		GRANULAR_NETWORK_MEMORY_TRACKING_TRACK("LifetimeCustomDeltaProperties", LifetimeCustomDeltaProperties.CountBytes(Ar));
+		GRANULAR_NETWORK_MEMORY_TRACKING_TRACK("LifetimeCustomDeltaIndexLookup", LifetimeCustomDeltaIndexLookup.CountBytes(Ar));
 	}
 
 	const uint16 GetNumCustomDeltaProperties() const
@@ -458,12 +460,15 @@ struct FDeltaArrayHistoryItem
 
 	void CountBytes(FArchive& Ar) const
 	{
-		ChangelistByID.CountBytes(Ar);
+		GRANULAR_NETWORK_MEMORY_TRACKING_INIT(Ar, "FDeltaArrayHistoryItem::CountBytes");
 
-		for (const auto& KVP : ChangelistByID)
-		{
-			KVP.Value.CountBytes(Ar);
-		}
+		GRANULAR_NETWORK_MEMORY_TRACKING_TRACK("ChangelistByID",
+			ChangelistByID.CountBytes(Ar);
+			for (const auto& KVP : ChangelistByID)
+			{
+				KVP.Value.CountBytes(Ar);
+			}
+		);
 	}
 };
 
@@ -496,12 +501,16 @@ struct FDeltaArrayHistoryState
 
 	void CountBytes(FArchive& Ar) const
 	{
-		IDToIndexMap.CountBytes(Ar);
+		GRANULAR_NETWORK_MEMORY_TRACKING_INIT(Ar, "FDeltaArrayHistoryState::CountBytes");
 
-		for (const FDeltaArrayHistoryItem& HistoryItem : ChangeHistory)
-		{
-			HistoryItem.CountBytes(Ar);
-		}
+		GRANULAR_NETWORK_MEMORY_TRACKING_TRACK("IDToIndexMap", IDToIndexMap.CountBytes(Ar));
+
+		GRANULAR_NETWORK_MEMORY_TRACKING_TRACK("ChangeHistory",
+			for (const FDeltaArrayHistoryItem& HistoryItem : ChangeHistory)
+			{
+				HistoryItem.CountBytes(Ar);
+			}
+		);
 	}
 };
 
@@ -1185,6 +1194,14 @@ bool FRepChangelistState::HasAnyDirtyProperties() const
 void FRepChangelistState::CountBytes(FArchive& Ar) const
 {
 	GRANULAR_NETWORK_MEMORY_TRACKING_INIT(Ar, "FRepChangelistState::CountBytes");
+
+	GRANULAR_NETWORK_MEMORY_TRACKING_TRACK("ChangeHistory",
+		for (const FRepChangedHistory& HistoryItem : ChangeHistory)
+		{
+			HistoryItem.CountBytes(Ar);
+		}
+	);
+
 	GRANULAR_NETWORK_MEMORY_TRACKING_TRACK("StaticBuffer", StaticBuffer.CountBytes(Ar));
 	GRANULAR_NETWORK_MEMORY_TRACKING_TRACK("SharedSerialization", SharedSerialization.CountBytes(Ar));
 
@@ -1215,7 +1232,8 @@ FReplicationChangelistMgr::~FReplicationChangelistMgr()
 
 void FReplicationChangelistMgr::CountBytes(FArchive& Ar) const
 {
-	RepChangelistState.CountBytes(Ar);
+	GRANULAR_NETWORK_MEMORY_TRACKING_INIT(Ar, "FReplicationChangelistMgr::CountBytes");
+	GRANULAR_NETWORK_MEMORY_TRACKING_TRACK("RepChangelistState", RepChangelistState.CountBytes(Ar));
 }
 
 FReceivingRepState::FReceivingRepState(FRepStateStaticBuffer&& InStaticBuffer) 
diff --git a/Engine/Source/Runtime/Net/Core/Classes/Net/Serialization/FastArraySerializer.h b/Engine/Source/Runtime/Net/Core/Classes/Net/Serialization/FastArraySerializer.h
--- a/Engine/Source/Runtime/Net/Core/Classes/Net/Serialization/FastArraySerializer.h
+++ b/Engine/Source/Runtime/Net/Core/Classes/Net/Serialization/FastArraySerializer.h
@@ -234,6 +234,11 @@ public:
 		return true;
 	}
 
+	virtual void CountBytes(FArchive& Ar) const override
+	{
+		IDToCLMap.CountBytes(Ar);
+	}
+
 	/** Maps an element's Replication ID to Index. */
 	TMap<int32, int32> IDToCLMap;
 
