Rename fixed tick to Async Physics Tick

diff --git a/Engine/Source/Runtime/Engine/Classes/Components/ActorComponent.h b/Engine/Source/Runtime/Engine/Classes/Components/ActorComponent.h
--- a/Engine/Source/Runtime/Engine/Classes/Components/ActorComponent.h
+++ b/Engine/Source/Runtime/Engine/Classes/Components/ActorComponent.h
@@ -292,8 +292,8 @@ private:
 
 public:
 
-	/** Whether to use use the fixed physics tick with this component. Requires async fixed tick to be enabled (see project settings) */
-	uint8 bFixedTickEnabled : 1;
+	/** Whether to use use the async physics tick with this component. Requires async physics to be enabled (see project settings) */
+	uint8 bAsyncPhysicsTickEnabled : 1;
 
 	/** Describes how a component instance will be created */
 	UPROPERTY()
@@ -681,14 +681,14 @@ public:
 	virtual void TickComponent(float DeltaTime, enum ELevelTick TickType, FActorComponentTickFunction *ThisTickFunction);
 
 	/**
-	 * Function called for every physics step on this ActorComponent. Override this function to implement custom logic to be executed every physics step.
-	 * Only executes if the component is registered, and also bFixedTick must be set to true.
-	 * Requires async fixed tick to be enabled (see physics project settings)
+	 * Override this function to implement custom logic to be executed every async physics step.
+	 * Only executes if the component is registered, and also bAsyncPhysicsTick must be set to true.
+	 * Requires async physics to be enabled (see physics project settings)
 	 *	
-	 * @param DeltaTime - The fixed tick delta time
+	 * @param DeltaTime - The async physics time associated with each step (see project settings)
 	 * @param SimTime - This is the total sim time since the sim began, it is an integer multiple of DeltaTime
 	 */
-	virtual void FixedTickComponent(float DeltaTime, float SimTime) { ReceiveFixedTick(DeltaTime, SimTime); }
+	virtual void AsyncPhysicsTickComponent(float DeltaTime, float SimTime) { ReceiveAsyncPhysicsTick(DeltaTime, SimTime); }
 	
 	/** 
 	 * Set up a tick function for a component in the standard way. 
@@ -1005,9 +1005,9 @@ public:
 	UFUNCTION(BlueprintImplementableEvent, meta=(DisplayName = "Tick"))
 	void ReceiveTick(float DeltaSeconds);
 
-	/** Event called every physics tick if EnabledFixedTick is true and async fixed tick is enabled (see physics project settings)*/
-	UFUNCTION(BlueprintImplementableEvent, meta = (DisplayName = "Fixed Tick"))
-	void ReceiveFixedTick(float DeltaSeconds, float SimSeconds);
+	/** Event called every async physics tick if bAsyncPhysicsTickEnabled is true and async physics is enabled (see physics project settings)*/
+	UFUNCTION(BlueprintImplementableEvent, meta = (DisplayName = "Async Physics Tick"))
+	void ReceiveAsyncPhysicsTick(float DeltaSeconds, float SimSeconds);
 	
 	/** 
 	 *  Called by owner actor on position shifting
diff --git a/Engine/Source/Runtime/Engine/Classes/Components/PrimitiveComponent.h b/Engine/Source/Runtime/Engine/Classes/Components/PrimitiveComponent.h
--- a/Engine/Source/Runtime/Engine/Classes/Components/PrimitiveComponent.h
+++ b/Engine/Source/Runtime/Engine/Classes/Components/PrimitiveComponent.h
@@ -1954,16 +1954,16 @@ public:
 	virtual FBodyInstance* GetBodyInstance(FName BoneName = NAME_None, bool bGetWelded = true, int32 Index = INDEX_NONE) const;
 
 	/**
-	 * Returns BodyInstanceFixedTickHandle of the component. For use in the FixedTick event
+	 * Returns BodyInstanceAsyncPhysicsTickHandle of the component. For use in the Async Physics Tick event
 	*
 	* @param BoneName				Used to get body associated with specific bone. NAME_None automatically gets the root most body
 	* @param bGetWelded				If the component has been welded to another component and bGetWelded is true we return the single welded BodyInstance that is used in the simulation
 	* @param Index					Index used in Components with multiple body instances
 	*
-	* @return		Returns the BodyInstanceFixedTickHandle based on various states (does component have multiple bodies? Is the body welded to another body?)
+	* @return		Returns the BodyInstanceAsyncPhysicsTickHandle based on various states (does component have multiple bodies? Is the body welded to another body?)
 	*/
 	UFUNCTION(BlueprintPure, Category = "Physics")
-	FBodyInstanceFixedTickHandle GetBodyInstanceFixedTickHandle(FName BoneName = NAME_None, bool bGetWelded = true, int32 Index = -1) const;
+		FBodyInstanceAsyncPhysicsTickHandle GetBodyInstanceAsyncPhysicsTickHandle(FName BoneName = NAME_None, bool bGetWelded = true, int32 Index = -1) const;
 
 	/** 
 	 * Returns The square of the distance to closest Body Instance surface. 
diff --git a/Engine/Source/Runtime/Engine/Classes/GameFramework/Actor.h b/Engine/Source/Runtime/Engine/Classes/GameFramework/Actor.h
--- a/Engine/Source/Runtime/Engine/Classes/GameFramework/Actor.h
+++ b/Engine/Source/Runtime/Engine/Classes/GameFramework/Actor.h
@@ -452,9 +452,9 @@ private:
 	/** Set if an Actor tries to be destroyed while it is beginning play so that once BeginPlay ends we can issue the destroy call. */
 	uint8 bActorWantsDestroyDuringBeginPlay : 1;
 	
-	/** Whether to use use the fixed physics tick with this actor. Requires async fixed tick to be enabled (see project settings) */
+	/** Whether to use use the async physics tick with this actor. Requires async physics be enabled (see project settings) */
 	UPROPERTY(EditAnywhere, Category=Physics)
-	uint8 bFixedTickEnabled : 1;
+	uint8 bAsyncPhysicsTickEnabled : 1;
 
 	/** Enum defining if BeginPlay has started or finished */
 	enum class EActorBeginPlayState : uint8
@@ -1751,9 +1751,9 @@ public:
 	UFUNCTION(BlueprintImplementableEvent, meta=(DisplayName = "Tick"))
 	void ReceiveTick(float DeltaSeconds);
 
-	/** Event called every physics tick if EnabledFixedTick is true and async fixed tick is enabled (see physics project settings)*/
-	UFUNCTION(BlueprintImplementableEvent, meta = (DisplayName = "Fixed Tick"))
-	void ReceiveFixedTick(float DeltaSeconds, float SimSeconds);
+	/** Event called every async physics tick if bAsyncPhysicsTickEnabled is true and async physics is enabled (see physics project settings)*/
+	UFUNCTION(BlueprintImplementableEvent, meta = (DisplayName = "Async Physics Tick"))
+	void ReceiveAsyncPhysicsTick(float DeltaSeconds, float SimSeconds);
 
 	/** 
 	 *	Event when this actor overlaps another actor, for example a player walking into a trigger.
@@ -2482,14 +2482,14 @@ public:
 	virtual void TickActor( float DeltaTime, enum ELevelTick TickType, FActorTickFunction& ThisTickFunction );
 
 	/**
-	 * Function called for every physics step on this Actor. Override this function to implement custom logic to be executed every physics step.
-	 * bFixedTick must be set to true.
-	 * Requires async fixed tick to be enabled (see physics project settings)
+	 * Override this function to implement custom logic to be executed every async physics step.
+	 * bAsyncPhysicsTick must be set to true.
+	 * Requires async physics to be enabled (see physics project settings)
 	 *	
-	 * @param DeltaTime - The fixed tick delta time
+	 * @param DeltaTime - The physics step delta time (see physics project settings)
 	 * @param SimTime - This is the total sim time since the sim began, it is an integer multiple of DeltaTime
 	 */
-	virtual void FixedTickActor(float DeltaTime, float SimTime) { ReceiveFixedTick(DeltaTime, SimTime); }
+	virtual void AsyncPhysicsTickActor(float DeltaTime, float SimTime) { ReceiveAsyncPhysicsTick(DeltaTime, SimTime); }
 
 	/**
 	 * Called when an actor is done spawning into the world (from UWorld::SpawnActor), both in the editor and during gameplay
diff --git a/Engine/Source/Runtime/Engine/Classes/PhysicsEngine/BodyInstance.h b/Engine/Source/Runtime/Engine/Classes/PhysicsEngine/BodyInstance.h
--- a/Engine/Source/Runtime/Engine/Classes/PhysicsEngine/BodyInstance.h
+++ b/Engine/Source/Runtime/Engine/Classes/PhysicsEngine/BodyInstance.h
@@ -254,7 +254,7 @@ namespace Chaos
 }
 
 USTRUCT(BlueprintType)
-struct ENGINE_API FBodyInstanceFixedTickHandle
+struct ENGINE_API FBodyInstanceAsyncPhysicsTickHandle
 {
 	GENERATED_BODY()
 	FPhysicsActorHandle Proxy = nullptr;
@@ -569,7 +569,7 @@ public:
 	// Internal physics representation of our body instance
 	FPhysicsActorHandle ActorHandle;
 
-	FBodyInstanceFixedTickHandle GetBodyInstancePhysicsThreadHandle() const { return FBodyInstanceFixedTickHandle{ ActorHandle }; }
+	FBodyInstanceAsyncPhysicsTickHandle GetBodyInstanceAsyncPhysicsTickHandle() const { return FBodyInstanceAsyncPhysicsTickHandle{ ActorHandle }; }
 
 #if USE_BODYINSTANCE_DEBUG_NAMES
 	TSharedPtr<TArray<ANSICHAR>> CharDebugName;
diff --git a/Engine/Source/Runtime/Engine/Private/Actor.cpp b/Engine/Source/Runtime/Engine/Private/Actor.cpp
--- a/Engine/Source/Runtime/Engine/Private/Actor.cpp
+++ b/Engine/Source/Runtime/Engine/Private/Actor.cpp
@@ -108,7 +108,7 @@ void AActor::InitializeDefaults()
 	PrimaryActorTick.bCanEverTick = false;
 	PrimaryActorTick.bStartWithTickEnabled = true;
 	PrimaryActorTick.SetTickFunctionEnable(false); 
-	bFixedTickEnabled = false;
+	bAsyncPhysicsTickEnabled = false;
 
 	CustomTimeDilation = 1.0f;
 
@@ -1145,17 +1145,17 @@ void AActor::RegisterAllActorTickFunctions(bool bRegister, bool bDoComponents)
 		}
 
 #if WITH_CHAOS
-		if (bFixedTickEnabled)
+		if (bAsyncPhysicsTickEnabled)
 		{
 			if (FPhysScene_Chaos* Scene = static_cast<FPhysScene_Chaos*>(GetWorld()->GetPhysicsScene()))
 			{
 				if (bRegister)
 				{
-					Scene->RegisterFixedTickActor(this);
+					Scene->RegisterAsyncPhysicsTickActor(this);
 				}
 				else
 				{
-					Scene->UnregisterFixedTickActor(this);
+					Scene->UnregisterAsyncPhysicsTickActor(this);
 				}
 			}
 		}
diff --git a/Engine/Source/Runtime/Engine/Private/Components/ActorComponent.cpp b/Engine/Source/Runtime/Engine/Private/Components/ActorComponent.cpp
--- a/Engine/Source/Runtime/Engine/Private/Components/ActorComponent.cpp
+++ b/Engine/Source/Runtime/Engine/Private/Components/ActorComponent.cpp
@@ -327,7 +327,7 @@ UActorComponent::UActorComponent(const FObjectInitializer& ObjectInitializer /*=
 	bNavigationRelevant = false;
 
 	bMarkedForPreEndOfFrameSync = false;
-	bFixedTickEnabled = false;
+	bAsyncPhysicsTickEnabled = false;
 }
 
 void UActorComponent::PostInitProperties()
@@ -1193,17 +1193,17 @@ void UActorComponent::RegisterAllComponentTickFunctions(bool bRegister)
 		}
 
 #if WITH_CHAOS
-		if(bFixedTickEnabled)
+		if(bAsyncPhysicsTickEnabled)
 		{
 			if(FPhysScene_Chaos* Scene = static_cast<FPhysScene_Chaos*>(WorldPrivate->GetPhysicsScene()))
 			{
 				if(bRegister)
 				{
-					Scene->RegisterFixedTickComponent(this);
+					Scene->RegisterAsyncPhysicsTickComponent(this);
 				}
 				else
 				{
-					Scene->UnregisterFixedTickComponent(this);
+					Scene->UnregisterAsyncPhysicsTickComponent(this);
 				}
 			}
 		}
diff --git a/Engine/Source/Runtime/Engine/Private/PhysicsEngine/BodyInstance.cpp b/Engine/Source/Runtime/Engine/Private/PhysicsEngine/BodyInstance.cpp
--- a/Engine/Source/Runtime/Engine/Private/PhysicsEngine/BodyInstance.cpp
+++ b/Engine/Source/Runtime/Engine/Private/PhysicsEngine/BodyInstance.cpp
@@ -4534,7 +4534,7 @@ FRecalculatedMassProperties& FBodyInstance::OnRecalculatedMassProperties()
 }
 
 
-bool FBodyInstanceFixedTickHandle::IsValid() const
+bool FBodyInstanceAsyncPhysicsTickHandle::IsValid() const
 {
 #if WITH_CHAOS
 	return Proxy && Proxy->GetPhysicsThreadAPI() != nullptr;
@@ -4543,7 +4543,7 @@ bool FBodyInstanceFixedTickHandle::IsValid() const
 #endif
 }
 
-Chaos::FRigidBodyHandle_Internal* FBodyInstanceFixedTickHandle::operator->()
+Chaos::FRigidBodyHandle_Internal* FBodyInstanceAsyncPhysicsTickHandle::operator->()
 {
 #if WITH_CHAOS
 	return Proxy->GetPhysicsThreadAPI();
diff --git a/Engine/Source/Runtime/Engine/Private/PhysicsEngine/Experimental/PhysScene_Chaos.cpp b/Engine/Source/Runtime/Engine/Private/PhysicsEngine/Experimental/PhysScene_Chaos.cpp
--- a/Engine/Source/Runtime/Engine/Private/PhysicsEngine/Experimental/PhysScene_Chaos.cpp
+++ b/Engine/Source/Runtime/Engine/Private/PhysicsEngine/Experimental/PhysScene_Chaos.cpp
@@ -71,31 +71,31 @@ DECLARE_CYCLE_STAT(TEXT("Update Kinematics On Deferred SkelMeshes"), STAT_Update
 #include "Editor.h"
 #endif
 
-class FFixedTickCallback : public Chaos::TSimCallbackObject<>
+class FAsyncPhysicsTickCallback : public Chaos::TSimCallbackObject<>
 {
 public:
-	FFixedTickCallback()
+	FAsyncPhysicsTickCallback()
 		: TSimCallbackObject<>(/*InRunOnFrozenGameThread=*/ true)
 	{
 
 	}
 
-	TSet<UActorComponent*> FixedTickComponents;
-	TSet<AActor*> FixedTickActors;
+	TSet<UActorComponent*> AsyncPhysicsTickComponents;
+	TSet<AActor*> AsyncPhysicsTickActors;
 	virtual void OnPreSimulate_Internal() override
 	{
 		using namespace Chaos;
 		const FReal DeltaTime = GetDeltaTime_Internal();
 		const FReal SimTime = GetSimTime_Internal();
-		//TODO: handle case where callbacks modify FixedTickComponents or FixedTickActors
-		for (UActorComponent* Component : FixedTickComponents)
+		//TODO: handle case where callbacks modify AsyncPhysicsTickComponents or AsyncPhysicsTickActors
+		for (UActorComponent* Component : AsyncPhysicsTickComponents)
 		{
-			Component->FixedTickComponent(DeltaTime, SimTime);
+			Component->AsyncPhysicsTickComponent(DeltaTime, SimTime);
 		}
 
-		for(AActor* Actor : FixedTickActors)
+		for(AActor* Actor : AsyncPhysicsTickActors)
 		{
-			Actor->FixedTickActor(DeltaTime, SimTime);
+			Actor->AsyncPhysicsTickActor(DeltaTime, SimTime);
 		}
 	}
 };
@@ -448,9 +448,9 @@ FPhysScene_Chaos::~FPhysScene_Chaos()
 {
 #if WITH_CHAOS
 
-	if (FixedTickCallback)
+	if (AsyncPhysicsTickCallback)
 	{
-		SceneSolver->UnregisterAndFreeSimCallbackObject_External(FixedTickCallback);
+		SceneSolver->UnregisterAndFreeSimCallbackObject_External(AsyncPhysicsTickCallback);
 	}
 
 	// Must ensure deferred components do not hold onto scene pointer.
@@ -1847,39 +1847,39 @@ void FPhysScene_Chaos::ResimNFrames(const int32 NumFramesRequested)
 #endif
 }
 
-void FPhysScene_Chaos::EnableFixedTickCallback()
+void FPhysScene_Chaos::EnableAsyncPhysicsTickCallback()
 {
-	if (FixedTickCallback == nullptr)
+	if (AsyncPhysicsTickCallback == nullptr)
 	{
-		FixedTickCallback = SceneSolver->CreateAndRegisterSimCallbackObject_External<FFixedTickCallback>();
+		AsyncPhysicsTickCallback = SceneSolver->CreateAndRegisterSimCallbackObject_External<FAsyncPhysicsTickCallback>();
 	}
 }
 
-void FPhysScene_Chaos::RegisterFixedTickComponent(UActorComponent* Component)
+void FPhysScene_Chaos::RegisterAsyncPhysicsTickComponent(UActorComponent* Component)
 {
-	EnableFixedTickCallback();
-	FixedTickCallback->FixedTickComponents.Add(Component);
+	EnableAsyncPhysicsTickCallback();
+	AsyncPhysicsTickCallback->AsyncPhysicsTickComponents.Add(Component);
 }
 
-void FPhysScene_Chaos::UnregisterFixedTickComponent(UActorComponent* Component)
+void FPhysScene_Chaos::UnregisterAsyncPhysicsTickComponent(UActorComponent* Component)
 {
-	if (FixedTickCallback)
+	if (AsyncPhysicsTickCallback)
 	{
-		FixedTickCallback->FixedTickComponents.Remove(Component);
+		AsyncPhysicsTickCallback->AsyncPhysicsTickComponents.Remove(Component);
 	}
 }
 
-void FPhysScene_Chaos::RegisterFixedTickActor(AActor* Actor)
+void FPhysScene_Chaos::RegisterAsyncPhysicsTickActor(AActor* Actor)
 {
-	EnableFixedTickCallback();
-	FixedTickCallback->FixedTickActors.Add(Actor);
+	EnableAsyncPhysicsTickCallback();
+	AsyncPhysicsTickCallback->AsyncPhysicsTickActors.Add(Actor);
 }
 
-void FPhysScene_Chaos::UnregisterFixedTickActor(AActor* Actor)
+void FPhysScene_Chaos::UnregisterAsyncPhysicsTickActor(AActor* Actor)
 {
-	if (FixedTickCallback)
+	if (AsyncPhysicsTickCallback)
 	{
-		FixedTickCallback->FixedTickActors.Remove(Actor);
+		AsyncPhysicsTickCallback->AsyncPhysicsTickActors.Remove(Actor);
 	}
 }
 
diff --git a/Engine/Source/Runtime/Engine/Private/PrimitiveComponentPhysics.cpp b/Engine/Source/Runtime/Engine/Private/PrimitiveComponentPhysics.cpp
--- a/Engine/Source/Runtime/Engine/Private/PrimitiveComponentPhysics.cpp
+++ b/Engine/Source/Runtime/Engine/Private/PrimitiveComponentPhysics.cpp
@@ -824,15 +824,15 @@ FBodyInstance* UPrimitiveComponent::GetBodyInstance(FName BoneName, bool bGetWel
 	return const_cast<FBodyInstance*>((bGetWelded && BodyInstance.WeldParent) ? BodyInstance.WeldParent : &BodyInstance);
 }
 
-FBodyInstanceFixedTickHandle UPrimitiveComponent::GetBodyInstanceFixedTickHandle(FName BoneName, bool bGetWelded, int32 Index) const
+FBodyInstanceAsyncPhysicsTickHandle UPrimitiveComponent::GetBodyInstanceAsyncPhysicsTickHandle(FName BoneName, bool bGetWelded, int32 Index) const
 {
 	if(FBodyInstance* BI = GetBodyInstance(BoneName, bGetWelded, Index))
 	{
-		return BI->GetBodyInstancePhysicsThreadHandle();
+		return BI->GetBodyInstanceAsyncPhysicsTickHandle();
 	}
 	else
 	{
-		return FBodyInstanceFixedTickHandle();
+		return FBodyInstanceAsyncPhysicsTickHandle();
 	}
 	
 }
diff --git a/Engine/Source/Runtime/Engine/Public/Physics/Experimental/PhysScene_Chaos.h b/Engine/Source/Runtime/Engine/Public/Physics/Experimental/PhysScene_Chaos.h
--- a/Engine/Source/Runtime/Engine/Public/Physics/Experimental/PhysScene_Chaos.h
+++ b/Engine/Source/Runtime/Engine/Public/Physics/Experimental/PhysScene_Chaos.h
@@ -122,11 +122,11 @@ public:
 
 	void UnRegisterForCollisionEvents(UPrimitiveComponent* Component);
 
-	void RegisterFixedTickComponent(UActorComponent* Component);
-	void UnregisterFixedTickComponent(UActorComponent* Component);
+	void RegisterAsyncPhysicsTickComponent(UActorComponent* Component);
+	void UnregisterAsyncPhysicsTickComponent(UActorComponent* Component);
 
-	void RegisterFixedTickActor(AActor* Actor);
-	void UnregisterFixedTickActor(AActor* Actor);
+	void RegisterAsyncPhysicsTickActor(AActor* Actor);
+	void UnregisterAsyncPhysicsTickActor(AActor* Actor);
 
 	/**
 	 * Called during creation of the physics state for gamethread objects to pass off an object to the physics thread
@@ -300,7 +300,7 @@ private:
 	virtual void OnSyncBodies(Chaos::FPhysicsSolverBase* Solver) override;
 #endif
 
-	void EnableFixedTickCallback();
+	void EnableAsyncPhysicsTickCallback();
 
 #if 0
 	void SetKinematicTransform(FPhysicsActorHandle& InActorReference,const Chaos::TRigidTransform<float,3>& NewTransform)
@@ -372,7 +372,7 @@ private:
 #if WITH_CHAOS
 	Chaos::FReal LastEventDispatchTime;
 #endif 
-	class FFixedTickCallback* FixedTickCallback = nullptr;
+	class FAsyncPhysicsTickCallback* AsyncPhysicsTickCallback = nullptr;
 
 #if WITH_EDITOR
 	// Counter used to check a match with the single step status.
diff --git a/Engine/Source/Runtime/Engine/Public/Physics/Experimental/PhysicsThreadLibrary.h b/Engine/Source/Runtime/Engine/Public/Physics/Experimental/PhysicsThreadLibrary.h
--- a/Engine/Source/Runtime/Engine/Public/Physics/Experimental/PhysicsThreadLibrary.h
+++ b/Engine/Source/Runtime/Engine/Public/Physics/Experimental/PhysicsThreadLibrary.h
@@ -29,7 +29,7 @@ class ENGINE_API UPhysicsThreadLibrary : public UBlueprintFunctionLibrary
 	 *  @param  bAccelChange If true, Force is taken as a change in acceleration instead of a physical force (i.e. mass will have no effect).
 	 */
 	UFUNCTION(BlueprintCallable, Category = "Utilities")
-	static void AddForce(FBodyInstanceFixedTickHandle Handle, FVector Force, bool bAccelChange = false)
+	static void AddForce(FBodyInstanceAsyncPhysicsTickHandle Handle, FVector Force, bool bAccelChange = false)
 	{
 		if (Handle)
 		{
