Fixed: GPUScene specific validation assert in a mobile preview

diff --git a/Engine/Source/Runtime/Engine/Private/PrimitiveSceneProxy.cpp b/Engine/Source/Runtime/Engine/Private/PrimitiveSceneProxy.cpp
--- a/Engine/Source/Runtime/Engine/Private/PrimitiveSceneProxy.cpp
+++ b/Engine/Source/Runtime/Engine/Private/PrimitiveSceneProxy.cpp
@@ -1005,11 +1005,13 @@ void FPrimitiveSceneProxy::EnableGPUSceneSupportFlags()
 {
 	const ERHIFeatureLevel::Type FeatureLevel = GetScene().GetFeatureLevel();
 	const bool bUseGPUScene = UseGPUScene(GMaxRHIShaderPlatform, FeatureLevel);
+	const bool bMobilePath = (FeatureLevel == ERHIFeatureLevel::ES3_1);
 
 	// Skip primitive uniform buffer if we will be using local vertex factory which gets it's data from GPUScene.
 	// Vertex shaders on mobile may still use PrimitiveUB with GPUScene enabled
-	bVFRequiresPrimitiveUniformBuffer = !bUseGPUScene || (FeatureLevel == ERHIFeatureLevel::ES3_1);
-	bSupportsGPUScene = bUseGPUScene;
+	bVFRequiresPrimitiveUniformBuffer = !bUseGPUScene || bMobilePath;
+	// For mobile we always assume that proxy does not support GPUScene, as it depends on vertex factory setup which happens later
+	bSupportsGPUScene = bMobilePath ? false : bUseGPUScene;
 }
 
 /**
