//UE5/Release-5.0 - Build ContentExamples Win64 - asset out of date with source animation
#fix ensure that source animation is fully loaded before comparing RawDataGUIDs

diff --git a/Engine/Source/Runtime/Engine/Private/Animation/PoseAsset.cpp b/Engine/Source/Runtime/Engine/Private/Animation/PoseAsset.cpp
--- a/Engine/Source/Runtime/Engine/Private/Animation/PoseAsset.cpp
+++ b/Engine/Source/Runtime/Engine/Private/Animation/PoseAsset.cpp
@@ -913,14 +913,27 @@ void UPoseAsset::PostLoad()
 
 	if (GetLinkerCustomVersion(FUE5ReleaseStreamObjectVersion::GUID) >= FUE5ReleaseStreamObjectVersion::PoseAssetRawDataGUID)
 	{
-		if (SourceAnimation != nullptr && (!SourceAnimationRawDataGUID.IsValid() || SourceAnimationRawDataGUID != SourceAnimation->GetRawDataGuid()))
+		if (SourceAnimation)
 		{
-			FFormatNamedArguments Args;
-			Args.Add(TEXT("AssetName"), FText::FromString(GetPathName()));
-			Args.Add(TEXT("SourceAsset"), FText::FromString(SourceAnimation->GetPathName()));
-			const FText ResultText = FText::Format(LOCTEXT("PoseAssetSourceOutOfDate", "PoseAsset {AssetName} is out-of-date with its source animation {SourceAsset}"), Args);
-			UE_LOG(LogAnimation, Warning,TEXT("%s"), *ResultText.ToString());
-		}
+			// Fully load the source animation to ensure its RawDataGUID is populated
+			if(SourceAnimation->HasAnyFlags(RF_NeedLoad))
+			{
+				if (FLinkerLoad* Linker = SourceAnimation->GetLinker())
+				{
+					Linker->Preload(SourceAnimation);
+				}
+			}
+			SourceAnimation->ConditionalPostLoad();
+			
+			if (!SourceAnimationRawDataGUID.IsValid() || SourceAnimationRawDataGUID != SourceAnimation->GetRawDataGuid())
+			{
+				FFormatNamedArguments Args;
+				Args.Add(TEXT("AssetName"), FText::FromString(GetPathName()));
+				Args.Add(TEXT("SourceAsset"), FText::FromString(SourceAnimation->GetPathName()));
+				const FText ResultText = FText::Format(LOCTEXT("PoseAssetSourceOutOfDate", "PoseAsset {AssetName} is out-of-date with its source animation {SourceAsset}"), Args);
+				UE_LOG(LogAnimation, Warning,TEXT("%s"), *ResultText.ToString());
+			}
+		}	
 	}	
 #endif // WITH_EDITOR
 
