Serialize space that HISM cluster tree is built in and always use that for translated instance space.
This is to make sure that we always handle old data after the recent addition of translated instance space.
#rnx

diff --git a/Engine/Source/Runtime/Engine/Classes/Components/HierarchicalInstancedStaticMeshComponent.h b/Engine/Source/Runtime/Engine/Classes/Components/HierarchicalInstancedStaticMeshComponent.h
--- a/Engine/Source/Runtime/Engine/Classes/Components/HierarchicalInstancedStaticMeshComponent.h
+++ b/Engine/Source/Runtime/Engine/Classes/Components/HierarchicalInstancedStaticMeshComponent.h
@@ -143,6 +143,15 @@ class ENGINE_API UHierarchicalInstancedStaticMeshComponent : public UInstancedSt
 
 	TSharedPtr<TArray<FClusterNode>, ESPMode::ThreadSafe> ClusterTreePtr;
 
+	// If true then we allow a translated space when building the cluster tree.
+	// This can help for impementations (foliage) where we can have instances with offsets to large for single float precision.
+	UPROPERTY()
+	uint32 bUseTranslatedInstanceSpace : 1;
+
+	// Origin of the translated space used when building the cluster tree.
+	UPROPERTY()
+	FVector TranslatedInstanceSpaceOrigin;
+
 	// Table for remapping instances from cluster tree to PerInstanceSMData order
 	UPROPERTY()
 	TArray<int32> SortedInstances;
@@ -273,6 +282,7 @@ protected:
 	void ApplyEmpty();
 	void SetPerInstanceLightMapAndEditorData(FStaticMeshInstanceData& PerInstanceData, const TArray<TRefCountPtr<HHitProxy>>& HitProxies);
 
+	FVector CalcTranslatedInstanceSpaceOrigin() const;
 	void GetInstanceTransforms(TArray<FMatrix>& InstanceTransforms, FVector const& Offset) const;
 	void InitializeInstancingRandomSeed();
 
@@ -288,6 +298,8 @@ protected:
 
 	virtual void OnPostLoadPerInstanceData() override;
 
+	virtual FVector GetTranslatedInstanceSpaceOrigin() const override { return TranslatedInstanceSpaceOrigin; }
+
 	virtual void GetNavigationPerInstanceTransforms(const FBox& AreaBox, TArray<FTransform>& InstanceData) const override;
 	virtual void PartialNavigationUpdate(int32 InstanceIdx) override;
 	virtual bool SupportsPartialNavigationUpdate() const override { return true; }
diff --git a/Engine/Source/Runtime/Engine/Private/HierarchicalInstancedStaticMesh.cpp b/Engine/Source/Runtime/Engine/Private/HierarchicalInstancedStaticMesh.cpp
--- a/Engine/Source/Runtime/Engine/Private/HierarchicalInstancedStaticMesh.cpp
+++ b/Engine/Source/Runtime/Engine/Private/HierarchicalInstancedStaticMesh.cpp
@@ -2015,6 +2015,8 @@ static FBox GetClusterTreeBounds(TArray<FClusterNode> const& InClusterTree, FVec
 UHierarchicalInstancedStaticMeshComponent::UHierarchicalInstancedStaticMeshComponent(const FObjectInitializer& ObjectInitializer)
 	: Super(ObjectInitializer)
 	, ClusterTreePtr(MakeShareable(new TArray<FClusterNode>))
+	, bUseTranslatedInstanceSpace(false)
+	, TranslatedInstanceSpaceOrigin(ForceInitToZero)
 	, NumBuiltInstances(0)
 	, NumBuiltRenderInstances(0)
 	, UnbuiltInstanceBounds(ForceInit)
@@ -2124,7 +2126,7 @@ void UHierarchicalInstancedStaticMeshComponent::Serialize(FArchive& Ar)
 	if (Ar.IsLoading() && !BuiltInstanceBounds.IsValid)
 	{
 		TArray<FClusterNode>& ClusterTree = *ClusterTreePtr.Get();
-		BuiltInstanceBounds = GetClusterTreeBounds(ClusterTree, GetTranslatedInstanceSpaceOrigin());
+		BuiltInstanceBounds = GetClusterTreeBounds(ClusterTree, TranslatedInstanceSpaceOrigin);
 	}
 }
 
@@ -2512,7 +2514,8 @@ int32 UHierarchicalInstancedStaticMeshComponent::AddInstance(const FTransform& I
 			++InstanceCountToRender;
 		}
 
-		InstanceUpdateCmdBuffer.AddInstance(PerInstanceSMData[InstanceIndex].Transform);
+		// CmdBuffer takes final render transforms so should include the translated space.
+		InstanceUpdateCmdBuffer.AddInstance(PerInstanceSMData[InstanceIndex].Transform.ConcatTranslation(-TranslatedInstanceSpaceOrigin));
 
 		const FBox NewInstanceBounds = GetStaticMesh()->GetBounds().GetBox().TransformBy(PerInstanceSMData[InstanceIndex].Transform);
 		UnbuiltInstanceBounds += NewInstanceBounds;
@@ -2719,7 +2722,7 @@ void UHierarchicalInstancedStaticMeshComponent::BuildTree()
 	if (PerInstanceSMData.Num() > 0 && GetStaticMesh() && !GetStaticMesh()->IsCompiling() && GetStaticMesh()->HasValidRenderData())
 	{
 		// Build the tree in translated space to maintain precision.
-		const FVector TranslatedInstanceSpaceOrigin = GetTranslatedInstanceSpaceOrigin();
+		TranslatedInstanceSpaceOrigin = CalcTranslatedInstanceSpaceOrigin();
 
 		InitializeInstancingRandomSeed();
 		TArray<FMatrix> InstanceTransforms;
@@ -2772,7 +2775,7 @@ void UHierarchicalInstancedStaticMeshComponent::AcceptPrebuiltTree(TArray<FClust
 	// this is only for prebuild data, already in the correct order
 	check(!PerInstanceSMData.Num());
 	NumBuiltInstances = 0;
-	check(GetTranslatedInstanceSpaceOrigin() == FVector::Zero());
+	TranslatedInstanceSpaceOrigin = FVector::Zero();
 	check(PerInstanceRenderData.IsValid());
 	NumBuiltRenderInstances = InNumBuiltRenderInstances;
 	check(NumBuiltRenderInstances);
@@ -2878,7 +2881,7 @@ void UHierarchicalInstancedStaticMeshComponent::ApplyBuildTree(FClusterBuilder&
 
 	// Get the new bounds taking into account the translated space used when building the tree.
 	const TArray<FClusterNode>& ClusterTree = *ClusterTreePtr;
-	BuiltInstanceBounds = GetClusterTreeBounds(ClusterTree, GetTranslatedInstanceSpaceOrigin());
+	BuiltInstanceBounds = GetClusterTreeBounds(ClusterTree, TranslatedInstanceSpaceOrigin);
 
 	UnbuiltInstanceBounds.Init();
 	UnbuiltInstanceBoundsList.Empty();
@@ -2969,6 +2972,18 @@ bool UHierarchicalInstancedStaticMeshComponent::BuildTreeIfOutdated(bool Async,
 	return false;
 }
 
+FVector UHierarchicalInstancedStaticMeshComponent::CalcTranslatedInstanceSpaceOrigin() const
+{
+	// Foliage is often built in world space which can cause problems with large world coordinates because
+	// the instance transforms in the renderer are single precision, and the HISM culling is also single precision.
+	// We should fix the HISM culling to be double precision.
+	// But the instance transforms (relative to the owner primitive) will probably stay single precision to not bloat memory.
+	// A fix for that is to have authoring tools set sensible primitive transforms (instead of identity).
+	// But until that happens we set a translated instance space here.
+	// For simplicity we use the first instance as the origin of the translated space.
+	return bUseTranslatedInstanceSpace && PerInstanceSMData.Num() ? PerInstanceSMData[0].Transform.GetOrigin() : FVector::Zero();
+}
+
 void UHierarchicalInstancedStaticMeshComponent::GetInstanceTransforms(TArray<FMatrix>& InstanceTransforms, FVector const& Offset) const
 {
 	double StartTime = FPlatformTime::Seconds();
@@ -3019,7 +3034,7 @@ void UHierarchicalInstancedStaticMeshComponent::BuildTreeAsync()
 		double StartTime = FPlatformTime::Seconds();
 		
 		// Build the tree in translated space to maintain precision.
-		const FVector TranslatedInstanceSpaceOrigin = GetTranslatedInstanceSpaceOrigin();
+		TranslatedInstanceSpaceOrigin = CalcTranslatedInstanceSpaceOrigin();
 
 		InitializeInstancingRandomSeed();
 		TArray<FMatrix> InstanceTransforms;
@@ -3457,7 +3472,7 @@ void UHierarchicalInstancedStaticMeshComponent::FlushAccumulatedNavigationUpdate
 		const TArray<FClusterNode>& ClusterTree = *ClusterTreePtr;
 		if (ClusterTree.Num())
 		{
-			const FBox NewBounds = GetClusterTreeBounds(ClusterTree, GetTranslatedInstanceSpaceOrigin()).TransformBy(GetComponentTransform());
+			const FBox NewBounds = GetClusterTreeBounds(ClusterTree, TranslatedInstanceSpaceOrigin).TransformBy(GetComponentTransform());
 			FNavigationSystem::OnComponentBoundsChanged(*this, NewBounds, AccumulatedNavigationDirtyArea);
 		}
 			
diff --git a/Engine/Source/Runtime/Engine/Private/InstancedStaticMesh.cpp b/Engine/Source/Runtime/Engine/Private/InstancedStaticMesh.cpp
--- a/Engine/Source/Runtime/Engine/Private/InstancedStaticMesh.cpp
+++ b/Engine/Source/Runtime/Engine/Private/InstancedStaticMesh.cpp
@@ -2300,8 +2300,7 @@ FPrimitiveSceneProxy* UInstancedStaticMeshComponent::CreateSceneProxy()
 FMatrix UInstancedStaticMeshComponent::GetRenderMatrix() const
 {
 	// Apply the translated space to the render matrix.
-	const FVector TranslatedInstanceSpaceOrigin = GetTranslatedInstanceSpaceOrigin();
-	return GetComponentTransform().ToMatrixWithScale().ConcatTranslation(TranslatedInstanceSpaceOrigin);
+	return (FTransform(GetTranslatedInstanceSpaceOrigin()) * GetComponentTransform()).ToMatrixWithScale();
 }
 
 void UInstancedStaticMeshComponent::CreateHitProxyData(TArray<TRefCountPtr<HHitProxy>>& HitProxies)
diff --git a/Engine/Source/Runtime/Foliage/Private/InstancedFoliage.cpp b/Engine/Source/Runtime/Foliage/Private/InstancedFoliage.cpp
--- a/Engine/Source/Runtime/Foliage/Private/InstancedFoliage.cpp
+++ b/Engine/Source/Runtime/Foliage/Private/InstancedFoliage.cpp
@@ -5625,23 +5625,18 @@ UFoliageInstancedStaticMeshComponent::UFoliageInstancedStaticMeshComponent(const
 	: Super(ObjectInitializer)
 	, bEnableDiscardOnLoad(false)
 {
+	// Foliage is often built in world space which can cause problems with large world coordinates.
+	// We use a translated space to deal with that.
+	// todo: Maybe we can address this on the tool side and strip out all of the translated instance space code. It would require:
+	// * foliage components to always be placed near their instances (instead of at origin).
+	// * warnings and a fixup option for components which have instances with very large coordinates.
+	bUseTranslatedInstanceSpace = true;
+
 #if WITH_EDITORONLY_DATA
 	bEnableAutoLODGeneration = false;
 #endif
 }
 
-FVector UFoliageInstancedStaticMeshComponent::GetTranslatedInstanceSpaceOrigin() const
-{
-	// Foliage is often built in world space which can cause problems with large world coordinates because
-	// the instance transforms in the renderer are single precision, and the HISM culling is also single precision.
-	// We should fix the HISM culling to be double precision.
-	// But the instance transforms (relative to the owner primitive) will probably stay single precision to not bloat memory.
-	// A fix for that is to have authoring tools set sensible primitive transforms (instead of identity).
-	// But until that happens we set a translated instance space here.
-	// For simplicity we use the first instance as the origin of the translated space.
-	return PerInstanceSMData.Num() ? PerInstanceSMData[0].Transform.GetOrigin() : FVector::Zero();
-}
-
 void UFoliageInstancedStaticMeshComponent::ReceiveComponentDamage(float DamageAmount, FDamageEvent const& DamageEvent, AController* EventInstigator, AActor* DamageCauser)
 {
 	Super::ReceiveComponentDamage(DamageAmount, DamageEvent, EventInstigator, DamageCauser);
diff --git a/Engine/Source/Runtime/Foliage/Public/FoliageInstancedStaticMeshComponent.h b/Engine/Source/Runtime/Foliage/Public/FoliageInstancedStaticMeshComponent.h
--- a/Engine/Source/Runtime/Foliage/Public/FoliageInstancedStaticMeshComponent.h
+++ b/Engine/Source/Runtime/Foliage/Public/FoliageInstancedStaticMeshComponent.h
@@ -43,9 +43,6 @@ public:
 	virtual uint64 GetHiddenEditorViews() const override;
 #endif// WITH_EDITOR
 
-	// UInstancedStaticMesh interface
-	virtual FVector GetTranslatedInstanceSpaceOrigin() const override;
-
 	/** Used by procedural generation to link generated component with its creator */
 	void SetGenerationGuid(const FGuid& InGuid) { GenerationGuid = InGuid; }
 	const FGuid& GetGenerationGuid() const { return GenerationGuid; }
