SubObject replication
* Actors and ActorComponents can now replicate subobjects by registering them via AddReplicatedSubObject instead of overriding the ReplicateSubObjects method.
* Their bReplicateUsingRegisteredSubObjectList flag needs to be enabled for the list to be used.
* When registering subobjects you can pass a ELifetimeCondition to choose on which connection the object will be replicated to.

diff --git a/Engine/Source/Runtime/Engine/Classes/Components/ActorComponent.h b/Engine/Source/Runtime/Engine/Classes/Components/ActorComponent.h
--- a/Engine/Source/Runtime/Engine/Classes/Components/ActorComponent.h
+++ b/Engine/Source/Runtime/Engine/Classes/Components/ActorComponent.h
@@ -164,6 +164,13 @@ protected:
 	UPROPERTY()
 	uint8 bNetAddressable:1;
 
+	/**
+	* When true the replication system will only replicate the registered subobjects list
+	* When false the replication system will instead call the virtual ReplicateSubObjects() function where the subobjects need to be manually replicated.
+	*/
+	UPROPERTY(Config, EditDefaultsOnly, BlueprintReadOnly, Category=Replication, AdvancedDisplay)
+	uint8 bReplicateUsingRegisteredSubObjectList : 1;
+
 private:
 	/** Is this component currently replicating? Should the network code consider it for replication? Owning Actor must be replicating first! */
 	UPROPERTY(EditDefaultsOnly, BlueprintReadOnly, Replicated, Category=ComponentReplication,meta=(DisplayName = "Component Replicates", AllowPrivateAccess = "true"))
@@ -467,15 +474,39 @@ public:
 		return bReplicates;
 	}
 
-	/** Allows a component to replicate other subobject on the actor  */
+	/** 
+	* Allows a component to replicate other subobject on the actor
+	* This method is used only when bReplicateUsingRegisteredSubObjectList is false.
+	* Otherwise this function is not called and only the ReplicatedSubObjects list is used.
+	*/
 	virtual bool ReplicateSubobjects(class UActorChannel *Channel, class FOutBunch *Bunch, FReplicationFlags *RepFlags);
 
+	/** 
+	* Returns if this component is replicating subobjects via the registration list or via the virtual ReplicateSubObjects method.
+	* Note: the owning actor of this component must also have it's bReplicateUsingRegisteredSubObjectList flag set to true.
+	*/
+	bool IsUsingRegisteredSubObjectList() const { return bReplicateUsingRegisteredSubObjectList; }
+
 	/** Called on the component right before replication occurs */
 	virtual void PreReplication(IRepChangedPropertyTracker & ChangedPropertyTracker);
 
 	/** Returns true if this type of component can ever replicate, override to disable the default behavior */
 	virtual bool GetComponentClassCanReplicate() const;
 
+	/**
+	* Register a SubObject that will get replicated along with the actor component.
+	* The subobject needs to be manually removed from the list before it gets deleted.
+	* @param SubObject The subobject to replicate
+	* @param NetCondition Optional condition to select which type of connection we will replicate the object to.
+	*/
+	void AddReplicatedSubObject(UObject* SubObject, ELifetimeCondition NetCondition = COND_None);
+
+	/**
+	* Unregister a SubObject so it stops being replicated.
+	* @param SubObject The subobject to remove
+	*/
+	void RemoveReplicatedSubObject(UObject* SubObject);
+
 #if WITH_EDITORONLY_DATA
 	/** Returns whether this component is an editor-only object or not */
 	virtual bool IsEditorOnly() const override { return bIsEditorOnly; }
diff --git a/Engine/Source/Runtime/Engine/Classes/Engine/ActorChannel.h b/Engine/Source/Runtime/Engine/Classes/Engine/ActorChannel.h
--- a/Engine/Source/Runtime/Engine/Classes/Engine/ActorChannel.h
+++ b/Engine/Source/Runtime/Engine/Classes/Engine/ActorChannel.h
@@ -273,6 +273,7 @@ public:
 
 	/** utility template for replicating list of replicated subobjects */
 	template<typename Type>
+	UE_DEPRECATED(5.1, "This function will be deleted. Register your subobjects using AddReplicatedSubObject instead.")
 	bool ReplicateSubobjectList(TArray<Type*> &ObjectList, FOutBunch &Bunch, const FReplicationFlags &RepFlags)
 	{
 		bool WroteSomething = false;
@@ -362,6 +363,17 @@ protected:
 	/** Closes the actor channel but with a 'dormant' flag set so it can be reopened */
 	virtual void BecomeDormant() override;
 
+private:
+
+	/** Handle the replication of subobjects for this actor. Returns true if data was written into the Bunch. */
+	bool DoSubObjectReplication(FOutBunch& Bunch, const FReplicationFlags& RepFlags);
+
+	/** Replicate Subobjects using the actor's registered list and its replicated actor component list */
+	bool ReplicateRegisteredSubObjects(FOutBunch& Bunch, const FReplicationFlags& RepFlags);
+
+    /** Write the replicated bits into the bunch data */
+	bool WriteSubObjectInBunch(UObject* Obj, FOutBunch& Bunch, const FReplicationFlags& RepFlags);
+
 private:
 
 	// TODO: It would be nice to merge the tracking of these with PendingGuidResolves, to not duplicate memory,
diff --git a/Engine/Source/Runtime/Engine/Classes/GameFramework/Actor.h b/Engine/Source/Runtime/Engine/Classes/GameFramework/Actor.h
--- a/Engine/Source/Runtime/Engine/Classes/GameFramework/Actor.h
+++ b/Engine/Source/Runtime/Engine/Classes/GameFramework/Actor.h
@@ -17,6 +17,7 @@
 #include "RenderCommandFence.h"
 #include "Misc/ITransaction.h"
 #include "Engine/Level.h"
+#include "Net/Core/Misc/NetSubObjectRegistry.h"
 
 #if WITH_EDITOR
 #include "WorldPartition/DataLayer/ActorDataLayer.h"
@@ -416,6 +417,13 @@ protected:
 	/** Flag indicating we have checked initial simulating physics state to sync networked proxies to the server. */
 	uint8 bNetCheckedInitialPhysicsState : 1;
 
+	/**
+	* When true the replication system will only replicate the registered subobjects and the replicated actor components list
+	* When false the replication system will instead call the virtual ReplicateSubobjects() function where the subobjects and actor components need to be manually replicated.
+	*/
+	UPROPERTY(Config, EditDefaultsOnly, BlueprintReadOnly, Category=Replication, AdvancedDisplay)
+	uint8 bReplicateUsingRegisteredSubObjectList : 1;
+
 private:
 	/** Whether FinishSpawning has been called for this Actor.  If it has not, the Actor is in a malformed state */
 	uint8 bHasFinishedSpawning:1;
@@ -704,6 +712,12 @@ public:
 	/** Returns name of the net driver associated with this actor (all RPCs will go out via this connection) */
 	FName GetNetDriverName() const { return NetDriverName; }
 
+	/** 
+    * Returns true if this actor is replicating SubObjects & ActorComponents via the registration list.
+    * Returns false when he replicates them via the virtual ReplicateSubobjects method 
+    */
+	bool IsUsingRegisteredSubObjectList() const { return bReplicateUsingRegisteredSubObjectList; }
+
 	/** Method that allows an actor to replicate subobjects on its actor channel */
 	virtual bool ReplicateSubobjects(class UActorChannel *Channel, class FOutBunch *Bunch, FReplicationFlags *RepFlags);
 
@@ -3550,6 +3564,54 @@ public:
 		return ReplicatedComponents; 
 	}
 
+	/**
+	* Register a SubObject that will get replicated along with the actor.
+	* The subobject needs to be manually removed from the list before it gets deleted.
+	* @param SubObject The SubObject to replicate
+	* @param NetCondition Optional condition to select which type of connection we will replicate the object to.
+	*/
+	void AddReplicatedSubObject(UObject* SubObject, ELifetimeCondition NetCondition = COND_None);
+
+	/**
+	* Unregister a SubObject so it stops being replicated.
+	* @param SubObject The SubObject to remove
+	*/
+	void RemoveReplicatedSubObject(UObject* SubObject);
+
+	/**
+	* Register a SubObject that will get replicated along with the actor component owning it.
+	* The subobject needs to be manually removed from the list before it gets deleted.
+	* @param SubObject The SubObject to replicate
+	* @param NetCondition Optional condition to select which type of connection we will replicate the object to.
+	*/
+	void AddActorComponentReplicatedSubObject(UActorComponent* OwnerComponent, UObject* SubObject, ELifetimeCondition NetCondition = COND_None);
+
+	/**
+	* Unregister a SubObject owned by an ActorComponent so it stops being replicated.
+	* @param SubObject The SubObject to remove
+	*/
+	void RemoveActorComponentReplicatedSubObject(UActorComponent* OwnerComponent, UObject* SubObject);
+
+private:
+	/** Collection of SubObjects that get replicated when this actor gets replicated. */
+	UE::Net::FSubObjectRegistry ReplicatedSubObjects;
+	friend class FSubObjectGetter;
+
+	/** Keep track of replicated components and their subobject list */
+	struct FReplicatedComponentInfo
+	{
+		// Component that will be replicated
+		UActorComponent* Component = nullptr;
+
+		// Collection of subobjects replicated with this component
+		UE::Net::FSubObjectRegistry SubObjects;
+
+		bool operator==(const FReplicatedComponentInfo& rhs) const { return Component == rhs.Component; }
+		bool operator==(const UActorComponent* rhs) const { return Component == rhs; }
+	};
+	/** Array of replicated components and the list of replicated subobjects they own. Replaces the deprecated ReplicatedCompoments array. */
+	TArray<FReplicatedComponentInfo> ReplicatedComponentsInfo;
+
 protected:
 	/** Set of replicated components, stored as an array to save space as this is generally not very large */
 	TArray<UActorComponent*> ReplicatedComponents;
diff --git a/Engine/Source/Runtime/Engine/Private/ActorReplication.cpp b/Engine/Source/Runtime/Engine/Private/ActorReplication.cpp
--- a/Engine/Source/Runtime/Engine/Private/ActorReplication.cpp
+++ b/Engine/Source/Runtime/Engine/Private/ActorReplication.cpp
@@ -483,6 +483,71 @@ bool AActor::ReplicateSubobjects(UActorChannel *Channel, FOutBunch *Bunch, FRepl
 	return WroteSomething;
 }
 
+void AActor::AddReplicatedSubObject(UObject* SubObject, ELifetimeCondition NetCondition)
+{
+	check(IsValid(SubObject));
+
+	ensureMsgf(IsUsingRegisteredSubObjectList(), TEXT("%s is registering subobjects but bReplicateUsingRegisteredSubObjectList is false. Without the flag set to true the registered subobjects will not be replicated."), *GetName());
+
+	UE::Net::FSubObjectRegistry::EResult Result = ReplicatedSubObjects.AddSubObjectUnique(SubObject, NetCondition);
+
+	UE_CLOG(Result == UE::Net::FSubObjectRegistry::EResult::NewEntry, LogNetSubObject, Verbose, TEXT("%s (0x%p) added replicated subobject %s (0x%p) [%s]"), 
+		*GetName(), this, *SubObject->GetName(), SubObject, *StaticEnum<ELifetimeCondition>()->GetValueAsString(NetCondition));
+
+	// Warn if the subobject was registered with a different net condition.
+	ensureMsgf(Result != UE::Net::FSubObjectRegistry::EResult::NetConditionConflict, TEXT("%s(0x%p) Registered subobject %s (0x%p) again with a different net condition. Active [%s] New [%s]. New condition will be ignored"),
+		*GetName(), this, *SubObject->GetName(), SubObject, *StaticEnum<ELifetimeCondition>()->GetValueAsString(ReplicatedSubObjects.GetNetCondition(SubObject)), *StaticEnum<ELifetimeCondition>()->GetValueAsString(NetCondition));
+}
+
+void AActor::RemoveReplicatedSubObject(UObject* SubObject)
+{
+	check(SubObject);
+	bool bWasRemoved = ReplicatedSubObjects.RemoveSubObject(SubObject);
+
+	UE_CLOG(bWasRemoved, LogNetSubObject, Verbose, TEXT("%s (0x%p) removed replicated subobject %s (0x%p)"), *GetName(), this, *SubObject->GetName(), SubObject);
+}
+
+void AActor::AddActorComponentReplicatedSubObject(UActorComponent* OwnerComponent, UObject* SubObject, ELifetimeCondition NetCondition)
+{
+	check(IsValid(OwnerComponent));
+	check(IsValid(SubObject));
+
+	FReplicatedComponentInfo* ComponentInfo = ReplicatedComponentsInfo.FindByKey(OwnerComponent);
+	if (ComponentInfo)
+	{
+		// Add the subobject to the component's list
+		UE::Net::FSubObjectRegistry::EResult Result = ComponentInfo->SubObjects.AddSubObjectUnique(SubObject, NetCondition);
+
+		UE_CLOG(Result == UE::Net::FSubObjectRegistry::EResult::NewEntry, LogNetSubObject, Verbose, TEXT("%s::%s (0x%p) added replicated subobject %s (0x%p) [%s]"),
+			*GetName(), *OwnerComponent->GetName(), OwnerComponent, *SubObject->GetName(), SubObject, *StaticEnum<ELifetimeCondition>()->GetValueAsString(NetCondition));
+
+		// Warn if the subobject was registered with a different net condition.
+		ensureMsgf(Result != UE::Net::FSubObjectRegistry::EResult::NetConditionConflict, TEXT("%s::%s (0x%p) Registered subobject %s (0x%p) again with a different net condition. Active [%s] New [%s]. New condition will be ignored"),
+			*GetName(), *OwnerComponent->GetName(), OwnerComponent, *SubObject->GetName(), SubObject, *StaticEnum<ELifetimeCondition>()->GetValueAsString(ComponentInfo->SubObjects.GetNetCondition(SubObject)), *StaticEnum<ELifetimeCondition>()->GetValueAsString(NetCondition));
+	}
+	else
+	{
+		const int32 Index = ReplicatedComponentsInfo.Emplace(FReplicatedComponentInfo{OwnerComponent});
+		ReplicatedComponentsInfo[Index].SubObjects.AddSubObjectUnique(SubObject, NetCondition);
+
+		UE_LOG(LogNetSubObject, Verbose, TEXT("%s::%s (0x%p) added replicated subobject %s (0x%p) [%s]"),
+			*GetName(), *OwnerComponent->GetName(), OwnerComponent, *SubObject->GetName(), SubObject, *StaticEnum<ELifetimeCondition>()->GetValueAsString(NetCondition));
+	}
+}
+
+void AActor::RemoveActorComponentReplicatedSubObject(UActorComponent* OwnerComponent, UObject* SubObject)
+{
+	check(OwnerComponent);
+	check(SubObject);
+
+	if (FReplicatedComponentInfo* ComponentInfo = ReplicatedComponentsInfo.FindByKey(OwnerComponent))
+	{
+		bool bWasRemoved = ComponentInfo->SubObjects.RemoveSubObject(SubObject);
+		
+		UE_CLOG(bWasRemoved, LogNetSubObject, Verbose, TEXT("%s::%s (0x%p) removed replicated subobject %s (0x%p)"), *GetName(), *OwnerComponent->GetName(), OwnerComponent, *SubObject->GetName(), SubObject);
+	}
+}
+
 void AActor::GetSubobjectsWithStableNamesForNetworking(TArray<UObject*> &ObjList)
 {	
 	// For experimenting with replicating ALL stably named components initially
diff --git a/Engine/Source/Runtime/Engine/Private/Components/ActorComponent.cpp b/Engine/Source/Runtime/Engine/Private/Components/ActorComponent.cpp
--- a/Engine/Source/Runtime/Engine/Private/Components/ActorComponent.cpp
+++ b/Engine/Source/Runtime/Engine/Private/Components/ActorComponent.cpp
@@ -2079,7 +2079,23 @@ void UActorComponent::SetIsReplicated(bool bShouldReplicate)
 	}
 }
 
-bool UActorComponent::ReplicateSubobjects(class UActorChannel *Channel, class FOutBunch *Bunch, FReplicationFlags *RepFlags)
+void UActorComponent::AddReplicatedSubObject(UObject* SubObject, ELifetimeCondition NetCondition)
+{
+	if (AActor* MyOwner=GetOwner())
+	{
+		MyOwner->AddActorComponentReplicatedSubObject(this, SubObject, NetCondition);
+	}
+}
+
+void UActorComponent::RemoveReplicatedSubObject(UObject* SubObject)
+{
+	if (AActor* MyOwner=GetOwner())
+	{
+		MyOwner->RemoveActorComponentReplicatedSubObject(this, SubObject);
+	}
+}
+
+bool UActorComponent::ReplicateSubobjects(class UActorChannel* Channel, class FOutBunch* Bunch, FReplicationFlags* RepFlags)
 {
 	return false;
 }
diff --git a/Engine/Source/Runtime/Engine/Private/DataChannel.cpp b/Engine/Source/Runtime/Engine/Private/DataChannel.cpp
--- a/Engine/Source/Runtime/Engine/Private/DataChannel.cpp
+++ b/Engine/Source/Runtime/Engine/Private/DataChannel.cpp
@@ -23,11 +23,13 @@
 #include "ProfilingDebugging/CsvProfiler.h"
 #include "Net/NetworkGranularMemoryLogging.h"
 #include "Net/Core/Trace/NetTrace.h"
+#include "Net/Core/Misc/NetSubObjectRegistry.h"
 #include "Misc/NetworkVersion.h"
 #include "Net/Core/PushModel/PushModel.h"
 #include "HAL/LowLevelMemStats.h"
 
 DEFINE_LOG_CATEGORY(LogNet);
+DEFINE_LOG_CATEGORY(LogNetSubObject);
 DEFINE_LOG_CATEGORY(LogRep);
 DEFINE_LOG_CATEGORY(LogNetPlayerMovement);
 DEFINE_LOG_CATEGORY(LogNetTraffic);
@@ -122,6 +124,25 @@ static const bool IsBunchTooLarge(UNetConnection* Connection, T* Bunch)
 	return !Connection->IsInternalAck() && Bunch != nullptr && Bunch->GetNumBytes() > NetMaxConstructedPartialBunchSizeBytes;
 }
 
+/** Helper class to restrict access to the subobject list of actors and actor components */
+class FSubObjectGetter final
+{
+public:
+	FSubObjectGetter() = delete;
+	~FSubObjectGetter() = delete;
+
+	static const TArray<UE::Net::FSubObjectRegistry::FEntry>& GetSubObjects(AActor* InActor)
+	{
+		return InActor->ReplicatedSubObjects.GetRegistryList();
+	}
+
+	static const TArray<UE::Net::FSubObjectRegistry::FEntry>* GetSubObjectsOfActorCompoment(AActor* InActor, UActorComponent* InActorComp)
+	{
+		AActor::FReplicatedComponentInfo* ComponentInfo = InActor->ReplicatedComponentsInfo.FindByKey(InActorComp);
+		return ComponentInfo ? &(ComponentInfo->SubObjects.GetRegistryList()) : nullptr;
+	}
+};
+
 /*-----------------------------------------------------------------------------
 	UChannel implementation.
 -----------------------------------------------------------------------------*/
@@ -3248,7 +3269,6 @@ int64 UActorChannel::ReplicateActor()
 	RepFlags.bNetSimulated	= (Actor->GetRemoteRole() == ROLE_SimulatedProxy);
 	RepFlags.bRepPhysics	= Actor->GetReplicatedMovement().bRepPhysics;
 	RepFlags.bReplay		= bReplay;
-	//RepFlags.bNetInitial	= RepFlags.bNetInitial;
 	RepFlags.bForceInitialDirty = Connection->IsForceInitialDirty();
 
 	UE_LOG(LogNetTraffic, Log, TEXT("Replicate %s, bNetInitial: %d, bNetOwner: %d"), *Actor->GetName(), RepFlags.bNetInitial, RepFlags.bNetOwner);
@@ -3275,8 +3295,7 @@ int64 UActorChannel::ReplicateActor()
 			}
 		}
 
-		// The SubObjects
-		bWroteSomethingImportant |= Actor->ReplicateSubobjects(this, &Bunch, &RepFlags);
+		bWroteSomethingImportant |= DoSubObjectReplication(Bunch, RepFlags);
 
 		if (Connection->ResendAllDataState != EResendAllDataState::None)
 		{
@@ -3398,6 +3417,77 @@ int64 UActorChannel::ReplicateActor()
 	return NumBitsWrote;
 }
 
+bool UActorChannel::DoSubObjectReplication(FOutBunch& Bunch, const FReplicationFlags& RepFlags)
+{
+	bool bWroteSomethingImportant = false;
+
+	if (Actor->IsUsingRegisteredSubObjectList())
+	{
+		// If the actor replicates it's subobjects and those of it's components via the list.
+		bWroteSomethingImportant |= ReplicateRegisteredSubObjects(Bunch, RepFlags);
+	}
+	else
+	{
+		// Replicate the subobjects using the virtual method.
+		bWroteSomethingImportant |= Actor->ReplicateSubobjects(this, &Bunch, const_cast<FReplicationFlags*>(&RepFlags));
+	}
+
+	return bWroteSomethingImportant;
+}
+
+bool UActorChannel::ReplicateRegisteredSubObjects(FOutBunch& Bunch, const FReplicationFlags& RepFlags)
+{
+	check(Actor->IsUsingRegisteredSubObjectList());
+
+	const TStaticBitArray<COND_Max> ConditionMap = FSendingRepState::BuildConditionMapFromRepFlags(RepFlags);
+
+	bool bWroteSomethingImportant = false;
+
+	// Start with the Actor's subobjects
+	{
+		const TArray<UE::Net::FSubObjectRegistry::FEntry>& ActorSubObjects = FSubObjectGetter::GetSubObjects(Actor);
+		for (const UE::Net::FSubObjectRegistry::FEntry& SubObjectInfo : ActorSubObjects)
+		{
+			if (ConditionMap[SubObjectInfo.NetCondition])
+			{
+				bWroteSomethingImportant |= WriteSubObjectInBunch(SubObjectInfo.SubObject, Bunch, RepFlags);
+			}
+		}
+	}
+
+	// Now the replicated actor components
+	for (UActorComponent* ReplicatedComponent : Actor->GetReplicatedComponents())
+	{
+		check(ReplicatedComponent);
+
+		if (ReplicatedComponent->IsUsingRegisteredSubObjectList())
+		{
+			if (const TArray<UE::Net::FSubObjectRegistry::FEntry>* ComponentSubObjects = FSubObjectGetter::GetSubObjectsOfActorCompoment(Actor, ReplicatedComponent))
+			{
+				for (const UE::Net::FSubObjectRegistry::FEntry& SubObjectInfo : *ComponentSubObjects)
+				{
+					checkf(IsValid(SubObjectInfo.SubObject), TEXT("Found invalid subobject (%s) registered in %s"), *GetNameSafe(SubObjectInfo.SubObject), *GetNameSafe(ReplicatedComponent));
+
+					if (ConditionMap[SubObjectInfo.NetCondition])
+					{
+						bWroteSomethingImportant |= WriteSubObjectInBunch(SubObjectInfo.SubObject, Bunch, RepFlags);
+					}
+				}
+			}
+		}
+		// If a component is not using the registered list, collect the subobjects to replicate via the virtual function.
+		else
+		{
+			// Replicate the subobjects using the virtual method
+			ReplicatedComponent->ReplicateSubobjects(this, &Bunch, const_cast<FReplicationFlags*>(&RepFlags));
+		}
+
+		// Finally replicate the component itself after it's subobjects. They need to be created before the component on the receiving end.
+		bWroteSomethingImportant |= WriteSubObjectInBunch(ReplicatedComponent, Bunch, RepFlags);
+	}
+
+	return bWroteSomethingImportant;
+}
 
 FString UActorChannel::Describe()
 {
@@ -4381,16 +4471,28 @@ void UActorChannel::AddedToChannelPool()
 	PendingObjKeys.Empty();
 }
 
-bool UActorChannel::ReplicateSubobject(UObject *Obj, FOutBunch &Bunch, const FReplicationFlags &RepFlags)
+bool UActorChannel::ReplicateSubobject(UObject* Obj, FOutBunch& Bunch, const FReplicationFlags& RepFlags)
 {
-	SCOPE_CYCLE_UOBJECT(ActorChannelRepSubObj, Obj);
+	if (!IsValid(Obj))
+	{
+		return false;
+	}
 
+	const bool bWroteSomethingImportant = WriteSubObjectInBunch(Obj, Bunch, RepFlags);
+
+	return bWroteSomethingImportant;
+}
+
+bool UActorChannel::WriteSubObjectInBunch(UObject * Obj, FOutBunch & Bunch, const FReplicationFlags & RepFlags)
+{
 	if (!IsValid(Obj))
 	{
 		return false;
 	}
 
-	if (RepFlags.bUseCustomSubobjectReplication)
+    SCOPE_CYCLE_UOBJECT(ActorChannelRepSubObj, Obj);
+
+    if (RepFlags.bUseCustomSubobjectReplication)
 	{
 		return ReplicateSubobjectCustom(Obj, Bunch, RepFlags);
 	}
diff --git a/Engine/Source/Runtime/Engine/Public/EngineLogs.h b/Engine/Source/Runtime/Engine/Public/EngineLogs.h
--- a/Engine/Source/Runtime/Engine/Public/EngineLogs.h
+++ b/Engine/Source/Runtime/Engine/Public/EngineLogs.h
@@ -19,6 +19,7 @@ ENGINE_API DECLARE_LOG_CATEGORY_EXTERN(LogLevel, Log, All);
 ENGINE_API DECLARE_LOG_CATEGORY_EXTERN(LogSkeletalMesh, Log, All);
 ENGINE_API DECLARE_LOG_CATEGORY_EXTERN(LogStaticMesh, Log, All);
 ENGINE_API DECLARE_LOG_CATEGORY_EXTERN(LogNet, Log, All);
+ENGINE_API DECLARE_LOG_CATEGORY_EXTERN(LogNetSubObject, Log, All);
 ENGINE_API DECLARE_LOG_CATEGORY_EXTERN(LogRep, Log, All);
 ENGINE_API DECLARE_LOG_CATEGORY_EXTERN(LogNetPlayerMovement, Warning, All);
 ENGINE_API DECLARE_LOG_CATEGORY_EXTERN(LogNetTraffic, Warning, All);
diff --git a/Engine/Source/Runtime/Net/Core/Private/Net/Core/Misc/NetSubObjectRegistry.cpp b/Engine/Source/Runtime/Net/Core/Private/Net/Core/Misc/NetSubObjectRegistry.cpp
new file mode 100644
--- /dev/null
+++ b/Engine/Source/Runtime/Net/Core/Private/Net/Core/Misc/NetSubObjectRegistry.cpp
@@ -0,0 +1,43 @@
+// Copyright Epic Games, Inc. All Rights Reserved.
+
+#include "Net/Core/Misc/NetSubObjectRegistry.h"
+
+namespace UE {
+namespace Net {
+
+// ----------------------------------------------------------------
+//	FSubObjectRegistry
+// ----------------------------------------------------------------
+FSubObjectRegistry::EResult FSubObjectRegistry::AddSubObjectUnique(UObject* InSubObject, ELifetimeCondition InNetCondition)
+{
+	check(InSubObject);
+
+	const FEntry SubObjectEntry{ InSubObject, InNetCondition };
+
+	const int32 CurrentIndex = Registry.Find(SubObjectEntry);
+	if (CurrentIndex == INDEX_NONE)
+	{
+		Registry.Add(SubObjectEntry);
+		return EResult::NewEntry;
+	}
+	else
+	{
+		// Check if the net condition is different.
+		const FEntry& ActiveEntry = Registry[CurrentIndex];
+		return (ActiveEntry.NetCondition == InNetCondition) ? EResult::AlreadyRegistered : EResult::NetConditionConflict;
+	}
+}
+
+ELifetimeCondition FSubObjectRegistry::GetNetCondition(UObject* InSubObject) const
+{
+	const int32 CurrentIndex = Registry.Find(FEntry{ InSubObject });
+	return CurrentIndex == INDEX_NONE ? COND_Max : Registry[CurrentIndex].NetCondition;
+}
+
+bool FSubObjectRegistry::RemoveSubObject(UObject* InSubObject)
+{
+	return Registry.RemoveSingle(FEntry{ InSubObject }) == 1;
+}
+
+} //namespace Net
+} //namespace UE
\ No newline at end of file
diff --git a/Engine/Source/Runtime/Net/Core/Public/Net/Core/Misc/NetSubObjectRegistry.h b/Engine/Source/Runtime/Net/Core/Public/Net/Core/Misc/NetSubObjectRegistry.h
new file mode 100644
--- /dev/null
+++ b/Engine/Source/Runtime/Net/Core/Public/Net/Core/Misc/NetSubObjectRegistry.h
@@ -0,0 +1,58 @@
+// Copyright Epic Games, Inc. All Rights Reserved.
+
+#pragma once
+
+#include "CoreMinimal.h"
+#include "UObject/CoreNetTypes.h"
+
+namespace UE {
+namespace Net {
+
+/**
+* Stores the SubObjects to replicate and the network condition dictating to which connection they can replicate to.
+*/
+struct NETCORE_API FSubObjectRegistry
+{
+public:
+
+	enum class EResult : uint8
+	{
+		/** SubObject is new in the registry */
+		NewEntry,
+		/** SubObject already was in the registry */
+		AlreadyRegistered,
+		/** SubObject was already registered but the NetCondition passed is different from the entry */
+		NetConditionConflict,
+	};
+
+	/** Adds a subobject and returns if it's a new entry or existing entry */
+	FSubObjectRegistry::EResult AddSubObjectUnique(UObject* InSubObject, ELifetimeCondition InNetCondition);
+
+	/** Remove the subobject from the replicated list. Returns true if the subobject had been registered. */
+	bool RemoveSubObject(UObject* InSubObject);
+
+	/** Find the NetCondition of a SubObject. Returns COND_MAX if not registered */
+	ELifetimeCondition GetNetCondition(UObject* SubObject) const;
+
+	struct FEntry
+	{
+		/** Raw pointer since users are obligated to call RemoveReplicatedSubobject before destroying the subobject otherwise it will cause a crash.  */
+		UObject* SubObject = nullptr;
+
+		/** The network condition that chooses which connection this subobject can be replicated to. Default is none which means all connections receive it. */
+		ELifetimeCondition NetCondition = COND_None;
+
+		bool operator==(const FEntry& rhs) const { return SubObject == rhs.SubObject; }
+	};
+
+	/** Returns the list of registered subobjects */
+	const TArray<FSubObjectRegistry::FEntry>& GetRegistryList() const { return Registry; }
+
+private:
+
+	TArray<FEntry> Registry;
+};
+
+} // namespace Net
+} // namespace UE
+
