Take into account whether primitive is set to always write velocity when building GPU Scene primitive flags.
Primitives that use world position offset use this to ensure that they write velocity during base or depth pass.
#rnx

diff --git a/Engine/Source/Runtime/Engine/Private/PrimitiveUniformShaderParameters.cpp b/Engine/Source/Runtime/Engine/Private/PrimitiveUniformShaderParameters.cpp
--- a/Engine/Source/Runtime/Engine/Private/PrimitiveUniformShaderParameters.cpp
+++ b/Engine/Source/Runtime/Engine/Private/PrimitiveUniformShaderParameters.cpp
@@ -152,7 +152,7 @@ FPrimitiveSceneShaderData::FPrimitiveSceneShaderData(const FPrimitiveSceneProxy*
 		.HasCapsuleRepresentation(Proxy->HasDynamicIndirectShadowCasterRepresentation())
 		.ReceivesDecals(Proxy->ReceivesDecals())
 		.DrawsVelocity(Proxy->DrawsVelocity())
-		.OutputVelocity(bOutputVelocity)
+		.OutputVelocity(bOutputVelocity || Proxy->AlwaysHasVelocity())
 		.CastContactShadow(Proxy->CastsContactShadow())
 		.CastShadow(Proxy->CastsDynamicShadow())
 		.CastHiddenShadow(Proxy->CastsHiddenShadow())
