Fix HISM occlusion volume not using translated space.
#rnx

diff --git a/Engine/Source/Runtime/Engine/Private/HierarchicalInstancedStaticMesh.cpp b/Engine/Source/Runtime/Engine/Private/HierarchicalInstancedStaticMesh.cpp
--- a/Engine/Source/Runtime/Engine/Private/HierarchicalInstancedStaticMesh.cpp
+++ b/Engine/Source/Runtime/Engine/Private/HierarchicalInstancedStaticMesh.cpp
@@ -939,7 +939,7 @@ public:
 			//int32 NumPerNode = (1 + ClusterTree[0].LastInstance - ClusterTree[0].FirstInstance) / NumNodes;
 			//UE_LOG(LogTemp, Display, TEXT("Occlusion level %d   %d inst / node"), NumNodes, NumPerNode);
 			OcclusionBounds.Reserve(NumNodes);
-			FMatrix XForm = InComponent->GetComponentTransform().ToMatrixWithScale();
+			FMatrix XForm = InComponent->GetRenderMatrix();
 			for (int32 Index = FirstOcclusionNode; Index <= LastOcclusionNode; Index++)
 			{
 				OcclusionBounds.Add(FBoxSphereBounds(FBox(ClusterTree[Index].BoundMin, ClusterTree[Index].BoundMax).TransformBy(XForm)));
