Disabled duplicated actor folder fixup when running commandlets as FActorFolders singleton is not initialized.

diff --git a/Engine/Source/Runtime/Engine/Private/Level.cpp b/Engine/Source/Runtime/Engine/Private/Level.cpp
--- a/Engine/Source/Runtime/Engine/Private/Level.cpp
+++ b/Engine/Source/Runtime/Engine/Private/Level.cpp
@@ -2352,16 +2352,19 @@ void ULevel::FixupActorFolders()
 		return true;
 	}, /*bSkipDeleted*/ true);
 
-	// Rename duplicates to a new valid/unique name
-	for (UActorFolder* ActorFolder : DuplicateFolders)
+	if (!IsRunningCommandlet())
 	{
-		FFolder Folder = ActorFolder->GetFolder();
-		const FFolder NewPath = FActorFolders::Get().GetFolderName(*GetWorld(), Folder.GetParent(), Folder.GetLeafName());
-		ActorFolder->SetLabel(NewPath.GetLeafName().ToString());
-		bool bIsAlreadyInSet = false;
-		FolderPaths.Add(ActorFolder->GetPath(), &bIsAlreadyInSet);
-		check(!bIsAlreadyInSet);
-		UE_LOG(LogLevel, Warning, TEXT("Found duplicate actor folder %s, renamed to %s."), *Folder.GetPath().ToString(), *NewPath.GetPath().ToString());
+		// Rename duplicates to a new valid/unique name
+		for (UActorFolder* ActorFolder : DuplicateFolders)
+		{
+			FFolder Folder = ActorFolder->GetFolder();
+			const FFolder NewPath = FActorFolders::Get().GetFolderName(*GetWorld(), Folder.GetParent(), Folder.GetLeafName());
+			ActorFolder->SetLabel(NewPath.GetLeafName().ToString());
+			bool bIsAlreadyInSet = false;
+			FolderPaths.Add(ActorFolder->GetPath(), &bIsAlreadyInSet);
+			check(!bIsAlreadyInSet);
+			UE_LOG(LogLevel, Warning, TEXT("Found duplicate actor folder %s, renamed to %s."), *Folder.GetPath().ToString(), *NewPath.GetPath().ToString());
+		}
 	}
 	
 	// Fixup actors actor folder
