Track min/max bounds on values in new HLSL translator.  Allows converting doubles->floats automatically after min(max()) chain, possibly other uses

diff --git a/Engine/Source/Runtime/Engine/Private/HLSLTree/HLSLTree.cpp b/Engine/Source/Runtime/Engine/Private/HLSLTree/HLSLTree.cpp
--- a/Engine/Source/Runtime/Engine/Private/HLSLTree/HLSLTree.cpp
+++ b/Engine/Source/Runtime/Engine/Private/HLSLTree/HLSLTree.cpp
@@ -397,12 +397,12 @@ void FExpressionLocalPHI::EmitValueShader(FEmitContext& Context, FEmitScope& Sco
 	if (!EmitExpression)
 	{
 		const int32 LocalPHIIndex = Context.NumExpressionLocalPHIs++;
-		const FPreparedType& LocalType = GetPreparedType();
+		const Shader::FType LocalType = GetType();
 
 		// This is the first time we've emitted shader code for this PHI
 		// Create an expression and add it to the map first, so if this is called recursively this path will only be taken the first time
 		EmitExpression = OutResult.Code = Context.EmitInlineExpression(Scope,
-			LocalType.GetType(),
+			LocalType,
 			TEXT("LocalPHI%"), LocalPHIIndex);
 		Context.EmitLocalPHIMap.Add(this, EmitExpression);
 
@@ -556,9 +556,15 @@ FRequestedType::FRequestedType(ERequestedType InType)
 	RequestedComponents.Init(true, NumComponents);
 }
 
-FRequestedType::FRequestedType(int32 NumComponents)
+FRequestedType::FRequestedType(Shader::EValueType InType)
 {
-	RequestedComponents.Init(true, NumComponents);
+	const Shader::FValueTypeDescription TypeDesc = Shader::GetValueTypeDescription(InType);
+	RequestedComponents.Init(true, TypeDesc.NumComponents);
+}
+
+FRequestedType::FRequestedType(const Shader::FType& InType)
+{
+	RequestedComponents.Init(true, InType.GetNumComponents());
 }
 
 int32 FRequestedType::GetNumComponents() const
@@ -643,20 +649,26 @@ FPreparedComponent CombineComponents(const FPreparedComponent& Lhs, const FPrepa
 	}
 	else
 	{
-		const EExpressionEvaluation Evaluation = CombineEvaluations(Lhs.Evaluation, Rhs.Evaluation);
-		FEmitScope* LoopScope = nullptr;
-		if (IsLoopEvaluation(Evaluation))
+		FPreparedComponent Result;
+		Result.Evaluation = CombineEvaluations(Lhs.Evaluation, Rhs.Evaluation);
+		if (IsLoopEvaluation(Result.Evaluation))
 		{
-			LoopScope = FEmitScope::FindSharedParent(Lhs.LoopScope, Rhs.LoopScope);
+			Result.LoopScope = FEmitScope::FindSharedParent(Lhs.LoopScope, Rhs.LoopScope);
 		}
-		FExpression* ForwardValue = nullptr;
-		int32 ForwardComponentIndex = INDEX_NONE;
 		if (Lhs.ForwardValue == Rhs.ForwardValue && Lhs.ForwardComponentIndex == Rhs.ForwardComponentIndex)
 		{
-			ForwardValue = Lhs.ForwardValue;
-			ForwardComponentIndex = Lhs.ForwardComponentIndex;
+			Result.ForwardValue = Lhs.ForwardValue;
+			Result.ForwardComponentIndex = Lhs.ForwardComponentIndex;
+		}
+		if (Lhs.Bounds.Min == Rhs.Bounds.Min)
+		{
+			Result.Bounds.Min = Lhs.Bounds.Min;
 		}
-		return FPreparedComponent(Evaluation, LoopScope, ForwardValue, ForwardComponentIndex);
+		if (Lhs.Bounds.Max == Rhs.Bounds.Max)
+		{
+			Result.Bounds.Max = Lhs.Bounds.Max;
+		}
+		return Result;
 	}
 }
 
@@ -875,6 +887,42 @@ void FPreparedType::SetForwardValue(const FRequestedType& RequestedType, FExpres
 	}
 }
 
+void FPreparedType::UpdateBounds(const FRequestedType& RequestedType, Shader::FComponentBounds Bounds)
+{
+	for (int32 Index = 0; Index < PreparedComponents.Num(); ++Index)
+	{
+		FPreparedComponent& Component = PreparedComponents[Index];
+		if (RequestedType.IsComponentRequested(Index) && !Component.IsNone())
+		{
+			Component.Bounds.Min = Shader::MaxBound(Component.Bounds.Min, Bounds.Min);
+			Component.Bounds.Max = Shader::MinBound(Component.Bounds.Max, Bounds.Max);
+		}
+	}
+}
+
+Shader::FComponentBounds FPreparedType::GetBounds(const FRequestedType& RequestedType) const
+{
+	Shader::FComponentBounds Result(Shader::EComponentBound::DoubleMax, Shader::EComponentBound::NegDoubleMax);
+	for (int32 Index = 0; Index < PreparedComponents.Num(); ++Index)
+	{
+		const FPreparedComponent& Component = PreparedComponents[Index];
+		if (RequestedType.IsComponentRequested(Index) && !Component.IsNone())
+		{
+			// Constrain the component bounds we've stored by the range of the component's type
+			const Shader::EValueComponentType ComponentType = StructType ? StructType->ComponentTypes[Index] : ValueComponentType;
+			check(ComponentType != Shader::EValueComponentType::Void);
+
+			const Shader::FValueComponentTypeDescription ComponentTypeDesc = Shader::GetValueComponentTypeDescription(ComponentType);
+			const Shader::EComponentBound MinBound = Shader::MaxBound(Component.Bounds.Min, ComponentTypeDesc.Bounds.Min);
+			const Shader::EComponentBound MaxBound = Shader::MinBound(Component.Bounds.Max, ComponentTypeDesc.Bounds.Max);
+	
+			Result.Min = Shader::MinBound(Result.Min, MinBound);
+			Result.Max = Shader::MaxBound(Result.Max, MaxBound);
+		}
+	}
+	return Result;
+}
+
 void FPreparedType::SetField(const Shader::FStructField* Field, const FPreparedType& FieldType)
 {
 	for (int32 Index = 0; Index < Field->GetNumComponents(); ++Index)
@@ -934,31 +982,6 @@ FPreparedType MergePreparedTypes(const FPreparedType& Lhs, const FPreparedType&
 	return Result;
 }
 
-FRequestedValueType::FRequestedValueType(Shader::EValueType InType)
-{
-	const Shader::FValueTypeDescription TypeDesc = Shader::GetValueTypeDescription(InType);
-	ComponentType = TypeDesc.ComponentType;
-	RequestedType.RequestedComponents.Init(true, TypeDesc.NumComponents);
-}
-
-FRequestedValueType::FRequestedValueType(const FPreparedType& InType)
-{
-	ComponentType = InType.ValueComponentType;
-	if (ComponentType != Shader::EValueComponentType::Void)
-	{
-		RequestedType = InType.GetRequestedType();
-	}
-}
-
-Shader::EValueType FRequestedValueType::GetType() const
-{
-	if (ComponentType != Shader::EValueComponentType::Void)
-	{
-		return Shader::MakeValueType(ComponentType, GetNumComponents());
-	}
-	return Shader::EValueType::Void;
-}
-
 bool FPrepareValueResult::TryMergePreparedType(FEmitContext& Context, const Shader::FStructType* StructType, Shader::EValueComponentType ComponentType)
 {
 	// If we previously had a forwarded value set, reset that and start over
@@ -1020,7 +1043,17 @@ bool FPrepareValueResult::SetType(FEmitContext& Context, const FRequestedType& R
 
 bool FPrepareValueResult::SetType(FEmitContext& Context, const FRequestedType& RequestedType, const FPreparedType& Type)
 {
-	if (TryMergePreparedType(Context, Type.StructType, Type.ValueComponentType))
+	Shader::EValueComponentType ComponentType = Type.ValueComponentType;
+	if (ComponentType == Shader::EValueComponentType::Double)
+	{
+		const Shader::FComponentBounds Bounds = Type.GetBounds(RequestedType);
+		if (Shader::IsWithinBounds(Bounds, Shader::GetValueComponentTypeDescription(Shader::EValueComponentType::Float).Bounds))
+		{
+			ComponentType = Shader::EValueComponentType::Float;
+		}
+	}
+
+	if (TryMergePreparedType(Context, Type.StructType, ComponentType))
 	{
 		const int32 NumComponents = RequestedType.GetNumComponents();
 		for (int32 Index = 0; Index < NumComponents; ++Index)
@@ -1088,21 +1121,24 @@ FExpression* FExpression::GetForwardValue(const FRequestedType& RequestedType)
 			if (RequestedType.RequestedComponents[Index])
 			{
 				const FPreparedComponent Component = PreparedType.GetComponent(BaseComponentIndex + Index);
-				if (!Component.ForwardValue)
+				if (!Component.IsNone())
 				{
-					return nullptr;
-				}
+					if (!Component.ForwardValue)
+					{
+						return nullptr;
+					}
 
-				if (StartComponentIndex == INDEX_NONE)
-				{
-					StartComponentIndex = Index;
-					ForwardValue = Component.ForwardValue;
-					ForwardComponentIndex = Component.ForwardComponentIndex;
-				}
-				else if (ForwardValue != Component.ForwardValue ||
-					ForwardComponentIndex + Index - StartComponentIndex != Component.ForwardComponentIndex)
-				{
-					return nullptr;
+					if (StartComponentIndex == INDEX_NONE)
+					{
+						StartComponentIndex = Index;
+						ForwardValue = Component.ForwardValue;
+						ForwardComponentIndex = Component.ForwardComponentIndex;
+					}
+					else if (ForwardValue != Component.ForwardValue ||
+						ForwardComponentIndex + Index - StartComponentIndex != Component.ForwardComponentIndex)
+					{
+						return nullptr;
+					}
 				}
 			}
 		}
@@ -1158,6 +1194,16 @@ FEmitShaderExpression* FExpression::GetValueShader(FEmitContext& Context, FEmitS
 		Value = Result.Code;
 	}
 
+	if (Value->Type.IsNumeric() && GetPreparedType().IsNumeric())
+	{
+		const Shader::FValueTypeDescription TypeDest = Shader::GetValueTypeDescription(Value->Type);
+		if (TypeDest.ComponentType == Shader::EValueComponentType::Double && GetPreparedType().ValueComponentType == Shader::EValueComponentType::Float)
+		{
+			const Shader::EValueType CastType = Shader::MakeValueType(Shader::EValueComponentType::Float, TypeDest.NumComponents);
+			Value = Context.EmitCast(Scope, Value, CastType);
+		}
+	}
+
 	return Value;
 }
 
@@ -1167,14 +1213,14 @@ FEmitShaderExpression* FExpression::GetValueShader(FEmitContext& Context, FEmitS
 	return Context.EmitCast(Scope, Value, ResultType);
 }
 
-FEmitShaderExpression* FExpression::GetValueShader(FEmitContext& Context, FEmitScope& Scope, const FRequestedValueType& ResultType)
+FEmitShaderExpression* FExpression::GetValueShader(FEmitContext& Context, FEmitScope& Scope, const Shader::FType& ResultType)
 {
-	return GetValueShader(Context, Scope, ResultType.RequestedType, ResultType.GetType());
+	return GetValueShader(Context, Scope, ResultType, ResultType);
 }
 
-FEmitShaderExpression* FExpression::GetValueShader(FEmitContext& Context, FEmitScope& Scope, const FPreparedType& ResultType)
+FEmitShaderExpression* FExpression::GetValueShader(FEmitContext& Context, FEmitScope& Scope, Shader::EValueType ResultType)
 {
-	return GetValueShader(Context, Scope, ResultType.GetRequestedType(), ResultType.GetType());
+	return GetValueShader(Context, Scope, ResultType, ResultType);
 }
 
 FEmitShaderExpression* FExpression::GetValueShader(FEmitContext& Context, FEmitScope& Scope)
@@ -1255,7 +1301,12 @@ Shader::FValue FExpression::GetValueConstant(FEmitContext& Context, FEmitScope&
 
 Shader::FValue FExpression::GetValueConstant(FEmitContext& Context, FEmitScope& Scope, const Shader::FType& ResultType)
 {
-	return GetValueConstant(Context, Scope, FRequestedType(ResultType.GetNumComponents()), ResultType);
+	return GetValueConstant(Context, Scope, FRequestedType(ResultType), ResultType);
+}
+
+Shader::FValue FExpression::GetValueConstant(FEmitContext& Context, FEmitScope& Scope, Shader::EValueType ResultType)
+{
+	return GetValueConstant(Context, Scope, FRequestedType(ResultType), ResultType);
 }
 
 bool FScope::HasParentScope(const FScope& InParentScope) const
diff --git a/Engine/Source/Runtime/Engine/Private/HLSLTree/HLSLTreeCommon.cpp b/Engine/Source/Runtime/Engine/Private/HLSLTree/HLSLTreeCommon.cpp
--- a/Engine/Source/Runtime/Engine/Private/HLSLTree/HLSLTreeCommon.cpp
+++ b/Engine/Source/Runtime/Engine/Private/HLSLTree/HLSLTreeCommon.cpp
@@ -783,7 +783,7 @@ bool FExpressionUnaryOp::PrepareValue(FEmitContext& Context, FEmitScope& Scope,
 	case EUnaryOp::Length:
 	case EUnaryOp::Normalize:
 		// Each component of result is influenced by all components of input
-		InputRequestedType = FRequestedType(4);
+		InputRequestedType = ERequestedType::Vector4;
 		break;
 	default:
 		InputRequestedType = RequestedType;
@@ -830,46 +830,44 @@ namespace Private
 {
 struct FUnaryOpTypes
 {
-	FUnaryOpTypes(const FRequestedValueType& InType) : InputType(InType), ResultType(InType.GetType()) {}
-	FUnaryOpTypes(const FRequestedValueType& InInputType, const Shader::FType& InResultType) : InputType(InInputType), ResultType(InResultType) {}
-
-	FRequestedValueType InputType;
-	Shader::FType ResultType;
+	Shader::EValueType InputType;
+	Shader::EValueType ResultType;
+	bool bIsLWC;
 };
-FUnaryOpTypes GetUnaryyOpTypes(const FRequestedType& RequestedType, EUnaryOp Op, const FPreparedType& InputPreparedType)
+FUnaryOpTypes GetUnaryyOpTypes(EUnaryOp Op, Shader::EValueType InputType)
 {
-	const Shader::EValueComponentType InputComponentType = InputPreparedType.ValueComponentType;
-	const Shader::EValueComponentType NonLWCInputComponentType = Shader::MakeNonLWCType(InputComponentType);
-	const FRequestedValueType InputType(InputComponentType, RequestedType);
-	const FRequestedValueType NonLWCInputType(NonLWCInputComponentType, RequestedType);
-	const int32 NumInputComponents = InputPreparedType.GetNumComponents();
-	const int32 NumRequestedComponents = RequestedType.GetNumComponents();
+	const Shader::FValueTypeDescription InputTypeDesc = Shader::GetValueTypeDescription(InputType);
 
+	FUnaryOpTypes Types;
+	Types.InputType = InputType;
+	Types.ResultType = InputType;
+	Types.bIsLWC = InputTypeDesc.ComponentType == Shader::EValueComponentType::Double;
 	switch (Op)
 	{
 	case EUnaryOp::Length:
-		return FUnaryOpTypes(InputPreparedType, Shader::MakeValueType(InputComponentType, 1));
+		Types.ResultType = Shader::MakeValueType(InputTypeDesc.ComponentType, 1);
+		break;
 	case EUnaryOp::Normalize:
-		return FUnaryOpTypes(InputPreparedType, Shader::MakeValueType(NonLWCInputComponentType, NumInputComponents));
 	case EUnaryOp::Frac:
 	case EUnaryOp::Rcp:
-		return FUnaryOpTypes(InputType, Shader::MakeValueType(NonLWCInputComponentType, NumRequestedComponents));
+		Types.ResultType = Shader::MakeNonLWCType(InputType);
+		break;
 	default:
-		return InputType;
+		break;
 	}
+	return Types;
 }
 } // namespace Private
 
 void FExpressionUnaryOp::EmitValueShader(FEmitContext& Context, FEmitScope& Scope, const FRequestedType& RequestedType, FEmitValueShaderResult& OutResult) const
 {
-	const Private::FUnaryOpTypes Types = Private::GetUnaryyOpTypes(RequestedType, Op, Input->GetPreparedType());
-	const bool bLWC = Types.InputType.ComponentType == Shader::EValueComponentType::Double;
+	const Private::FUnaryOpTypes Types = Private::GetUnaryyOpTypes(Op, Input->GetType());
 	FEmitShaderExpression* InputValue = Input->GetValueShader(Context, Scope, Types.InputType);
 
 	switch (Op)
 	{
 	case EUnaryOp::Neg:
-		if (bLWC)
+		if (Types.bIsLWC)
 		{
 			OutResult.Code = Context.EmitExpression(Scope, Types.ResultType, TEXT("LWCNegate(%)"), InputValue);
 		}
@@ -879,7 +877,7 @@ void FExpressionUnaryOp::EmitValueShader(FEmitContext& Context, FEmitScope& Scop
 		}
 		break;
 	case EUnaryOp::Rcp:
-		if (bLWC)
+		if (Types.bIsLWC)
 		{
 			OutResult.Code = Context.EmitExpression(Scope, Types.ResultType, TEXT("LWCRcp(%)"), InputValue);
 		}
@@ -889,7 +887,7 @@ void FExpressionUnaryOp::EmitValueShader(FEmitContext& Context, FEmitScope& Scop
 		}
 		break;
 	case EUnaryOp::Frac:
-		if (bLWC)
+		if (Types.bIsLWC)
 		{
 			OutResult.Code = Context.EmitExpression(Scope, Types.ResultType, TEXT("LWCFrac(%)"), InputValue);
 		}
@@ -899,7 +897,7 @@ void FExpressionUnaryOp::EmitValueShader(FEmitContext& Context, FEmitScope& Scop
 		}
 		break;
 	case EUnaryOp::Length:
-		if (bLWC)
+		if (Types.bIsLWC)
 		{
 			OutResult.Code = Context.EmitExpression(Scope, Types.ResultType, TEXT("LWCLength(%)"), InputValue);
 		}
@@ -909,7 +907,7 @@ void FExpressionUnaryOp::EmitValueShader(FEmitContext& Context, FEmitScope& Scop
 		}
 		break;
 	case EUnaryOp::Normalize:
-		if (bLWC)
+		if (Types.bIsLWC)
 		{
 			OutResult.Code = Context.EmitExpression(Scope, Types.ResultType, TEXT("LWCNormalize(%)"), InputValue);
 		}
@@ -926,11 +924,11 @@ void FExpressionUnaryOp::EmitValueShader(FEmitContext& Context, FEmitScope& Scop
 
 void FExpressionUnaryOp::EmitValuePreshader(FEmitContext& Context, FEmitScope& Scope, const FRequestedType& RequestedType, FEmitValuePreshaderResult& OutResult) const
 {
-	const Private::FUnaryOpTypes Types = Private::GetUnaryyOpTypes(RequestedType, Op, Input->GetPreparedType());
+	const Private::FUnaryOpTypes Types = Private::GetUnaryyOpTypes(Op, Input->GetType());
 	const FUnaryOpDescription OpDesc = GetUnaryOpDesription(Op);
 	check(OpDesc.PreshaderOpcode != Shader::EPreshaderOpcode::Nop);
 
-	Input->GetValuePreshader(Context, Scope, Types.InputType.RequestedType, OutResult.Preshader);
+	Input->GetValuePreshader(Context, Scope, Types.InputType, OutResult.Preshader);
 	OutResult.Preshader.WriteOpcode(OpDesc.PreshaderOpcode);
 	OutResult.Type = Types.ResultType;
 }
@@ -1043,19 +1041,19 @@ bool FExpressionBinaryOp::PrepareValue(FEmitContext& Context, FEmitScope& Scope,
 	{
 	case EBinaryOp::VecMulMatrix3:
 	case EBinaryOp::VecMulMatrix4:
-		LhsRequestedType = FRequestedType(3);
-		RhsRequestedType = FRequestedType(16);
+		LhsRequestedType = ERequestedType::Vector3;
+		RhsRequestedType = ERequestedType::Matrix4x4;
 		bIsMatrixOperation = true;
 		break;
 	case EBinaryOp::Matrix3MulVec:
 	case EBinaryOp::Matrix4MulVec:
-		LhsRequestedType = FRequestedType(16);
-		RhsRequestedType = FRequestedType(3);
+		LhsRequestedType = ERequestedType::Matrix4x4;
+		RhsRequestedType = ERequestedType::Vector3;
 		bIsMatrixOperation = true;
 		break;
 	case EBinaryOp::Dot:
-		LhsRequestedType = FRequestedType(4);
-		RhsRequestedType = FRequestedType(4);
+		LhsRequestedType = ERequestedType::Vector4;
+		RhsRequestedType = ERequestedType::Vector4;
 		break;
 	default:
 		LhsRequestedType = RequestedType;
@@ -1094,6 +1092,13 @@ bool FExpressionBinaryOp::PrepareValue(FEmitContext& Context, FEmitScope& Scope,
 		{
 			ResultType.ValueComponentType = Shader::EValueComponentType::Bool;
 		}
+		if (Op == EBinaryOp::Min || Op == EBinaryOp::Max)
+		{
+			const Shader::FComponentBounds LhsBounds = LhsType.GetBounds(RequestedType);
+			const Shader::FComponentBounds RhsBounds = RhsType.GetBounds(RequestedType);
+			const Shader::FComponentBounds Bounds = (Op == EBinaryOp::Min) ? Shader::MinBound(LhsBounds, RhsBounds) : Shader::MaxBound(LhsBounds, RhsBounds);
+			ResultType.UpdateBounds(RequestedType, Bounds);
+		}
 	}
 
 	if (OpDesc.PreshaderOpcode == Shader::EPreshaderOpcode::Nop)
@@ -1109,56 +1114,72 @@ namespace Private
 {
 struct FBinaryOpTypes
 {
-	FBinaryOpTypes(const FRequestedValueType& InType) : LhsType(InType), RhsType(LhsType), ResultType(InType.GetType()) {}
-	FBinaryOpTypes(const FRequestedValueType& InLhsType, const FRequestedValueType& InRhsType, const Shader::FType& InResultType) : LhsType(InLhsType), RhsType(InRhsType), ResultType(InResultType) {}
-
-	FRequestedValueType LhsType;
-	FRequestedValueType RhsType;
-	Shader::FType ResultType;
+	Shader::EValueType LhsType;
+	Shader::EValueType RhsType;
+	Shader::EValueType ResultType;
+	bool bIsLWC;
 };
-FBinaryOpTypes GetBinaryOpTypes(const FRequestedType& RequestedType, EBinaryOp Op, const FPreparedType& LhsType, const FPreparedType& RhsType)
+FBinaryOpTypes GetBinaryOpTypes(EBinaryOp Op, Shader::EValueType LhsType, Shader::EValueType RhsType)
 {
-	const Shader::EValueComponentType InputComponentType = Shader::CombineComponentTypes(LhsType.ValueComponentType, RhsType.ValueComponentType);
-	const Shader::EValueComponentType NonLWCInputComponentType = Shader::MakeNonLWCType(InputComponentType);
-	const FRequestedValueType InputType(InputComponentType, RequestedType);
-	const FRequestedValueType NonLWCInputType(NonLWCInputComponentType, RequestedType);
-	const int32 NumComponents = RequestedType.GetNumComponents();
+	const Shader::FValueTypeDescription LhsTypeDesc = Shader::GetValueTypeDescription(LhsType);
+	const Shader::FValueTypeDescription RhsTypeDesc = Shader::GetValueTypeDescription(RhsType);
 
+	Shader::EValueComponentType ComponentType = Shader::CombineComponentTypes(LhsTypeDesc.ComponentType, RhsTypeDesc.ComponentType);
+	int32 NumComponents = FMath::Max(LhsTypeDesc.NumComponents, RhsTypeDesc.NumComponents);
+	Shader::EValueType IntermediateType = Shader::MakeValueType(ComponentType, NumComponents);
+
+	FBinaryOpTypes Types;
+	Types.LhsType = IntermediateType;
+	Types.RhsType = IntermediateType;
+	Types.ResultType = IntermediateType;
+	Types.bIsLWC = (ComponentType == Shader::EValueComponentType::Double);
 	switch (Op)
 	{
 	case EBinaryOp::Less:
 	case EBinaryOp::Greater:
-		return FBinaryOpTypes(InputType, InputType, Shader::MakeValueType(Shader::EValueComponentType::Bool, NumComponents));
+		Types.ResultType = Shader::MakeValueType(Shader::EValueComponentType::Bool, NumComponents);
+		break;
 	case EBinaryOp::Fmod:
-		return FBinaryOpTypes(InputType, NonLWCInputType, NonLWCInputType.GetType());
+		Types.RhsType = Types.ResultType = Shader::MakeNonLWCType(IntermediateType);
+		break;
 	case EBinaryOp::Dot:
-		return FBinaryOpTypes(InputType, InputType, Shader::MakeValueType(InputComponentType, 1));
+		Types.ResultType = Shader::MakeValueType(ComponentType, 1);
+		break;
 	case EBinaryOp::VecMulMatrix3:
 		// No LWC for matrix3
-		return FBinaryOpTypes(Shader::EValueType::Float3, Shader::EValueType::Float4x4, Shader::EValueType::Float3);
+		Types.LhsType = Shader::EValueType::Float3;
+		Types.RhsType = Shader::EValueType::Float4x4;
+		Types.ResultType = Shader::EValueType::Float3;
+		break;
 	case EBinaryOp::VecMulMatrix4:
-		return FBinaryOpTypes(Shader::EValueType::Float3, Shader::EValueType::Float4x4, Shader::EValueType::Float3);
+		Types.LhsType = Shader::EValueType::Float3;
+		Types.RhsType = Shader::EValueType::Float4x4;
+		Types.ResultType = Shader::EValueType::Float3;
+		break;
 	case EBinaryOp::Matrix3MulVec:
 	case EBinaryOp::Matrix4MulVec:
 		// No LWC for transpose matrices
-		return FBinaryOpTypes(Shader::EValueType::Float4x4, Shader::EValueType::Float3, Shader::EValueType::Float3);
+		Types.LhsType = Shader::EValueType::Float4x4;
+		Types.RhsType = Shader::EValueType::Float3;
+		Types.ResultType = Shader::EValueType::Float3;
+		break;
 	default:
-		return InputType;
+		break;
 	}
+	return Types;
 }
 } // namespace Private
 
 void FExpressionBinaryOp::EmitValueShader(FEmitContext& Context, FEmitScope& Scope, const FRequestedType& RequestedType, FEmitValueShaderResult& OutResult) const
 {
-	const Private::FBinaryOpTypes Types = Private::GetBinaryOpTypes(RequestedType, Op, Lhs->GetPreparedType(), Rhs->GetPreparedType());
-	const bool bLWC = Types.LhsType.ComponentType == Shader::EValueComponentType::Double;
+	const Private::FBinaryOpTypes Types = Private::GetBinaryOpTypes(Op, Lhs->GetType(), Rhs->GetType());
 	FEmitShaderExpression* LhsValue = Lhs->GetValueShader(Context, Scope, Types.LhsType);
 	FEmitShaderExpression* RhsValue = Rhs->GetValueShader(Context, Scope, Types.RhsType);
 
 	switch (Op)
 	{
 	case EBinaryOp::Add:
-		if (bLWC)
+		if (Types.bIsLWC)
 		{
 			OutResult.Code = Context.EmitExpression(Scope, Types.ResultType, TEXT("LWCAdd(%, %)"), LhsValue, RhsValue);
 		}
@@ -1168,7 +1189,7 @@ void FExpressionBinaryOp::EmitValueShader(FEmitContext& Context, FEmitScope& Sco
 		}
 		break;
 	case EBinaryOp::Sub:
-		if (bLWC)
+		if (Types.bIsLWC)
 		{
 			OutResult.Code = Context.EmitExpression(Scope, Types.ResultType, TEXT("LWCSubtract(%, %)"), LhsValue, RhsValue);
 		}
@@ -1178,7 +1199,7 @@ void FExpressionBinaryOp::EmitValueShader(FEmitContext& Context, FEmitScope& Sco
 		}
 		break;
 	case EBinaryOp::Mul:
-		if (bLWC)
+		if (Types.bIsLWC)
 		{
 			OutResult.Code = Context.EmitExpression(Scope, Types.ResultType, TEXT("LWCMultiply(%, %)"), LhsValue, RhsValue);
 		}
@@ -1188,7 +1209,7 @@ void FExpressionBinaryOp::EmitValueShader(FEmitContext& Context, FEmitScope& Sco
 		}
 		break;
 	case EBinaryOp::Div:
-		if (bLWC)
+		if (Types.bIsLWC)
 		{
 			OutResult.Code = Context.EmitExpression(Scope, Types.ResultType, TEXT("LWCDivide(%, %)"), LhsValue, RhsValue);
 		}
@@ -1198,7 +1219,7 @@ void FExpressionBinaryOp::EmitValueShader(FEmitContext& Context, FEmitScope& Sco
 		}
 		break;
 	case EBinaryOp::Fmod:
-		if (bLWC)
+		if (Types.bIsLWC)
 		{
 			OutResult.Code = Context.EmitExpression(Scope, Types.ResultType, TEXT("LWCFmod(%, %)"), LhsValue, RhsValue);
 		}
@@ -1208,7 +1229,7 @@ void FExpressionBinaryOp::EmitValueShader(FEmitContext& Context, FEmitScope& Sco
 		}
 		break;
 	case EBinaryOp::Dot:
-		if (bLWC)
+		if (Types.bIsLWC)
 		{
 			OutResult.Code = Context.EmitExpression(Scope, Types.ResultType, TEXT("LWCDot(%, %)"), LhsValue, RhsValue);
 		}
@@ -1218,7 +1239,7 @@ void FExpressionBinaryOp::EmitValueShader(FEmitContext& Context, FEmitScope& Sco
 		}
 		break;
 	case EBinaryOp::Min:
-		if (bLWC)
+		if (Types.bIsLWC)
 		{
 			OutResult.Code = Context.EmitExpression(Scope, Types.ResultType, TEXT("LWCMin(%, %)"), LhsValue, RhsValue);
 		}
@@ -1228,7 +1249,7 @@ void FExpressionBinaryOp::EmitValueShader(FEmitContext& Context, FEmitScope& Sco
 		}
 		break;
 	case EBinaryOp::Max:
-		if (bLWC)
+		if (Types.bIsLWC)
 		{
 			OutResult.Code = Context.EmitExpression(Scope, Types.ResultType, TEXT("LWCMax(%, %)"), LhsValue, RhsValue);
 		}
@@ -1238,7 +1259,7 @@ void FExpressionBinaryOp::EmitValueShader(FEmitContext& Context, FEmitScope& Sco
 		}
 		break;
 	case EBinaryOp::Less:
-		if (bLWC)
+		if (Types.bIsLWC)
 		{
 			OutResult.Code = Context.EmitExpression(Scope, Types.ResultType, TEXT("LWCLess(%, %)"), LhsValue, RhsValue);
 		}
@@ -1248,7 +1269,7 @@ void FExpressionBinaryOp::EmitValueShader(FEmitContext& Context, FEmitScope& Sco
 		}
 		break;
 	case EBinaryOp::Greater:
-		if (bLWC)
+		if (Types.bIsLWC)
 		{
 			OutResult.Code = Context.EmitExpression(Scope, Types.ResultType, TEXT("LWCGreater(%, %)"), LhsValue, RhsValue);
 		}
@@ -1258,7 +1279,7 @@ void FExpressionBinaryOp::EmitValueShader(FEmitContext& Context, FEmitScope& Sco
 		}
 		break;
 	case EBinaryOp::VecMulMatrix3:
-		if (bLWC)
+		if (Types.bIsLWC)
 		{
 			OutResult.Code = Context.EmitExpression(Scope, Types.ResultType, TEXT("LWCMlutiplyVector(%, %)"), LhsValue, RhsValue);
 		}
@@ -1268,7 +1289,7 @@ void FExpressionBinaryOp::EmitValueShader(FEmitContext& Context, FEmitScope& Sco
 		}
 		break;
 	case EBinaryOp::VecMulMatrix4:
-		if (bLWC)
+		if (Types.bIsLWC)
 		{
 			OutResult.Code = Context.EmitExpression(Scope, Types.ResultType, TEXT("LWCMlutiply(%, %)"), LhsValue, RhsValue);
 		}
@@ -1291,12 +1312,12 @@ void FExpressionBinaryOp::EmitValueShader(FEmitContext& Context, FEmitScope& Sco
 
 void FExpressionBinaryOp::EmitValuePreshader(FEmitContext& Context, FEmitScope& Scope, const FRequestedType& RequestedType, FEmitValuePreshaderResult& OutResult) const
 {
-	const Private::FBinaryOpTypes Types = Private::GetBinaryOpTypes(RequestedType, Op, Lhs->GetPreparedType(), Rhs->GetPreparedType());
+	const Private::FBinaryOpTypes Types = Private::GetBinaryOpTypes(Op, Lhs->GetType(), Rhs->GetType());
 	const FBinaryOpDescription OpDesc = GetBinaryOpDesription(Op);
 	check(OpDesc.PreshaderOpcode != Shader::EPreshaderOpcode::Nop);
 
-	Lhs->GetValuePreshader(Context, Scope, Types.LhsType.RequestedType, OutResult.Preshader);
-	Rhs->GetValuePreshader(Context, Scope, Types.RhsType.RequestedType, OutResult.Preshader);
+	Lhs->GetValuePreshader(Context, Scope, Types.LhsType, OutResult.Preshader);
+	Rhs->GetValuePreshader(Context, Scope, Types.RhsType, OutResult.Preshader);
 
 	check(Context.PreshaderStackPosition > 0);
 	Context.PreshaderStackPosition--;
@@ -1452,48 +1473,40 @@ namespace Private
 {
 struct FAppendTypes
 {
-	Shader::FType ResultType;
-	Shader::EValueComponentType ResultComponentType;
-	FRequestedValueType LhsType;
-	FRequestedValueType RhsType;
+	Shader::EValueType ResultType;
+	Shader::EValueType LhsType;
+	Shader::EValueType RhsType;
+	bool bIsLWC;
 };
-FAppendTypes GetAppendTypes(const FRequestedType& RequestedType, const FPreparedType& LhsType, const FPreparedType& RhsType)
+FAppendTypes GetAppendTypes(Shader::EValueType LhsType, Shader::EValueType RhsType)
 {
-	const int32 NumResultComponents = RequestedType.GetNumComponents();
-	const int32 NumLhsComponents = FMath::Min<int32>(NumResultComponents, LhsType.GetNumComponents());
-	const int32 NumRhsComponents = FMath::Min<int32>(NumResultComponents - NumLhsComponents, RhsType.GetNumComponents());
-	const Shader::EValueComponentType ComponentType = Shader::CombineComponentTypes(LhsType.ValueComponentType, RhsType.ValueComponentType);
+	const Shader::FValueTypeDescription LhsTypeDesc = Shader::GetValueTypeDescription(LhsType);
+	const Shader::FValueTypeDescription RhsTypeDesc = Shader::GetValueTypeDescription(RhsType);
+	const Shader::EValueComponentType ComponentType = Shader::CombineComponentTypes(LhsTypeDesc.ComponentType, RhsTypeDesc.ComponentType);
+	const int32 NumComponents = FMath::Min(LhsTypeDesc.NumComponents + RhsTypeDesc.NumComponents, 4);
 
 	FAppendTypes Types;
-	Types.ResultComponentType = ComponentType;
-	Types.ResultType = Shader::MakeValueType(ComponentType, NumLhsComponents + NumRhsComponents);
-	Types.LhsType.ComponentType = ComponentType;
-	Types.RhsType.ComponentType = ComponentType;
-	for (int32 Index = 0; Index < NumLhsComponents; ++Index)
-	{
-		Types.LhsType.RequestedType.SetComponentRequest(Index, RequestedType.IsComponentRequested(Index));
-	}
-	for (int32 Index = 0; Index < NumRhsComponents; ++Index)
-	{
-		Types.RhsType.RequestedType.SetComponentRequest(Index, RequestedType.IsComponentRequested(NumLhsComponents + Index));
-	}
+	Types.ResultType = Shader::MakeValueType(ComponentType, NumComponents);
+	Types.LhsType = Shader::MakeValueType(ComponentType, LhsTypeDesc.NumComponents);
+	Types.RhsType = Shader::MakeValueType(ComponentType, FMath::Min<int32>(RhsTypeDesc.NumComponents, 4 - LhsTypeDesc.NumComponents));
+	Types.bIsLWC = ComponentType == Shader::EValueComponentType::Double;
 	return Types;
 }
 } // namespace Private
 
 void FExpressionAppend::EmitValueShader(FEmitContext& Context, FEmitScope& Scope, const FRequestedType& RequestedType, FEmitValueShaderResult& OutResult) const
 {
-	const Private::FAppendTypes Types = Private::GetAppendTypes(RequestedType, Lhs->GetPreparedType(), Rhs->GetPreparedType());
+	const Private::FAppendTypes Types = Private::GetAppendTypes(Lhs->GetType(), Rhs->GetType());
 	FEmitShaderExpression* LhsValue = Lhs->GetValueShader(Context, Scope, Types.LhsType);
 
-	if (Types.RhsType.IsVoid())
+	if (Types.RhsType == Shader::EValueType::Void)
 	{
 		OutResult.Code = LhsValue;
 	}
 	else
 	{
 		FEmitShaderExpression* RhsValue = Rhs->GetValueShader(Context, Scope, Types.RhsType);
-		if (Types.ResultComponentType == Shader::EValueComponentType::Double)
+		if (Types.bIsLWC)
 		{
 			OutResult.Code = Context.EmitInlineExpression(Scope, Types.ResultType, TEXT("MakeLWCVector(%, %)"),
 				LhsValue,
@@ -1502,7 +1515,7 @@ void FExpressionAppend::EmitValueShader(FEmitContext& Context, FEmitScope& Scope
 		else
 		{
 			OutResult.Code = Context.EmitInlineExpression(Scope, Types.ResultType, TEXT("%(%, %)"),
-				Types.ResultType.GetName(),
+				Shader::GetValueTypeDescription(Types.ResultType).Name,
 				LhsValue,
 				RhsValue);
 		}
@@ -1511,11 +1524,11 @@ void FExpressionAppend::EmitValueShader(FEmitContext& Context, FEmitScope& Scope
 
 void FExpressionAppend::EmitValuePreshader(FEmitContext& Context, FEmitScope& Scope, const FRequestedType& RequestedType, FEmitValuePreshaderResult& OutResult) const
 {
-	const Private::FAppendTypes Types = Private::GetAppendTypes(RequestedType, Lhs->GetPreparedType(), Rhs->GetPreparedType());
-	Lhs->GetValuePreshader(Context, Scope, Types.LhsType.RequestedType, OutResult.Preshader);
-	if (Types.RhsType.IsVoid())
+	const Private::FAppendTypes Types = Private::GetAppendTypes(Lhs->GetType(), Rhs->GetType());
+	Lhs->GetValuePreshader(Context, Scope, Types.LhsType, OutResult.Preshader);
+	if (Types.RhsType != Shader::EValueType::Void)
 	{
-		Rhs->GetValuePreshader(Context, Scope, Types.RhsType.RequestedType, OutResult.Preshader);
+		Rhs->GetValuePreshader(Context, Scope, Types.RhsType, OutResult.Preshader);
 
 		check(Context.PreshaderStackPosition > 0);
 		Context.PreshaderStackPosition--;
@@ -1624,7 +1637,7 @@ void FStatementIf::EmitShader(FEmitContext& Context, FEmitScope& Scope) const
 
 void FStatementIf::EmitPreshader(FEmitContext& Context, FEmitScope& Scope, const FRequestedType& RequestedType, TArrayView<const FEmitPreshaderScope> Scopes, Shader::FPreshaderData& OutPreshader) const
 {
-	ConditionExpression->GetValuePreshader(Context, Scope, FRequestedType(1), OutPreshader);
+	ConditionExpression->GetValuePreshader(Context, Scope, ERequestedType::Scalar, OutPreshader);
 
 	check(Context.PreshaderStackPosition > 0);
 	Context.PreshaderStackPosition--;
diff --git a/Engine/Source/Runtime/Engine/Private/HLSLTree/HLSLTreeEmit.cpp b/Engine/Source/Runtime/Engine/Private/HLSLTree/HLSLTreeEmit.cpp
--- a/Engine/Source/Runtime/Engine/Private/HLSLTree/HLSLTreeEmit.cpp
+++ b/Engine/Source/Runtime/Engine/Private/HLSLTree/HLSLTreeEmit.cpp
@@ -885,6 +885,7 @@ FEmitShaderExpression* FEmitContext::EmitCast(FEmitScope& Scope, FEmitShaderExpr
 
 	TStringBuilder<1024> FormattedCode;
 	Shader::FType IntermediateType = DestType;
+	bool bInline = true;
 
 	if (SourceTypeDesc.NumComponents > 0 && DestTypeDesc.NumComponents > 0)
 	{
@@ -904,6 +905,7 @@ FEmitShaderExpression* FEmitContext::EmitCast(FEmitScope& Scope, FEmitShaderExpr
 				//LWC->float
 				FormattedCode.Appendf(TEXT("LWCToFloat(%s)"), ShaderValue->Reference);
 				IntermediateType = Shader::MakeValueType(Shader::EValueComponentType::Float, SourceTypeDesc.NumComponents);
+				bInline = false; // LWCToFloat has non-zero cost
 			}
 		}
 		else
@@ -989,7 +991,15 @@ FEmitShaderExpression* FEmitContext::EmitCast(FEmitScope& Scope, FEmitShaderExpr
 	}
 
 	check(IntermediateType != ShaderValue->Type);
-	ShaderValue = EmitInlineExpressionWithDependency(Scope, ShaderValue, IntermediateType, FormattedCode.ToView());
+	if (bInline)
+	{
+		ShaderValue = EmitInlineExpressionWithDependency(Scope, ShaderValue, IntermediateType, FormattedCode.ToView());
+	}
+	else
+	{
+		ShaderValue = EmitExpressionWithDependency(Scope, ShaderValue, IntermediateType, FormattedCode.ToView());
+	}
+
 	if (ShaderValue->Type != DestType)
 	{
 		// May need to cast through multiple intermediate types to reach our destination type
diff --git a/Engine/Source/Runtime/Engine/Private/Shader/ShaderValue.cpp b/Engine/Source/Runtime/Engine/Private/Shader/ShaderValue.cpp
--- a/Engine/Source/Runtime/Engine/Private/Shader/ShaderValue.cpp
+++ b/Engine/Source/Runtime/Engine/Private/Shader/ShaderValue.cpp
@@ -375,29 +375,16 @@ bool FValue::AsBoolScalar() const
 	return false;
 }
 
-const TCHAR* GetComponentTypeName(EValueComponentType Type)
+FValueComponentTypeDescription GetValueComponentTypeDescription(EValueComponentType Type)
 {
 	switch (Type)
 	{
-	case EValueComponentType::Void: return TEXT("void");
-	case EValueComponentType::Float: return TEXT("float");
-	case EValueComponentType::Double: return TEXT("double");
-	case EValueComponentType::Int: return TEXT("int");
-	case EValueComponentType::Bool: return TEXT("bool");
-	default: checkNoEntry() return TEXT("void");
-	}
-}
-
-uint32 GetComponentTypeSizeInBytes(EValueComponentType Type)
-{
-	switch (Type)
-	{
-	case EValueComponentType::Void: return 0u;
-	case EValueComponentType::Float: return sizeof(float);
-	case EValueComponentType::Double: return sizeof(double);
-	case EValueComponentType::Int: return sizeof(int32);
-	case EValueComponentType::Bool: return 1u;
-	default: checkNoEntry(); return 0u;
+	case EValueComponentType::Void: return FValueComponentTypeDescription(TEXT("void"), 0u, EComponentBound::Zero, EComponentBound::Zero);
+	case EValueComponentType::Float: return FValueComponentTypeDescription(TEXT("float"), 4u, EComponentBound::NegFloatMax, EComponentBound::FloatMax);
+	case EValueComponentType::Double: return FValueComponentTypeDescription(TEXT("double"), 8u, EComponentBound::NegDoubleMax, EComponentBound::DoubleMax);
+	case EValueComponentType::Int: return FValueComponentTypeDescription(TEXT("int"), 4u, EComponentBound::IntMin, EComponentBound::IntMax);
+	case EValueComponentType::Bool: return FValueComponentTypeDescription(TEXT("bool"), 1u, EComponentBound::Zero, EComponentBound::One);
+	default: checkNoEntry() return FValueComponentTypeDescription();
 	}
 }
 
@@ -660,6 +647,12 @@ EValueType MakeValueTypeWithRequestedNumComponents(EValueType BaseType, int8 Req
 	return MakeValueType(TypeDesc.ComponentType, FMath::Min(TypeDesc.NumComponents, RequestedNumComponents));
 }
 
+EValueType MakeNonLWCType(EValueType Type)
+{
+	const FValueTypeDescription TypeDesc = GetValueTypeDescription(Type);
+	return MakeValueType(MakeNonLWCType(TypeDesc.ComponentType), TypeDesc.NumComponents);
+}
+
 EValueType MakeArithmeticResultType(EValueType Lhs, EValueType Rhs, FString& OutErrorMessage)
 {
 	const FValueTypeDescription LhsDesc = GetValueTypeDescription(Lhs);
diff --git a/Engine/Source/Runtime/Engine/Public/HLSLTree/HLSLTree.h b/Engine/Source/Runtime/Engine/Public/HLSLTree/HLSLTree.h
--- a/Engine/Source/Runtime/Engine/Public/HLSLTree/HLSLTree.h
+++ b/Engine/Source/Runtime/Engine/Public/HLSLTree/HLSLTree.h
@@ -251,7 +251,8 @@ class FRequestedType
 public:
 	FRequestedType() = default;
 	FRequestedType(ERequestedType InType);
-	explicit FRequestedType(int32 NumComponents);
+	FRequestedType(Shader::EValueType InType);
+	FRequestedType(const Shader::FType& InType);
 
 	int32 GetNumComponents() const;
 	bool IsComponentRequested(int32 Index) const { return RequestedComponents.IsValidIndex(Index) ? (bool)RequestedComponents[Index] : false; }
@@ -297,11 +298,7 @@ inline void AppendHash(FHasher& Hasher, const FRequestedType& Value)
 struct FPreparedComponent
 {
 	FPreparedComponent() = default;
-	FPreparedComponent(EExpressionEvaluation InEvaluation, FEmitScope* InLoopScope = nullptr, FExpression* InForwardValue = nullptr, int32 InForwardComponentIndex = INDEX_NONE)
-		: LoopScope(InLoopScope), ForwardValue(InForwardValue), ForwardComponentIndex(InForwardComponentIndex), Evaluation(InEvaluation)
-	{
-		check(!IsLoopEvaluation(InEvaluation) || InLoopScope);
-	}
+	FPreparedComponent(EExpressionEvaluation InEvaluation) : Evaluation(InEvaluation) {}
 
 	inline bool IsNone() const { return Evaluation == EExpressionEvaluation::None; }
 
@@ -310,11 +307,13 @@ struct FPreparedComponent
 	FEmitScope* LoopScope = nullptr;
 	FExpression* ForwardValue = nullptr;
 	int32 ForwardComponentIndex = INDEX_NONE;
+	Shader::FComponentBounds Bounds;
 	EExpressionEvaluation Evaluation = EExpressionEvaluation::None;
 };
 inline bool operator==(const FPreparedComponent& Lhs, const FPreparedComponent& Rhs)
 {
 	return Lhs.Evaluation == Rhs.Evaluation &&
+		Lhs.Bounds == Rhs.Bounds &&
 		Lhs.LoopScope == Rhs.LoopScope &&
 		Lhs.ForwardValue == Rhs.ForwardValue &&
 		Lhs.ForwardComponentIndex == Rhs.ForwardComponentIndex;
@@ -341,6 +340,7 @@ public:
 	void MergeEvaluation(EExpressionEvaluation Evaluation);
 	void SetLoopEvaluation(FEmitScope& Scope, const FRequestedType& RequestedType);
 	void SetForwardValue(const FRequestedType& RequestedType, FExpression* ForwardValue);
+	void UpdateBounds(const FRequestedType& RequestedType, Shader::FComponentBounds Bounds);
 
 	void SetField(const Shader::FStructField* Field, const FPreparedType& FieldType);
 	FPreparedType GetFieldType(const Shader::FStructField* Field) const;
@@ -358,6 +358,7 @@ public:
 	EExpressionEvaluation GetEvaluation(const FEmitScope& Scope) const;
 	EExpressionEvaluation GetEvaluation(const FEmitScope& Scope, const FRequestedType& RequestedType) const;
 	EExpressionEvaluation GetFieldEvaluation(const FEmitScope& Scope, int32 ComponentIndex, int32 NumComponents) const;
+	Shader::FComponentBounds GetBounds(const FRequestedType& RequestedType) const;
 	FPreparedComponent GetComponent(int32 Index) const;
 
 	void SetComponent(int32 Index, const FPreparedComponent& InComponent);
@@ -385,21 +386,6 @@ inline bool operator!=(const FPreparedType& Lhs, const FPreparedType& Rhs)
 
 FPreparedType MergePreparedTypes(const FPreparedType& Lhs, const FPreparedType& Rhs);
 
-struct FRequestedValueType
-{
-	FRequestedValueType() = default;
-	FRequestedValueType(Shader::EValueComponentType InComponentType, const FRequestedType& InRequestedType) : ComponentType(InComponentType), RequestedType(InRequestedType) {}
-	FRequestedValueType(Shader::EValueType InType);
-	FRequestedValueType(const FPreparedType& InType);
-
-	int32 GetNumComponents() const { return RequestedType.GetNumComponents(); }
-	Shader::EValueType GetType() const;
-	inline bool IsVoid() const { return ComponentType == Shader::EValueComponentType::Void || GetNumComponents() == 0; }
-
-	Shader::EValueComponentType ComponentType = Shader::EValueComponentType::Void;
-	FRequestedType RequestedType;
-};
-
 class FPrepareValueResult
 {
 public:
@@ -472,12 +458,13 @@ public:
 	Shader::FValue GetValueConstant(FEmitContext& Context, FEmitScope& Scope, const FRequestedType& RequestedType);
 
 	FEmitShaderExpression* GetValueShader(FEmitContext& Context, FEmitScope& Scope, const FRequestedType& RequestedType, const Shader::FType& ResultType);
-	FEmitShaderExpression* GetValueShader(FEmitContext& Context, FEmitScope& Scope, const FRequestedValueType& ResultType);
-	FEmitShaderExpression* GetValueShader(FEmitContext& Context, FEmitScope& Scope, const FPreparedType& ResultType);
+	FEmitShaderExpression* GetValueShader(FEmitContext& Context, FEmitScope& Scope, const Shader::FType& ResultType);
+	FEmitShaderExpression* GetValueShader(FEmitContext& Context, FEmitScope& Scope, Shader::EValueType ResultType);
 	FEmitShaderExpression* GetValueShader(FEmitContext& Context, FEmitScope& Scope);
 
 	Shader::FValue GetValueConstant(FEmitContext& Context, FEmitScope& Scope, const FRequestedType& RequestedType, const Shader::FType& ResultType);
 	Shader::FValue GetValueConstant(FEmitContext& Context, FEmitScope& Scope, const Shader::FType& ResultType);
+	Shader::FValue GetValueConstant(FEmitContext& Context, FEmitScope& Scope, Shader::EValueType ResultType);
 
 protected:
 	virtual void ComputeAnalyticDerivatives(FTree& Tree, FExpressionDerivatives& OutResult) const;
diff --git a/Engine/Source/Runtime/Engine/Public/HLSLTree/HLSLTreeEmit.h b/Engine/Source/Runtime/Engine/Public/HLSLTree/HLSLTreeEmit.h
--- a/Engine/Source/Runtime/Engine/Public/HLSLTree/HLSLTreeEmit.h
+++ b/Engine/Source/Runtime/Engine/Public/HLSLTree/HLSLTreeEmit.h
@@ -244,6 +244,12 @@ public:
 		return InternalEmitExpression(Scope, LocalDependencies, false, Type, String.ToView());
 	}
 
+	template<typename FormatType, typename... Types>
+	FEmitShaderExpression* EmitExpressionWithDependency(FEmitScope& Scope, FEmitShaderNode* Dependency, const Shader::FType& Type, const FormatType& Format, Types... Args)
+	{
+		return EmitExpressionWithDependencies(Scope, MakeDependencies(Dependency), Type, Format, Forward<Types>(Args)...);
+	}
+
 	template<typename FormatType, typename... Types>
 	FEmitShaderExpression* EmitInlineExpressionWithDependencies(FEmitScope& Scope, TArrayView<FEmitShaderNode*> Dependencies, const Shader::FType& Type, const FormatType& Format, Types... Args)
 	{
diff --git a/Engine/Source/Runtime/Engine/Public/Shader/ShaderTypes.h b/Engine/Source/Runtime/Engine/Public/Shader/ShaderTypes.h
--- a/Engine/Source/Runtime/Engine/Public/Shader/ShaderTypes.h
+++ b/Engine/Source/Runtime/Engine/Public/Shader/ShaderTypes.h
@@ -15,6 +15,43 @@ namespace Shader
 
 struct FStructType;
 
+enum class EComponentBound : uint8
+{
+	NegDoubleMax,
+	NegFloatMax,
+	IntMin,
+	NegOne,
+	Zero,
+	One,
+	IntMax,
+	FloatMax,
+	DoubleMax,
+};
+
+inline EComponentBound MinBound(EComponentBound Lhs, EComponentBound Rhs) { return (EComponentBound)FMath::Min((uint8)Lhs, (uint8)Rhs); }
+inline EComponentBound MaxBound(EComponentBound Lhs, EComponentBound Rhs) { return (EComponentBound)FMath::Max((uint8)Lhs, (uint8)Rhs); }
+
+struct FComponentBounds
+{
+	FComponentBounds() = default;
+	FComponentBounds(EComponentBound InMin, EComponentBound InMax) : Min(InMin), Max(InMax) {}
+
+	EComponentBound Min = EComponentBound::NegDoubleMax;
+	EComponentBound Max = EComponentBound::DoubleMax;
+};
+inline bool operator==(const FComponentBounds& Lhs, const FComponentBounds& Rhs)
+{
+	return Lhs.Min == Rhs.Min && Lhs.Max == Rhs.Max;
+}
+inline bool operator!=(const FComponentBounds& Lhs, const FComponentBounds& Rhs)
+{
+	return !operator==(Lhs, Rhs);
+}
+
+inline FComponentBounds MinBound(FComponentBounds Lhs, FComponentBounds Rhs) { return FComponentBounds(MinBound(Lhs.Min, Rhs.Min), MinBound(Lhs.Max, Rhs.Max)); }
+inline FComponentBounds MaxBound(FComponentBounds Lhs, FComponentBounds Rhs) { return FComponentBounds(MaxBound(Lhs.Min, Rhs.Min), MaxBound(Lhs.Max, Rhs.Max)); }
+inline bool IsWithinBounds(FComponentBounds Lhs, FComponentBounds Rhs) { return (uint8)Lhs.Min >= (uint8)Rhs.Min && (uint8)Lhs.Max <= (uint8)Rhs.Max; }
+
 enum class EValueComponentType : uint8
 {
 	Void,
@@ -24,22 +61,25 @@ enum class EValueComponentType : uint8
 	Bool,
 };
 
-const TCHAR* GetComponentTypeName(EValueComponentType Type);
-uint32 GetComponentTypeSizeInBytes(EValueComponentType Type);
-EValueComponentType CombineComponentTypes(EValueComponentType Lhs, EValueComponentType Rhs);
-
-inline EValueComponentType MakeNonLWCType(EValueComponentType Type) { return Type == EValueComponentType::Double ? EValueComponentType::Float : Type; }
-
-struct FValueTypeDescription
+struct FValueComponentTypeDescription
 {
-	FValueTypeDescription() : Name(nullptr), ComponentType(EValueComponentType::Void), NumComponents(0) {}
-	FValueTypeDescription(const TCHAR* InName, EValueComponentType InComponentType, int8 InNumComponents) : Name(InName), ComponentType(InComponentType), NumComponents(InNumComponents) {}
+	FValueComponentTypeDescription() = default;
+	FValueComponentTypeDescription(const TCHAR* InName, uint32_t InSizeInBytes, EComponentBound InMin, EComponentBound InMax) : Name(InName), SizeInBytes(InSizeInBytes), Bounds(InMin, InMax) {}
 
-	const TCHAR* Name;
-	EValueComponentType ComponentType;
-	int8 NumComponents;
+	const TCHAR* Name = nullptr;
+	uint32_t SizeInBytes = 0u;
+	FComponentBounds Bounds;
 };
 
+FValueComponentTypeDescription GetValueComponentTypeDescription(EValueComponentType Type);
+inline const TCHAR* GetComponentTypeName(EValueComponentType Type) { return GetValueComponentTypeDescription(Type).Name; }
+inline uint32 GetComponentTypeSizeInBytes(EValueComponentType Type) { return GetValueComponentTypeDescription(Type).SizeInBytes; }
+inline bool IsComponentTypeWithinBounds(EValueComponentType Type, FComponentBounds Bounds) { return IsWithinBounds(GetValueComponentTypeDescription(Type).Bounds, Bounds); }
+
+EValueComponentType CombineComponentTypes(EValueComponentType Lhs, EValueComponentType Rhs);
+
+inline EValueComponentType MakeNonLWCType(EValueComponentType Type) { return Type == EValueComponentType::Double ? EValueComponentType::Float : Type; }
+
 enum class EValueType : uint8
 {
 	Void,
@@ -75,10 +115,21 @@ enum class EValueType : uint8
 	Struct,
 };
 
+struct FValueTypeDescription
+{
+	FValueTypeDescription() : Name(nullptr), ComponentType(EValueComponentType::Void), NumComponents(0) {}
+	FValueTypeDescription(const TCHAR* InName, EValueComponentType InComponentType, int8 InNumComponents) : Name(InName), ComponentType(InComponentType), NumComponents(InNumComponents) {}
+
+	const TCHAR* Name;
+	EValueComponentType ComponentType;
+	int8 NumComponents;
+};
+
 FValueTypeDescription GetValueTypeDescription(EValueType Type);
 EValueType MakeValueType(EValueComponentType ComponentType, int32 NumComponents);
 EValueType MakeValueType(EValueType BaseType, int32 NumComponents);
 EValueType MakeValueTypeWithRequestedNumComponents(EValueType BaseType, int8 RequestedNumComponents);
+EValueType MakeNonLWCType(EValueType Type);
 EValueType MakeArithmeticResultType(EValueType Lhs, EValueType Rhs, FString& OutErrorMessage);
 EValueType MakeComparisonResultType(EValueType Lhs, EValueType Rhs, FString& OutErrorMessage);
 
