#ue5 - Fix Characters using PerchRadiusThreshold to not stand on the edge of geometry using the wrong location for the perch check.

Using the incoming Hit.Location as the capsule location was incorrect when the previous tests shrank the capsule height to avoid penetrations, placing the origination of the perch check often starting in penetration and causing a fallback to another shrunken test, which would waste performance and sometimes even fail. This was made more common when using bUseFlatBaseForFloorChecks, as the tests were often rejected when physics returns points on the edge of the box which was too close to the radius to be acceptable. In the worst case, no floor checks would pass when using bUseFlatBaseForFloorChecks (normal tests would be too close to edge, perch check would be in penetration).

diff --git a/Engine/Source/Runtime/Engine/Private/Components/CharacterMovementComponent.cpp b/Engine/Source/Runtime/Engine/Private/Components/CharacterMovementComponent.cpp
--- a/Engine/Source/Runtime/Engine/Private/Components/CharacterMovementComponent.cpp
+++ b/Engine/Source/Runtime/Engine/Private/Components/CharacterMovementComponent.cpp
@@ -6716,13 +6716,14 @@ bool UCharacterMovementComponent::ComputePerchResult(const float TestRadius, con
 	// Sweep further than actual requested distance, because a reduced capsule radius means we could miss some hits that the normal radius would contact.
 	float PawnRadius, PawnHalfHeight;
 	CharacterOwner->GetCapsuleComponent()->GetScaledCapsuleSize(PawnRadius, PawnHalfHeight);
+	const FVector CapsuleLocation = InHit.TraceStart;
 
-	const float InHitAboveBase = FMath::Max<float>(0.f, InHit.ImpactPoint.Z - (InHit.Location.Z - PawnHalfHeight));
+	const float InHitAboveBase = FMath::Max<float>(0.f, InHit.ImpactPoint.Z - (CapsuleLocation.Z - PawnHalfHeight));
 	const float PerchLineDist = FMath::Max(0.f, InMaxFloorDist - InHitAboveBase);
 	const float PerchSweepDist = FMath::Max(0.f, InMaxFloorDist);
 
 	const float ActualSweepDist = PerchSweepDist + PawnRadius;
-	ComputeFloorDist(InHit.Location, PerchLineDist, ActualSweepDist, OutPerchFloorResult, TestRadius);
+	ComputeFloorDist(CapsuleLocation, PerchLineDist, ActualSweepDist, OutPerchFloorResult, TestRadius);
 
 	if (!OutPerchFloorResult.IsWalkableFloor())
 	{
