Fixes for initialisation of the BlendSpace filtering - bugs introduced in the fixes for UE-141074, UE-141123

diff --git a/Engine/Source/Runtime/Engine/Private/Animation/AnimInterpFilter.cpp b/Engine/Source/Runtime/Engine/Private/Animation/AnimInterpFilter.cpp
--- a/Engine/Source/Runtime/Engine/Private/Animation/AnimInterpFilter.cpp
+++ b/Engine/Source/Runtime/Engine/Private/Animation/AnimInterpFilter.cpp
@@ -28,7 +28,7 @@ void FFIRFilter::CalculateCoefficient(EFilterInterpolationType InterpolationType
 	if ( IsValid() )
 	{
 		float Sum=0.f;
-		for (int32 I=0; I<Coefficients.Num(); ++I)
+		for (int32 I=0; I != Coefficients.Num(); ++I)
 		{
 			Coefficients[I] = GetInterpolationCoefficient(InterpolationType, I);
 			Sum += Coefficients[I];
@@ -37,7 +37,7 @@ void FFIRFilter::CalculateCoefficient(EFilterInterpolationType InterpolationType
 		// now normalize it, if not 1
 		if ( fabs(Sum-1.f) > ZERO_ANIMWEIGHT_THRESH )
 		{
-			for (int32 I=0; I<Coefficients.Num(); ++I)
+			for (int32 I=0; I != Coefficients.Num(); ++I)
 			{
 				Coefficients[I]/=Sum;
 			}
@@ -49,11 +49,9 @@ float FFIRFilter::GetFilteredData(float Input)
 {
 	if ( IsValid() )
 	{
-		FilterWindow[CurrentStack] = Input;
+		FilterData[CurrentStack] = Input;
 		float Result = CalculateFilteredOutput();
-
-		CurrentStack = CurrentStack+1;
-		if ( CurrentStack > FilterWindow.Num()-1 ) 
+		if (++CurrentStack > FilterData.Num() - 1)
 		{
 			CurrentStack = 0;
 		}
@@ -73,11 +71,10 @@ float FFIRFilter::CalculateFilteredOutput() const
 
 	for ( int32 I=Coefficients.Num()-1; I>=0; --I )
 	{
-		Output += FilterWindow[StackIndex]*Coefficients[I];
-		StackIndex = StackIndex-1;
-		if (StackIndex < 0)
+		Output += FilterData[StackIndex]*Coefficients[I];
+		if (--StackIndex < 0)
 		{
-			StackIndex = FilterWindow.Num()-1;
+			StackIndex = FilterData.Num() - 1;
 		}
 	}
 
@@ -90,10 +87,10 @@ float FFIRFilter::CalculateFilteredOutput() const
 int32 FFIRFilterTimeBased::GetSafeCurrentStackIndex()
 {
 	// if valid range
-	check ( CurrentStackIndex < FilterWindow.Num() );
+	check ( CurrentStackIndex < FilterData.Num() );
 
 	// see if it's expired yet
-	if ( !FilterWindow[CurrentStackIndex].IsValid() )
+	if ( !FilterData[CurrentStackIndex].IsValid() )
 	{
 		return CurrentStackIndex;
 	}
@@ -103,15 +100,15 @@ int32 FFIRFilterTimeBased::GetSafeCurrentStackIndex()
 	// this should not be the case because most of times
 	// current one should be the oldest one, but since 
 	// we jumps when reallocation happens, we still do this
-	for (int32 I=0; I<FilterWindow.Num(); ++I)
+	for (int32 I=0; I != FilterData.Num(); ++I)
 	{
 		int32 NewIndex = CurrentStackIndex + I;
-		if (NewIndex >= FilterWindow.Num())
+		if (NewIndex >= FilterData.Num())
 		{
-			NewIndex = NewIndex - FilterWindow.Num();
+			NewIndex = NewIndex - FilterData.Num();
 		}
 
-		if ( !FilterWindow[NewIndex].IsValid() )
+		if ( !FilterData[NewIndex].IsValid() )
 		{
 			return NewIndex;
 		}
@@ -119,25 +116,25 @@ int32 FFIRFilterTimeBased::GetSafeCurrentStackIndex()
 
 	// if current one isn't available anymore 
 	// that means we need more stack
-	const int32 NewIndex = FilterWindow.Num();
-	FilterWindow.AddZeroed(5);
+	const int32 NewIndex = FilterData.Num();
+	FilterData.AddZeroed(5);
 	return NewIndex;
 }
 
 void FFIRFilterTimeBased::RefreshValidFilters()
 {
-	NumValidFilter = 0;
-
-	if ( WindowDuration > 0.f )
+	if (FilterData.IsEmpty())
+	{
+		FilterData.Empty(10);
+		FilterData.AddZeroed(10);
+		CurrentStackIndex = 0;
+	}
+	else
 	{
-		// run validation test
-		for (int32 I=0; I<FilterWindow.Num(); ++I)
+		// Ensure the current time is valid
+		for (int32 I=0; I != FilterData.Num(); ++I)
 		{
-			FilterWindow[I].CheckValidation(CurrentTime, WindowDuration);
-			if ( FilterWindow[I].IsValid() )
-			{
-				++NumValidFilter;
-			}
+			FilterData[I].EnsureTimeIsValid(CurrentTime, WindowDuration);
 		}
 	}
 }
@@ -155,9 +152,9 @@ void FFIRFilterTimeBased::WrapToValue(float Input, float Range)
 	case BSIT_ExponentialDecay:
 	case BSIT_SpringDamper:
 	{
-		if (FilterWindow.Num() != 0)
+		if (FilterData.Num() != 0)
 		{
-			FilterWindow[0].Input = FMath::Wrap(FilterWindow[0].Input, Input - HalfRange, Input + HalfRange);
+			FilterData[0].Input = FMath::Wrap(FilterData[0].Input, Input - HalfRange, Input + HalfRange);
 		}
 	}
 	break;
@@ -170,9 +167,9 @@ void FFIRFilterTimeBased::WrapToValue(float Input, float Range)
 			if (Delta)
 			{
 				LastOutput = NewLastOutput;
-				for (int32 Index = 0; Index != FilterWindow.Num() ; ++Index)
+				for (int32 Index = 0; Index != FilterData.Num() ; ++Index)
 				{
-					FilterWindow[Index].Input += Delta;
+					FilterData[Index].Input += Delta;
 				}
 			}
 		}
@@ -202,75 +199,75 @@ float FFIRFilterTimeBased::UpdateAndGetFilteredData(float Input, float DeltaTime
 		{
 			case BSIT_ExponentialDecay:
 			{
-				if (FilterWindow.Num() != 1)
+				if (FilterData.Num() != 1)
 				{
-					FilterWindow.Empty(1);
-					FilterWindow.Push(FFilterData(Input, 0.0f));
+					FilterData.Empty(1);
+					FilterData.Push(FFilterData(Input, 0.0f));
 				}
-				const float OrigValue = FilterWindow[0].Input;
-				FMath::ExponentialSmoothingApprox(FilterWindow[0].Input, Input, DeltaTime, WindowDuration / EULERS_NUMBER);
+				const float OrigValue = FilterData[0].Input;
+				FMath::ExponentialSmoothingApprox(FilterData[0].Input, Input, DeltaTime, WindowDuration / EULERS_NUMBER);
 				if (MaxSpeed > 0.0f)
 				{
 					// Clamp the speed
-					FilterWindow[0].Input = FMath::Clamp(FilterWindow[0].Input, OrigValue - MaxSpeed * DeltaTime,
+					FilterData[0].Input = FMath::Clamp(FilterData[0].Input, OrigValue - MaxSpeed * DeltaTime,
                                                          OrigValue + MaxSpeed * DeltaTime); 
 				}
-				Result = FilterWindow[0].Input;
+				Result = FilterData[0].Input;
 			}
 			break;
 			case BSIT_SpringDamper:
 			{
-				if (FilterWindow.Num() != 2)
+				if (FilterData.Num() != 2)
 				{
-					FilterWindow.Empty(2);
+					FilterData.Empty(2);
 					// [0] element is the value, [1] element is the rate
-					FilterWindow.Push(FFilterData(Input, 0.0f));
-					FilterWindow.Push(FFilterData(0.0f, 0.0f));
+					FilterData.Push(FFilterData(Input, 0.0f));
+					FilterData.Push(FFilterData(0.0f, 0.0f));
 				}
-				const float OrigValue = FilterWindow[0].Input;
-				FMath::SpringDamperSmoothing(FilterWindow[0].Input, FilterWindow[1].Input, Input, 0.0f, DeltaTime,
+				const float OrigValue = FilterData[0].Input;
+				FMath::SpringDamperSmoothing(FilterData[0].Input, FilterData[1].Input, Input, 0.0f, DeltaTime,
 				                             WindowDuration / EULERS_NUMBER, DampingRatio);
 				if (MaxSpeed > 0.0f)
 				{
 					// Clamp the speed
-					FilterWindow[0].Input = FMath::Clamp(FilterWindow[0].Input, OrigValue - MaxSpeed * DeltaTime,
+					FilterData[0].Input = FMath::Clamp(FilterData[0].Input, OrigValue - MaxSpeed * DeltaTime,
                                                          OrigValue + MaxSpeed * DeltaTime); 
-					FilterWindow[1].Input = FMath::Clamp(FilterWindow[1].Input, -MaxSpeed, MaxSpeed);
+					FilterData[1].Input = FMath::Clamp(FilterData[1].Input, -MaxSpeed, MaxSpeed);
 				}
 				if (bClamp)
 				{
 					// Clamp the value
-					if (FilterWindow[0].Input > MaxValue)
+					if (FilterData[0].Input > MaxValue)
 					{
-						FilterWindow[0].Input = MaxValue;
-						if (FilterWindow[1].Input > 0.0f)
+						FilterData[0].Input = MaxValue;
+						if (FilterData[1].Input > 0.0f)
 						{
-							FilterWindow[1].Input = 0.0f;
+							FilterData[1].Input = 0.0f;
 						}
 					}
-					if (FilterWindow[0].Input < MinValue)
+					if (FilterData[0].Input < MinValue)
 					{
-						FilterWindow[0].Input = MinValue;
-						if (FilterWindow[1].Input < 0.0f)
+						FilterData[0].Input = MinValue;
+						if (FilterData[1].Input < 0.0f)
 						{
-							FilterWindow[1].Input = 0.0f;
+							FilterData[1].Input = 0.0f;
 						}
 					}
 				}
-				Result = FilterWindow[0].Input;
+				Result = FilterData[0].Input;
 			}
 			break;
 			default:
 			{
+				// This handles the array based filters
 				if (IsValid())
 				{
 					RefreshValidFilters();
 
 					CurrentStackIndex = GetSafeCurrentStackIndex();
-					FilterWindow[CurrentStackIndex].SetInput(Input, CurrentTime);
+					FilterData[CurrentStackIndex].SetInput(Input, CurrentTime);
 					Result = CalculateFilteredOutput();
-					CurrentStackIndex = CurrentStackIndex + 1;
-					if (CurrentStackIndex > FilterWindow.Num() - 1)
+					if (++CurrentStackIndex > FilterData.Num() - 1)
 					{
 						CurrentStackIndex = 0;
 					}
@@ -321,13 +318,13 @@ float FFIRFilterTimeBased::CalculateFilteredOutput()
 	check ( IsValid() );
 	float SumCoefficient = 0.f;
 	float SumInputs = 0.f;
-	for (int32 I=0; I<FilterWindow.Num(); ++I)
+	for (int32 I=0; I != FilterData.Num(); ++I)
 	{
-		const float Coefficient = GetInterpolationCoefficient(FilterWindow[I]);
+		const float Coefficient = GetInterpolationCoefficient(FilterData[I]);
 		if (Coefficient > 0.f)
 		{
 			SumCoefficient += Coefficient;
-			SumInputs += Coefficient * FilterWindow[I].Input;
+			SumInputs += Coefficient * FilterData[I].Input;
 		}
 	}
 	return SumCoefficient > 0.f ? SumInputs/SumCoefficient : 0.f;
diff --git a/Engine/Source/Runtime/Engine/Public/AnimInterpFilter.h b/Engine/Source/Runtime/Engine/Public/AnimInterpFilter.h
--- a/Engine/Source/Runtime/Engine/Public/AnimInterpFilter.h
+++ b/Engine/Source/Runtime/Engine/Public/AnimInterpFilter.h
@@ -26,13 +26,13 @@ public:
 	{
 		if ( WindowLen > 0 )
 		{
-			FilterWindow.AddZeroed(WindowLen);
+			FilterData.AddZeroed(WindowLen);
 			Coefficients.AddZeroed(WindowLen);
 			CurrentStack = 0;
 		}
 		else
 		{
-			FilterWindow.Reset();
+			FilterData.Reset();
 			Coefficients.Reset();
 			CurrentStack = 0;		
 		}
@@ -40,13 +40,13 @@ public:
 
 	void CalculateCoefficient(EFilterInterpolationType InterpolationType);
 	float GetFilteredData(float Input);
-	bool IsValid() const { return FilterWindow.Num() > 0; }
+	bool IsValid() const { return FilterData.Num() > 0; }
 	float				LastOutput;
 private:
 
 	// CurrentStack is latest till CurrentStack + 1 is oldest
 	// note that this works in reverse order with Coefficient
-	TArray<float>		FilterWindow;
+	TArray<float>		FilterData;
 	// n-1 is latest till 0 is oldest
 	TArray<float>		Coefficients;
 	int32					CurrentStack;
@@ -77,7 +77,7 @@ struct FFilterData
 	{
 	}
 
-	void CheckValidation(const float CurrentTime, const float ValidationWindow)
+	void EnsureTimeIsValid(const float CurrentTime, const float ValidationWindow)
 	{
 		if (Diff(CurrentTime) > ValidationWindow)
 		{
@@ -119,9 +119,8 @@ public:
 	void Initialize(float InWindowDuration, EFilterInterpolationType InInterpolationType, float InDampingRatio,
 	                float InMinValue, float InMaxValue, float InMaxSpeed, bool bInClamp)
 	{
-		FilterWindow.Empty();
 		InterpolationType = InInterpolationType;
-		NumValidFilter = 0;
+		FilterData.Empty();
 		CurrentStackIndex = 0;
 		WindowDuration = InWindowDuration;
 		DampingRatio = InDampingRatio;
@@ -149,8 +148,8 @@ public:
 	// Wraps the internal state by steps of Range so that it is as close as possible to Input
 	void WrapToValue(float Input, float Range);
 
+	// Filter is considered valid if the WindowDuration is > 0
 	bool IsValid() const { return WindowDuration > 0.f; }
-	float LastOutput;
 
 	void SetWindowDuration(float InWindowDuration)
 	{
@@ -164,6 +163,8 @@ public:
 	}
 #endif // WITH_EDITOR
 
+	float LastOutput;
+
 private:
 	EFilterInterpolationType InterpolationType;
 	int32 CurrentStackIndex;
@@ -173,9 +174,8 @@ private:
 	float MaxValue;
 	float MaxSpeed;
 	bool  bClamp;
-	int32 NumValidFilter;
 	float CurrentTime;
-	TArray<FFilterData> FilterWindow;
+	TArray<FFilterData> FilterData;
 
 	float GetInterpolationCoefficient(const FFilterData& Data) const;
 	float CalculateFilteredOutput();
