Make UAnimInstance LinkAnimClassLayers / UnlinkAnimClassLayers virtual

#rnx

diff --git a/Engine/Source/Runtime/Engine/Classes/Animation/AnimInstance.h b/Engine/Source/Runtime/Engine/Classes/Animation/AnimInstance.h
--- a/Engine/Source/Runtime/Engine/Classes/Animation/AnimInstance.h
+++ b/Engine/Source/Runtime/Engine/Classes/Animation/AnimInstance.h
@@ -816,7 +816,7 @@ public:
 	 * If InClass is null, then all layers are reset to their defaults.
 	 */
 	UFUNCTION(BlueprintCallable, Category = "Animation Blueprint Linking")
-	void LinkAnimClassLayers(TSubclassOf<UAnimInstance> InClass);
+	virtual void LinkAnimClassLayers(TSubclassOf<UAnimInstance> InClass);
 
 	UE_DEPRECATED(4.24, "Function renamed, please use UnlinkAnimClassLayers")
 	void ClearLayerOverlay(TSubclassOf<UAnimInstance> InClass) { UnlinkAnimClassLayers(InClass); }
@@ -827,7 +827,7 @@ public:
 	 * If InClass is null, does nothing.
 	 */
 	UFUNCTION(BlueprintCallable, Category = "Animation Blueprint Linking")
-	void UnlinkAnimClassLayers(TSubclassOf<UAnimInstance> InClass);
+	virtual void UnlinkAnimClassLayers(TSubclassOf<UAnimInstance> InClass);
 
 	UE_DEPRECATED(4.24, "Function renamed, please use GetLinkedAnimLayerInstanceByGroup")
 	UAnimInstance* GetLayerSubInstanceByGroup(FName InGroup) const { return GetLinkedAnimLayerInstanceByGroup(InGroup); }
