Fix constant path tracer invalidation from PostProcess material blending

Avoid invalidating the path tracer when PostProcess or UI materials are being changed as they cannot have an impact on the lighting calculations.

This case was noticed on the MetaHuman sample scene.

diff --git a/Engine/Source/Runtime/Engine/Private/Materials/MaterialInstance.cpp b/Engine/Source/Runtime/Engine/Private/Materials/MaterialInstance.cpp
--- a/Engine/Source/Runtime/Engine/Private/Materials/MaterialInstance.cpp
+++ b/Engine/Source/Runtime/Engine/Private/Materials/MaterialInstance.cpp
@@ -345,7 +345,16 @@ void GameThread_UpdateMIParameter(const UMaterialInstance* Instance, const Param
 {
 	if (FApp::CanEverRender())
 	{
-		GetRendererModule().InvalidatePathTracedOutput();
+		const UMaterial* Material = Instance->GetMaterial();
+		if (Material != nullptr)
+		{
+			EMaterialDomain Domain = Material->MaterialDomain;
+			// check if this material has any relevance to path tracing
+			if (Domain != MD_PostProcess && Domain != MD_UI)
+			{
+				GetRendererModule().InvalidatePathTracedOutput();
+			}
+		}
 		FMaterialInstanceResource* Resource = Instance->Resource;
 		const FMaterialParameterInfo& ParameterInfo = Parameter.ParameterInfo;
 		typename ParameterType::ValueType Value = ParameterType::GetValue(Parameter);
