fixed audio pipe tasks priority and set it to Normal as initially intended BackgroundLow or BackgroundNormal leads to slower frame time on some platforms

diff --git a/Engine/Source/Runtime/Engine/Private/AudioThread.cpp b/Engine/Source/Runtime/Engine/Private/AudioThread.cpp
--- a/Engine/Source/Runtime/Engine/Private/AudioThread.cpp
+++ b/Engine/Source/Runtime/Engine/Private/AudioThread.cpp
@@ -84,7 +84,7 @@ struct FAudioThreadInteractor
 	}
 };
 
-UE::Tasks::ETaskPriority GAudioTaskPriority = UE::Tasks::ETaskPriority::BackgroundLow;
+UE::Tasks::ETaskPriority GAudioTaskPriority = UE::Tasks::ETaskPriority::Normal;
 
 static void SetAudioTaskPriority(const TArray<FString>& Args)
 {
@@ -377,7 +377,8 @@ struct FAudioAsyncBatcher
 				{
 					Work();
 				}
-			}
+			},
+			GAudioTaskPriority
 		);
 		WorkItems.Reset();
 	}
