Merge UnaryOp/BinaryOp into a single generalized 'Operation', collapses some duplicate code, makes it easier to potentially add ternary operations (or any other number of inputs)

diff --git a/Engine/Source/Runtime/Engine/Classes/Materials/MaterialExpressionBinaryOp.h b/Engine/Source/Runtime/Engine/Classes/Materials/MaterialExpressionBinaryOp.h
--- a/Engine/Source/Runtime/Engine/Classes/Materials/MaterialExpressionBinaryOp.h
+++ b/Engine/Source/Runtime/Engine/Classes/Materials/MaterialExpressionBinaryOp.h
@@ -29,7 +29,7 @@ class UMaterialExpressionBinaryOp : public UMaterialExpression
 	UPROPERTY(EditAnywhere, Category = MaterialExpressionAdd, meta = (OverridingInputProperty = "B"))
 	float ConstB;
 
-	virtual UE::HLSLTree::EBinaryOp GetBinaryOp() const PURE_VIRTUAL(UMaterialExpressionBinaryOp::GetBinaryOp, return UE::HLSLTree::EBinaryOp::None;);
+	virtual UE::HLSLTree::EOperation GetBinaryOp() const PURE_VIRTUAL(UMaterialExpressionBinaryOp::GetBinaryOp, return UE::HLSLTree::EOperation::None;);
 
 	//~ Begin UMaterialExpression Interface
 #if WITH_EDITOR
@@ -46,5 +46,5 @@ UCLASS(MinimalAPI, meta = (MaterialNewHLSLGenerator))
 class UMaterialExpressionLess : public UMaterialExpressionBinaryOp
 {
 	GENERATED_UCLASS_BODY()
-	virtual UE::HLSLTree::EBinaryOp GetBinaryOp() const override { return UE::HLSLTree::EBinaryOp::Less; }
+	virtual UE::HLSLTree::EOperation GetBinaryOp() const override { return UE::HLSLTree::EOperation::Less; }
 };
diff --git a/Engine/Source/Runtime/Engine/Private/HLSLTree/HLSLTree.cpp b/Engine/Source/Runtime/Engine/Private/HLSLTree/HLSLTree.cpp
--- a/Engine/Source/Runtime/Engine/Private/HLSLTree/HLSLTree.cpp
+++ b/Engine/Source/Runtime/Engine/Private/HLSLTree/HLSLTree.cpp
@@ -80,56 +80,42 @@ public:
 	int32 OutputIndex;
 };
 
-FUnaryOpDescription::FUnaryOpDescription()
-	: Name(nullptr), Operator(nullptr), PreshaderOpcode(Shader::EPreshaderOpcode::Nop)
-{}
-
-FUnaryOpDescription::FUnaryOpDescription(const TCHAR* InName, const TCHAR* InOperator, Shader::EPreshaderOpcode InOpcode)
-	: Name(InName), Operator(InOperator), PreshaderOpcode(InOpcode)
+FOperationDescription::FOperationDescription(const TCHAR* InName, const TCHAR* InOperator, int8 InNumInputs, Shader::EPreshaderOpcode InOpcode)
+	: Name(InName), Operator(InOperator), NumInputs(InNumInputs), PreshaderOpcode(InOpcode)
 {}
 
-FBinaryOpDescription::FBinaryOpDescription()
+FOperationDescription::FOperationDescription()
 	: Name(nullptr), Operator(nullptr), PreshaderOpcode(Shader::EPreshaderOpcode::Nop)
 {}
 
-FBinaryOpDescription::FBinaryOpDescription(const TCHAR* InName, const TCHAR* InOperator, Shader::EPreshaderOpcode InOpcode)
-	: Name(InName), Operator(InOperator), PreshaderOpcode(InOpcode)
-{}
-
-FUnaryOpDescription GetUnaryOpDesription(EUnaryOp Op)
-{
-	switch (Op)
-	{
-	case EUnaryOp::None: return FUnaryOpDescription(TEXT("None"), TEXT(""), Shader::EPreshaderOpcode::Nop); break;
-	case EUnaryOp::Neg: return FUnaryOpDescription(TEXT("Neg"), TEXT("-"), Shader::EPreshaderOpcode::Neg); break;
-	case EUnaryOp::Rcp: return FUnaryOpDescription(TEXT("Rcp"), TEXT("/"), Shader::EPreshaderOpcode::Rcp); break;
-	case EUnaryOp::Frac: return FUnaryOpDescription(TEXT("Frac"), TEXT("frac"), Shader::EPreshaderOpcode::Frac); break;
-	case EUnaryOp::Length: return FUnaryOpDescription(TEXT("Length"), TEXT("length"), Shader::EPreshaderOpcode::Length); break;
-	case EUnaryOp::Normalize: return FUnaryOpDescription(TEXT("Normalize"), TEXT("normalize"), Shader::EPreshaderOpcode::Normalize); break;
-	default: checkNoEntry(); return FUnaryOpDescription();
-	}
-}
 
-FBinaryOpDescription GetBinaryOpDesription(EBinaryOp Op)
+FOperationDescription GetOperationDescription(EOperation Op)
 {
 	switch (Op)
 	{
-	case EBinaryOp::None: return FBinaryOpDescription(TEXT("None"), TEXT(""), Shader::EPreshaderOpcode::Nop); break;
-	case EBinaryOp::Add: return FBinaryOpDescription(TEXT("Add"), TEXT("+"), Shader::EPreshaderOpcode::Add); break;
-	case EBinaryOp::Sub: return FBinaryOpDescription(TEXT("Subtract"), TEXT("-"), Shader::EPreshaderOpcode::Sub); break;
-	case EBinaryOp::Mul: return FBinaryOpDescription(TEXT("Multiply"), TEXT("*"), Shader::EPreshaderOpcode::Mul); break;
-	case EBinaryOp::Div: return FBinaryOpDescription(TEXT("Divide"), TEXT("/"), Shader::EPreshaderOpcode::Div); break;
-	case EBinaryOp::Fmod: return FBinaryOpDescription(TEXT("Fmod"), TEXT("%"), Shader::EPreshaderOpcode::Fmod); break;
-	case EBinaryOp::Dot: return FBinaryOpDescription(TEXT("Dot"), TEXT("dot"), Shader::EPreshaderOpcode::Dot); break;
-	case EBinaryOp::Min: return FBinaryOpDescription(TEXT("Min"), TEXT("min"), Shader::EPreshaderOpcode::Min); break;
-	case EBinaryOp::Max: return FBinaryOpDescription(TEXT("Max"), TEXT("max"), Shader::EPreshaderOpcode::Max); break;
-	case EBinaryOp::Less: return FBinaryOpDescription(TEXT("Less"), TEXT("<"), Shader::EPreshaderOpcode::Less); break;
-	case EBinaryOp::Greater: return FBinaryOpDescription(TEXT("Greater"), TEXT(">"), Shader::EPreshaderOpcode::Greater); break;
-	case EBinaryOp::VecMulMatrix3: return FBinaryOpDescription(TEXT("VecMulMatrix3"), TEXT("mul"), Shader::EPreshaderOpcode::Nop); break;
-	case EBinaryOp::VecMulMatrix4: return FBinaryOpDescription(TEXT("VecMulMatrix4"), TEXT("mul"), Shader::EPreshaderOpcode::Nop); break;
-	case EBinaryOp::Matrix3MulVec: return FBinaryOpDescription(TEXT("Matrix3MulVec"), TEXT("mul"), Shader::EPreshaderOpcode::Nop); break;
-	case EBinaryOp::Matrix4MulVec: return FBinaryOpDescription(TEXT("Matrix4MulVec"), TEXT("mul"), Shader::EPreshaderOpcode::Nop); break;
-	default: checkNoEntry(); return FBinaryOpDescription();
+	// Unary
+	case EOperation::None: return FOperationDescription(TEXT("None"), TEXT(""), 0, Shader::EPreshaderOpcode::Nop); break;
+	case EOperation::Neg: return FOperationDescription(TEXT("Neg"), TEXT("-"), 1, Shader::EPreshaderOpcode::Neg); break;
+	case EOperation::Rcp: return FOperationDescription(TEXT("Rcp"), TEXT("/"), 1, Shader::EPreshaderOpcode::Rcp); break;
+	case EOperation::Frac: return FOperationDescription(TEXT("Frac"), TEXT("frac"), 1, Shader::EPreshaderOpcode::Frac); break;
+	case EOperation::Length: return FOperationDescription(TEXT("Length"), TEXT("length"), 1, Shader::EPreshaderOpcode::Length); break;
+	case EOperation::Normalize: return FOperationDescription(TEXT("Normalize"), TEXT("normalize"), 1, Shader::EPreshaderOpcode::Normalize); break;
+	// Binary
+	case EOperation::Add: return FOperationDescription(TEXT("Add"), TEXT("+"), 2, Shader::EPreshaderOpcode::Add); break;
+	case EOperation::Sub: return FOperationDescription(TEXT("Subtract"), TEXT("-"), 2, Shader::EPreshaderOpcode::Sub); break;
+	case EOperation::Mul: return FOperationDescription(TEXT("Multiply"), TEXT("*"), 2, Shader::EPreshaderOpcode::Mul); break;
+	case EOperation::Div: return FOperationDescription(TEXT("Divide"), TEXT("/"), 2, Shader::EPreshaderOpcode::Div); break;
+	case EOperation::Fmod: return FOperationDescription(TEXT("Fmod"), TEXT("%"), 2, Shader::EPreshaderOpcode::Fmod); break;
+	case EOperation::Dot: return FOperationDescription(TEXT("Dot"), TEXT("dot"), 2, Shader::EPreshaderOpcode::Dot); break;
+	case EOperation::Min: return FOperationDescription(TEXT("Min"), TEXT("min"), 2, Shader::EPreshaderOpcode::Min); break;
+	case EOperation::Max: return FOperationDescription(TEXT("Max"), TEXT("max"), 2, Shader::EPreshaderOpcode::Max); break;
+	case EOperation::Less: return FOperationDescription(TEXT("Less"), TEXT("<"), 2, Shader::EPreshaderOpcode::Less); break;
+	case EOperation::Greater: return FOperationDescription(TEXT("Greater"), TEXT(">"), 2, Shader::EPreshaderOpcode::Greater); break;
+	case EOperation::VecMulMatrix3: return FOperationDescription(TEXT("VecMulMatrix3"), TEXT("mul"), 2, Shader::EPreshaderOpcode::Nop); break;
+	case EOperation::VecMulMatrix4: return FOperationDescription(TEXT("VecMulMatrix4"), TEXT("mul"), 2, Shader::EPreshaderOpcode::Nop); break;
+	case EOperation::Matrix3MulVec: return FOperationDescription(TEXT("Matrix3MulVec"), TEXT("mul"), 2, Shader::EPreshaderOpcode::Nop); break;
+	case EOperation::Matrix4MulVec: return FOperationDescription(TEXT("Matrix4MulVec"), TEXT("mul"), 2, Shader::EPreshaderOpcode::Nop); break;
+	default: checkNoEntry(); return FOperationDescription();
 	}
 }
 
diff --git a/Engine/Source/Runtime/Engine/Private/HLSLTree/HLSLTreeCommon.cpp b/Engine/Source/Runtime/Engine/Private/HLSLTree/HLSLTreeCommon.cpp
--- a/Engine/Source/Runtime/Engine/Private/HLSLTree/HLSLTreeCommon.cpp
+++ b/Engine/Source/Runtime/Engine/Private/HLSLTree/HLSLTreeCommon.cpp
@@ -732,371 +732,245 @@ void FExpressionSelect::EmitValuePreshader(FEmitContext& Context, FEmitScope& Sc
 	check(false); // TODO
 }
 
-FExpression* FTree::NewUnaryOp(EUnaryOp Op, FExpression* Input)
+FExpression* FTree::NewUnaryOp(EOperation Op, FExpression* Input)
 {
-	return NewExpression<FExpressionUnaryOp>(Op, Input);
+	FExpression* Inputs[1] = { Input };
+	return NewExpression<FExpressionOperation>(Op, Inputs);
 }
 
-void FExpressionUnaryOp::ComputeAnalyticDerivatives(FTree& Tree, FExpressionDerivatives& OutResult) const
+FExpression* FTree::NewBinaryOp(EOperation Op, FExpression* Lhs, FExpression* Rhs)
 {
-	const FExpressionDerivatives InputDerivatives = Tree.GetAnalyticDerivatives(Input);
-	if (InputDerivatives.IsValid())
-	{
-		switch (Op)
-		{
-		case EUnaryOp::Neg:
-			OutResult.ExpressionDdx = Tree.NewNeg(InputDerivatives.ExpressionDdx);
-			OutResult.ExpressionDdy = Tree.NewNeg(InputDerivatives.ExpressionDdy);
-			break;
-		case EUnaryOp::Rcp:
-		{
-			FExpression* Result = Tree.NewRcp(Input);
-			FExpression* dFdA = Tree.NewNeg(Tree.NewMul(Result, Result));
-			OutResult.ExpressionDdx = Tree.NewMul(dFdA, InputDerivatives.ExpressionDdx);
-			OutResult.ExpressionDdy = Tree.NewMul(dFdA, InputDerivatives.ExpressionDdy);
-			break;
-		}
-		case EUnaryOp::Frac:
-			OutResult = InputDerivatives;
-			break;
-		case EUnaryOp::Length:
-		case EUnaryOp::Normalize:
-			// TODO
-			break;
-		default:
-			checkNoEntry();
-			break;
-		}
-	}
+	FExpression* Inputs[2] = { Lhs, Rhs };
+	return NewExpression<FExpressionOperation>(Op, Inputs);
 }
 
-FExpression* FExpressionUnaryOp::ComputePreviousFrame(FTree& Tree, const FRequestedType& RequestedType) const
+FExpressionOperation::FExpressionOperation(EOperation InOp, TConstArrayView<FExpression*> InInputs) : Op(InOp)
 {
-	return Tree.NewUnaryOp(Op, Tree.GetPreviousFrame(Input, RequestedType));
+	const FOperationDescription OpDesc = GetOperationDescription(InOp);
+	check(OpDesc.NumInputs == InInputs.Num());
+	for (int32 i = 0; i < OpDesc.NumInputs; ++i)
+	{
+		Inputs[i] = InInputs[i];
+		check(Inputs[i]);
+	}
 }
 
-bool FExpressionUnaryOp::PrepareValue(FEmitContext& Context, FEmitScope& Scope, const FRequestedType& RequestedType, FPrepareValueResult& OutResult) const
+void FExpressionOperation::ComputeAnalyticDerivatives(FTree& Tree, FExpressionDerivatives& OutResult) const
 {
-	FRequestedType InputRequestedType;
+	// Operations with constant derivatives
 	switch (Op)
 	{
-	case EUnaryOp::Length:
-	case EUnaryOp::Normalize:
-		// Each component of result is influenced by all components of input
-		InputRequestedType = ERequestedType::Vector4;
+	case EOperation::Less:
+	case EOperation::Greater:
+		OutResult.ExpressionDdx = Tree.NewConstant(0.0f);
+		OutResult.ExpressionDdx = OutResult.ExpressionDdy;
 		break;
 	default:
-		InputRequestedType = RequestedType;
 		break;
 	}
 
-	const FPreparedType& InputType = Context.PrepareExpression(Input, Scope, InputRequestedType);
-	if (InputType.IsVoid())
-	{
-		return false;
-	}
-
-	if (!InputType.IsNumeric())
+	if (OutResult.IsValid())
 	{
-		return Context.Errors->AddError(TEXT("Invalid operation on non-numeric type"));
+		return;
 	}
 
-	FPreparedType ResultType;
-	switch (Op)
+	const FOperationDescription OpDesc = GetOperationDescription(Op);
+	FExpressionDerivatives InputDerivatives[MaxInputs];
+	for (int32 Index = 0; Index < OpDesc.NumInputs; ++Index)
 	{
-	case EUnaryOp::Length:
-		ResultType = FPreparedType(Shader::MakeValueType(InputType.ValueComponentType, 1), InputType.GetEvaluation(Scope));
-		break;
-	default:
-		ResultType = InputType;
-		if (Op == EUnaryOp::Normalize)
+		InputDerivatives[Index] = Tree.GetAnalyticDerivatives(Inputs[Index]);
+		if (!InputDerivatives[Index].IsValid())
 		{
-			ResultType.ValueComponentType = Shader::MakeNonLWCType(ResultType.ValueComponentType);
+			return;
 		}
-		break;
 	}
 
-	const FUnaryOpDescription OpDesc = GetUnaryOpDesription(Op);
-	if (OpDesc.PreshaderOpcode == Shader::EPreshaderOpcode::Nop)
-	{
-		// No preshader support
-		ResultType.SetEvaluation(EExpressionEvaluation::Shader);
-	}
-
-	return OutResult.SetType(Context, RequestedType, ResultType);
-}
-
-namespace Private
-{
-struct FUnaryOpTypes
-{
-	Shader::EValueType InputType;
-	Shader::EValueType ResultType;
-	bool bIsLWC;
-};
-FUnaryOpTypes GetUnaryyOpTypes(EUnaryOp Op, Shader::EValueType InputType)
-{
-	const Shader::FValueTypeDescription InputTypeDesc = Shader::GetValueTypeDescription(InputType);
-
-	FUnaryOpTypes Types;
-	Types.InputType = InputType;
-	Types.ResultType = InputType;
-	Types.bIsLWC = InputTypeDesc.ComponentType == Shader::EValueComponentType::Double;
 	switch (Op)
 	{
-	case EUnaryOp::Length:
-		Types.ResultType = Shader::MakeValueType(InputTypeDesc.ComponentType, 1);
-		break;
-	case EUnaryOp::Normalize:
-	case EUnaryOp::Frac:
-	case EUnaryOp::Rcp:
-		Types.ResultType = Shader::MakeNonLWCType(InputType);
+	case EOperation::Neg:
+		OutResult.ExpressionDdx = Tree.NewNeg(InputDerivatives[0].ExpressionDdx);
+		OutResult.ExpressionDdy = Tree.NewNeg(InputDerivatives[0].ExpressionDdy);
 		break;
-	default:
+	case EOperation::Rcp:
+	{
+		FExpression* Result = Tree.NewRcp(Inputs[0]);
+		FExpression* dFdA = Tree.NewNeg(Tree.NewMul(Result, Result));
+		OutResult.ExpressionDdx = Tree.NewMul(dFdA, InputDerivatives[0].ExpressionDdx);
+		OutResult.ExpressionDdy = Tree.NewMul(dFdA, InputDerivatives[0].ExpressionDdy);
 		break;
 	}
-	return Types;
-}
-} // namespace Private
-
-void FExpressionUnaryOp::EmitValueShader(FEmitContext& Context, FEmitScope& Scope, const FRequestedType& RequestedType, FEmitValueShaderResult& OutResult) const
-{
-	const Private::FUnaryOpTypes Types = Private::GetUnaryyOpTypes(Op, Input->GetType());
-	FEmitShaderExpression* InputValue = Input->GetValueShader(Context, Scope, Types.InputType);
-
-	switch (Op)
-	{
-	case EUnaryOp::Neg:
-		if (Types.bIsLWC)
-		{
-			OutResult.Code = Context.EmitExpression(Scope, Types.ResultType, TEXT("LWCNegate(%)"), InputValue);
-		}
-		else
-		{
-			OutResult.Code = Context.EmitInlineExpression(Scope, Types.ResultType, TEXT("(-%)"), InputValue);
-		}
+	case EOperation::Frac:
+		OutResult = InputDerivatives[0];
 		break;
-	case EUnaryOp::Rcp:
-		if (Types.bIsLWC)
-		{
-			OutResult.Code = Context.EmitExpression(Scope, Types.ResultType, TEXT("LWCRcp(%)"), InputValue);
-		}
-		else
-		{
-			OutResult.Code = Context.EmitExpression(Scope, Types.ResultType, TEXT("rcp(%)"), InputValue);
-		}
+	case EOperation::Length:
+	case EOperation::Normalize:
+		// TODO
 		break;
-	case EUnaryOp::Frac:
-		if (Types.bIsLWC)
-		{
-			OutResult.Code = Context.EmitExpression(Scope, Types.ResultType, TEXT("LWCFrac(%)"), InputValue);
-		}
-		else
-		{
-			OutResult.Code = Context.EmitExpression(Scope, Types.ResultType, TEXT("frac(%)"), InputValue);
-		}
+	case EOperation::Add:
+		OutResult.ExpressionDdx = Tree.NewAdd(InputDerivatives[0].ExpressionDdx, InputDerivatives[1].ExpressionDdx);
+		OutResult.ExpressionDdy = Tree.NewAdd(InputDerivatives[0].ExpressionDdy, InputDerivatives[1].ExpressionDdy);
 		break;
-	case EUnaryOp::Length:
-		if (Types.bIsLWC)
-		{
-			OutResult.Code = Context.EmitExpression(Scope, Types.ResultType, TEXT("LWCLength(%)"), InputValue);
-		}
-		else
-		{
-			OutResult.Code = Context.EmitExpression(Scope, Types.ResultType, TEXT("length(%)"), InputValue);
-		}
+	case EOperation::Sub:
+		OutResult.ExpressionDdx = Tree.NewSub(InputDerivatives[0].ExpressionDdx, InputDerivatives[1].ExpressionDdx);
+		OutResult.ExpressionDdy = Tree.NewSub(InputDerivatives[0].ExpressionDdy, InputDerivatives[1].ExpressionDdy);
 		break;
-	case EUnaryOp::Normalize:
-		if (Types.bIsLWC)
-		{
-			OutResult.Code = Context.EmitExpression(Scope, Types.ResultType, TEXT("LWCNormalize(%)"), InputValue);
-		}
-		else
-		{
-			OutResult.Code = Context.EmitExpression(Scope, Types.ResultType, TEXT("normalize(%)"), InputValue);
-		}
+	case EOperation::Mul:
+		OutResult.ExpressionDdx = Tree.NewAdd(Tree.NewMul(InputDerivatives[0].ExpressionDdx, Inputs[1]), Tree.NewMul(InputDerivatives[1].ExpressionDdx, Inputs[0]));
+		OutResult.ExpressionDdy = Tree.NewAdd(Tree.NewMul(InputDerivatives[0].ExpressionDdy, Inputs[1]), Tree.NewMul(InputDerivatives[1].ExpressionDdy, Inputs[0]));
 		break;
-	default:
-		checkNoEntry();
+	case EOperation::Div:
+	{
+		FExpression* Denom = Tree.NewRcp(Tree.NewMul(Inputs[1], Inputs[1]));
+		FExpression* dFdA = Tree.NewMul(Inputs[1], Denom);
+		FExpression* dFdB = Tree.NewNeg(Tree.NewMul(Inputs[0], Denom));
+		OutResult.ExpressionDdx = Tree.NewAdd(Tree.NewMul(dFdA, InputDerivatives[0].ExpressionDdx), Tree.NewMul(dFdB, InputDerivatives[1].ExpressionDdx));
+		OutResult.ExpressionDdy = Tree.NewAdd(Tree.NewMul(dFdA, InputDerivatives[0].ExpressionDdy), Tree.NewMul(dFdB, InputDerivatives[1].ExpressionDdy));
 		break;
 	}
-}
-
-void FExpressionUnaryOp::EmitValuePreshader(FEmitContext& Context, FEmitScope& Scope, const FRequestedType& RequestedType, FEmitValuePreshaderResult& OutResult) const
-{
-	const Private::FUnaryOpTypes Types = Private::GetUnaryyOpTypes(Op, Input->GetType());
-	const FUnaryOpDescription OpDesc = GetUnaryOpDesription(Op);
-	check(OpDesc.PreshaderOpcode != Shader::EPreshaderOpcode::Nop);
-
-	Input->GetValuePreshader(Context, Scope, Types.InputType, OutResult.Preshader);
-	OutResult.Preshader.WriteOpcode(OpDesc.PreshaderOpcode);
-	OutResult.Type = Types.ResultType;
-}
-
-FExpression* FTree::NewBinaryOp(EBinaryOp Op, FExpression* Lhs, FExpression* Rhs)
-{
-	return NewExpression<FExpressionBinaryOp>(Op, Lhs, Rhs);
-}
-
-void FExpressionBinaryOp::ComputeAnalyticDerivatives(FTree& Tree, FExpressionDerivatives& OutResult) const
-{
-	// Operations with constant derivatives
-	switch (Op)
-	{
-	case EBinaryOp::Less:
-	case EBinaryOp::Greater:
-		OutResult.ExpressionDdx = Tree.NewConstant(0.0f);
-		OutResult.ExpressionDdx = OutResult.ExpressionDdy;
+	case EOperation::Fmod:
+		// Only valid when B derivatives are zero.
+		// We can't really do anything meaningful in the non-zero case.
+		OutResult = InputDerivatives[0];
 		break;
-	default:
+	case EOperation::Dot:
+	{
+		// Dot means multiply the values, then sum the resulting components
+		FExpression* MulDdx = Tree.NewAdd(Tree.NewMul(InputDerivatives[0].ExpressionDdx, Inputs[1]), Tree.NewMul(InputDerivatives[1].ExpressionDdx, Inputs[0]));
+		FExpression* MulDdy = Tree.NewAdd(Tree.NewMul(InputDerivatives[0].ExpressionDdy, Inputs[1]), Tree.NewMul(InputDerivatives[1].ExpressionDdy, Inputs[0]));
+		// Dot the products with 1 to sum them
+		FExpression* Const1 = Tree.NewConstant(FVector4f(1.0f, 1.0f, 1.0f, 1.0f));
+		OutResult.ExpressionDdx = Tree.NewDot(MulDdx, Const1);
+		OutResult.ExpressionDdy = Tree.NewDot(MulDdy, Const1);
 		break;
 	}
-
-	if (OutResult.IsValid())
+	case EOperation::Min:
 	{
-		return;
+		FExpression* Cond = Tree.NewLess(Inputs[0], Inputs[1]);
+		OutResult.ExpressionDdx = Tree.NewExpression<FExpressionSelect>(Cond, InputDerivatives[0].ExpressionDdx, InputDerivatives[1].ExpressionDdx);
+		OutResult.ExpressionDdy = Tree.NewExpression<FExpressionSelect>(Cond, InputDerivatives[0].ExpressionDdy, InputDerivatives[1].ExpressionDdy);
+		break;
 	}
-
-	const FExpressionDerivatives LhsDerivatives = Tree.GetAnalyticDerivatives(Lhs);
-	const FExpressionDerivatives RhsDerivatives = Tree.GetAnalyticDerivatives(Rhs);
-	if (LhsDerivatives.IsValid() && RhsDerivatives.IsValid())
+	case EOperation::Max:
 	{
-		switch (Op)
-		{
-		case EBinaryOp::Add:
-			OutResult.ExpressionDdx = Tree.NewAdd(LhsDerivatives.ExpressionDdx, RhsDerivatives.ExpressionDdx);
-			OutResult.ExpressionDdy = Tree.NewAdd(LhsDerivatives.ExpressionDdy, RhsDerivatives.ExpressionDdy);
-			break;
-		case EBinaryOp::Sub:
-			OutResult.ExpressionDdx = Tree.NewSub(LhsDerivatives.ExpressionDdx, RhsDerivatives.ExpressionDdx);
-			OutResult.ExpressionDdy = Tree.NewSub(LhsDerivatives.ExpressionDdy, RhsDerivatives.ExpressionDdy);
-			break;
-		case EBinaryOp::Mul:
-			OutResult.ExpressionDdx = Tree.NewAdd(Tree.NewMul(LhsDerivatives.ExpressionDdx, Rhs), Tree.NewMul(RhsDerivatives.ExpressionDdx, Lhs));
-			OutResult.ExpressionDdy = Tree.NewAdd(Tree.NewMul(LhsDerivatives.ExpressionDdy, Rhs), Tree.NewMul(RhsDerivatives.ExpressionDdy, Lhs));
-			break;
-		case EBinaryOp::Div:
-		{
-			FExpression* Denom = Tree.NewRcp(Tree.NewMul(Rhs, Rhs));
-			FExpression* dFdA = Tree.NewMul(Rhs, Denom);
-			FExpression* dFdB = Tree.NewNeg(Tree.NewMul(Lhs, Denom));
-			OutResult.ExpressionDdx = Tree.NewAdd(Tree.NewMul(dFdA, LhsDerivatives.ExpressionDdx), Tree.NewMul(dFdB, RhsDerivatives.ExpressionDdx));
-			OutResult.ExpressionDdy = Tree.NewAdd(Tree.NewMul(dFdA, LhsDerivatives.ExpressionDdy), Tree.NewMul(dFdB, RhsDerivatives.ExpressionDdy));
-			break;
-		}
-		case EBinaryOp::Fmod:
-			// Only valid when B derivatives are zero.
-			// We can't really do anything meaningful in the non-zero case.
-			OutResult = LhsDerivatives;
-			break;
-		case EBinaryOp::Dot:
-		{
-			// Dot means multiply the values, then sum the resulting components
-			FExpression* MulDdx = Tree.NewAdd(Tree.NewMul(LhsDerivatives.ExpressionDdx, Rhs), Tree.NewMul(RhsDerivatives.ExpressionDdx, Lhs));
-			FExpression* MulDdy = Tree.NewAdd(Tree.NewMul(LhsDerivatives.ExpressionDdy, Rhs), Tree.NewMul(RhsDerivatives.ExpressionDdy, Lhs));
-			// Dot the products with 1 to sum them
-			FExpression* Const1 = Tree.NewConstant(1.0f);
-			OutResult.ExpressionDdx = Tree.NewDot(MulDdx, Const1);
-			OutResult.ExpressionDdy = Tree.NewDot(MulDdy, Const1);
-			break;
-		}
-		case EBinaryOp::Min:
-		{
-			FExpression* Cond = Tree.NewLess(Lhs, Rhs);
-			OutResult.ExpressionDdx = Tree.NewExpression<FExpressionSelect>(Cond, LhsDerivatives.ExpressionDdx, RhsDerivatives.ExpressionDdx);
-			OutResult.ExpressionDdy = Tree.NewExpression<FExpressionSelect>(Cond, LhsDerivatives.ExpressionDdy, RhsDerivatives.ExpressionDdy);
-			break;
-		}
-		case EBinaryOp::Max:
-		{
-			FExpression* Cond = Tree.NewGreater(Lhs, Rhs);
-			OutResult.ExpressionDdx = Tree.NewExpression<FExpressionSelect>(Cond, LhsDerivatives.ExpressionDdx, RhsDerivatives.ExpressionDdx);
-			OutResult.ExpressionDdy = Tree.NewExpression<FExpressionSelect>(Cond, LhsDerivatives.ExpressionDdy, RhsDerivatives.ExpressionDdy);
-			break;
-		}
-		case EBinaryOp::VecMulMatrix3:
-		case EBinaryOp::VecMulMatrix4:
-		case EBinaryOp::Matrix3MulVec:
-		case EBinaryOp::Matrix4MulVec:
-			// TODO
-			break;
-		default:
-			checkNoEntry();
-			break;
-		}
+		FExpression* Cond = Tree.NewGreater(Inputs[0], Inputs[1]);
+		OutResult.ExpressionDdx = Tree.NewExpression<FExpressionSelect>(Cond, InputDerivatives[0].ExpressionDdx, InputDerivatives[1].ExpressionDdx);
+		OutResult.ExpressionDdy = Tree.NewExpression<FExpressionSelect>(Cond, InputDerivatives[0].ExpressionDdy, InputDerivatives[1].ExpressionDdy);
+		break;
+	}
+	case EOperation::VecMulMatrix3:
+	case EOperation::VecMulMatrix4:
+	case EOperation::Matrix3MulVec:
+	case EOperation::Matrix4MulVec:
+		// TODO
+		break;
+	default:
+		checkNoEntry();
+		break;
 	}
 }
 
-FExpression* FExpressionBinaryOp::ComputePreviousFrame(FTree& Tree, const FRequestedType& RequestedType) const
+FExpression* FExpressionOperation::ComputePreviousFrame(FTree& Tree, const FRequestedType& RequestedType) const
 {
-	return Tree.NewBinaryOp(Op, Tree.GetPreviousFrame(Lhs, RequestedType), Tree.GetPreviousFrame(Rhs, RequestedType));
+	const FOperationDescription OpDesc = GetOperationDescription(Op);
+	FExpression* PrevFrameInputs[MaxInputs];
+	for (int32 Index = 0; Index < OpDesc.NumInputs; ++Index)
+	{
+		PrevFrameInputs[Index] = Tree.GetPreviousFrame(Inputs[Index], RequestedType);
+	}
+
+	return Tree.NewExpression<FExpressionOperation>(Op, MakeArrayView(PrevFrameInputs, OpDesc.NumInputs));
 }
 
-bool FExpressionBinaryOp::PrepareValue(FEmitContext& Context, FEmitScope& Scope, const FRequestedType& RequestedType, FPrepareValueResult& OutResult) const
+bool FExpressionOperation::PrepareValue(FEmitContext& Context, FEmitScope& Scope, const FRequestedType& RequestedType, FPrepareValueResult& OutResult) const
 {
-	FRequestedType LhsRequestedType;
-	FRequestedType RhsRequestedType;
+	const FOperationDescription OpDesc = GetOperationDescription(Op);
+	FRequestedType InputRequestedType[MaxInputs];
 	bool bIsMatrixOperation = false;
 	switch (Op)
 	{
-	case EBinaryOp::VecMulMatrix3:
-	case EBinaryOp::VecMulMatrix4:
-		LhsRequestedType = ERequestedType::Vector3;
-		RhsRequestedType = ERequestedType::Matrix4x4;
+	case EOperation::Length:
+	case EOperation::Normalize:
+		// Each component of result is influenced by all components of input
+		InputRequestedType[0] = ERequestedType::Vector4;
+		break;
+	case EOperation::VecMulMatrix3:
+	case EOperation::VecMulMatrix4:
+		InputRequestedType[0] = ERequestedType::Vector3;
+		InputRequestedType[1] = ERequestedType::Matrix4x4;
 		bIsMatrixOperation = true;
 		break;
-	case EBinaryOp::Matrix3MulVec:
-	case EBinaryOp::Matrix4MulVec:
-		LhsRequestedType = ERequestedType::Matrix4x4;
-		RhsRequestedType = ERequestedType::Vector3;
+	case EOperation::Matrix3MulVec:
+	case EOperation::Matrix4MulVec:
+		InputRequestedType[0] = ERequestedType::Matrix4x4;
+		InputRequestedType[1] = ERequestedType::Vector3;
 		bIsMatrixOperation = true;
 		break;
-	case EBinaryOp::Dot:
-		LhsRequestedType = ERequestedType::Vector4;
-		RhsRequestedType = ERequestedType::Vector4;
+	case EOperation::Dot:
+		InputRequestedType[0] = ERequestedType::Vector4;
+		InputRequestedType[1] = ERequestedType::Vector4;
 		break;
 	default:
-		LhsRequestedType = RequestedType;
-		RhsRequestedType = RequestedType;
+		for (int32 Index = 0; Index < OpDesc.NumInputs; ++Index)
+		{
+			InputRequestedType[Index] = RequestedType;
+		}
 		break;
 	}
 
-	const FPreparedType& LhsType = Context.PrepareExpression(Lhs, Scope, LhsRequestedType);
-	const FPreparedType& RhsType = Context.PrepareExpression(Rhs, Scope, RhsRequestedType);
-	if (LhsType.IsVoid() || RhsType.IsVoid())
+	FPreparedType InputType[MaxInputs];
+	for (int32 Index = 0; Index < OpDesc.NumInputs; ++Index)
 	{
-		return false;
-	}
+		InputType[Index] = Context.PrepareExpression(Inputs[Index], Scope, InputRequestedType[Index]);
+		if (InputType[Index].IsVoid())
+		{
+			return false;
+		}
 
-	if (!LhsType.IsNumeric() || !RhsType.IsNumeric())
-	{
-		return Context.Errors->AddError(TEXT("Invalid arithmetic between non-numeric types"));
+		if (!InputType[Index].IsNumeric())
+		{
+			return Context.Errors->AddError(TEXT("Invalid arithmetic between non-numeric types"));
+		}
 	}
 
-	const FBinaryOpDescription OpDesc = GetBinaryOpDesription(Op);
 	FPreparedType ResultType;
 	if (bIsMatrixOperation)
 	{
-		const EExpressionEvaluation Evaluation = CombineEvaluations(LhsType.GetEvaluation(Scope), RhsType.GetEvaluation(Scope));
-		const Shader::EValueComponentType ComponentType = Shader::CombineComponentTypes(LhsType.ValueComponentType, RhsType.ValueComponentType);
+		const EExpressionEvaluation Evaluation = CombineEvaluations(InputType[0].GetEvaluation(Scope), InputType[1].GetEvaluation(Scope));
+		const Shader::EValueComponentType ComponentType = Shader::CombineComponentTypes(InputType[0].ValueComponentType, InputType[1].ValueComponentType);
 		ResultType = FPreparedType(Shader::MakeValueType(ComponentType, 3), Evaluation);
 	}
-	else
+	else if (OpDesc.NumInputs == 1)
+	{
+		switch (Op)
+		{
+		case EOperation::Length:
+			ResultType = FPreparedType(Shader::MakeValueType(InputType[0].ValueComponentType, 1), InputType[0].GetEvaluation(Scope));
+			break;
+		default:
+			ResultType = InputType[0];
+			if (Op == EOperation::Normalize)
+			{
+				ResultType.ValueComponentType = Shader::MakeNonLWCType(ResultType.ValueComponentType);
+			}
+			break;
+		}
+	}
+	else if(OpDesc.NumInputs == 2)
 	{
-		if (LhsType.GetNumComponents() != 1 && RhsType.GetNumComponents() != 1 && LhsType.GetNumComponents() != RhsType.GetNumComponents())
+		if (InputType[0].GetNumComponents() != 1 && InputType[1].GetNumComponents() != 1 && InputType[0].GetNumComponents() != InputType[1].GetNumComponents())
 		{
-			return Context.Errors->AddErrorf(TEXT("Invalid arithmetic between %s and %s"), LhsType.GetType().GetName(), RhsType.GetType().GetName());
+			return Context.Errors->AddErrorf(TEXT("Invalid arithmetic between %s and %s"), InputType[0].GetType().GetName(), InputType[1].GetType().GetName());
 		}
-		ResultType = MergePreparedTypes(LhsType, RhsType);
-		if (Op == EBinaryOp::Less || Op == EBinaryOp::Greater)
+		ResultType = MergePreparedTypes(InputType[0], InputType[1]);
+		if (Op == EOperation::Less || Op == EOperation::Greater)
 		{
 			ResultType.ValueComponentType = Shader::EValueComponentType::Bool;
 		}
-		if (Op == EBinaryOp::Min || Op == EBinaryOp::Max)
+		if (Op == EOperation::Min || Op == EOperation::Max)
 		{
-			const Shader::FComponentBounds LhsBounds = LhsType.GetBounds(RequestedType);
-			const Shader::FComponentBounds RhsBounds = RhsType.GetBounds(RequestedType);
-			const Shader::FComponentBounds Bounds = (Op == EBinaryOp::Min) ? Shader::MinBound(LhsBounds, RhsBounds) : Shader::MaxBound(LhsBounds, RhsBounds);
+			const Shader::FComponentBounds LhsBounds = InputType[0].GetBounds(RequestedType);
+			const Shader::FComponentBounds RhsBounds = InputType[1].GetBounds(RequestedType);
+			const Shader::FComponentBounds Bounds = (Op == EOperation::Min) ? Shader::MinBound(LhsBounds, RhsBounds) : Shader::MaxBound(LhsBounds, RhsBounds);
 			ResultType.UpdateBounds(RequestedType, Bounds);
 		}
 	}
@@ -1112,55 +986,68 @@ bool FExpressionBinaryOp::PrepareValue(FEmitContext& Context, FEmitScope& Scope,
 
 namespace Private
 {
-struct FBinaryOpTypes
+struct FOperationTypes
 {
-	Shader::EValueType LhsType;
-	Shader::EValueType RhsType;
+	Shader::EValueType InputType[FExpressionOperation::MaxInputs];
 	Shader::EValueType ResultType;
 	bool bIsLWC;
 };
-FBinaryOpTypes GetBinaryOpTypes(EBinaryOp Op, Shader::EValueType LhsType, Shader::EValueType RhsType)
+FOperationTypes GetOperationTypes(EOperation Op, TConstArrayView<Shader::EValueType> InputTypes)
 {
-	const Shader::FValueTypeDescription LhsTypeDesc = Shader::GetValueTypeDescription(LhsType);
-	const Shader::FValueTypeDescription RhsTypeDesc = Shader::GetValueTypeDescription(RhsType);
+	Shader::EValueComponentType ComponentType = Shader::EValueComponentType::Void;
+	int32 NumComponents = 0;
+	for (int32 Index = 0; Index < InputTypes.Num(); ++Index)
+	{
+		const Shader::FValueTypeDescription InputTypeDesc = Shader::GetValueTypeDescription(InputTypes[Index]);
+		ComponentType = Shader::CombineComponentTypes(ComponentType, InputTypeDesc.ComponentType);
+		NumComponents = FMath::Max<int32>(NumComponents, InputTypeDesc.NumComponents);
+	}
 
-	Shader::EValueComponentType ComponentType = Shader::CombineComponentTypes(LhsTypeDesc.ComponentType, RhsTypeDesc.ComponentType);
-	int32 NumComponents = FMath::Max(LhsTypeDesc.NumComponents, RhsTypeDesc.NumComponents);
 	Shader::EValueType IntermediateType = Shader::MakeValueType(ComponentType, NumComponents);
 
-	FBinaryOpTypes Types;
-	Types.LhsType = IntermediateType;
-	Types.RhsType = IntermediateType;
+	FOperationTypes Types;
+	for (int32 Index = 0; Index < InputTypes.Num(); ++Index)
+	{
+		Types.InputType[Index] = IntermediateType;
+	}
 	Types.ResultType = IntermediateType;
 	Types.bIsLWC = (ComponentType == Shader::EValueComponentType::Double);
 	switch (Op)
 	{
-	case EBinaryOp::Less:
-	case EBinaryOp::Greater:
+	case EOperation::Length:
+		Types.ResultType = Shader::MakeValueType(ComponentType, 1);
+		break;
+	case EOperation::Normalize:
+	case EOperation::Frac:
+	case EOperation::Rcp:
+		Types.ResultType = Shader::MakeNonLWCType(IntermediateType);
+		break;
+	case EOperation::Less:
+	case EOperation::Greater:
 		Types.ResultType = Shader::MakeValueType(Shader::EValueComponentType::Bool, NumComponents);
 		break;
-	case EBinaryOp::Fmod:
-		Types.RhsType = Types.ResultType = Shader::MakeNonLWCType(IntermediateType);
+	case EOperation::Fmod:
+		Types.InputType[1] = Types.ResultType = Shader::MakeNonLWCType(IntermediateType);
 		break;
-	case EBinaryOp::Dot:
+	case EOperation::Dot:
 		Types.ResultType = Shader::MakeValueType(ComponentType, 1);
 		break;
-	case EBinaryOp::VecMulMatrix3:
+	case EOperation::VecMulMatrix3:
 		// No LWC for matrix3
-		Types.LhsType = Shader::EValueType::Float3;
-		Types.RhsType = Shader::EValueType::Float4x4;
+		Types.InputType[0] = Shader::EValueType::Float3;
+		Types.InputType[1] = Shader::EValueType::Float4x4;
 		Types.ResultType = Shader::EValueType::Float3;
 		break;
-	case EBinaryOp::VecMulMatrix4:
-		Types.LhsType = Shader::EValueType::Float3;
-		Types.RhsType = Shader::EValueType::Float4x4;
+	case EOperation::VecMulMatrix4:
+		Types.InputType[0] = Shader::EValueType::Float3;
+		Types.InputType[1] = Shader::EValueType::Float4x4;
 		Types.ResultType = Shader::EValueType::Float3;
 		break;
-	case EBinaryOp::Matrix3MulVec:
-	case EBinaryOp::Matrix4MulVec:
+	case EOperation::Matrix3MulVec:
+	case EOperation::Matrix4MulVec:
 		// No LWC for transpose matrices
-		Types.LhsType = Shader::EValueType::Float4x4;
-		Types.RhsType = Shader::EValueType::Float3;
+		Types.InputType[0] = Shader::EValueType::Float4x4;
+		Types.InputType[1] = Shader::EValueType::Float3;
 		Types.ResultType = Shader::EValueType::Float3;
 		break;
 	default:
@@ -1170,139 +1057,198 @@ FBinaryOpTypes GetBinaryOpTypes(EBinaryOp Op, Shader::EValueType LhsType, Shader
 }
 } // namespace Private
 
-void FExpressionBinaryOp::EmitValueShader(FEmitContext& Context, FEmitScope& Scope, const FRequestedType& RequestedType, FEmitValueShaderResult& OutResult) const
+void FExpressionOperation::EmitValueShader(FEmitContext& Context, FEmitScope& Scope, const FRequestedType& RequestedType, FEmitValueShaderResult& OutResult) const
 {
-	const Private::FBinaryOpTypes Types = Private::GetBinaryOpTypes(Op, Lhs->GetType(), Rhs->GetType());
-	FEmitShaderExpression* LhsValue = Lhs->GetValueShader(Context, Scope, Types.LhsType);
-	FEmitShaderExpression* RhsValue = Rhs->GetValueShader(Context, Scope, Types.RhsType);
+	const FOperationDescription OpDesc = GetOperationDescription(Op);
+	Shader::EValueType InputTypes[MaxInputs];
+	for (int32 Index = 0; Index < OpDesc.NumInputs; ++Index)
+	{
+		InputTypes[Index] = Inputs[Index]->GetType();
+	}
+	const Private::FOperationTypes Types = Private::GetOperationTypes(Op, MakeArrayView(InputTypes, OpDesc.NumInputs));
+	FEmitShaderExpression* InputValue[MaxInputs] = { nullptr };
+	for (int32 Index = 0; Index < OpDesc.NumInputs; ++Index)
+	{
+		InputValue[Index] = Inputs[Index]->GetValueShader(Context, Scope, Types.InputType[Index]);
+	}
 
 	switch (Op)
 	{
-	case EBinaryOp::Add:
+	case EOperation::Neg:
 		if (Types.bIsLWC)
 		{
-			OutResult.Code = Context.EmitExpression(Scope, Types.ResultType, TEXT("LWCAdd(%, %)"), LhsValue, RhsValue);
+			OutResult.Code = Context.EmitExpression(Scope, Types.ResultType, TEXT("LWCNegate(%)"), InputValue[0]);
 		}
 		else
 		{
-			OutResult.Code = Context.EmitExpression(Scope, Types.ResultType, TEXT("(% + %)"), LhsValue, RhsValue);
+			OutResult.Code = Context.EmitInlineExpression(Scope, Types.ResultType, TEXT("(-%)"), InputValue[0]);
 		}
 		break;
-	case EBinaryOp::Sub:
+	case EOperation::Rcp:
 		if (Types.bIsLWC)
 		{
-			OutResult.Code = Context.EmitExpression(Scope, Types.ResultType, TEXT("LWCSubtract(%, %)"), LhsValue, RhsValue);
+			OutResult.Code = Context.EmitExpression(Scope, Types.ResultType, TEXT("LWCRcp(%)"), InputValue[0]);
 		}
 		else
 		{
-			OutResult.Code = Context.EmitExpression(Scope, Types.ResultType, TEXT("(% - %)"), LhsValue, RhsValue);
+			OutResult.Code = Context.EmitExpression(Scope, Types.ResultType, TEXT("rcp(%)"), InputValue[0]);
 		}
 		break;
-	case EBinaryOp::Mul:
+	case EOperation::Frac:
 		if (Types.bIsLWC)
 		{
-			OutResult.Code = Context.EmitExpression(Scope, Types.ResultType, TEXT("LWCMultiply(%, %)"), LhsValue, RhsValue);
+			OutResult.Code = Context.EmitExpression(Scope, Types.ResultType, TEXT("LWCFrac(%)"), InputValue[0]);
 		}
 		else
 		{
-			OutResult.Code = Context.EmitExpression(Scope, Types.ResultType, TEXT("(% * %)"), LhsValue, RhsValue);
+			OutResult.Code = Context.EmitExpression(Scope, Types.ResultType, TEXT("frac(%)"), InputValue[0]);
 		}
 		break;
-	case EBinaryOp::Div:
+	case EOperation::Length:
 		if (Types.bIsLWC)
 		{
-			OutResult.Code = Context.EmitExpression(Scope, Types.ResultType, TEXT("LWCDivide(%, %)"), LhsValue, RhsValue);
+			OutResult.Code = Context.EmitExpression(Scope, Types.ResultType, TEXT("LWCLength(%)"), InputValue[0]);
 		}
 		else
 		{
-			OutResult.Code = Context.EmitExpression(Scope, Types.ResultType, TEXT("(% / %)"), LhsValue, RhsValue);
+			OutResult.Code = Context.EmitExpression(Scope, Types.ResultType, TEXT("length(%)"), InputValue[0]);
 		}
 		break;
-	case EBinaryOp::Fmod:
+	case EOperation::Normalize:
 		if (Types.bIsLWC)
 		{
-			OutResult.Code = Context.EmitExpression(Scope, Types.ResultType, TEXT("LWCFmod(%, %)"), LhsValue, RhsValue);
+			OutResult.Code = Context.EmitExpression(Scope, Types.ResultType, TEXT("LWCNormalize(%)"), InputValue[0]);
 		}
 		else
 		{
-			OutResult.Code = Context.EmitExpression(Scope, Types.ResultType, TEXT("fmod(%, %)"), LhsValue, RhsValue);
+			OutResult.Code = Context.EmitExpression(Scope, Types.ResultType, TEXT("normalize(%)"), InputValue[0]);
 		}
 		break;
-	case EBinaryOp::Dot:
+	case EOperation::Add:
 		if (Types.bIsLWC)
 		{
-			OutResult.Code = Context.EmitExpression(Scope, Types.ResultType, TEXT("LWCDot(%, %)"), LhsValue, RhsValue);
+			OutResult.Code = Context.EmitExpression(Scope, Types.ResultType, TEXT("LWCAdd(%, %)"), InputValue[0], InputValue[1]);
 		}
 		else
 		{
-			OutResult.Code = Context.EmitExpression(Scope, Types.ResultType, TEXT("dot(%, %)"), LhsValue, RhsValue);
+			OutResult.Code = Context.EmitExpression(Scope, Types.ResultType, TEXT("(% + %)"), InputValue[0], InputValue[1]);
 		}
 		break;
-	case EBinaryOp::Min:
+	case EOperation::Sub:
 		if (Types.bIsLWC)
 		{
-			OutResult.Code = Context.EmitExpression(Scope, Types.ResultType, TEXT("LWCMin(%, %)"), LhsValue, RhsValue);
+			OutResult.Code = Context.EmitExpression(Scope, Types.ResultType, TEXT("LWCSubtract(%, %)"), InputValue[0], InputValue[1]);
 		}
 		else
 		{
-			OutResult.Code = Context.EmitExpression(Scope, Types.ResultType, TEXT("min(%, %)"), LhsValue, RhsValue);
+			OutResult.Code = Context.EmitExpression(Scope, Types.ResultType, TEXT("(% - %)"), InputValue[0], InputValue[1]);
 		}
 		break;
-	case EBinaryOp::Max:
+	case EOperation::Mul:
 		if (Types.bIsLWC)
 		{
-			OutResult.Code = Context.EmitExpression(Scope, Types.ResultType, TEXT("LWCMax(%, %)"), LhsValue, RhsValue);
+			OutResult.Code = Context.EmitExpression(Scope, Types.ResultType, TEXT("LWCMultiply(%, %)"), InputValue[0], InputValue[1]);
 		}
 		else
 		{
-			OutResult.Code = Context.EmitExpression(Scope, Types.ResultType, TEXT("max(%, %)"), LhsValue, RhsValue);
+			OutResult.Code = Context.EmitExpression(Scope, Types.ResultType, TEXT("(% * %)"), InputValue[0], InputValue[1]);
 		}
 		break;
-	case EBinaryOp::Less:
+	case EOperation::Div:
 		if (Types.bIsLWC)
 		{
-			OutResult.Code = Context.EmitExpression(Scope, Types.ResultType, TEXT("LWCLess(%, %)"), LhsValue, RhsValue);
+			OutResult.Code = Context.EmitExpression(Scope, Types.ResultType, TEXT("LWCDivide(%, %)"), InputValue[0], InputValue[1]);
 		}
 		else
 		{
-			OutResult.Code = Context.EmitExpression(Scope, Types.ResultType, TEXT("(% < %)"), LhsValue, RhsValue);
+			OutResult.Code = Context.EmitExpression(Scope, Types.ResultType, TEXT("(% / %)"), InputValue[0], InputValue[1]);
 		}
 		break;
-	case EBinaryOp::Greater:
+	case EOperation::Fmod:
 		if (Types.bIsLWC)
 		{
-			OutResult.Code = Context.EmitExpression(Scope, Types.ResultType, TEXT("LWCGreater(%, %)"), LhsValue, RhsValue);
+			OutResult.Code = Context.EmitExpression(Scope, Types.ResultType, TEXT("LWCFmod(%, %)"), InputValue[0], InputValue[1]);
 		}
 		else
 		{
-			OutResult.Code = Context.EmitExpression(Scope, Types.ResultType, TEXT("(% > %)"), LhsValue, RhsValue);
+			OutResult.Code = Context.EmitExpression(Scope, Types.ResultType, TEXT("fmod(%, %)"), InputValue[0], InputValue[1]);
 		}
 		break;
-	case EBinaryOp::VecMulMatrix3:
+	case EOperation::Dot:
 		if (Types.bIsLWC)
 		{
-			OutResult.Code = Context.EmitExpression(Scope, Types.ResultType, TEXT("LWCMlutiplyVector(%, %)"), LhsValue, RhsValue);
+			OutResult.Code = Context.EmitExpression(Scope, Types.ResultType, TEXT("LWCDot(%, %)"), InputValue[0], InputValue[1]);
 		}
 		else
 		{
-			OutResult.Code = Context.EmitExpression(Scope, Types.ResultType, TEXT("mul(%, (float3x3)%)"), LhsValue, RhsValue);
+			OutResult.Code = Context.EmitExpression(Scope, Types.ResultType, TEXT("dot(%, %)"), InputValue[0], InputValue[1]);
 		}
 		break;
-	case EBinaryOp::VecMulMatrix4:
+	case EOperation::Min:
 		if (Types.bIsLWC)
 		{
-			OutResult.Code = Context.EmitExpression(Scope, Types.ResultType, TEXT("LWCMlutiply(%, %)"), LhsValue, RhsValue);
+			OutResult.Code = Context.EmitExpression(Scope, Types.ResultType, TEXT("LWCMin(%, %)"), InputValue[0], InputValue[1]);
 		}
 		else
 		{
-			OutResult.Code = Context.EmitExpression(Scope, Types.ResultType, TEXT("mul(%, %)"), LhsValue, RhsValue);
+			OutResult.Code = Context.EmitExpression(Scope, Types.ResultType, TEXT("min(%, %)"), InputValue[0], InputValue[1]);
 		}
 		break;
-	case EBinaryOp::Matrix3MulVec:
-		OutResult.Code = Context.EmitExpression(Scope, Types.ResultType, TEXT("mul((float3x3)%, %)"), LhsValue, RhsValue);
+	case EOperation::Max:
+		if (Types.bIsLWC)
+		{
+			OutResult.Code = Context.EmitExpression(Scope, Types.ResultType, TEXT("LWCMax(%, %)"), InputValue[0], InputValue[1]);
+		}
+		else
+		{
+			OutResult.Code = Context.EmitExpression(Scope, Types.ResultType, TEXT("max(%, %)"), InputValue[0], InputValue[1]);
+		}
 		break;
-	case EBinaryOp::Matrix4MulVec:
-		OutResult.Code = Context.EmitExpression(Scope, Types.ResultType, TEXT("mul(%, %)"), LhsValue, RhsValue);
+	case EOperation::Less:
+		if (Types.bIsLWC)
+		{
+			OutResult.Code = Context.EmitExpression(Scope, Types.ResultType, TEXT("LWCLess(%, %)"), InputValue[0], InputValue[1]);
+		}
+		else
+		{
+			OutResult.Code = Context.EmitExpression(Scope, Types.ResultType, TEXT("(% < %)"), InputValue[0], InputValue[1]);
+		}
+		break;
+	case EOperation::Greater:
+		if (Types.bIsLWC)
+		{
+			OutResult.Code = Context.EmitExpression(Scope, Types.ResultType, TEXT("LWCGreater(%, %)"), InputValue[0], InputValue[1]);
+		}
+		else
+		{
+			OutResult.Code = Context.EmitExpression(Scope, Types.ResultType, TEXT("(% > %)"), InputValue[0], InputValue[1]);
+		}
+		break;
+	case EOperation::VecMulMatrix3:
+		if (Types.bIsLWC)
+		{
+			OutResult.Code = Context.EmitExpression(Scope, Types.ResultType, TEXT("LWCMlutiplyVector(%, %)"), InputValue[0], InputValue[1]);
+		}
+		else
+		{
+			OutResult.Code = Context.EmitExpression(Scope, Types.ResultType, TEXT("mul(%, (float3x3)%)"), InputValue[0], InputValue[1]);
+		}
+		break;
+	case EOperation::VecMulMatrix4:
+		if (Types.bIsLWC)
+		{
+			OutResult.Code = Context.EmitExpression(Scope, Types.ResultType, TEXT("LWCMlutiply(%, %)"), InputValue[0], InputValue[1]);
+		}
+		else
+		{
+			OutResult.Code = Context.EmitExpression(Scope, Types.ResultType, TEXT("mul(%, %)"), InputValue[0], InputValue[1]);
+		}
+		break;
+	case EOperation::Matrix3MulVec:
+		OutResult.Code = Context.EmitExpression(Scope, Types.ResultType, TEXT("mul((float3x3)%, %)"), InputValue[0], InputValue[1]);
+		break;
+	case EOperation::Matrix4MulVec:
+		OutResult.Code = Context.EmitExpression(Scope, Types.ResultType, TEXT("mul(%, %)"), InputValue[0], InputValue[1]);
 		break;
 	default:
 		checkNoEntry();
@@ -1310,14 +1256,21 @@ void FExpressionBinaryOp::EmitValueShader(FEmitContext& Context, FEmitScope& Sco
 	}
 }
 
-void FExpressionBinaryOp::EmitValuePreshader(FEmitContext& Context, FEmitScope& Scope, const FRequestedType& RequestedType, FEmitValuePreshaderResult& OutResult) const
+void FExpressionOperation::EmitValuePreshader(FEmitContext& Context, FEmitScope& Scope, const FRequestedType& RequestedType, FEmitValuePreshaderResult& OutResult) const
 {
-	const Private::FBinaryOpTypes Types = Private::GetBinaryOpTypes(Op, Lhs->GetType(), Rhs->GetType());
-	const FBinaryOpDescription OpDesc = GetBinaryOpDesription(Op);
+	const FOperationDescription OpDesc = GetOperationDescription(Op);
+	Shader::EValueType InputTypes[MaxInputs];
+	for (int32 Index = 0; Index < OpDesc.NumInputs; ++Index)
+	{
+		InputTypes[Index] = Inputs[Index]->GetType();
+	}
+	const Private::FOperationTypes Types = Private::GetOperationTypes(Op, MakeArrayView(InputTypes, OpDesc.NumInputs));
 	check(OpDesc.PreshaderOpcode != Shader::EPreshaderOpcode::Nop);
 
-	Lhs->GetValuePreshader(Context, Scope, Types.LhsType, OutResult.Preshader);
-	Rhs->GetValuePreshader(Context, Scope, Types.RhsType, OutResult.Preshader);
+	for (int32 Index = 0; Index < OpDesc.NumInputs; ++Index)
+	{
+		Inputs[Index]->GetValuePreshader(Context, Scope, Types.InputType[Index], OutResult.Preshader);
+	}
 
 	check(Context.PreshaderStackPosition > 0);
 	Context.PreshaderStackPosition--;
diff --git a/Engine/Source/Runtime/Engine/Private/Materials/MaterialExpressionHLSL.cpp b/Engine/Source/Runtime/Engine/Private/Materials/MaterialExpressionHLSL.cpp
--- a/Engine/Source/Runtime/Engine/Private/Materials/MaterialExpressionHLSL.cpp
+++ b/Engine/Source/Runtime/Engine/Private/Materials/MaterialExpressionHLSL.cpp
@@ -354,7 +354,7 @@ bool UMaterialExpressionAdd::GenerateHLSLExpression(FMaterialHLSLGenerator& Gene
 	{
 		return false;
 	}
-	OutExpression = Generator.GetTree().NewBinaryOp(UE::HLSLTree::EBinaryOp::Add, Lhs, Rhs);
+	OutExpression = Generator.GetTree().NewAdd(Lhs, Rhs);
 	return true;
 }
 
@@ -366,7 +366,7 @@ bool UMaterialExpressionSubtract::GenerateHLSLExpression(FMaterialHLSLGenerator&
 	{
 		return false;
 	}
-	OutExpression = Generator.GetTree().NewBinaryOp(UE::HLSLTree::EBinaryOp::Sub, Lhs, Rhs);
+	OutExpression = Generator.GetTree().NewSub(Lhs, Rhs);
 	return true;
 }
 
@@ -378,7 +378,7 @@ bool UMaterialExpressionMultiply::GenerateHLSLExpression(FMaterialHLSLGenerator&
 	{
 		return false;
 	}
-	OutExpression = Generator.GetTree().NewBinaryOp(UE::HLSLTree::EBinaryOp::Mul, Lhs, Rhs);
+	OutExpression = Generator.GetTree().NewMul(Lhs, Rhs);
 	return true;
 }
 
@@ -390,7 +390,7 @@ bool UMaterialExpressionDivide::GenerateHLSLExpression(FMaterialHLSLGenerator& G
 	{
 		return false;
 	}
-	OutExpression = Generator.GetTree().NewBinaryOp(UE::HLSLTree::EBinaryOp::Div, Lhs, Rhs);
+	OutExpression = Generator.GetTree().NewDiv(Lhs, Rhs);
 	return true;
 }
 
@@ -402,7 +402,7 @@ bool UMaterialExpressionDotProduct::GenerateHLSLExpression(FMaterialHLSLGenerato
 	{
 		return false;
 	}
-	OutExpression = Generator.GetTree().NewBinaryOp(UE::HLSLTree::EBinaryOp::Dot, Lhs, Rhs);
+	OutExpression = Generator.GetTree().NewDot(Lhs, Rhs);
 	return true;
 }
 
@@ -414,7 +414,7 @@ bool UMaterialExpressionMin::GenerateHLSLExpression(FMaterialHLSLGenerator& Gene
 	{
 		return false;
 	}
-	OutExpression = Generator.GetTree().NewBinaryOp(UE::HLSLTree::EBinaryOp::Min, Lhs, Rhs);
+	OutExpression = Generator.GetTree().NewMin(Lhs, Rhs);
 	return true;
 }
 
@@ -426,7 +426,7 @@ bool UMaterialExpressionMax::GenerateHLSLExpression(FMaterialHLSLGenerator& Gene
 	{
 		return false;
 	}
-	OutExpression = Generator.GetTree().NewBinaryOp(UE::HLSLTree::EBinaryOp::Max, Lhs, Rhs);
+	OutExpression = Generator.GetTree().NewMax(Lhs, Rhs);
 	return true;
 }
 
@@ -622,7 +622,7 @@ static UE::HLSLTree::FExpression* TransformBase(UE::HLSLTree::FTree& Tree,
 
 	FExpression* Result = nullptr;
 	EMaterialCommonBasis IntermediaryBasis = MCB_World;
-	EBinaryOp Op = bWComponent ? EBinaryOp::VecMulMatrix4 : EBinaryOp::VecMulMatrix3;
+	EOperation Op = bWComponent ? EOperation::VecMulMatrix4 : EOperation::VecMulMatrix3;
 	switch (SourceCoordBasis)
 	{
 	case MCB_Tangent:
diff --git a/Engine/Source/Runtime/Engine/Private/Materials/MaterialExpressions.cpp b/Engine/Source/Runtime/Engine/Private/Materials/MaterialExpressions.cpp
--- a/Engine/Source/Runtime/Engine/Private/Materials/MaterialExpressions.cpp
+++ b/Engine/Source/Runtime/Engine/Private/Materials/MaterialExpressions.cpp
@@ -22370,13 +22370,13 @@ UMaterialExpressionLess::UMaterialExpressionLess(const FObjectInitializer& Objec
 #if WITH_EDITOR
 FText UMaterialExpressionBinaryOp::GetKeywords() const
 {
-	const UE::HLSLTree::FBinaryOpDescription Description = UE::HLSLTree::GetBinaryOpDesription(GetBinaryOp());
+	const UE::HLSLTree::FOperationDescription Description = UE::HLSLTree::GetOperationDescription(GetBinaryOp());
 	return FText::FromString(Description.Operator);
 }
 
 void UMaterialExpressionBinaryOp::GetCaption(TArray<FString>& OutCaptions) const
 {
-	const UE::HLSLTree::FBinaryOpDescription Description = UE::HLSLTree::GetBinaryOpDesription(GetBinaryOp());
+	const UE::HLSLTree::FOperationDescription Description = UE::HLSLTree::GetOperationDescription(GetBinaryOp());
 	FString ret = Description.Name;
 	FExpressionInput ATraced = A.GetTracedInput();
 	FExpressionInput BTraced = B.GetTracedInput();
diff --git a/Engine/Source/Runtime/Engine/Public/HLSLTree/HLSLTree.h b/Engine/Source/Runtime/Engine/Public/HLSLTree/HLSLTree.h
--- a/Engine/Source/Runtime/Engine/Public/HLSLTree/HLSLTree.h
+++ b/Engine/Source/Runtime/Engine/Public/HLSLTree/HLSLTree.h
@@ -637,25 +637,25 @@ public:
 
 	/** Shortcuts to create various common expression types */
 	FExpression* NewConstant(const Shader::FValue& Value);
-	FExpression* NewUnaryOp(EUnaryOp Op, FExpression* Input);
-	FExpression* NewBinaryOp(EBinaryOp Op, FExpression* Lhs, FExpression* Rhs);
-
-	FExpression* NewNeg(FExpression* Input) { return NewUnaryOp(EUnaryOp::Neg, Input); }
-	FExpression* NewRcp(FExpression* Input) { return NewUnaryOp(EUnaryOp::Rcp, Input); }
-	FExpression* NewFrac(FExpression* Input) { return NewUnaryOp(EUnaryOp::Frac, Input); }
-	FExpression* NewLength(FExpression* Input) { return NewUnaryOp(EUnaryOp::Length, Input); }
-	FExpression* NewNormalize(FExpression* Input) { return NewUnaryOp(EUnaryOp::Normalize, Input); }
-
-	FExpression* NewAdd(FExpression* Lhs, FExpression* Rhs) { return NewBinaryOp(EBinaryOp::Add, Lhs, Rhs); }
-	FExpression* NewSub(FExpression* Lhs, FExpression* Rhs) { return NewBinaryOp(EBinaryOp::Sub, Lhs, Rhs); }
-	FExpression* NewMul(FExpression* Lhs, FExpression* Rhs) { return NewBinaryOp(EBinaryOp::Mul, Lhs, Rhs); }
-	FExpression* NewDiv(FExpression* Lhs, FExpression* Rhs) { return NewBinaryOp(EBinaryOp::Div, Lhs, Rhs); }
-	FExpression* NewFmod(FExpression* Lhs, FExpression* Rhs) { return NewBinaryOp(EBinaryOp::Fmod, Lhs, Rhs); }
-	FExpression* NewDot(FExpression* Lhs, FExpression* Rhs) { return NewBinaryOp(EBinaryOp::Dot, Lhs, Rhs); }
-	FExpression* NewMin(FExpression* Lhs, FExpression* Rhs) { return NewBinaryOp(EBinaryOp::Min, Lhs, Rhs); }
-	FExpression* NewMax(FExpression* Lhs, FExpression* Rhs) { return NewBinaryOp(EBinaryOp::Max, Lhs, Rhs); }
-	FExpression* NewLess(FExpression* Lhs, FExpression* Rhs) { return NewBinaryOp(EBinaryOp::Less, Lhs, Rhs); }
-	FExpression* NewGreater(FExpression* Lhs, FExpression* Rhs) { return NewBinaryOp(EBinaryOp::Greater, Lhs, Rhs); }
+	FExpression* NewUnaryOp(EOperation Op, FExpression* Input);
+	FExpression* NewBinaryOp(EOperation Op, FExpression* Lhs, FExpression* Rhs);
+
+	FExpression* NewNeg(FExpression* Input) { return NewUnaryOp(EOperation::Neg, Input); }
+	FExpression* NewRcp(FExpression* Input) { return NewUnaryOp(EOperation::Rcp, Input); }
+	FExpression* NewFrac(FExpression* Input) { return NewUnaryOp(EOperation::Frac, Input); }
+	FExpression* NewLength(FExpression* Input) { return NewUnaryOp(EOperation::Length, Input); }
+	FExpression* NewNormalize(FExpression* Input) { return NewUnaryOp(EOperation::Normalize, Input); }
+
+	FExpression* NewAdd(FExpression* Lhs, FExpression* Rhs) { return NewBinaryOp(EOperation::Add, Lhs, Rhs); }
+	FExpression* NewSub(FExpression* Lhs, FExpression* Rhs) { return NewBinaryOp(EOperation::Sub, Lhs, Rhs); }
+	FExpression* NewMul(FExpression* Lhs, FExpression* Rhs) { return NewBinaryOp(EOperation::Mul, Lhs, Rhs); }
+	FExpression* NewDiv(FExpression* Lhs, FExpression* Rhs) { return NewBinaryOp(EOperation::Div, Lhs, Rhs); }
+	FExpression* NewFmod(FExpression* Lhs, FExpression* Rhs) { return NewBinaryOp(EOperation::Fmod, Lhs, Rhs); }
+	FExpression* NewDot(FExpression* Lhs, FExpression* Rhs) { return NewBinaryOp(EOperation::Dot, Lhs, Rhs); }
+	FExpression* NewMin(FExpression* Lhs, FExpression* Rhs) { return NewBinaryOp(EOperation::Min, Lhs, Rhs); }
+	FExpression* NewMax(FExpression* Lhs, FExpression* Rhs) { return NewBinaryOp(EOperation::Max, Lhs, Rhs); }
+	FExpression* NewLess(FExpression* Lhs, FExpression* Rhs) { return NewBinaryOp(EOperation::Less, Lhs, Rhs); }
+	FExpression* NewGreater(FExpression* Lhs, FExpression* Rhs) { return NewBinaryOp(EOperation::Greater, Lhs, Rhs); }
 
 	FTextureParameterDeclaration* NewTextureParameterDeclaration(const FName& Name, const FTextureDescription& DefaultValue);
 
diff --git a/Engine/Source/Runtime/Engine/Public/HLSLTree/HLSLTreeCommon.h b/Engine/Source/Runtime/Engine/Public/HLSLTree/HLSLTreeCommon.h
--- a/Engine/Source/Runtime/Engine/Public/HLSLTree/HLSLTreeCommon.h
+++ b/Engine/Source/Runtime/Engine/Public/HLSLTree/HLSLTreeCommon.h
@@ -253,36 +253,15 @@ public:
 	virtual void EmitValuePreshader(FEmitContext& Context, FEmitScope& Scope, const FRequestedType& RequestedType, FEmitValuePreshaderResult& OutResult) const override;
 };
 
-class FExpressionUnaryOp : public FExpression
+class FExpressionOperation : public FExpression
 {
 public:
-	FExpressionUnaryOp(EUnaryOp InOp, FExpression* InInput)
-		: Op(InOp)
-		, Input(InInput)
-	{}
+	FExpressionOperation(EOperation InOp, TConstArrayView<FExpression*> InInputs);
 
-	EUnaryOp Op;
-	FExpression* Input;
+	static constexpr int8 MaxInputs = 2;
 
-	virtual void ComputeAnalyticDerivatives(FTree& Tree, FExpressionDerivatives& OutResult) const override;
-	virtual FExpression* ComputePreviousFrame(FTree& Tree, const FRequestedType& RequestedType) const override;
-	virtual bool PrepareValue(FEmitContext& Context, FEmitScope& Scope, const FRequestedType& RequestedType, FPrepareValueResult& OutResult) const override;
-	virtual void EmitValueShader(FEmitContext& Context, FEmitScope& Scope, const FRequestedType& RequestedType, FEmitValueShaderResult& OutResult) const override;
-	virtual void EmitValuePreshader(FEmitContext& Context, FEmitScope& Scope, const FRequestedType& RequestedType, FEmitValuePreshaderResult& OutResult) const override;
-};
-
-class FExpressionBinaryOp : public FExpression
-{
-public:
-	FExpressionBinaryOp(EBinaryOp InOp, FExpression* InLhs, FExpression* InRhs)
-		: Op(InOp)
-		, Lhs(InLhs)
-		, Rhs(InRhs)
-	{}
-
-	EBinaryOp Op;
-	FExpression* Lhs;
-	FExpression* Rhs;
+	EOperation Op;
+	FExpression* Inputs[MaxInputs];
 
 	virtual void ComputeAnalyticDerivatives(FTree& Tree, FExpressionDerivatives& OutResult) const override;
 	virtual FExpression* ComputePreviousFrame(FTree& Tree, const FRequestedType& RequestedType) const override;
diff --git a/Engine/Source/Runtime/Engine/Public/HLSLTree/HLSLTreeTypes.h b/Engine/Source/Runtime/Engine/Public/HLSLTree/HLSLTreeTypes.h
--- a/Engine/Source/Runtime/Engine/Public/HLSLTree/HLSLTreeTypes.h
+++ b/Engine/Source/Runtime/Engine/Public/HLSLTree/HLSLTreeTypes.h
@@ -59,31 +59,18 @@ inline bool IsLoopEvaluation(EExpressionEvaluation Evaluation)
 	return Evaluation == EExpressionEvaluation::PreshaderLoop || Evaluation == EExpressionEvaluation::ConstantLoop;
 }
 
-enum class EUnaryOp : uint8
+enum class EOperation : uint8
 {
 	None,
+
+	// Unary Ops
 	Neg,
 	Rcp,
 	Frac,
 	Length,
 	Normalize,
-};
-
-struct FUnaryOpDescription
-{
-	FUnaryOpDescription();
-	FUnaryOpDescription(const TCHAR* InName, const TCHAR* InOperator, Shader::EPreshaderOpcode InOpcode);
-
-	const TCHAR* Name;
-	const TCHAR* Operator;
-	Shader::EPreshaderOpcode PreshaderOpcode;
-};
 
-FUnaryOpDescription GetUnaryOpDesription(EUnaryOp Op);
-
-enum class EBinaryOp : uint8
-{
-	None,
+	// Binary Ops
 	Add,
 	Sub,
 	Mul,
@@ -101,17 +88,18 @@ enum class EBinaryOp : uint8
 	Matrix4MulVec,
 };
 
-struct FBinaryOpDescription
+struct FOperationDescription
 {
-	FBinaryOpDescription();
-	FBinaryOpDescription(const TCHAR* InName, const TCHAR* InOperator, Shader::EPreshaderOpcode InOpcode);
+	FOperationDescription();
+	FOperationDescription(const TCHAR* InName, const TCHAR* InOperator, int8 InNumInputs, Shader::EPreshaderOpcode InOpcode);
 
 	const TCHAR* Name;
 	const TCHAR* Operator;
+	int8 NumInputs;
 	Shader::EPreshaderOpcode PreshaderOpcode;
 };
 
-FBinaryOpDescription GetBinaryOpDesription(EBinaryOp Op);
+FOperationDescription GetOperationDescription(EOperation Op);
 
 struct FTextureDescription
 {
