Converted history index values to uint16
Removed sanity tag from property retirements, removed validation from test builds

diff --git a/Engine/Source/Runtime/CoreUObject/Public/UObject/CoreNet.h b/Engine/Source/Runtime/CoreUObject/Public/UObject/CoreNet.h
--- a/Engine/Source/Runtime/CoreUObject/Public/UObject/CoreNet.h
+++ b/Engine/Source/Runtime/CoreUObject/Public/UObject/CoreNet.h
@@ -228,23 +228,14 @@ struct FPacketIdRange
 /** Information for tracking retirement and retransmission of a property. */
 struct FPropertyRetirement
 {
-#if !UE_BUILD_SHIPPING
-	static const uint32 ExpectedSanityTag = 0xDF41C9A3;
-
-	uint32 SanityTag;
-#endif
-
 	FPropertyRetirement* Next;
 
 	TSharedPtr<class INetDeltaBaseState> DynamicState;
 
 	FPacketIdRange OutPacketIdRange;
-	uint32 FastArrayChangelistHistory;
+	uint16 FastArrayChangelistHistory;
 
 	FPropertyRetirement() :
-#if !UE_BUILD_SHIPPING
-		SanityTag( ExpectedSanityTag ),
-#endif
 		 Next(nullptr)
 		, DynamicState(nullptr)
 		, FastArrayChangelistHistory(0)
@@ -435,15 +426,15 @@ public:
 	 */
 	virtual void CountBytes(FArchive& Ar) const {}
 
-	uint32 GetLastAckedHistory() const { return LastAckedHistory; }
-	void SetLastAckedHistory(uint32 InAckedHistory) { LastAckedHistory = InAckedHistory; }
+	uint16 GetLastAckedHistory() const { return LastAckedHistory; }
+	void SetLastAckedHistory(uint16 InAckedHistory) { LastAckedHistory = InAckedHistory; }
 
-	uint32 GetChangelistHistory() const { return ChangelistHistory; }
-	void SetChangelistHistory(uint32 InChangelistHistory) { ChangelistHistory = InChangelistHistory; }
+	uint16 GetChangelistHistory() const { return ChangelistHistory; }
+	void SetChangelistHistory(uint16 InChangelistHistory) { ChangelistHistory = InChangelistHistory; }
 
 private:
-	uint32 LastAckedHistory;
-	uint32 ChangelistHistory;
+	uint16 LastAckedHistory;
+	uint16 ChangelistHistory;
 };
 
 struct FNetDeltaSerializeInfo;
diff --git a/Engine/Source/Runtime/Engine/Private/DataReplication.cpp b/Engine/Source/Runtime/Engine/Private/DataReplication.cpp
--- a/Engine/Source/Runtime/Engine/Private/DataReplication.cpp
+++ b/Engine/Source/Runtime/Engine/Private/DataReplication.cpp
@@ -691,16 +691,13 @@ void FObjectReplicator::StartReplicating(class UActorChannel * InActorChannel)
 
 void ValidateRetirementHistory(const FPropertyRetirement & Retire, const UObject* Object)
 {
-#if !UE_BUILD_SHIPPING
-	checkf( Retire.SanityTag == FPropertyRetirement::ExpectedSanityTag, TEXT( "Invalid Retire.SanityTag. Object: %s" ), *GetFullNameSafe(Object) );
-
-	FPropertyRetirement * Rec = Retire.Next;	// Note the first element is 'head' that we dont actually use
+#if !(UE_BUILD_SHIPPING || UE_BUILD_TEST)
+	FPropertyRetirement * Rec = Retire.Next;	// Note the first element is 'head' that we don't actually use
 
 	FPacketIdRange LastRange;
 
 	while ( Rec != nullptr )
 	{
-		checkf( Rec->SanityTag == FPropertyRetirement::ExpectedSanityTag, TEXT( "Invalid Rec->SanityTag. Object: %s" ), *GetFullNameSafe(Object) );
 		checkf( Rec->OutPacketIdRange.Last >= Rec->OutPacketIdRange.First, TEXT( "Invalid packet id range (Last < First). Object: %s" ), *GetFullNameSafe(Object) );
 		checkf( Rec->OutPacketIdRange.First >= LastRange.Last, TEXT( "Invalid packet id range (First < LastRange.Last). Object: %s" ), *GetFullNameSafe(Object) );		// Bunch merging and queuing can cause this overlap
 
@@ -865,9 +862,9 @@ void FObjectReplicator::ReceivedNak( int32 NakPacketId )
 		if (FSendingRepState* SendingRepState = RepState.IsValid() ? RepState->GetSendingRepState() : nullptr)
 		{
 			// Go over properties tracked with histories, and mark them as needing to be resent.
-			for (int32 i = SendingRepState->HistoryStart; i < SendingRepState->HistoryEnd; ++i)
+			for (uint16 i = SendingRepState->HistoryStart; i < SendingRepState->HistoryEnd; ++i)
 			{
-				const int32 HistoryIndex = i % FSendingRepState::MAX_CHANGE_HISTORY;
+				const uint16 HistoryIndex = i % FSendingRepState::MAX_CHANGE_HISTORY;
 
 				FRepChangedHistory& HistoryItem = SendingRepState->ChangeHistory[HistoryIndex];
 
@@ -887,7 +884,7 @@ void FObjectReplicator::ReceivedNak( int32 NakPacketId )
 
 				// If this is a dynamic array property, we have to look through the list of retirement records to see if we need to reset the base state
 				FPropertyRetirement* Rec = Retirement.Next; // Retirement[i] is head and not actually used in this case
-				uint32 LastAcknowledged = 0;
+				uint16 LastAcknowledged = 0;
 				while (Rec != nullptr)
 				{
 					if (NakPacketId > Rec->OutPacketIdRange.Last)
@@ -1507,7 +1504,7 @@ void FObjectReplicator::PostReceivedBunch()
 
 static FORCEINLINE FPropertyRetirement** UpdateAckedRetirements(
 	FPropertyRetirement& Retire,
-	uint32& LastAcknowledged,
+	uint16& LastAcknowledged,
 	const int32 OutAckPacketId,
 	const UObject* Object)
 {
@@ -1644,7 +1641,7 @@ void FObjectReplicator::ReplicateCustomDeltaProperties( FNetBitWriter & Bunch, F
 
 			// Update Retirement records with this new state so we can handle packet drops.
 			// LastNext will be pointer to the last "Next" pointer in the list (so pointer to a pointer)
-			uint32 LastAcknowledged = 0;
+			uint16 LastAcknowledged = 0;
 			LastNext = UpdateAckedRetirements(Retire, LastAcknowledged, Connection->OutAckPacketId, Object);
 
 			if (LastAcknowledged != 0 && OldState.IsValid())
@@ -1887,20 +1884,20 @@ void FObjectReplicator::PostSendBunch( FPacketIdRange & PacketRange, uint8 bReli
 	check(RepLayout);
 
 	// Don't update retirement records for reliable properties. This is ok to do only if we also pause replication on the channel until the acks have gone through.
-	bool SkipRetirementUpdate = OwningChannel->bPausedUntilReliableACK;
+	const bool bSkipRetirementUpdate = OwningChannel->bPausedUntilReliableACK;
 
 	const FRepLayout& LocalRepLayout = *RepLayout;
 
 	if (FSendingRepState* SendingRepState = RepState.IsValid() ? RepState->GetSendingRepState() : nullptr)
 	{
-		if (!SkipRetirementUpdate)
+		if (!bSkipRetirementUpdate)
 		{
 			// Don't call if reliable, since the bunch will be resent. We dont want this to end up in the changelist history
 			// But is that enough? How does it know to delta against this latest state?
 
-			for (int32 i = SendingRepState->HistoryStart; i < SendingRepState->HistoryEnd; ++i)
+			for (uint16 i = SendingRepState->HistoryStart; i < SendingRepState->HistoryEnd; ++i)
 			{
-				const int32 HistoryIndex = i % FSendingRepState::MAX_CHANGE_HISTORY;
+				const uint16 HistoryIndex = i % FSendingRepState::MAX_CHANGE_HISTORY;
 
 				FRepChangedHistory & HistoryItem = SendingRepState->ChangeHistory[HistoryIndex];
 
@@ -1929,7 +1926,7 @@ void FObjectReplicator::PostSendBunch( FPacketIdRange & PacketRange, uint8 bReli
 			// possible that we either didn't send anything, or the last thing we sent wasn't reliable.
 			if (SendingRepState->HistoryEnd > SendingRepState->HistoryStart)
 			{
-				const int32 HistoryIndex = (SendingRepState->HistoryEnd - 1) % FSendingRepState::MAX_CHANGE_HISTORY;
+				const uint16 HistoryIndex = (SendingRepState->HistoryEnd - 1) % FSendingRepState::MAX_CHANGE_HISTORY;
 				FRepChangedHistory& HistoryItem = SendingRepState->ChangeHistory[HistoryIndex];
 	
 				if (!HistoryItem.WasSent())
@@ -1951,7 +1948,7 @@ void FObjectReplicator::PostSendBunch( FPacketIdRange & PacketRange, uint8 bReli
 				// (we have to wait until we actually send the bunch to know the packetID, which is why we look for .First==INDEX_NONE)
 				if (Next->OutPacketIdRange.First == INDEX_NONE)
 				{
-					if (!SkipRetirementUpdate)
+					if (!bSkipRetirementUpdate)
 					{
 						Next->OutPacketIdRange = PacketRange;
 
diff --git a/Engine/Source/Runtime/Engine/Private/RepLayout.cpp b/Engine/Source/Runtime/Engine/Private/RepLayout.cpp
--- a/Engine/Source/Runtime/Engine/Private/RepLayout.cpp
+++ b/Engine/Source/Runtime/Engine/Private/RepLayout.cpp
@@ -475,7 +475,7 @@ struct FDeltaArrayHistoryItem
 struct FDeltaArrayHistoryState
 {
 	/** The maximum number of individual changelists allowed.*/
-	static const uint32 MAX_CHANGE_HISTORY = FRepChangelistState::MAX_CHANGE_HISTORY;
+	static const uint16 MAX_CHANGE_HISTORY = FRepChangelistState::MAX_CHANGE_HISTORY;
 
 	//~ TODO: Investigate either making this a Dynamically sized container,
 	//~			or potentially changing ArrayStates to unique pointers / 
@@ -494,10 +494,10 @@ struct FDeltaArrayHistoryState
 	TMap<int32, int32> IDToIndexMap;
 
 	/** Index in the buffer where changelist history starts (i.e., the Oldest changelist). */
-	uint32 HistoryStart = 0;
+	uint16 HistoryStart = 0;
 
 	/** Index in the buffer where changelist history ends (i.e., the Newest changelist). */
-	uint32 HistoryEnd = 0;
+	uint16 HistoryEnd = 0;
 
 	void CountBytes(FArchive& Ar) const
 	{
@@ -1689,10 +1689,20 @@ ERepLayoutResult FRepLayout::CompareProperties(
 		return ERepLayoutResult::Empty;
 	}
 
-	RepChangelistState->CompareIndex++;
+	static_assert(sizeof(RepChangelistState->CompareIndex) == sizeof(uint16), "CompareIndex should be a uint16");
+
+	if (RepChangelistState->CompareIndex == MAX_uint16)
+	{
+		// other systems test for index > 1 to indicate that we've ever done a comparison
+		RepChangelistState->CompareIndex = 2;
+	}
+	else
+	{
+		RepChangelistState->CompareIndex++;
+	}
 
 	check((RepChangelistState->HistoryEnd - RepChangelistState->HistoryStart) < FRepChangelistState::MAX_CHANGE_HISTORY);
-	const int32 HistoryIndex = RepChangelistState->HistoryEnd % FRepChangelistState::MAX_CHANGE_HISTORY;
+	const uint16 HistoryIndex = RepChangelistState->HistoryEnd % FRepChangelistState::MAX_CHANGE_HISTORY;
 
 	FRepChangedHistory& NewHistoryItem = RepChangelistState->ChangeHistory[HistoryIndex];
 
@@ -1799,11 +1809,11 @@ ERepLayoutResult FRepLayout::CompareProperties(
 	// If we're full, merge the oldest up, so we always have room for a new entry
 	if ((RepChangelistState->HistoryEnd - RepChangelistState->HistoryStart) == FRepChangelistState::MAX_CHANGE_HISTORY)
 	{
-		const int32 FirstHistoryIndex = RepChangelistState->HistoryStart % FRepChangelistState::MAX_CHANGE_HISTORY;
+		const uint16 FirstHistoryIndex = RepChangelistState->HistoryStart % FRepChangelistState::MAX_CHANGE_HISTORY;
 
 		RepChangelistState->HistoryStart++;
 
-		const int32 SecondHistoryIndex = RepChangelistState->HistoryStart % FRepChangelistState::MAX_CHANGE_HISTORY;
+		const uint16 SecondHistoryIndex = RepChangelistState->HistoryStart % FRepChangelistState::MAX_CHANGE_HISTORY;
 
 		TArray<uint16>& FirstChangelistRef = RepChangelistState->ChangeHistory[FirstHistoryIndex].Changed;
 		TArray<uint16> SecondChangelistCopy = MoveTemp(RepChangelistState->ChangeHistory[SecondHistoryIndex].Changed);
@@ -1918,9 +1928,9 @@ bool FRepLayout::ReplicateProperties(
 		if (RepState->LifetimeChangelist.Num() == 0)
 		{
 			// If this object was dormant, the lifetime list will be empty if the replicator was discarded, so rebuild it from the changelist state.
-			for (int32 i = RepChangelistState->HistoryStart; i < RepChangelistState->HistoryEnd; ++i)
+			for (uint16 i = RepChangelistState->HistoryStart; i < RepChangelistState->HistoryEnd; ++i)
 			{
-				const int32 HistoryIndex = i % FRepChangelistState::MAX_CHANGE_HISTORY;
+				const uint16 HistoryIndex = i % FRepChangelistState::MAX_CHANGE_HISTORY;
 
 				FRepChangedHistory& HistoryItem = RepChangelistState->ChangeHistory[HistoryIndex];
 
@@ -1982,7 +1992,7 @@ bool FRepLayout::ReplicateProperties(
 		RepState->LastChangelistIndex = RepChangelistState->HistoryStart;
 	}
 
-	const int32 PossibleNewHistoryIndex = RepState->HistoryEnd % FSendingRepState::MAX_CHANGE_HISTORY;
+	const uint16 PossibleNewHistoryIndex = RepState->HistoryEnd % FSendingRepState::MAX_CHANGE_HISTORY;
 
 	FRepChangedHistory& PossibleNewHistoryItem = RepState->ChangeHistory[PossibleNewHistoryIndex];
 
@@ -1992,9 +2002,9 @@ bool FRepLayout::ReplicateProperties(
 	check(Changed.Num() == 0);
 
 	// Gather all change lists that are new since we last looked, and merge them all together into a single CL
-	for (int32 i = RepState->LastChangelistIndex; i < RepChangelistState->HistoryEnd; ++i)
+	for (uint16 i = RepState->LastChangelistIndex; i < RepChangelistState->HistoryEnd; ++i)
 	{
-		const int32 HistoryIndex = i % FRepChangelistState::MAX_CHANGE_HISTORY;
+		const uint16 HistoryIndex = i % FRepChangelistState::MAX_CHANGE_HISTORY;
 
 		FRepChangedHistory& HistoryItem = RepChangelistState->ChangeHistory[HistoryIndex];
 
@@ -2103,12 +2113,12 @@ void FRepLayout::UpdateChangelistHistory(
 {
 	check(RepState->HistoryEnd >= RepState->HistoryStart);
 
-	const int32 HistoryCount = RepState->HistoryEnd - RepState->HistoryStart;
-	const bool DumpHistory = HistoryCount == FSendingRepState::MAX_CHANGE_HISTORY;
+	const uint16 HistoryCount = RepState->HistoryEnd - RepState->HistoryStart;
+	const bool bDumpHistory = HistoryCount == FSendingRepState::MAX_CHANGE_HISTORY;
 	const int32 AckPacketId = Connection->OutAckPacketId;
 
 	// If our buffer is currently full, forcibly send the entire history
-	if (DumpHistory)
+	if (bDumpHistory)
 	{
 		UE_LOG(LogRep, Verbose, TEXT("FRepLayout::UpdateChangelistHistory: History overflow, forcing history dump %s, %s"), *ObjectClass->GetName(), *Connection->Describe());
 	}
@@ -2117,9 +2127,9 @@ void FRepLayout::UpdateChangelistHistory(
 
 	if (bDeltaCheckpoint)
 	{
-		for (int32 i = RepState->HistoryStart; i < RepState->HistoryEnd - 1; i++)
+		for (uint16 i = RepState->HistoryStart; i < RepState->HistoryEnd - 1; i++)
 		{
-			const int32 HistoryIndex = i % FSendingRepState::MAX_CHANGE_HISTORY;
+			const uint16 HistoryIndex = i % FSendingRepState::MAX_CHANGE_HISTORY;
 
 			FRepChangedHistory& HistoryItem = RepState->ChangeHistory[HistoryIndex];
 
@@ -2132,9 +2142,9 @@ void FRepLayout::UpdateChangelistHistory(
 	}
 	else
 	{
-		for (int32 i = RepState->HistoryStart; i < RepState->HistoryEnd; i++)
+		for (uint16 i = RepState->HistoryStart; i < RepState->HistoryEnd; i++)
 		{
-			const int32 HistoryIndex = i % FSendingRepState::MAX_CHANGE_HISTORY;
+			const uint16 HistoryIndex = i % FSendingRepState::MAX_CHANGE_HISTORY;
 
 			FRepChangedHistory& HistoryItem = RepState->ChangeHistory[HistoryIndex];
 
@@ -2148,9 +2158,9 @@ void FRepLayout::UpdateChangelistHistory(
 			// All active history items should contain a change list
 			check(HistoryItem.Changed.Num() > 0);
 
-			if (AckPacketId >= HistoryItem.OutPacketIdRange.Last || HistoryItem.Resend || DumpHistory)
+			if (AckPacketId >= HistoryItem.OutPacketIdRange.Last || HistoryItem.Resend || bDumpHistory)
 			{
-				if (HistoryItem.Resend || DumpHistory)
+				if (HistoryItem.Resend || bDumpHistory)
 				{
 					// Merge in nak'd change lists
 					check(OutMerged != NULL);
@@ -2174,7 +2184,7 @@ void FRepLayout::UpdateChangelistHistory(
 	}
 
 	// Remove any tiling in the history markers to keep them from wrapping over time
-	const int32 NewHistoryCount	= RepState->HistoryEnd - RepState->HistoryStart;
+	const uint16 NewHistoryCount = RepState->HistoryEnd - RepState->HistoryStart;
 
 	check(NewHistoryCount < FSendingRepState::MAX_CHANGE_HISTORY);
 
@@ -7313,8 +7323,8 @@ void FRepLayout::PreSendCustomDeltaProperties(
 							const int32 FastArrayReplicationKey = CustomDeltaProperty.GetFastArrayArrayReplicationKey(FastArraySerializer);
 							if (FastArrayHistoryState.ArrayReplicationKey != FastArrayReplicationKey)
 							{
-								const uint32 HistoryDelta = FastArrayHistoryState.HistoryEnd - FastArrayHistoryState.HistoryStart;
-								const uint32 CurrentHistoryIndex = FastArrayHistoryState.HistoryEnd % FDeltaArrayHistoryState::MAX_CHANGE_HISTORY;
+								const uint16 HistoryDelta = FastArrayHistoryState.HistoryEnd - FastArrayHistoryState.HistoryStart;
+								const uint16 CurrentHistoryIndex = FastArrayHistoryState.HistoryEnd % FDeltaArrayHistoryState::MAX_CHANGE_HISTORY;
 								const FDeltaArrayHistoryItem& CurrentHistory = FastArrayHistoryState.ChangeHistory[CurrentHistoryIndex];
 
 								// If we don't have any history items, go ahead and create one.
@@ -7331,7 +7341,7 @@ void FRepLayout::PreSendCustomDeltaProperties(
 									}
 
 									++FastArrayHistoryState.HistoryEnd;
-									const uint32 NewHistory = FastArrayHistoryState.HistoryEnd % FDeltaArrayHistoryState::MAX_CHANGE_HISTORY;
+									const uint16 NewHistory = FastArrayHistoryState.HistoryEnd % FDeltaArrayHistoryState::MAX_CHANGE_HISTORY;
 									FastArrayHistoryState.ChangeHistory[NewHistory].Reset();
 								}
 							}
@@ -7450,7 +7460,6 @@ ERepLayoutResult FRepLayout::DeltaSerializeFastArrayProperty(FFastArrayDeltaSeri
 				const int32 FastArrayNumber = CustomDeltaProperty.FastArrayNumber;
 				FDeltaArrayHistoryState& FastArrayState = DeltaChangelistState.ArrayStates[FastArrayNumber];
 
-
 				// Params.WriteBaseState should be valid, and have the most up to date IDToChangelist map for the Fast Array.
 				// However, it's ChangelistHistory will be to the last History Number sent to the Fast TArray on the specific
 				// connection we're replicating from.
@@ -7458,9 +7467,9 @@ ERepLayoutResult FRepLayout::DeltaSerializeFastArrayProperty(FFastArrayDeltaSeri
 
 				// Cache off the newest history, the last history sent to this connection, and then update the state
 				// to notify that we're going to send it the newest history.
-				const uint32 NewChangelistHistory = FastArrayState.HistoryEnd;
-				const uint32 LastAckedHistory = NewArrayDeltaState->GetLastAckedHistory();
-				const uint32 LastAckedChangelistDelta = NewChangelistHistory - LastAckedHistory;
+				const uint16 NewChangelistHistory = FastArrayState.HistoryEnd;
+				const uint16 LastAckedHistory = NewArrayDeltaState->GetLastAckedHistory();
+				const uint16 LastAckedChangelistDelta = NewChangelistHistory - LastAckedHistory;
 
 				NewArrayDeltaState->SetChangelistHistory(NewChangelistHistory);
 
@@ -7475,8 +7484,8 @@ ERepLayoutResult FRepLayout::DeltaSerializeFastArrayProperty(FFastArrayDeltaSeri
 
 				// Check to see whether or not we need to update the global changelist shared between connections.
 				{
-					const uint32 RelativeNewHistory = NewChangelistHistory % FDeltaArrayHistoryState::MAX_CHANGE_HISTORY;
-					const uint32 CompareChangelistDelta = NewChangelistHistory - FastArrayState.HistoryStart;
+					const uint16 RelativeNewHistory = NewChangelistHistory % FDeltaArrayHistoryState::MAX_CHANGE_HISTORY;
+					const uint16 CompareChangelistDelta = NewChangelistHistory - FastArrayState.HistoryStart;
 					FDeltaArrayHistoryItem& HistoryItem = FastArrayState.ChangeHistory[RelativeNewHistory];
 
 					if (!HistoryItem.bWasUpdated)
diff --git a/Engine/Source/Runtime/Engine/Public/Net/RepLayout.h b/Engine/Source/Runtime/Engine/Public/Net/RepLayout.h
--- a/Engine/Source/Runtime/Engine/Public/Net/RepLayout.h
+++ b/Engine/Source/Runtime/Engine/Public/Net/RepLayout.h
@@ -501,7 +501,7 @@ public:
 	~FRepChangelistState();
 
 	/** The maximum number of individual changelists allowed.*/
-	static const int32 MAX_CHANGE_HISTORY = 64;
+	static const uint16 MAX_CHANGE_HISTORY = 64;
 
 	/** Circular buffer of changelists. */
 	FRepChangedHistory ChangeHistory[MAX_CHANGE_HISTORY];
@@ -510,13 +510,13 @@ public:
 	TUniquePtr<struct FCustomDeltaChangelistState> CustomDeltaChangelistState;
 
 	/** Index in the buffer where changelist history starts (i.e., the Oldest changelist). */
-	int32 HistoryStart;
+	uint16 HistoryStart;
 
 	/** Index in the buffer where changelist history ends (i.e., the Newest changelist). */
-	int32 HistoryEnd;
+	uint16 HistoryEnd;
 
 	/** Number of times that properties have been compared */
-	int32 CompareIndex;
+	uint16 CompareIndex;
 
 	/** Latest state of all property data. Not used on Clients, only used on Servers if Shadow State is enabled. */
 	FRepStateStaticBuffer StaticBuffer;
@@ -647,19 +647,19 @@ public:
 	TEnumAsByte<ENetRole> SavedRole = ROLE_MAX;
 
 	/** Index in the buffer where changelist history starts (i.e., the Oldest changelist). */
-	int32 HistoryStart;
+	uint16 HistoryStart;
 
 	/** Index in the buffer where changelist history ends (i.e., the Newest changelist). */
-	int32 HistoryEnd;
+	uint16 HistoryEnd;
 
 	/** Number of Changelist history entries that have outstanding Naks. */
-	int32 NumNaks;
+	uint16 NumNaks;
 
 	/**
 	 * The last change list history item we replicated from FRepChangelistState.
 	 * (If we are caught up to FRepChangelistState::HistoryEnd, there are no new changelists to replicate).
 	 */
-	int32 LastChangelistIndex;
+	uint16 LastChangelistIndex;
 
 	/**
 	 * Tracks the last time this RepState actually replicated data.
@@ -670,14 +670,14 @@ public:
 	 *
 	 * Note, we can't solely rely on on LastChangelistIndex, since changelists are stored in circular buffers.
 	 */
-	int32 LastCompareIndex;
+	uint16 LastCompareIndex;
 
 	FReplicationFlags RepFlags;
 
 	TSharedPtr<FRepChangedPropertyTracker> RepChangedPropertyTracker;
 
 	/** The maximum number of individual changelists allowed.*/
-	static constexpr int32 MAX_CHANGE_HISTORY = 32;
+	static constexpr uint16 MAX_CHANGE_HISTORY = 32;
 
 	/** Circular buffer of changelists. */
 	FRepChangedHistory ChangeHistory[MAX_CHANGE_HISTORY];
diff --git a/Engine/Source/Runtime/Net/Core/Classes/Net/Serialization/FastArraySerializer.h b/Engine/Source/Runtime/Net/Core/Classes/Net/Serialization/FastArraySerializer.h
--- a/Engine/Source/Runtime/Net/Core/Classes/Net/Serialization/FastArraySerializer.h
+++ b/Engine/Source/Runtime/Net/Core/Classes/Net/Serialization/FastArraySerializer.h
@@ -1601,8 +1601,8 @@ bool FFastArraySerializer::FastArrayDeltaSerialize_DeltaSerializeStructs(TArray<
 
 		// We set ChangelistHistory and LastAckedHistory to 0, as that's the default RepLayout expects.
 		// Setting them to anything else may cause issues with initial sends.
-		uint32 OldChangelistHistory = 0;
-		uint32 OldLastAckedHistory = 0;
+		uint16 OldChangelistHistory = 0;
+		uint16 OldLastAckedHistory = 0;
 
 		// See if the array changed at all. If the ArrayReplicationKey matches we can skip checking individual items
 		if (Parms.OldState)
