Terminating PIE will no longer assert when PendingKill is disabled and level leaks are detected. Added an ensure when reporting level leaks through UEngine::FindAndPrintStaleReferencesToObject() with PendingKill disabled.

diff --git a/Engine/Source/Editor/UnrealEd/Private/Cooker/IoStoreCookOnTheFlyRequestManager.cpp b/Engine/Source/Editor/UnrealEd/Private/Cooker/IoStoreCookOnTheFlyRequestManager.cpp
--- a/Engine/Source/Editor/UnrealEd/Private/Cooker/IoStoreCookOnTheFlyRequestManager.cpp
+++ b/Engine/Source/Editor/UnrealEd/Private/Cooker/IoStoreCookOnTheFlyRequestManager.cpp
@@ -701,7 +701,7 @@ private:
 				if (Package)
 				{
 					UE_LOG(LogCookOnTheFly, Warning, TEXT("Can't recook package '%s'"), *PackageName.ToString());
-					UEngine::FindAndPrintStaleReferencesToObject(Package, ELogVerbosity::Display);
+					UEngine::FindAndPrintStaleReferencesToObject(Package, EPrintStaleReferencesOptions::Display);
 				}
 				else
 				{
diff --git a/Engine/Source/Editor/UnrealEd/Private/EditorEngine.cpp b/Engine/Source/Editor/UnrealEd/Private/EditorEngine.cpp
--- a/Engine/Source/Editor/UnrealEd/Private/EditorEngine.cpp
+++ b/Engine/Source/Editor/UnrealEd/Private/EditorEngine.cpp
@@ -7047,7 +7047,8 @@ void UEditorEngine::VerifyLoadMapWorldCleanup()
 				UE_LOG(LogLoad, Error, TEXT("Previously active world %s not cleaned up by garbage collection!"), *World->GetPathName());
 				UE_LOG(LogLoad, Error, TEXT("Once a world has become active, it cannot be reused and must be destroyed and reloaded. World referenced by:"));
 			
-				FindAndPrintStaleReferencesToObject(World, UObjectBaseUtility::IsPendingKillEnabled() ? ELogVerbosity::Fatal : ELogVerbosity::Error);
+				FindAndPrintStaleReferencesToObject(World, 
+					UObjectBaseUtility::IsPendingKillEnabled() ? EPrintStaleReferencesOptions::Fatal : (EPrintStaleReferencesOptions::Error | EPrintStaleReferencesOptions::Ensure));
 			}
 		}
 	}
diff --git a/Engine/Source/Editor/UnrealEd/Private/EditorLevelUtils.cpp b/Engine/Source/Editor/UnrealEd/Private/EditorLevelUtils.cpp
--- a/Engine/Source/Editor/UnrealEd/Private/EditorLevelUtils.cpp
+++ b/Engine/Source/Editor/UnrealEd/Private/EditorLevelUtils.cpp
@@ -859,7 +859,7 @@ bool UEditorLevelUtils::RemoveLevelsFromWorld(TArray<ULevel*> InLevels, bool bCl
 	// Reset transaction buffer and run GC to clear out the destroyed level
 	GEditor->Cleanse(bClearSelection, false, TransResetText, bResetTransBuffer);
 
-	auto CheckPackage = [](const TArray<FName>& PackageNames, ELogVerbosity::Type Verbosity)
+	auto CheckPackage = [](const TArray<FName>& PackageNames, EPrintStaleReferencesOptions Options)
 	{
 		// Check Package no longer exists
 		for (const FName& LevelPackageName : PackageNames)
@@ -871,7 +871,7 @@ bool UEditorLevelUtils::RemoveLevelsFromWorld(TArray<ULevel*> InLevels, bool bCl
 				UWorld* TheWorld = UWorld::FindWorldInPackage(LevelPackage->GetOutermost());
 				if (TheWorld != nullptr)
 				{
-					UEngine::FindAndPrintStaleReferencesToObject(TheWorld, Verbosity);
+					UEngine::FindAndPrintStaleReferencesToObject(TheWorld, Options);
 					return false;
 				}
 			}
@@ -880,17 +880,17 @@ bool UEditorLevelUtils::RemoveLevelsFromWorld(TArray<ULevel*> InLevels, bool bCl
 	};
 
 	// Check that packages no longer exist
-	ELogVerbosity::Type Verbosity = ELogVerbosity::Log;
+	EPrintStaleReferencesOptions PrintStaleReferencesOptions = EPrintStaleReferencesOptions::Log;
 	if (bResetTransBuffer)
 	{
-		Verbosity = UObjectBaseUtility::IsPendingKillEnabled() ? ELogVerbosity::Fatal : ELogVerbosity::Error;
+		PrintStaleReferencesOptions = UObjectBaseUtility::IsPendingKillEnabled() ? EPrintStaleReferencesOptions::Fatal : (EPrintStaleReferencesOptions::Error | EPrintStaleReferencesOptions::Ensure);
 	}
-	bool bFailed = !CheckPackage(PackageNames, Verbosity);
-	if (bFailed && Verbosity != ELogVerbosity::Fatal)
+	bool bFailed = !CheckPackage(PackageNames, PrintStaleReferencesOptions);
+	if (bFailed && !(PrintStaleReferencesOptions & EPrintStaleReferencesOptions::Fatal))
 	{
 		// We tried avoiding clearing the Transaction buffer but it failed. Plan B.
 		GEditor->Cleanse(bClearSelection, false, TransResetText, true);
-		CheckPackage(PackageNames, Verbosity);
+		CheckPackage(PackageNames, PrintStaleReferencesOptions);
 	}
 
 	return true;
diff --git a/Engine/Source/Editor/UnrealEd/Private/EditorServer.cpp b/Engine/Source/Editor/UnrealEd/Private/EditorServer.cpp
--- a/Engine/Source/Editor/UnrealEd/Private/EditorServer.cpp
+++ b/Engine/Source/Editor/UnrealEd/Private/EditorServer.cpp
@@ -2018,7 +2018,7 @@ void UEditorEngine::CheckForWorldGCLeaks( UWorld* NewWorld, UPackage* WorldPacka
 		const bool bIsPersistantWorldType = (RemainingWorld->WorldType == EWorldType::Inactive) || (RemainingWorld->WorldType == EWorldType::EditorPreview) || (RemainingWorld->WorldType == EWorldType::GamePreview);
 		if(!bIsNewWorld && !bIsPersistantWorldType && !WorldHasValidContext(RemainingWorld))
 		{
-			FindAndPrintStaleReferencesToObject(RemainingWorld, ELogVerbosity::Error);
+			FindAndPrintStaleReferencesToObject(RemainingWorld, EPrintStaleReferencesOptions::Error);
 			NumFailedToCleanup++;
 		}
 	}
@@ -2032,7 +2032,7 @@ void UEditorEngine::CheckForWorldGCLeaks( UWorld* NewWorld, UPackage* WorldPacka
 			const bool bIsNewWorldPackage = (NewWorldPackage && RemainingPackage == NewWorldPackage);
 			if(!bIsNewWorldPackage && RemainingPackage == WorldPackage)
 			{
-				FindAndPrintStaleReferencesToObject(RemainingPackage, ELogVerbosity::Error);
+				FindAndPrintStaleReferencesToObject(RemainingPackage, EPrintStaleReferencesOptions::Error);
 				NumFailedToCleanup++;
 			}
 		}
diff --git a/Engine/Source/Editor/UnrealEd/Private/PlayLevel.cpp b/Engine/Source/Editor/UnrealEd/Private/PlayLevel.cpp
--- a/Engine/Source/Editor/UnrealEd/Private/PlayLevel.cpp
+++ b/Engine/Source/Editor/UnrealEd/Private/PlayLevel.cpp
@@ -446,9 +446,29 @@ void UEditorEngine::EndPlayMap()
 			Arguments.Add(TEXT("Path"), FText::FromString(RefChainSearch.GetRootPath()));
 				
 			// We cannot safely recover from this.
-			FMessageLog(NAME_CategoryPIE).CriticalError()
-				->AddToken(FUObjectToken::Create(Object, FText::FromString(Object->GetFullName())))
-				->AddToken(FTextToken::Create(FText::Format(LOCTEXT("PIEObjectStillReferenced", "Object from PIE level still referenced. Shortest path from root: {Path}"), Arguments)));
+			if (UObjectBaseUtility::IsPendingKillEnabled())
+			{
+				FMessageLog(NAME_CategoryPIE).CriticalError()
+					->AddToken(FUObjectToken::Create(Object, FText::FromString(Object->GetFullName())))
+					->AddToken(FTextToken::Create(FText::Format(LOCTEXT("PIEObjectStillReferenced", "Object from PIE level still referenced. Shortest path from root: {Path}"), Arguments)));
+			}
+			else
+			{
+				// And yet we still try and recover
+				UPackage* ObjectPackage = Object->GetOutermost();
+				ObjectPackage->ClearPackageFlags(PKG_PlayInEditor);
+				ObjectPackage->ClearFlags(RF_Standalone);
+				// We let it leak but it needs to be renamed to not collide with future attempts at creating the same object(s)
+				ObjectPackage->Rename(*MakeUniqueObjectName(nullptr, UPackage::StaticClass()).ToString(), nullptr, REN_ForceNoResetLoaders | REN_DontCreateRedirectors | REN_DoNotDirty | REN_NonTransactional);
+
+				FText ErrorMessage = FText::Format(LOCTEXT("PIEObjectStillReferenced", "Object from PIE level still referenced. Shortest path from root: {Path}"), Arguments);
+				// Throw an ensure so that crash reporter logs this error
+				ensureAlwaysMsgf(false, TEXT("%s"), *ErrorMessage.ToString());
+				// Display the error in the editor window
+				FMessageLog(NAME_CategoryPIE).Error()
+					->AddToken(FUObjectToken::Create(Object, FText::FromString(Object->GetFullName())))
+					->AddToken(FTextToken::Create(ErrorMessage));
+			}
 		}
 	}
 
diff --git a/Engine/Source/Runtime/Engine/Classes/Engine/Engine.h b/Engine/Source/Runtime/Engine/Classes/Engine/Engine.h
--- a/Engine/Source/Runtime/Engine/Classes/Engine/Engine.h
+++ b/Engine/Source/Runtime/Engine/Classes/Engine/Engine.h
@@ -613,6 +613,20 @@ struct FPluginRedirect
 	FString NewPluginName;
 };
 
+/* Options for UEngine::FindAndPrintStaleReferencesToObject function */
+enum class EPrintStaleReferencesOptions
+{
+	None = 0,
+	Log = 1,
+	Display = 2,
+	Warning = 3,
+	Error = 4,
+	Fatal = 5,
+	Ensure = 0x00000100,
+
+	VerbosityMask = 0x000000ff
+};
+ENUM_CLASS_FLAGS(EPrintStaleReferencesOptions);
 
 /** Game thread events for dynamic resolution state. */
 enum class EDynamicResolutionStateEvent : uint8;
@@ -3133,8 +3147,16 @@ public:
 	 * @param ObjectToFindReferencesTo World or its package (or any object from the world package that should've been destroyed)
 	 * @param Verbosity Verbosity (can be fatal or non-fatal) with which to print the error message with
 	 */
+	UE_DEPRECATED(5.1, "Please use FindAndPrintStaleReferencesToObject overload that takes EPrintStaleReferencesOptions parameter")
 	static void FindAndPrintStaleReferencesToObject(UObject* ObjectToFindReferencesTo, ELogVerbosity::Type Verbosity);
 
+	/**
+	 * Attempts to find what is referencing a world that should have been garbage collected
+	 * @param ObjectToFindReferencesTo World or its package (or any object from the world package that should've been destroyed)
+	 * @param Options Determines how the stale references messages should be logged
+	 */
+	static void FindAndPrintStaleReferencesToObject(UObject* ObjectToFindReferencesTo, EPrintStaleReferencesOptions Options);
+
 	FWorldContext& CreateNewWorldContext(EWorldType::Type WorldType);
 
 	virtual void DestroyWorldContext(UWorld * InWorld);
diff --git a/Engine/Source/Runtime/Engine/Private/LevelStreaming.cpp b/Engine/Source/Runtime/Engine/Private/LevelStreaming.cpp
--- a/Engine/Source/Runtime/Engine/Private/LevelStreaming.cpp
+++ b/Engine/Source/Runtime/Engine/Private/LevelStreaming.cpp
@@ -1163,7 +1163,7 @@ bool ULevelStreaming::RequestLevel(UWorld* PersistentWorld, bool bAllowLevelLoad
 				UE_LOG(LogLevelStreaming, Error, TEXT("World exists but PersistentLevel doesn't for %s, most likely caused by reference to world of unloaded level and GC setting reference to null while keeping world object"), *World->GetOutermost()->GetName());
 				UE_LOG(LogLevelStreaming, Error, TEXT("Most likely caused by reference to world of unloaded level and GC setting reference to null while keeping world object. Referenced by:"));
 
-				UEngine::FindAndPrintStaleReferencesToObject(World, UObjectBaseUtility::IsPendingKillEnabled() ? ELogVerbosity::Fatal : ELogVerbosity::Error);
+				UEngine::FindAndPrintStaleReferencesToObject(World, UObjectBaseUtility::IsPendingKillEnabled() ? EPrintStaleReferencesOptions::Fatal : (EPrintStaleReferencesOptions::Error | EPrintStaleReferencesOptions::Ensure));
 
 				return false;
 			}
diff --git a/Engine/Source/Runtime/Engine/Private/UnrealEngine.cpp b/Engine/Source/Runtime/Engine/Private/UnrealEngine.cpp
--- a/Engine/Source/Runtime/Engine/Private/UnrealEngine.cpp
+++ b/Engine/Source/Runtime/Engine/Private/UnrealEngine.cpp
@@ -15004,7 +15004,7 @@ void UEngine::VerifyLoadMapWorldCleanup()
 				UE_LOG(LogLoad, Error, TEXT("Previously active world %s not cleaned up by garbage collection!"), *World->GetPathName());
 				UE_LOG(LogLoad, Error, TEXT("Once a world has become active, it cannot be reused and must be destroyed and reloaded. World referenced by:"));
 
-				FindAndPrintStaleReferencesToObject(World, UObjectBaseUtility::IsPendingKillEnabled() ? ELogVerbosity::Fatal : ELogVerbosity::Error);
+				FindAndPrintStaleReferencesToObject(World, UObjectBaseUtility::IsPendingKillEnabled() ? EPrintStaleReferencesOptions::Fatal : (EPrintStaleReferencesOptions::Error | EPrintStaleReferencesOptions::Ensure));
 			}
 		}
 	}
@@ -15061,7 +15061,41 @@ static FString GetPathToStaleObjectReferencer(UObject* ObjectToFindReferencesTo,
 
 void UEngine::FindAndPrintStaleReferencesToObject(UObject* ObjectToFindReferencesTo, ELogVerbosity::Type Verbosity)
 {
-	check(ObjectToFindReferencesTo);
+	EPrintStaleReferencesOptions Options = EPrintStaleReferencesOptions::None;
+	switch (Verbosity)
+	{
+	case ELogVerbosity::NoLogging:
+		Options = EPrintStaleReferencesOptions::None;
+		break;
+	case ELogVerbosity::Log:
+		Options = EPrintStaleReferencesOptions::Log;
+		break;
+	case ELogVerbosity::Display:
+		Options = EPrintStaleReferencesOptions::Display;
+		break;
+	case ELogVerbosity::Warning:
+		Options = EPrintStaleReferencesOptions::Warning;
+		break;
+	case ELogVerbosity::Error:
+		Options = EPrintStaleReferencesOptions::Error;
+		break;
+	case ELogVerbosity::Fatal:
+		Options = EPrintStaleReferencesOptions::Fatal;
+		break;
+	default:
+		checkf(false, TEXT("Unsupported Verbosity: %d"), Verbosity);
+		break;
+	}
+	FindAndPrintStaleReferencesToObject(ObjectToFindReferencesTo, Options);
+}
+
+void UEngine::FindAndPrintStaleReferencesToObject(UObject* ObjectToFindReferencesTo, EPrintStaleReferencesOptions Options)
+{
+	checkf(ObjectToFindReferencesTo, TEXT("ObjectToFindReferencesTo cannot be null"));
+	const ELogVerbosity::Type OptionsToVerbosityMapping[] = { ELogVerbosity::NoLogging, ELogVerbosity::Log, ELogVerbosity::Display, ELogVerbosity::Warning, ELogVerbosity::Error, ELogVerbosity::Fatal };
+	const int32 VerbosityIndex = int32(Options & EPrintStaleReferencesOptions::VerbosityMask);
+	checkf(VerbosityIndex < UE_ARRAY_COUNT(OptionsToVerbosityMapping), TEXT("Unsupported verbosity option provided: %d"), VerbosityIndex);
+	ELogVerbosity::Type Verbosity = OptionsToVerbosityMapping[VerbosityIndex];
 
 	UE_LOG(LogLoad, Log, TEXT("Looking for existing references to %s..."), *ObjectToFindReferencesTo->GetFullName());
 
@@ -15122,6 +15156,14 @@ void UEngine::FindAndPrintStaleReferencesToObject(UObject* ObjectToFindReference
 		}
 	}
 
+	if ((Options & EPrintStaleReferencesOptions::Ensure) == EPrintStaleReferencesOptions::Ensure)
+	{
+		ensureAlwaysMsgf(false, TEXT("Old %s not cleaned up by GC! %s:") LINE_TERMINATOR TEXT("%s"),
+			*ObjectToFindReferencesTo->GetFullName(),
+			*GarbageErrorMessage,
+			*PathToCulprit);
+	}
+
 	if (Verbosity == ELogVerbosity::Fatal)
 	{
 		UE_LOG(LogLoad, Fatal, TEXT("Old %s not cleaned up by GC! %s:") LINE_TERMINATOR TEXT("%s"),
