4#jira UE-137617
Make logic for r.VertexDeformationOutputsVelocity a bit easier to follow.
Now it has an effect under all conditions and has setting off/on/auto.
The auto setting does the best thing according to the velocity pass setting.

diff --git a/Engine/Source/Runtime/Engine/Classes/Engine/RendererSettings.h b/Engine/Source/Runtime/Engine/Classes/Engine/RendererSettings.h
--- a/Engine/Source/Runtime/Engine/Classes/Engine/RendererSettings.h
+++ b/Engine/Source/Runtime/Engine/Classes/Engine/RendererSettings.h
@@ -79,7 +79,7 @@ namespace EEarlyZPass
 }
 
 /**
- * Enumerates available options for early Z-passes.
+ * Enumerates available options for velocity pass.
  */
 UENUM()
 namespace EVelocityOutputPass
@@ -92,6 +92,20 @@ namespace EVelocityOutputPass
 	};
 }
 
+/**
+ * Enumerates available options for Vertex Deformation Outputs Velocity.
+ */
+UENUM()
+namespace EVertexDeformationOutputsVelocity
+{
+	enum Type
+	{
+		Off = 0 UMETA(ToolTip = "Always off"),
+		On = 1 UMETA(ToolTip = "Always on"),
+		Auto = 2 UMETA(ToolTip = "On when the performance cost is low (velocity in depth or base pass)."),
+	};
+}
+
 /**
  * Enumerates available options for alpha channel through post processing. The renderer will always generate premultiplied RGBA
  * with alpha as translucency (0 = fully opaque; 1 = fully translucent).
@@ -705,10 +719,8 @@ class ENGINE_API URendererSettings : public UDeveloperSettings
 	UPROPERTY(config, EditAnywhere, Category=Optimizations, meta=(
 		DisplayName="Output velocities due to vertex deformation",
 		ConsoleVariable="r.VertexDeformationOutputsVelocity",
-		EditCondition = "VelocityPass == EVelocityOutputPass::AfterBasePass",
-		ToolTip="Enables materials with World Position Offset and/or World Displacement to output velocities during velocity pass even when the actor has not moved. This incurs a performance cost and can be quite significant if many objects are using WPO, such as a forest of trees. Instead consider outputting velocities during depth or base pass (r.VelocityOutputPass)."
-		))
-	uint32 bVertexDeformationOutputsVelocity:1;
+		ToolTip="Enables materials with World Position Offset and/or World Displacement to output velocities during the velocity pass even when the actor has not moved. \nIf the VelocityPass is set to 'Write after base pass' this can incur a performance cost due to additional draw calls. \nThat performance cost is higher if many objects are using World Position Offset. A forest of trees for example." ))
+	TEnumAsByte<EVertexDeformationOutputsVelocity::Type> VertexDeformationOutputsVelocity;
 
 	UPROPERTY(config, EditAnywhere, Category=Optimizations, meta=(
 		ConsoleVariable="r.SelectiveBasePassOutputs", DisplayName="Selectively output to the GBuffer rendertargets",
diff --git a/Engine/Source/Runtime/Engine/Private/PrimitiveSceneProxy.cpp b/Engine/Source/Runtime/Engine/Private/PrimitiveSceneProxy.cpp
--- a/Engine/Source/Runtime/Engine/Private/PrimitiveSceneProxy.cpp
+++ b/Engine/Source/Runtime/Engine/Private/PrimitiveSceneProxy.cpp
@@ -40,10 +40,11 @@ static TAutoConsoleVariable<int32> CVarCacheWPOPrimitives(
 
 static TAutoConsoleVariable<int32> CVarVertexDeformationOutputsVelocity(
 	TEXT("r.VertexDeformationOutputsVelocity"),
-	0,
+	2,
 	TEXT("Enables materials with World Position Offset and/or World Displacement to output velocities during velocity pass even when the actor has not moved. \n")
-	TEXT("This only has an impact if r.VelocityOutputPass=2. \n")
-	TEXT("This will incur a performance cost that can be quite significant if many objects are using WPO, such as a forest of trees."));
+	TEXT("0=Off, 1=On, 2=Auto(Default). \n")
+	TEXT("Auto setting is off if r.VelocityOutputPass=2, or else on. \n")
+	TEXT("When r.VelocityOutputPass=2 this can incur a performance cost due to additional draw calls."));
 
 bool CacheShadowDepthsFromPrimitivesUsingWPO()
 {
@@ -107,13 +108,10 @@ bool SupportsNaniteRendering(const FVertexFactory* RESTRICT VertexFactory, const
 
 static bool VertexDeformationOutputsVelocity()
 {
-	// This is only optional when velocity pass is after base pass (when it could create many more draw calls).
-	// Note this output pass setting is read only.
 	static const auto CVarVelocityOutputPass = IConsoleManager::Get().FindConsoleVariable(TEXT("r.VelocityOutputPass"));
-	bool bVertexDeformationOutputsVelocityForPass = CVarVelocityOutputPass && CVarVelocityOutputPass->GetInt() != 2;
-	
-	bool bVertexDeformationOutputsVelocity = CVarVertexDeformationOutputsVelocity.GetValueOnAnyThread() != 0;
-	return bVertexDeformationOutputsVelocityForPass || bVertexDeformationOutputsVelocity;
+	const bool bVertexDeformationOutputsVelocityDefault = CVarVelocityOutputPass && CVarVelocityOutputPass->GetInt() != 2;
+	const int32 VertexDeformationOutputsVelocity = CVarVertexDeformationOutputsVelocity.GetValueOnAnyThread();
+	return VertexDeformationOutputsVelocity == 1 || (VertexDeformationOutputsVelocity == 2 && bVertexDeformationOutputsVelocityDefault);
 }
 
 FPrimitiveSceneProxy::FPrimitiveSceneProxy(const UPrimitiveComponent* InComponent, FName InResourceName)
diff --git a/Engine/Source/Runtime/Renderer/Private/RendererScene.cpp b/Engine/Source/Runtime/Renderer/Private/RendererScene.cpp
--- a/Engine/Source/Runtime/Renderer/Private/RendererScene.cpp
+++ b/Engine/Source/Runtime/Renderer/Private/RendererScene.cpp
@@ -128,12 +128,6 @@ static TAutoConsoleVariable<int32> CVarBasePassWriteDepthEvenWithFullPrepass(
 	TEXT("0 to allow a readonly base pass, which skips an MSAA depth resolve, and allows masked materials to get EarlyZ (writing to depth while doing clip() disables EarlyZ) (default)\n")
 	TEXT("1 to force depth writes in the base pass.  Useful for debugging when the prepass and base pass don't match what they render."));
 
-static TAutoConsoleVariable<int32> CVarWPOPrimitivesOutputVelocity(
-	TEXT("r.WPOPrimitivesOutputVelocity"),
-	0,
-	TEXT("Whether primitives using World Position Offset (WPO) materials output velocity."),
-	ECVF_RenderThreadSafe);
-
 DECLARE_CYCLE_STAT(TEXT("DeferredShadingSceneRenderer MotionBlurStartFrame"), STAT_FDeferredShadingSceneRenderer_MotionBlurStartFrame, STATGROUP_SceneRendering);
 
 IMPLEMENT_GLOBAL_SHADER_PARAMETER_STRUCT(FDistanceCullFadeUniformShaderParameters, "PrimitiveFade");
@@ -4215,7 +4209,7 @@ struct FPrimitiveArraySortKey
 
 static bool ShouldPrimitiveOutputVelocity(const FPrimitiveSceneProxy* Proxy, const FStaticShaderPlatform ShaderPlatform)
 {
-	bool bShouldPrimitiveOutputVelocity = Proxy->DrawsVelocity() || (!!CVarWPOPrimitivesOutputVelocity.GetValueOnRenderThread() && Proxy->IsUsingWPOMaterial());
+	bool bShouldPrimitiveOutputVelocity = Proxy->DrawsVelocity();
 
 	bool bPlatformSupportsVelocityRendering = PlatformSupportsVelocityRendering(ShaderPlatform);
 
