Fix merge issue

diff --git a/Engine/Source/Runtime/Engine/Private/Level.cpp b/Engine/Source/Runtime/Engine/Private/Level.cpp
--- a/Engine/Source/Runtime/Engine/Private/Level.cpp
+++ b/Engine/Source/Runtime/Engine/Private/Level.cpp
@@ -3225,7 +3225,7 @@ TArray<UPackage*> ULevel::GetLoadedExternalObjectPackages() const
 			TArray<bool> PackageIsInExternalObjectsPath;
 			PackageIsInExternalObjectsPath.InsertUninitialized(0, Packages.Num());
 
-			ParallelFor(Packages.Num(), [&Packages = std::as_const(Packages), &ExternalObjectsPath = std::as_const(ExternalObjectsPath), &PackageIsInExternalObjectsPath](int32 Index) {
+			ParallelFor(Packages.Num(), [&Packages = std::as_const(Packages), &ExternalObjectsPaths = std::as_const(ExternalObjectsPaths), &PackageIsInExternalObjectsPath](int32 Index) {
 				UPackage* Package = static_cast<UPackage*>(Packages[Index]);
 
 				TStringBuilder<256> PackageName;
