Fix merge issue #2

diff --git a/Engine/Source/Runtime/Engine/Private/Level.cpp b/Engine/Source/Runtime/Engine/Private/Level.cpp
--- a/Engine/Source/Runtime/Engine/Private/Level.cpp
+++ b/Engine/Source/Runtime/Engine/Private/Level.cpp
@@ -3241,7 +3241,7 @@ TArray<UPackage*> ULevel::GetLoadedExternalObjectPackages() const
 					}
 				}
 				PackageIsInExternalObjectsPath[Index] = bIsInExternalObjectsPath;
-			}
+			});
 
 			for (int Index=0; Index<PackageIsInExternalObjectsPath.Num(); Index++)
 			{
