Fix merge #3

diff --git a/Engine/Source/Runtime/Engine/Private/Level.cpp b/Engine/Source/Runtime/Engine/Private/Level.cpp
--- a/Engine/Source/Runtime/Engine/Private/Level.cpp
+++ b/Engine/Source/Runtime/Engine/Private/Level.cpp
@@ -3247,7 +3247,7 @@ TArray<UPackage*> ULevel::GetLoadedExternalObjectPackages() const
 			{
 				if (PackageIsInExternalObjectsPath[Index])
 				{
-					ActorPackages.Add(static_cast<UPackage*>(Packages[Index]));
+					ExternalObjectPackages.Add(static_cast<UPackage*>(Packages[Index]));
 				}
 			}
 		}
