Fix merge #4

diff --git a/Engine/Source/Runtime/Engine/Private/Level.cpp b/Engine/Source/Runtime/Engine/Private/Level.cpp
--- a/Engine/Source/Runtime/Engine/Private/Level.cpp
+++ b/Engine/Source/Runtime/Engine/Private/Level.cpp
@@ -3219,36 +3219,35 @@ TArray<UPackage*> ULevel::GetLoadedExternalObjectPackages() const
 
 	if (!ExternalObjectsPaths.IsEmpty())
 	{
-			TArray<UObject*> Packages;
-			GetObjectsOfClass(UPackage::StaticClass(), Packages, /*bIncludeDerivedClasses =*/ true, /*EObjectFlags ExcludeFlags =*/ RF_ClassDefaultObject,/*EInternalObjectFlags ExclusionInternalFlags =*/ EInternalObjectFlags::None);
+		TArray<UObject*> Packages;
+		GetObjectsOfClass(UPackage::StaticClass(), Packages, /*bIncludeDerivedClasses =*/ true, /*EObjectFlags ExcludeFlags =*/ RF_ClassDefaultObject,/*EInternalObjectFlags ExclusionInternalFlags =*/ EInternalObjectFlags::None);
 
-			TArray<bool> PackageIsInExternalObjectsPath;
-			PackageIsInExternalObjectsPath.InsertUninitialized(0, Packages.Num());
+		TArray<bool> PackageIsInExternalObjectsPath;
+		PackageIsInExternalObjectsPath.InsertUninitialized(0, Packages.Num());
 
-			ParallelFor(Packages.Num(), [&Packages = std::as_const(Packages), &ExternalObjectsPaths = std::as_const(ExternalObjectsPaths), &PackageIsInExternalObjectsPath](int32 Index) {
-				UPackage* Package = static_cast<UPackage*>(Packages[Index]);
+		ParallelFor(Packages.Num(), [&Packages = std::as_const(Packages), &ExternalObjectsPaths = std::as_const(ExternalObjectsPaths), &PackageIsInExternalObjectsPath](int32 Index) {
+			UPackage* Package = static_cast<UPackage*>(Packages[Index]);
 
-				TStringBuilder<256> PackageName;
-				Package->GetLoadedPath().AppendPackageName(PackageName);
-				FStringView PackageNameStringView(PackageName);
-				bool bIsInExternalObjectsPath = false;
-				for (const FString& ExternalObjectsPath : ExternalObjectsPaths)
+			TStringBuilder<256> PackageName;
+			Package->GetLoadedPath().AppendPackageName(PackageName);
+			FStringView PackageNameStringView(PackageName);
+			bool bIsInExternalObjectsPath = false;
+			for (const FString& ExternalObjectsPath : ExternalObjectsPaths)
+			{
+				bIsInExternalObjectsPath = PackageNameStringView.Contains(ExternalObjectsPath);
+				if (bIsInExternalObjectsPath)
 				{
-					bIsInExternalObjectsPath = PackageNameStringView.Contains(ExternalObjectsPath);
-					if (bIsInExternalObjectsPath)
-					{
-						break;
-					}
+					break;
 				}
-				PackageIsInExternalObjectsPath[Index] = bIsInExternalObjectsPath;
-			});
+			}
+			PackageIsInExternalObjectsPath[Index] = bIsInExternalObjectsPath;
+		});
 
-			for (int Index=0; Index<PackageIsInExternalObjectsPath.Num(); Index++)
+		for (int Index=0; Index<PackageIsInExternalObjectsPath.Num(); Index++)
+		{
+			if (PackageIsInExternalObjectsPath[Index])
 			{
-				if (PackageIsInExternalObjectsPath[Index])
-				{
-					ExternalObjectPackages.Add(static_cast<UPackage*>(Packages[Index]));
-				}
+				ExternalObjectPackages.Add(static_cast<UPackage*>(Packages[Index]));
 			}
 		}
 	}
