Fix AssetManager ModifyCook to correctly apply cook rules to asset bundles so it handles labels properly

diff --git a/Engine/Source/Runtime/Engine/Private/AssetManager.cpp b/Engine/Source/Runtime/Engine/Private/AssetManager.cpp
--- a/Engine/Source/Runtime/Engine/Private/AssetManager.cpp
+++ b/Engine/Source/Runtime/Engine/Private/AssetManager.cpp
@@ -3932,19 +3932,17 @@ void UAssetManager::ModifyCook(TArray<FName>& PackagesToCook, TArray<FName>& Pac
 				// If this has an asset data, add that package name
 				AssetPackages.Add(AssetData.PackageName);
 			}
-			else
+
+			// Also add any bundle assets to handle cook rules for labels
+			TArray<FAssetBundleEntry> FoundEntries;
+			if (GetAssetBundleEntries(PrimaryAssetId, FoundEntries))
 			{
-				// If not, this may have bundles, so add those
-				TArray<FAssetBundleEntry> FoundEntries;
-				if (GetAssetBundleEntries(PrimaryAssetId, FoundEntries))
+				for (const FAssetBundleEntry& FoundEntry : FoundEntries)
 				{
-					for (const FAssetBundleEntry& FoundEntry : FoundEntries)
+					for (const FSoftObjectPath& FoundReference : FoundEntry.BundleAssets)
 					{
-						for (const FSoftObjectPath& FoundReference : FoundEntry.BundleAssets)
-						{
-							FName PackageName = FName(*FoundReference.GetLongPackageName());
-							AssetPackages.AddUnique(PackageName);
-						}
+						FName PackageName = FName(*FoundReference.GetLongPackageName());
+						AssetPackages.AddUnique(PackageName);
 					}
 				}
 			}
