WP HLOD - Add Landscape support

- Added landscape default flatten material

diff --git a/Engine/Config/BaseEngine.ini b/Engine/Config/BaseEngine.ini
--- a/Engine/Config/BaseEngine.ini
+++ b/Engine/Config/BaseEngine.ini
@@ -170,6 +170,7 @@ DefaultDeferredDecalMaterialName=/Engine/EngineMaterials/DefaultDeferredDecalMat
 DefaultPostProcessMaterialName=/Engine/EngineMaterials/DefaultPostProcessMaterial.DefaultPostProcessMaterial
 DefaultFlattenMaterialName=/Engine/EngineMaterials/FlattenMaterial.FlattenMaterial
 DefaultHLODFlattenMaterialName=/Engine/EngineMaterials/FlattenMaterial_VT.FlattenMaterial_VT
+DefaultLandscapeFlattenMaterialName=/Engine/EngineMaterials/FlattenMaterial_WS_Normal_VT.FlattenMaterial_WS_Normal_VT
 TimecodeProviderClassName=None
 ArrowMaterialName=/Engine/EngineMaterials/GizmoMaterial.GizmoMaterial
 ParticleEventManagerClassPath="/Script/Engine.ParticleEventManager"
diff --git a/Engine/Source/Runtime/Engine/Classes/Engine/Engine.h b/Engine/Source/Runtime/Engine/Classes/Engine/Engine.h
--- a/Engine/Source/Runtime/Engine/Classes/Engine/Engine.h
+++ b/Engine/Source/Runtime/Engine/Classes/Engine/Engine.h
@@ -1042,11 +1042,17 @@ public:
 	UPROPERTY(globalconfig)
 	FSoftObjectPath DefaultHLODFlattenMaterialName;
 
+	/** A material used to flatten materials to VT textures, with the normals being in world space. */
+	UPROPERTY(globalconfig)
+	FSoftObjectPath DefaultLandscapeFlattenMaterialName;
+
 	/** Materials used when flattening materials */
 	UPROPERTY()
 	TObjectPtr<class UMaterial> DefaultFlattenMaterial;
 	UPROPERTY()
 	TObjectPtr<class UMaterial> DefaultHLODFlattenMaterial;
+	UPROPERTY()
+	TObjectPtr<class UMaterial> DefaultLandscapeFlattenMaterial;
 #endif
 
 	/** A material used to render debug opaque material. Used in various animation editor viewport features. */
diff --git a/Engine/Source/Runtime/Engine/Private/UnrealEngine.cpp b/Engine/Source/Runtime/Engine/Private/UnrealEngine.cpp
--- a/Engine/Source/Runtime/Engine/Private/UnrealEngine.cpp
+++ b/Engine/Source/Runtime/Engine/Private/UnrealEngine.cpp
@@ -2999,6 +2999,7 @@ void UEngine::InitializeObjectReferences()
 	{
 		LoadSpecialMaterial(TEXT("DefaultFlattenMaterialName"), DefaultFlattenMaterialName.ToString(), DefaultFlattenMaterial, false);
 		LoadSpecialMaterial(TEXT("DefaultHLODFlattenMaterialName"), DefaultHLODFlattenMaterialName.ToString(), DefaultHLODFlattenMaterial, false);
+		LoadSpecialMaterial(TEXT("DefaultLandscapeFlattenMaterialName"), DefaultLandscapeFlattenMaterialName.ToString(), DefaultLandscapeFlattenMaterial, false);
 	}
 #endif
 
