WP HLOD - Add Landscape support

HLOD builder fixes
- Wait for assets compilation if any of the custom builders requires it.
- Static mesh component HLOD hash wasn't taking the static mesh/material/textures hashes like it used to (caused by recent changes)

diff --git a/Engine/Plugins/Editor/WorldPartitionHLODUtilities/Source/Private/WorldPartitionHLODUtilities.cpp b/Engine/Plugins/Editor/WorldPartitionHLODUtilities/Source/Private/WorldPartitionHLODUtilities.cpp
--- a/Engine/Plugins/Editor/WorldPartitionHLODUtilities/Source/Private/WorldPartitionHLODUtilities.cpp
+++ b/Engine/Plugins/Editor/WorldPartitionHLODUtilities/Source/Private/WorldPartitionHLODUtilities.cpp
@@ -16,8 +16,6 @@
 #include "Serialization/ArchiveCrc32.h"
 #include "Templates/UniquePtr.h"
 
-#include "AssetCompilingManager.h"
-
 #include "HLODBuilderInstancing.h"
 #include "HLODBuilderMeshMerge.h"
 #include "HLODBuilderMeshSimplify.h"
@@ -367,12 +365,6 @@ uint32 FWorldPartitionHLODUtilities::BuildHLOD(AWorldPartitionHLOD* InHLODActor)
 
 			HLODBuilder->SetHLODBuilderSettings(HLODLayer->GetHLODBuilderSettings());
 
-			if (HLODBuilder->RequiresCompiledAssets())
-			{
-				// Wait for compilation to finish
-				FAssetCompilingManager::Get().FinishAllCompilation();
-			}
-
 			FHLODBuildContext HLODBuildContext;
 			HLODBuildContext.World = InHLODActor->GetWorld();
 			HLODBuildContext.AssetsOuter = InHLODActor->GetPackage();
diff --git a/Engine/Source/Runtime/Engine/Private/WorldPartition/HLOD/HLODBuilder.cpp b/Engine/Source/Runtime/Engine/Private/WorldPartition/HLOD/HLODBuilder.cpp
--- a/Engine/Source/Runtime/Engine/Private/WorldPartition/HLOD/HLODBuilder.cpp
+++ b/Engine/Source/Runtime/Engine/Private/WorldPartition/HLOD/HLODBuilder.cpp
@@ -2,6 +2,8 @@
 
 #include "WorldPartition/HLOD/HLODBuilder.h"
 
+#include "AssetCompilingManager.h"
+
 #include "Engine/HLODProxy.h"
 
 DEFINE_LOG_CATEGORY(LogHLODBuilder);
@@ -71,8 +73,33 @@ uint32 UHLODBuilder::ComputeHLODHash(const UActorComponent* InSourceComponent) c
 	{
 		UE_LOG(LogHLODBuilder, VeryVerbose, TEXT(" - Component \'%s\' from actor \'%s\'"), *StaticMeshComponent->GetName(), *StaticMeshComponent->GetOwner()->GetName());
 
+		// CRC component
 		ComponentCRC = UHLODProxy::GetCRC(StaticMeshComponent);
 		UE_LOG(LogHLODBuilder, VeryVerbose, TEXT("     - StaticMeshComponent (%s) = %x"), *StaticMeshComponent->GetName(), ComponentCRC);
+
+		// CRC static mesh
+		if (UStaticMesh* StaticMesh = StaticMeshComponent->GetStaticMesh())
+		{
+			ComponentCRC = UHLODProxy::GetCRC(StaticMesh, ComponentCRC);
+		}
+
+		// CRC materials
+		const int32 NumMaterials = StaticMeshComponent->GetNumMaterials();
+		for (int32 MaterialIndex = 0; MaterialIndex < NumMaterials; ++MaterialIndex)
+		{
+			UMaterialInterface* MaterialInterface = StaticMeshComponent->GetMaterial(MaterialIndex);
+			if (MaterialInterface)
+			{
+				ComponentCRC = UHLODProxy::GetCRC(MaterialInterface, ComponentCRC);
+
+				TArray<UTexture*> Textures;
+				MaterialInterface->GetUsedTextures(Textures, EMaterialQualityLevel::High, true, ERHIFeatureLevel::SM5, true);
+				for (UTexture* Texture : Textures)
+				{
+					ComponentCRC = UHLODProxy::GetCRC(Texture, ComponentCRC);
+				}
+			}
+		}
 	}
 	else
 	{
@@ -117,7 +144,19 @@ TArray<UActorComponent*> UHLODBuilder::Build(const FHLODBuildContext& InHLODBuil
 		TSubclassOf<UHLODBuilder> HLODBuilderClass = SourceComponent->GetCustomHLODBuilderClass();
 		HLODBuildersForComponents.FindOrAdd(HLODBuilderClass).Add(SourceComponent);
 	}
+
+	// If any of the builders requires it, wait for assets compilation to finish
+	for (const auto& HLODBuilderPair : HLODBuildersForComponents)
+	{
+		const UHLODBuilder* HLODBuilder = HLODBuilderPair.Key ? HLODBuilderPair.Key->GetDefaultObject<UHLODBuilder>() : this;
+		if (HLODBuilder->RequiresCompiledAssets())
+		{
+			FAssetCompilingManager::Get().FinishAllCompilation();
+			break;
+		}
+	}	
 	
+	// Build HLOD components by sending source components to the individual builders, in batch
 	TArray<UActorComponent*> HLODComponents;
 	for (const auto& HLODBuilderPair : HLODBuildersForComponents)
 	{
