Fix assertion that can occur when checking the parent material of a material instance for path tracer invalidation

Should have been using GetMaterial_Concurrent() instead of GetMaterial()

diff --git a/Engine/Source/Runtime/Engine/Private/Materials/MaterialInstance.cpp b/Engine/Source/Runtime/Engine/Private/Materials/MaterialInstance.cpp
--- a/Engine/Source/Runtime/Engine/Private/Materials/MaterialInstance.cpp
+++ b/Engine/Source/Runtime/Engine/Private/Materials/MaterialInstance.cpp
@@ -345,7 +345,7 @@ void GameThread_UpdateMIParameter(const UMaterialInstance* Instance, const Param
 {
 	if (FApp::CanEverRender())
 	{
-		const UMaterial* Material = Instance->GetMaterial();
+		const UMaterial* Material = Instance->GetMaterial_Concurrent();
 		if (Material != nullptr)
 		{
 			EMaterialDomain Domain = Material->MaterialDomain;
