Move delegates from ReplaySubsystem to UnrealNetwork

diff --git a/Engine/Source/Runtime/Engine/Private/ReplayHelper.cpp b/Engine/Source/Runtime/Engine/Private/ReplayHelper.cpp
--- a/Engine/Source/Runtime/Engine/Private/ReplayHelper.cpp
+++ b/Engine/Source/Runtime/Engine/Private/ReplayHelper.cpp
@@ -153,6 +153,8 @@ void FReplayHelper::StartRecording(UNetConnection* Connection)
 	Params.bRecord = true;
 	Params.ReplayVersion = FNetworkVersion::GetReplayVersion();
 
+	FNetworkReplayDelegates::OnReplayRecordingStartAttempt.Broadcast(World.Get());
+
 	ReplayStreamer->StartStreaming(Params, FStartStreamingCallback::CreateRaw(this, &FReplayHelper::OnStartRecordingComplete));
 
 	AddNewLevel(GetNameSafe(World->GetOuter()));
diff --git a/Engine/Source/Runtime/Engine/Private/ReplaySubsystem.cpp b/Engine/Source/Runtime/Engine/Private/ReplaySubsystem.cpp
--- a/Engine/Source/Runtime/Engine/Private/ReplaySubsystem.cpp
+++ b/Engine/Source/Runtime/Engine/Private/ReplaySubsystem.cpp
@@ -121,7 +121,7 @@ void UReplaySubsystem::RecordReplay(const FString& Name, const FString& Friendly
 		return;
 	}
 
-	if(IsRecording())
+	if (IsRecording())
 	{
 		UE_LOG(LogDemo, Warning, TEXT("UReplaySubsystem::RecordReplay: A replay is already being recorded, cannot begin recording another one."));
 		return;
@@ -160,7 +160,7 @@ void UReplaySubsystem::RecordReplay(const FString& Name, const FString& Friendly
 		UE_LOG(LogDemo, Log, TEXT("UReplaySubsystem::RecordReplay: Starting recording with replay connection.  Name: %s FriendlyName: %s"), *Name, *FriendlyName);
 
 		Connection->StartRecording();
-		OnReplayRecordingStarted.Broadcast();
+
 		return;
 	}
 
@@ -218,7 +218,6 @@ void UReplaySubsystem::RecordReplay(const FString& Name, const FString& Friendly
 	}
 
 	UE_LOG(LogDemo, Verbose, TEXT("Num Network Actors: %i"), DemoNetDriver->GetNetworkObjectList().GetActiveObjects().Num());
-	OnReplayRecordingStarted.Broadcast();
 }
 
 bool UReplaySubsystem::PlayReplay(const FString& Name, UWorld* WorldOverride, const TArray<FString>& AdditionalOptions)
@@ -311,8 +310,6 @@ void UReplaySubsystem::StopExistingReplays(UWorld* InWorld)
 		Connection->CleanUp();
 		ReplayConnection = nullptr;
 	}
-
-	OnReplayRecordingStopped.Broadcast();
 }
 
 FString UReplaySubsystem::GetActiveReplayName() const
diff --git a/Engine/Source/Runtime/Engine/Private/UnrealNetwork.cpp b/Engine/Source/Runtime/Engine/Private/UnrealNetwork.cpp
--- a/Engine/Source/Runtime/Engine/Private/UnrealNetwork.cpp
+++ b/Engine/Source/Runtime/Engine/Private/UnrealNetwork.cpp
@@ -11,6 +11,7 @@ FOnReplayStartedDelegate FNetworkReplayDelegates::OnReplayStarted;
 FOnReplayStartFailureDelegate FNetworkReplayDelegates::OnReplayStartFailure;
 FOnReplayScrubCompleteDelegate FNetworkReplayDelegates::OnReplayScrubComplete;
 FOnReplayPlaybackCompleteDelegate FNetworkReplayDelegates::OnReplayPlaybackComplete;
+FOnReplayRecordingStartAttemptDelegate FNetworkReplayDelegates::OnReplayRecordingStartAttempt;
 FOnReplayRecordingCompleteDelegate FNetworkReplayDelegates::OnReplayRecordingComplete;
 FOnPauseChannelsChangedDelegate FNetworkReplayDelegates::OnPauseChannelsChanged;
 FOnReplayIDChangedDelegate FNetworkReplayDelegates::OnReplayIDChanged;
diff --git a/Engine/Source/Runtime/Engine/Public/Net/UnrealNetwork.h b/Engine/Source/Runtime/Engine/Public/Net/UnrealNetwork.h
--- a/Engine/Source/Runtime/Engine/Public/Net/UnrealNetwork.h
+++ b/Engine/Source/Runtime/Engine/Public/Net/UnrealNetwork.h
@@ -50,6 +50,7 @@ DECLARE_MULTICAST_DELEGATE_OneParam(FOnReplayStartedDelegate, UWorld* /*World*/)
 DECLARE_MULTICAST_DELEGATE_TwoParams(FOnReplayStartFailureDelegate, UWorld* /*World*/, EDemoPlayFailure::Type /*Error*/);
 DECLARE_MULTICAST_DELEGATE_OneParam(FOnReplayScrubCompleteDelegate, UWorld* /*World*/);
 DECLARE_MULTICAST_DELEGATE_OneParam(FOnReplayPlaybackCompleteDelegate, UWorld* /*World*/);
+DECLARE_MULTICAST_DELEGATE_OneParam(FOnReplayRecordingStartAttemptDelegate, UWorld* /*World*/);
 DECLARE_MULTICAST_DELEGATE_OneParam(FOnReplayRecordingCompleteDelegate, UWorld* /*World*/);
 DECLARE_MULTICAST_DELEGATE_TwoParams(FOnPauseChannelsChangedDelegate, UWorld* /*World*/, bool /*bPaused*/);
 DECLARE_MULTICAST_DELEGATE_TwoParams(FOnReplayIDChangedDelegate, UWorld* /*World*/, const FString& /*ReplayID*/);
@@ -79,6 +80,9 @@ struct ENGINE_API FNetworkReplayDelegates
 	/** Delegate for external systems to be notified when playback ends */
 	static FOnReplayPlaybackCompleteDelegate OnReplayPlaybackComplete;
 
+	/** Public Delegate for external systems to be notified when replay recording is about to start. */
+	static FOnReplayRecordingStartAttemptDelegate OnReplayRecordingStartAttempt;
+
 	/** Public Delegate for external systems to be notified when replay recording is about to finish. */
 	static FOnReplayRecordingCompleteDelegate OnReplayRecordingComplete;
 
diff --git a/Engine/Source/Runtime/Engine/Public/ReplaySubsystem.h b/Engine/Source/Runtime/Engine/Public/ReplaySubsystem.h
--- a/Engine/Source/Runtime/Engine/Public/ReplaySubsystem.h
+++ b/Engine/Source/Runtime/Engine/Public/ReplaySubsystem.h
@@ -8,9 +8,6 @@
 
 class UReplayNetConnection;
 
-DECLARE_DYNAMIC_MULTICAST_DELEGATE(FOnReplayRecordingStarted);
-DECLARE_DYNAMIC_MULTICAST_DELEGATE(FOnReplayRecordingStopped);
-
 UCLASS(DisplayName = "Replay Subsystem")
 class ENGINE_API UReplaySubsystem : public UGameInstanceSubsystem
 {
@@ -124,12 +121,6 @@ public:
 	UPROPERTY(EditAnywhere, Category=Replay)
 	bool bLoadDefaultMapOnStop = true;
 
-	UPROPERTY(BlueprintAssignable)
-	FOnReplayRecordingStarted OnReplayRecordingStarted;
-
-	UPROPERTY(BlueprintAssignable)
-	FOnReplayRecordingStopped OnReplayRecordingStopped;
-
 private:
 	void StopExistingReplays(UWorld* InWorld);
 
