Fix crash loading up ISM components with incorrect data.

- Perform PostLoad fixup on static mesh components that have incorrect data.
- This was happening due to Blueprints having ComponentTemplates with NumCustomFloatData set to 3, but not having any actual instances.
- ISM Components would be duplicated from this template object so they would start with the incorrect state, and later have instances added to them.  For some reason these components did not have their PerInstanceSMCustomData updated.
- The fixup tries to detect this _exact_ problem case and just fix it up:
    - If we have instances, and we NumCustomDataFloats, but we _don't_ have PerInstanceSMCustomData we should allocate enough space in PerInstanceSMCustomData for all the custom data.

diff --git a/Engine/Source/Runtime/Engine/Private/InstancedStaticMesh.cpp b/Engine/Source/Runtime/Engine/Private/InstancedStaticMesh.cpp
--- a/Engine/Source/Runtime/Engine/Private/InstancedStaticMesh.cpp
+++ b/Engine/Source/Runtime/Engine/Private/InstancedStaticMesh.cpp
@@ -4152,6 +4152,16 @@ void UInstancedStaticMeshComponent::PostLoad()
 {
 	Super::PostLoad();
 
+	// Ensure we have the proper amount of per instance custom float data.
+	// We have instances, and we have custom float data per instance, but we don't have custom float data allocated.
+	if (PerInstanceSMData.Num() > 0 && 
+		PerInstanceSMCustomData.Num() == 0 &&
+		NumCustomDataFloats > 0)
+	{
+		UE_LOG(LogStaticMesh, Warning, TEXT("%s has %d instances, and %d NumCustomDataFloats, but no PerInstanceSMCustomData.  Allocating %d custom floats."), *GetFullName(), PerInstanceSMData.Num(), NumCustomDataFloats, PerInstanceSMData.Num() * NumCustomDataFloats);
+		PerInstanceSMCustomData.AddZeroed(PerInstanceSMData.Num() * NumCustomDataFloats);
+	}
+
 	// Has different implementation in HISMC
 	OnPostLoadPerInstanceData();
 }
