VulkanRHI: Fix validation error "requires an image view of type VK_IMAGE_VIEW_TYPE_2D_ARRAY but got VK_IMAGE_VIEW_TYPE_2D" for SSProfilesPreIntegratedTexture.

diff --git a/Engine/Source/Runtime/Engine/Private/SceneManagement.cpp b/Engine/Source/Runtime/Engine/Private/SceneManagement.cpp
--- a/Engine/Source/Runtime/Engine/Private/SceneManagement.cpp
+++ b/Engine/Source/Runtime/Engine/Private/SceneManagement.cpp
@@ -844,7 +844,7 @@ FViewUniformShaderParameters::FViewUniformShaderParameters()
 	SSProfilesTransmissionSampler = TStaticSamplerState<SF_Bilinear, AM_Clamp, AM_Clamp, AM_Clamp>::GetRHI();
 
 	SSProfilesPreIntegratedTextureSizeAndInvSize = FVector4f(1.f,1.f,1.f,1.f);
-	SSProfilesPreIntegratedTexture = GBlackTextureWithSRV->TextureRHI;
+	SSProfilesPreIntegratedTexture = GBlackArrayTexture->TextureRHI;
 	SSProfilesPreIntegratedSampler = TStaticSamplerState<SF_Bilinear, AM_Clamp, AM_Clamp, AM_Clamp>::GetRHI();
 
 	//this can be deleted once sm4 support is removed.
