Support CustomHLSL expressions with new HLSL translator
Still TODO, support texture inputs to custom expression

diff --git a/Engine/Source/Runtime/Core/Public/Misc/MemStackUtility.h b/Engine/Source/Runtime/Core/Public/Misc/MemStackUtility.h
--- a/Engine/Source/Runtime/Core/Public/Misc/MemStackUtility.h
+++ b/Engine/Source/Runtime/Core/Public/Misc/MemStackUtility.h
@@ -25,6 +25,11 @@ inline const TCHAR* AllocateString(FMemStackBase& Allocator, FStringView String)
 	return AllocateString(Allocator, String.GetData(), String.Len());
 }
 
+inline FStringView AllocateStringView(FMemStackBase& Allocator, FStringView String)
+{
+	return FStringView(AllocateString(Allocator, String), String.Len());
+}
+
 inline const TCHAR* AllocateString(FMemStackBase& Allocator, const FStringBuilderBase& StringBuilder)
 {
 	return AllocateString(Allocator, StringBuilder.GetData(), StringBuilder.Len());
diff --git a/Engine/Source/Runtime/Engine/Classes/Materials/MaterialExpressionCustom.h b/Engine/Source/Runtime/Engine/Classes/Materials/MaterialExpressionCustom.h
--- a/Engine/Source/Runtime/Engine/Classes/Materials/MaterialExpressionCustom.h
+++ b/Engine/Source/Runtime/Engine/Classes/Materials/MaterialExpressionCustom.h
@@ -102,6 +102,7 @@ class UMaterialExpressionCustom : public UMaterialExpression
 	virtual uint32 GetInputType(int32 InputIndex) override {return MCT_Unknown;}
 	virtual uint32 GetOutputType(int32 OutputIndex) override;
 	virtual bool IsResultMaterialAttributes(int32 OutputIndex) override;
+	virtual bool GenerateHLSLExpression(FMaterialHLSLGenerator& Generator, UE::HLSLTree::FScope& Scope, int32 OutputIndex, UE::HLSLTree::FExpression*& OutExpression) override;
 #endif // WITH_EDITOR
 	//~ End UMaterialExpression Interface
 };
diff --git a/Engine/Source/Runtime/Engine/Private/HLSLTree/HLSLTreeCommon.cpp b/Engine/Source/Runtime/Engine/Private/HLSLTree/HLSLTreeCommon.cpp
--- a/Engine/Source/Runtime/Engine/Private/HLSLTree/HLSLTreeCommon.cpp
+++ b/Engine/Source/Runtime/Engine/Private/HLSLTree/HLSLTreeCommon.cpp
@@ -549,7 +549,7 @@ bool FExpressionGetStructField::PrepareValue(FEmitContext& Context, FEmitScope&
 		return Context.Errors->AddErrorf(TEXT("Expected type %s"), StructType->Name);
 	}
 
-	StructPreparedType.SetForwardValue(RequestedStructType, StructExpression);
+	//StructPreparedType.SetForwardValue(RequestedStructType, StructExpression);
 	return OutResult.SetType(Context, RequestedType, StructPreparedType.GetFieldType(Field));
 }
 
@@ -1502,6 +1502,25 @@ void FExpressionReflectionVector::EmitValueShader(FEmitContext& Context, FEmitSc
 	OutResult.Code = Context.EmitInlineExpression(Scope, Shader::EValueType::Float3, TEXT("Parameters.ReflectionVector"));
 }
 
+bool FExpressionCustomHLSL::PrepareValue(FEmitContext& Context, FEmitScope& Scope, const FRequestedType& RequestedType, FPrepareValueResult& OutResult) const
+{
+	for (const FCustomHLSLInput& Input : Inputs)
+	{
+		const FPreparedType& InputType = Context.PrepareExpression(Input.Expression, Scope, ERequestedType::Vector4);
+		if (InputType.IsVoid())
+		{
+			return false;
+		}
+	}
+
+	return OutResult.SetType(Context, RequestedType, EExpressionEvaluation::Shader, Shader::FType(OutputStructType));
+}
+
+void FExpressionCustomHLSL::EmitValueShader(FEmitContext& Context, FEmitScope& Scope, const FRequestedType& RequestedType, FEmitValueShaderResult& OutResult) const
+{
+	OutResult.Code = Context.EmitCustomHLSL(Scope, DeclarationCode, FunctionCode, Inputs, OutputStructType);
+}
+
 bool FStatementBreak::Prepare(FEmitContext& Context, FEmitScope& Scope) const
 {
 	return true;
diff --git a/Engine/Source/Runtime/Engine/Private/HLSLTree/HLSLTreeEmit.cpp b/Engine/Source/Runtime/Engine/Private/HLSLTree/HLSLTreeEmit.cpp
--- a/Engine/Source/Runtime/Engine/Private/HLSLTree/HLSLTreeEmit.cpp
+++ b/Engine/Source/Runtime/Engine/Private/HLSLTree/HLSLTreeEmit.cpp
@@ -205,6 +205,97 @@ FEmitContext::~FEmitContext()
 {
 }
 
+namespace Private
+{
+void EmitCustomHLSL(const FEmitCustomHLSL& EmitCustomHLSL, const TCHAR* ParametersTypeName, FStringBuilderBase& OutCode)
+{
+	const Shader::FStructType* OutputStruct = EmitCustomHLSL.OutputType;
+	
+	OutCode.Append(EmitCustomHLSL.DeclarationCode);
+	OutCode.Append(TEXT("\n"));
+
+	// Function that wraps custom HLSL, provides the interface expected by custom HLSL
+	// * LWC inputs have both LWC version and non-LWC version
+	// * First output is given through return value, additional outputs use inout function parameters
+	OutCode.Appendf(TEXT("%s CustomExpressionInternal%d(%s Parameters"), OutputStruct->Fields[0].Type.GetName(), EmitCustomHLSL.Index, ParametersTypeName);
+	for (const FEmitCustomHLSLInput& Input : EmitCustomHLSL.Inputs)
+	{
+		Shader::FType InputType = Input.Type;
+		if (InputType.IsNumericLWC())
+		{
+			// Add an additional input for LWC type with the LWC-prefix
+			OutCode.Appendf(TEXT(", %s LWC"), InputType.GetName());
+			OutCode.Append(Input.Name);
+			// Regular, unprefixed input uses non-LWC type
+			InputType = InputType.GetNonLWCType();
+		}
+
+		OutCode.Appendf(TEXT(", %s "), InputType.GetName());
+		OutCode.Append(Input.Name);
+	}
+	for (int32 OutputIndex = 1; OutputIndex < OutputStruct->Fields.Num(); ++OutputIndex)
+	{
+		const Shader::FStructField& OutputField = OutputStruct->Fields[OutputIndex];
+		OutCode.Appendf(TEXT(", inout %s %s"), OutputField.Type.GetName(), OutputField.Name);
+	}
+	OutCode.Append(TEXT(")\n{\n"));
+	OutCode.Append(EmitCustomHLSL.FunctionCode);
+	OutCode.Append(TEXT("\n}\n"));
+
+	// Function that calls the above wrapper, provides the interface expected by HLSLTree
+	// * No special handling for LWC inputs
+	// * All outputs are provided through a 'struct' type
+	OutCode.Appendf(TEXT("%s CustomExpression%d(%s Parameters"), OutputStruct->Name, EmitCustomHLSL.Index, ParametersTypeName);
+	for (const FEmitCustomHLSLInput& Input : EmitCustomHLSL.Inputs)
+	{
+		OutCode.Appendf(TEXT(", %s "), Input.Type.GetName());
+		OutCode.Append(Input.Name);
+	}
+	OutCode.Append(TEXT(")\n{\n"));
+	OutCode.Appendf(TEXT("\t%s Result = (%s)0;\n"), OutputStruct->Name, OutputStruct->Name);
+	OutCode.Appendf(TEXT("\tResult.%s = CustomExpressionInternal%d(Parameters"), OutputStruct->Fields[0].Name, EmitCustomHLSL.Index);
+	for (const FEmitCustomHLSLInput& Input : EmitCustomHLSL.Inputs)
+	{
+		OutCode.Append(TEXT(", "));
+		OutCode.Append(Input.Name);
+		if (Input.Type.IsNumericLWC())
+		{
+			OutCode.Append(TEXT(", LWCToFloat("));
+			OutCode.Append(Input.Name);
+			OutCode.Append(TEXT(")"));
+		}
+	}
+	for (int32 OutputIndex = 1; OutputIndex < OutputStruct->Fields.Num(); ++OutputIndex)
+	{
+		const Shader::FStructField& OutputField = OutputStruct->Fields[OutputIndex];
+		OutCode.Appendf(TEXT(", Result.%s"), OutputField.Name);
+	}
+	OutCode.Append(TEXT(");\n"));
+	OutCode.Append(TEXT("\treturn Result;\n"));
+	OutCode.Append(TEXT("}\n"));
+}
+} // namespace Private
+
+void FEmitContext::EmitDeclarationsCode(FStringBuilderBase& OutCode)
+{
+	for(const auto& It : EmitCustomHLSLMap)
+	{
+		const FEmitCustomHLSL& EmitCustomHLSL = It.Value;
+		uint32 ShaderFrequencyMask = EmitCustomHLSL.ShaderFrequencyMask;
+		if (ShaderFrequencyMask & (1u << SF_Vertex))
+		{
+			// Emit VS version if needed
+			Private::EmitCustomHLSL(EmitCustomHLSL, TEXT("FMaterialVertexParameters"), OutCode);
+			ShaderFrequencyMask &= ~(1u << SF_Vertex);
+		}
+		if (ShaderFrequencyMask != 0u)
+		{
+			// All other frequencies use PS version
+			Private::EmitCustomHLSL(EmitCustomHLSL, TEXT("FMaterialPixelParameters"), OutCode);
+		}
+	}
+}
+
 FPreparedType FEmitContext::PrepareExpression(FExpression* InExpression, FEmitScope& Scope, const FRequestedType& RequestedType)
 {
 	if (!InExpression || RequestedType.IsVoid())
@@ -1008,6 +1099,53 @@ FEmitShaderExpression* FEmitContext::EmitCast(FEmitScope& Scope, FEmitShaderExpr
 	return ShaderValue;
 }
 
+FEmitShaderExpression* FEmitContext::EmitCustomHLSL(FEmitScope& Scope, FStringView DeclarationCode, FStringView FunctionCode, TConstArrayView<FCustomHLSLInput> Inputs, const Shader::FStructType* OutputType)
+{
+	TArray<FEmitCustomHLSLInput, TInlineAllocator<8>> EmitInputs;
+	EmitInputs.Reserve(Inputs.Num());
+	for (const FCustomHLSLInput& Input : Inputs)
+	{
+		EmitInputs.Add(FEmitCustomHLSLInput{ Input.Name, Input.Expression->GetType() });
+	}
+
+	FHasher Hasher;
+	AppendHash(Hasher, DeclarationCode);
+	AppendHash(Hasher, FunctionCode);
+	AppendHash(Hasher, MakeArrayView(EmitInputs));
+	AppendHash(Hasher, OutputType);
+	const FXxHash64 Hash = Hasher.Finalize();
+
+	FEmitCustomHLSL* EmitCustomHLSL = EmitCustomHLSLMap.Find(Hash);
+	if (!EmitCustomHLSL)
+	{
+		const int32 Index = EmitCustomHLSLMap.Num();
+		EmitCustomHLSL = &EmitCustomHLSLMap.Emplace(Hash);
+		EmitCustomHLSL->DeclarationCode = DeclarationCode;
+		EmitCustomHLSL->FunctionCode = FunctionCode;
+		EmitCustomHLSL->Inputs = MemStack::AllocateArrayView(*Allocator, MakeArrayView(EmitInputs));
+		EmitCustomHLSL->OutputType = OutputType;
+		EmitCustomHLSL->Index = Index;
+	}
+
+	// Track the different shader frequencies that call this expression
+	EmitCustomHLSL->ShaderFrequencyMask |= (1u << ShaderFrequency);
+
+	FEmitShaderDependencies Dependencies;
+	Dependencies.Reserve(Inputs.Num());
+
+	TStringBuilder<1024> FormattedCode;
+	FormattedCode.Appendf(TEXT("CustomExpression%d(Parameters"), EmitCustomHLSL->Index);
+	for (const FCustomHLSLInput& Input : Inputs)
+	{
+		FEmitShaderExpression* EmitInputExpression = Input.Expression->GetValueShader(*this, Scope);
+		FormattedCode.Appendf(TEXT(", %s"), EmitInputExpression->Reference);
+		Dependencies.Add(EmitInputExpression);
+	}
+	FormattedCode.AppendChar(TEXT(')'));
+
+	return EmitExpressionWithDependencies(Scope, Dependencies, OutputType, FormattedCode.ToView());
+}
+
 void FEmitContext::Finalize()
 {
 	// Unlink all nodes from scopes
diff --git a/Engine/Source/Runtime/Engine/Private/Materials/MaterialExpressionHLSL.cpp b/Engine/Source/Runtime/Engine/Private/Materials/MaterialExpressionHLSL.cpp
--- a/Engine/Source/Runtime/Engine/Private/Materials/MaterialExpressionHLSL.cpp
+++ b/Engine/Source/Runtime/Engine/Private/Materials/MaterialExpressionHLSL.cpp
@@ -4,6 +4,7 @@
 #if WITH_EDITOR
 
 #include "MaterialHLSLGenerator.h"
+#include "Misc/MemStackUtility.h"
 #include "HLSLTree/HLSLTree.h"
 #include "HLSLTree/HLSLTreeCommon.h"
 
@@ -56,6 +57,7 @@
 #include "Materials/MaterialExpressionReflectionVectorWS.h"
 #include "Materials/MaterialExpressionTransform.h"
 #include "Materials/MaterialExpressionTransformPosition.h"
+#include "Materials/MaterialExpressionCustom.h"
 #include "Materials/MaterialExpressionSetLocal.h"
 #include "Materials/MaterialExpressionIfThenElse.h"
 #include "Materials/MaterialExpressionForLoop.h"
@@ -840,6 +842,109 @@ bool UMaterialExpressionTransformPosition::GenerateHLSLExpression(FMaterialHLSLG
 	return true;
 }
 
+static UE::Shader::FType GetCustomOutputType(const FMaterialHLSLGenerator& Generator, ECustomMaterialOutputType Type)
+{
+	using namespace UE::Shader;
+	switch (Type)
+	{
+	case CMOT_Float1: return EValueType::Float1;
+	case CMOT_Float2: return EValueType::Float2;
+	case CMOT_Float3: return EValueType::Float3;
+	case CMOT_Float4: return EValueType::Float4;
+	case CMOT_MaterialAttributes: return Generator.GetMaterialAttributesType();
+	default: checkNoEntry(); return EValueType::Void;
+	}
+}
+
+bool UMaterialExpressionCustom::GenerateHLSLExpression(FMaterialHLSLGenerator& Generator, UE::HLSLTree::FScope& Scope, int32 OutputIndex, UE::HLSLTree::FExpression*& OutExpression)
+{
+	using namespace UE::HLSLTree;
+	using namespace UE::Shader;
+
+	if (OutputIndex < 0 || OutputIndex > AdditionalOutputs.Num())
+	{
+		return Generator.GetErrors().AddErrorf(TEXT("Invalid output index %d"), OutputIndex);
+	}
+
+	FMemStackBase& Allocator = Generator.GetTree().GetAllocator();
+
+	TArray<FCustomHLSLInput, TInlineAllocator<8>> LocalInputs;
+	LocalInputs.Reserve(Inputs.Num());
+	for (int32 Index = 0; Index < Inputs.Num(); ++Index)
+	{
+		const FCustomInput& Input = Inputs[Index];
+		if (!Input.InputName.IsNone())
+		{
+			FExpression* Expression = Input.Input.AcquireHLSLExpression(Generator, Scope);
+			if (!Expression)
+			{
+				return false;
+			}
+			const FStringView InputName = UE::MemStack::AllocateStringView(Allocator, Input.InputName.ToString());
+			LocalInputs.Emplace(InputName, Expression);
+		}
+	}
+
+	TArray<FStructFieldInitializer, TInlineAllocator<8>> OutputFieldInitializers;
+	TArray<FString, TInlineAllocator<8>> OutputNames;
+	OutputFieldInitializers.Reserve(AdditionalOutputs.Num() + 1);
+	OutputNames.Reserve(AdditionalOutputs.Num());
+
+	const FType ReturnType = GetCustomOutputType(Generator, OutputType);
+	OutputFieldInitializers.Emplace(TEXT("Default"), ReturnType);
+	for (int32 Index = 0; Index < AdditionalOutputs.Num(); ++Index)
+	{
+		const FCustomOutput& Output = AdditionalOutputs[Index];
+		OutputNames.Add(Output.OutputName.ToString());
+		OutputFieldInitializers.Emplace(OutputNames.Last(), GetCustomOutputType(Generator, Output.OutputType));
+	}
+
+	FString OutputStructName = TEXT("FCustomOutput") + GetName();
+	FStructTypeInitializer OutputStructInitializer;
+	OutputStructInitializer.Name = OutputStructName;
+	OutputStructInitializer.Fields = OutputFieldInitializers;
+	const UE::Shader::FStructType* OutputStructType = Generator.GetTypeRegistry().NewType(OutputStructInitializer);
+
+	TStringBuilder<8 * 1024> DeclarationCode;
+	for (FCustomDefine DefineEntry : AdditionalDefines)
+	{
+		if (DefineEntry.DefineName.Len() > 0)
+		{
+			DeclarationCode.Appendf(TEXT("#ifndef %s\n#define %s %s\n#endif\n"), *DefineEntry.DefineName, *DefineEntry.DefineName, *DefineEntry.DefineValue);
+		}
+	}
+
+	for (FString IncludeFile : IncludeFilePaths)
+	{
+		if (IncludeFile.Len() > 0)
+		{
+			DeclarationCode.Appendf(TEXT("#include \"%s\"\n"), *IncludeFile);
+		}
+	}
+
+	FStringView FunctionCode;
+	if (Code.Contains(TEXT("return")))
+	{
+		// Can just reference to 'Code' field directly, the UMaterialExpressionCustom lifetime will be longer than the resulting HLSLTree
+		FunctionCode = Code;
+	}
+	else
+	{
+		TStringBuilder<8 * 1024> FormattedCode;
+		FormattedCode.Appendf(TEXT("return %s;"), *Code);
+		FunctionCode = UE::MemStack::AllocateStringView(Allocator, FormattedCode.ToView());
+	}
+
+	FExpression* ExpressionCustom = Generator.GetTree().NewExpression<FExpressionCustomHLSL>(
+		UE::MemStack::AllocateStringView(Allocator, DeclarationCode.ToView()),
+		FunctionCode,
+		LocalInputs,
+		OutputStructType);
+
+	OutExpression = Generator.GetTree().NewExpression<FExpressionGetStructField>(OutputStructType, &OutputStructType->Fields[OutputIndex], ExpressionCustom);
+	return true;
+}
+
 bool UMaterialExpressionExecBegin::GenerateHLSLStatements(FMaterialHLSLGenerator& Generator, UE::HLSLTree::FScope& Scope)
 {
 	Exec.GenerateHLSLStatements(Generator, Scope);
diff --git a/Engine/Source/Runtime/Engine/Private/Materials/MaterialHLSLEmitter.cpp b/Engine/Source/Runtime/Engine/Private/Materials/MaterialHLSLEmitter.cpp
--- a/Engine/Source/Runtime/Engine/Private/Materials/MaterialHLSLEmitter.cpp
+++ b/Engine/Source/Runtime/Engine/Private/Materials/MaterialHLSLEmitter.cpp
@@ -47,6 +47,7 @@ static FString GenerateMaterialTemplateHLSL(EShaderPlatform ShaderPlatform,
 	const FMaterial& Material,
 	const UE::HLSLTree::FEmitContext& EmitContext,
 	const TCHAR* DeclarationsCode,
+	const TCHAR* SharedShaderCode,
 	const TCHAR* VertexShaderCode,
 	const TCHAR* PixelShaderCodePhase0,
 	const TCHAR* PixelShaderCodePhase1)
@@ -292,7 +293,7 @@ static FString GenerateMaterialTemplateHLSL(EShaderPlatform ShaderPlatform,
 
 	//if (bEnableExecutionFlow)
 	{
-		FString EvaluateMaterialDeclaration;
+		FString EvaluateMaterialDeclaration = SharedShaderCode;
 
 		EvaluateMaterialDeclaration += TEXT("FMaterialAttributes EvaluateVertexMaterialAttributesInternal(FMaterialVertexParameters Parameters)" LINE_TERMINATOR);
 		EvaluateMaterialDeclaration += TEXT("{" LINE_TERMINATOR);
@@ -919,10 +920,14 @@ bool MaterialEmitHLSL(const FMaterialCompileTargetParameters& InCompilerTarget,
 	FStringBuilderMemstack Declarations(Allocator, 64 * 1024);
 	TypeRegistry.EmitDeclarationsCode(Declarations);
 
+	FStringBuilderMemstack SharedCode(Allocator, 64 * 1024);
+	EmitContext.EmitDeclarationsCode(SharedCode);
+
 	FString MaterialTemplateSource = GenerateMaterialTemplateHLSL(InCompilerTarget.ShaderPlatform,
 		InOutMaterial,
 		EmitContext,
 		Declarations.ToString(),
+		SharedCode.ToString(),
 		VertexCode.ToString(),
 		PixelCodePhase0.ToString(),
 		EmitContext.bReadMaterialNormal ? PixelCodePhase1.ToString() : nullptr);
diff --git a/Engine/Source/Runtime/Engine/Private/Shader/ShaderValue.cpp b/Engine/Source/Runtime/Engine/Private/Shader/ShaderValue.cpp
--- a/Engine/Source/Runtime/Engine/Private/Shader/ShaderValue.cpp
+++ b/Engine/Source/Runtime/Engine/Private/Shader/ShaderValue.cpp
@@ -1,7 +1,7 @@
 // Copyright Epic Games, Inc. All Rights Reserved.
 #include "Shader/ShaderTypes.h"
 #include "Misc/StringBuilder.h"
-#include "Misc/SecureHash.h"
+#include "Hash/xxhash.h"
 #include "Misc/MemStackUtility.h"
 #include "Misc/LargeWorldRenderPosition.h"
 
@@ -783,46 +783,56 @@ void SetFieldType(EValueType* FieldTypes, EValueComponentType* ComponentTypes, i
 
 const FStructType* FStructTypeRegistry::NewType(const FStructTypeInitializer& Initializer)
 {
-	FSHA1 Hasher;
-	Hasher.UpdateWithString(Initializer.Name.GetData(), Initializer.Name.Len());
-
 	TArray<FStructFieldInitializer, TInlineAllocator<16>> DerivativeFields;
 
 	const int32 NumFields = Initializer.Fields.Num();
+	FStructField* Fields = new(*Allocator) FStructField[NumFields];
 	int32 ComponentIndex = 0;
 	int32 FlatFieldIndex = 0;
-	FStructField* Fields = new(*Allocator) FStructField[NumFields];
-	for (int32 FieldIndex = 0; FieldIndex < NumFields; ++FieldIndex)
+	uint64 Hash = 0u;
 	{
-		const FStructFieldInitializer& FieldInitializer = Initializer.Fields[FieldIndex];
-		const FType& FieldType = FieldInitializer.Type;
+		FXxHash64Builder Hasher;
+		Hasher.Update(Initializer.Name.GetData(), Initializer.Name.Len() * sizeof(TCHAR));
 
-		Hasher.UpdateWithString(FieldInitializer.Name.GetData(), FieldInitializer.Name.Len());
-		if (FieldType.IsStruct())
-		{
-			Hasher.Update((uint8*)&FieldType.StructType->Hash, sizeof(FieldType.StructType->Hash));
-		}
-		else
+		for (int32 FieldIndex = 0; FieldIndex < NumFields; ++FieldIndex)
 		{
-			Hasher.Update((uint8*)&FieldType.ValueType, sizeof(FieldType.ValueType));
-		}
+			const FStructFieldInitializer& FieldInitializer = Initializer.Fields[FieldIndex];
+			const FType& FieldType = FieldInitializer.Type;
 
-		FStructField& Field = Fields[FieldIndex];
-		Field.Name = MemStack::AllocateString(*Allocator, FieldInitializer.Name);
-		Field.Type = FieldType;
-		Field.ComponentIndex = ComponentIndex;
-		Field.FlatFieldIndex = FlatFieldIndex;
-		ComponentIndex += FieldType.GetNumComponents();
-		FlatFieldIndex += FieldType.GetNumFlatFields();
+			Hasher.Update(FieldInitializer.Name.GetData(), FieldInitializer.Name.Len() * sizeof(TCHAR));
+			if (FieldType.IsStruct())
+			{
+				Hasher.Update(&FieldType.StructType->Hash, sizeof(FieldType.StructType->Hash));
+			}
+			else
+			{
+				Hasher.Update(&FieldType.ValueType, sizeof(FieldType.ValueType));
+			}
 
-		if (!Initializer.bIsDerivativeType)
-		{
-			const FType FieldDerivativeType = FieldType.GetDerivativeType();
-			if (!FieldDerivativeType.IsVoid())
+			FStructField& Field = Fields[FieldIndex];
+			Field.Name = MemStack::AllocateString(*Allocator, FieldInitializer.Name);
+			Field.Type = FieldType;
+			Field.ComponentIndex = ComponentIndex;
+			Field.FlatFieldIndex = FlatFieldIndex;
+			ComponentIndex += FieldType.GetNumComponents();
+			FlatFieldIndex += FieldType.GetNumFlatFields();
+
+			if (!Initializer.bIsDerivativeType)
 			{
-				DerivativeFields.Emplace(FieldInitializer.Name, FieldDerivativeType);
+				const FType FieldDerivativeType = FieldType.GetDerivativeType();
+				if (!FieldDerivativeType.IsVoid())
+				{
+					DerivativeFields.Emplace(FieldInitializer.Name, FieldDerivativeType);
+				}
 			}
 		}
+		Hash = Hasher.Finalize().Hash;
+	}
+
+	FStructType const* const* PrevStructType = Types.Find(Hash);
+	if (PrevStructType)
+	{
+		return *PrevStructType;
 	}
 
 	EValueComponentType* ComponentTypes = new(*Allocator) EValueComponentType[ComponentIndex];
@@ -833,16 +843,14 @@ const FStructType* FStructTypeRegistry::NewType(const FStructTypeInitializer& In
 		Private::SetFieldType(FlatFieldTypes, ComponentTypes, Field.FlatFieldIndex, Field.ComponentIndex, Field.Type);
 	}
 
-	const FSHAHash Hash = Hasher.Finalize();
-
 	FStructType* StructType = new(*Allocator) FStructType();
 	StructType->Name = MemStack::AllocateString(*Allocator, Initializer.Name);
-	StructType->Hash = *reinterpret_cast<const uint64*>(Hash.Hash);
+	StructType->Hash = Hash;
 	StructType->Fields = MakeArrayView(Fields, NumFields);
 	StructType->ComponentTypes = MakeArrayView(ComponentTypes, ComponentIndex);
 	StructType->FlatFieldTypes = MakeArrayView(FlatFieldTypes, FlatFieldIndex);
 
-	Types.Add(StructType->Hash, StructType);
+	Types.Add(Hash, StructType);
 
 	if (DerivativeFields.Num() > 0)
 	{
diff --git a/Engine/Source/Runtime/Engine/Public/HLSLTree/HLSLTree.h b/Engine/Source/Runtime/Engine/Public/HLSLTree/HLSLTree.h
--- a/Engine/Source/Runtime/Engine/Public/HLSLTree/HLSLTree.h
+++ b/Engine/Source/Runtime/Engine/Public/HLSLTree/HLSLTree.h
@@ -11,6 +11,7 @@
 #include "Misc/MemStack.h"
 #include "Hash/xxhash.h"
 #include "HLSLTree/HLSLTreeTypes.h"
+#include "HLSLTree/HLSLTreeHash.h"
 
 class FMaterial;
 class FMaterialCompilationOutput;
@@ -111,70 +112,6 @@ private:
 	int32 NumErrors = 0;
 };
 
-class FHasher
-{
-public:
-	FXxHash64 Finalize() { return Builder.Finalize(); }
-	void AppendData(const void* Data, uint64 Size) { Builder.Update(Data, Size); }
-
-private:
-	FXxHash64Builder Builder;
-};
-
-template<typename T>
-inline void AppendHash(FHasher& Hasher, const T& Value)
-{
-	Hasher.AppendData(&Value, sizeof(Value));
-}
-
-template<typename T>
-inline void AppendHash(FHasher& Hasher, TArrayView<T> Value)
-{
-	for (const T& Element : Value)
-	{
-		AppendHash(Hasher, Element);
-	}
-}
-
-
-template<typename Allocator>
-inline void AppendHash(FHasher& Hasher, const TBitArray<Allocator>& Value)
-{
-	const uint32 NumWords = FBitSet::CalculateNumWords(Value.Num());
-	const uint32* Data = Value.GetData();
-	Hasher.AppendData(Data, NumWords * sizeof(uint32));
-}
-
-inline void AppendHash(FHasher& Hasher, const FName& Value)
-{
-	AppendHash(Hasher, Value.GetComparisonIndex());
-	AppendHash(Hasher, Value.GetNumber());
-}
-
-inline void AppendHash(FHasher& Hasher, const Shader::FType& Type)
-{
-	if (Type.IsStruct()) AppendHash(Hasher, Type.StructType);
-	else AppendHash(Hasher, Type.ValueType);
-}
-
-inline void AppendHash(FHasher& Hasher, const Shader::FValue& Value)
-{
-	AppendHash(Hasher, Value.Type);
-	for (int32 i = 0; i < Value.Type.GetNumComponents(); ++i)
-	{
-		AppendHash(Hasher, Value.TryGetComponent(i));
-	}
-}
-
-inline void AppendHashes(FHasher& Hasher) {}
-
-template<typename T, typename... ArgTypes>
-inline void AppendHashes(FHasher& Hasher, const T& Value, ArgTypes&&... Args)
-{
-	AppendHash(Hasher, Value);
-	AppendHashes(Hasher, Forward<ArgTypes>(Args)...);
-}
-
 
 /** Root class of the HLSL AST */
 class FNode
diff --git a/Engine/Source/Runtime/Engine/Public/HLSLTree/HLSLTreeCommon.h b/Engine/Source/Runtime/Engine/Public/HLSLTree/HLSLTreeCommon.h
--- a/Engine/Source/Runtime/Engine/Public/HLSLTree/HLSLTreeCommon.h
+++ b/Engine/Source/Runtime/Engine/Public/HLSLTree/HLSLTreeCommon.h
@@ -329,10 +329,20 @@ public:
 class FExpressionCustomHLSL : public FExpression
 {
 public:
+	FExpressionCustomHLSL(FStringView InDeclarationCode, FStringView InFunctionCode, TArrayView<FCustomHLSLInput> InInputs, const Shader::FStructType* InOutputStructType)
+		: DeclarationCode(InDeclarationCode)
+		, FunctionCode(InFunctionCode)
+		, Inputs(InInputs)
+		, OutputStructType(InOutputStructType)
+	{}
+
 	virtual bool PrepareValue(FEmitContext& Context, FEmitScope& Scope, const FRequestedType& RequestedType, FPrepareValueResult& OutResult) const override;
 	virtual void EmitValueShader(FEmitContext& Context, FEmitScope& Scope, const FRequestedType& RequestedType, FEmitValueShaderResult& OutResult) const override;
 
-
+	FStringView DeclarationCode;
+	FStringView FunctionCode;
+	TArray<FCustomHLSLInput, TInlineAllocator<8>> Inputs;
+	const Shader::FStructType* OutputStructType = nullptr;
 };
 
 class FStatementReturn : public FStatement
diff --git a/Engine/Source/Runtime/Engine/Public/HLSLTree/HLSLTreeEmit.h b/Engine/Source/Runtime/Engine/Public/HLSLTree/HLSLTreeEmit.h
--- a/Engine/Source/Runtime/Engine/Public/HLSLTree/HLSLTreeEmit.h
+++ b/Engine/Source/Runtime/Engine/Public/HLSLTree/HLSLTreeEmit.h
@@ -7,9 +7,9 @@
 #include "Containers/StringView.h"
 #include "Misc/StringBuilder.h"
 #include "Misc/MemStack.h"
-#include "Hash/xxhash.h"
 #include "RHIDefinitions.h"
 #include "HLSLTree/HLSLTreeTypes.h"
+#include "HLSLTree/HLSLTreeHash.h"
 
 class FMaterial;
 class FMaterialCompilationOutput;
@@ -194,6 +194,29 @@ public:
 	EExpressionEvaluation Evaluation = EExpressionEvaluation::None;
 };
 
+struct FEmitCustomHLSLInput
+{
+	FStringView Name;
+	Shader::FType Type;
+};
+
+inline void AppendHash(FHasher& Hasher, const FEmitCustomHLSLInput& Value)
+{
+	AppendHash(Hasher, Value.Name);
+	AppendHash(Hasher, Value.Type);
+}
+
+class FEmitCustomHLSL
+{
+public:
+	FStringView DeclarationCode;
+	FStringView FunctionCode;
+	TConstArrayView<FEmitCustomHLSLInput> Inputs;
+	const Shader::FStructType* OutputType;
+	uint32 ShaderFrequencyMask;
+	int32 Index;
+};
+
 struct FPreshaderLocalPHIScope
 {
 	FPreshaderLocalPHIScope(const FExpression* InExpression, int32 InValueStackPosition) : ExpressionLocalPHI(InExpression), ValueStackPosition(InValueStackPosition) {}
@@ -209,6 +232,8 @@ public:
 	explicit FEmitContext(FMemStackBase& InAllocator, FErrorHandlerInterface& InErrors, const Shader::FStructTypeRegistry& InTypeRegistry);
 	~FEmitContext();
 
+	void EmitDeclarationsCode(FStringBuilderBase& OutCode);
+
 	FPreparedType PrepareExpression(FExpression* InExpression, FEmitScope& Scope, const FRequestedType& RequestedType);
 
 	FEmitScope* InternalPrepareScope(FScope* Scope, FScope* ParentScope);
@@ -374,6 +399,9 @@ public:
 	FEmitShaderExpression* EmitConstantZero(FEmitScope& Scope, const Shader::FType& Type);
 	FEmitShaderExpression* EmitCast(FEmitScope& Scope, FEmitShaderExpression* ShaderValue, const Shader::FType& DestType);
 
+	/** FStringViews passed to this method are expected to reference persistent memory */
+	FEmitShaderExpression* EmitCustomHLSL(FEmitScope& Scope, FStringView DeclarationCode, FStringView FunctionCode, TConstArrayView<FCustomHLSLInput> Inputs, const Shader::FStructType* OutputType);
+
 	FMemStackBase* Allocator = nullptr;
 	FErrorHandlerInterface* Errors = nullptr;
 	const Shader::FStructTypeRegistry* TypeRegistry = nullptr;
@@ -386,6 +414,7 @@ public:
 	TMap<FXxHash64, FEmitShaderExpression*> EmitExpressionMap;
 	TMap<FXxHash64, FEmitShaderExpression*> EmitPreshaderMap;
 	TMap<const FFunction*, FEmitShaderNode*> EmitFunctionMap;
+	TMap<FXxHash64, FEmitCustomHLSL> EmitCustomHLSLMap;
 	TArray<struct FPreshaderLoopScope*> PreshaderLoopScopes;
 	TArray<const FPreshaderLocalPHIScope*> PreshaderLocalPHIScopes;
 	int32 PreshaderStackPosition = 0;
diff --git a/Engine/Source/Runtime/Engine/Public/HLSLTree/HLSLTreeHash.h b/Engine/Source/Runtime/Engine/Public/HLSLTree/HLSLTreeHash.h
new file mode 100644
--- /dev/null
+++ b/Engine/Source/Runtime/Engine/Public/HLSLTree/HLSLTreeHash.h
@@ -0,0 +1,93 @@
+// Copyright Epic Games, Inc. All Rights Reserved.
+#pragma once
+
+#include "CoreMinimal.h"
+#include "Engine/EngineTypes.h"
+#include "Containers/ArrayView.h"
+#include "Containers/StringView.h"
+#include "Containers/BitArray.h"
+#include "Misc/StringBuilder.h"
+#include "Hash/xxhash.h"
+#include "HLSLTree/HLSLTreeTypes.h"
+
+namespace UE
+{
+namespace HLSLTree
+{
+
+class FHasher
+{
+public:
+	FXxHash64 Finalize() { return Builder.Finalize(); }
+	void AppendData(const void* Data, uint64 Size) { Builder.Update(Data, Size); }
+
+private:
+	FXxHash64Builder Builder;
+};
+
+template<typename T>
+inline void AppendHash(FHasher& Hasher, const T& Value)
+{
+	Hasher.AppendData(&Value, sizeof(Value));
+}
+
+template<typename T>
+inline void AppendHash(FHasher& Hasher, TArrayView<T> Value)
+{
+	for (const T& Element : Value)
+	{
+		AppendHash(Hasher, Element);
+	}
+}
+
+template<typename Allocator>
+inline void AppendHash(FHasher& Hasher, const TBitArray<Allocator>& Value)
+{
+	const uint32 NumWords = FBitSet::CalculateNumWords(Value.Num());
+	const uint32* Data = Value.GetData();
+	Hasher.AppendData(Data, NumWords * sizeof(uint32));
+}
+
+inline void AppendHash(FHasher& Hasher, const FName& Value)
+{
+	AppendHash(Hasher, Value.GetComparisonIndex());
+	AppendHash(Hasher, Value.GetNumber());
+}
+
+inline void AppendHash(FHasher& Hasher, FStringView Value)
+{
+	Hasher.AppendData(Value.GetData(), Value.Len() * sizeof(TCHAR));
+}
+
+inline void AppendHash(FHasher& Hasher, const FCustomHLSLInput& Value)
+{
+	AppendHash(Hasher, Value.Name);
+	AppendHash(Hasher, Value.Expression);
+}
+
+inline void AppendHash(FHasher& Hasher, const Shader::FType& Type)
+{
+	if (Type.IsStruct()) AppendHash(Hasher, Type.StructType);
+	else AppendHash(Hasher, Type.ValueType);
+}
+
+inline void AppendHash(FHasher& Hasher, const Shader::FValue& Value)
+{
+	AppendHash(Hasher, Value.Type);
+	for (int32 i = 0; i < Value.Type.GetNumComponents(); ++i)
+	{
+		AppendHash(Hasher, Value.TryGetComponent(i));
+	}
+}
+
+inline void AppendHashes(FHasher& Hasher) {}
+
+template<typename T, typename... ArgTypes>
+inline void AppendHashes(FHasher& Hasher, const T& Value, ArgTypes&&... Args)
+{
+	AppendHash(Hasher, Value);
+	AppendHashes(Hasher, Forward<ArgTypes>(Args)...);
+}
+
+} // namespace HLSLTree
+} // namespace UE
diff --git a/Engine/Source/Runtime/Engine/Public/HLSLTree/HLSLTreeTypes.h b/Engine/Source/Runtime/Engine/Public/HLSLTree/HLSLTreeTypes.h
--- a/Engine/Source/Runtime/Engine/Public/HLSLTree/HLSLTreeTypes.h
+++ b/Engine/Source/Runtime/Engine/Public/HLSLTree/HLSLTreeTypes.h
@@ -127,5 +127,14 @@ inline uint32 GetTypeHash(const FTextureDescription& Ref)
 	return HashCombine(GetTypeHash(Ref.Texture), GetTypeHash(Ref.SamplerType));
 }
 
+struct FCustomHLSLInput
+{
+	FCustomHLSLInput() = default;
+	FCustomHLSLInput(FStringView InName, FExpression* InExpression) : Name(InName), Expression(InExpression) {}
+
+	FStringView Name;
+	FExpression* Expression = nullptr;
+};
+
 } // namespace HLSLTree
 } // namespace UE
diff --git a/Engine/Source/Runtime/Engine/Public/MaterialHLSLGenerator.h b/Engine/Source/Runtime/Engine/Public/MaterialHLSLGenerator.h
--- a/Engine/Source/Runtime/Engine/Public/MaterialHLSLGenerator.h
+++ b/Engine/Source/Runtime/Engine/Public/MaterialHLSLGenerator.h
@@ -76,6 +76,7 @@ public:
 	bool Generate();
 
 	UE::HLSLTree::FTree& GetTree() const { return *HLSLTree; }
+	UE::Shader::FStructTypeRegistry& GetTypeRegistry() const { return *TypeRegistry; }
 	FMaterialHLSLErrorHandler& GetErrors() { return Errors; }
 
 	UE::HLSLTree::FExpression* GetResultExpression() { return ResultExpression; }
diff --git a/Engine/Source/Runtime/Engine/Public/Shader/ShaderTypes.h b/Engine/Source/Runtime/Engine/Public/Shader/ShaderTypes.h
--- a/Engine/Source/Runtime/Engine/Public/Shader/ShaderTypes.h
+++ b/Engine/Source/Runtime/Engine/Public/Shader/ShaderTypes.h
@@ -79,6 +79,7 @@ inline bool IsComponentTypeWithinBounds(EValueComponentType Type, FComponentBoun
 EValueComponentType CombineComponentTypes(EValueComponentType Lhs, EValueComponentType Rhs);
 
 inline EValueComponentType MakeNonLWCType(EValueComponentType Type) { return Type == EValueComponentType::Double ? EValueComponentType::Float : Type; }
+inline bool IsLWCType(EValueComponentType Type) { return Type == EValueComponentType::Double; }
 
 enum class EValueType : uint8
 {
@@ -126,6 +127,7 @@ struct FValueTypeDescription
 };
 
 FValueTypeDescription GetValueTypeDescription(EValueType Type);
+inline bool IsLWCType(EValueType Type) { return IsLWCType(GetValueTypeDescription(Type).ComponentType); }
 EValueType MakeValueType(EValueComponentType ComponentType, int32 NumComponents);
 EValueType MakeValueType(EValueType BaseType, int32 NumComponents);
 EValueType MakeValueTypeWithRequestedNumComponents(EValueType BaseType, int8 RequestedNumComponents);
@@ -141,9 +143,11 @@ struct FType
 
 	const TCHAR* GetName() const;
 	FType GetDerivativeType() const;
+	FType GetNonLWCType() const { return IsNumericLWC() ? FType(MakeNonLWCType(ValueType)) : *this; }
 	bool IsVoid() const { return ValueType == EValueType::Void; }
 	bool IsStruct() const { return ValueType == EValueType::Struct; }
 	bool IsNumeric() const { return !IsVoid() && !IsStruct(); }
+	bool IsNumericLWC() const { return IsNumeric() && IsLWCType(ValueType); }
 	int32 GetNumComponents() const;
 	int32 GetNumFlatFields() const;
 	EValueComponentType GetComponentType(int32 Index) const;
