Texture: Switch to 64-bit sizes when accessing DDC to support encoded textures larger than 2 GiB

diff --git a/Engine/Source/Runtime/Engine/Private/Streaming/Texture2DMipDataProvider_DDC.cpp b/Engine/Source/Runtime/Engine/Private/Streaming/Texture2DMipDataProvider_DDC.cpp
--- a/Engine/Source/Runtime/Engine/Private/Streaming/Texture2DMipDataProvider_DDC.cpp
+++ b/Engine/Source/Runtime/Engine/Private/Streaming/Texture2DMipDataProvider_DDC.cpp
@@ -124,7 +124,7 @@ void FTexture2DMipDataProvider_DDC::Init(const FTextureUpdateContext& Context, c
 	}
 }
 
-bool FTexture2DMipDataProvider_DDC::SerializeMipInfo(const FTextureUpdateContext& Context, FArchive& Ar, int32 MipIndex, int32 MipSize, const FTextureMipInfo& OutMipInfo)
+bool FTexture2DMipDataProvider_DDC::SerializeMipInfo(const FTextureUpdateContext& Context, FArchive& Ar, int32 MipIndex, int64 MipSize, const FTextureMipInfo& OutMipInfo)
 {
 	const uint32 DepthOrArraySize = FMath::Max<uint32>(OutMipInfo.ArraySize, OutMipInfo.SizeZ);
 	if (MipSize == OutMipInfo.DataSize)
@@ -132,7 +132,7 @@ bool FTexture2DMipDataProvider_DDC::SerializeMipInfo(const FTextureUpdateContext
 		Ar.Serialize(OutMipInfo.DestData, MipSize);
 		return true;
 	}
-	else if (MipSize < OutMipInfo.DataSize && DepthOrArraySize > 1 && OutMipInfo.DataSize % DepthOrArraySize == 0 && MipSize % DepthOrArraySize == 0)
+	else if (uint64(MipSize) < OutMipInfo.DataSize && DepthOrArraySize > 1 && OutMipInfo.DataSize % DepthOrArraySize == 0 && MipSize % DepthOrArraySize == 0)
 	{
 		UE_LOG(LogTexture, Verbose, TEXT("DDC mip size smaller than streaming buffer size. (%s, Mip %d): %d KB / %d KB."), *Context.Resource->GetTextureName().ToString(), ResourceState.MaxNumLODs - MipIndex, OutMipInfo.DataSize / 1024, MipSize / 1024);
 
@@ -178,13 +178,13 @@ int32 FTexture2DMipDataProvider_DDC::GetMips(
 				{
 					DDCHandles[MipIndex] = 0; // Clear the handle.
 
-					TArray<uint8> DerivedMipData;
+					TArray64<uint8> DerivedMipData;
 					if (GetDerivedDataCacheRef().GetAsynchronousResults(Handle, DerivedMipData))
 					{
 						const FTextureMipInfo& MipInfo = MipInfos[MipIndex];
 
 						// The result must be read from a memory reader!
-						FMemoryReader Ar(DerivedMipData, true);
+						FMemoryReaderView Ar(MakeMemoryView(DerivedMipData), true);
 						if (SerializeMipInfo(Context, Ar, MipIndex, DerivedMipData.Num(), MipInfo))
 						{
 							bSuccess = true;
diff --git a/Engine/Source/Runtime/Engine/Private/Streaming/Texture2DMipDataProvider_DDC.h b/Engine/Source/Runtime/Engine/Private/Streaming/Texture2DMipDataProvider_DDC.h
--- a/Engine/Source/Runtime/Engine/Private/Streaming/Texture2DMipDataProvider_DDC.h
+++ b/Engine/Source/Runtime/Engine/Private/Streaming/Texture2DMipDataProvider_DDC.h
@@ -40,7 +40,7 @@ protected:
 	// Helper to route incomplete DDC requests to the FAbandonedDDCHandleManager.
 	void ReleaseDDCResources();
 
-	bool SerializeMipInfo(const FTextureUpdateContext& Context, FArchive& Ar, int32 MipIndex, int32 MipSize, const FTextureMipInfo& OutMipInfo);
+	bool SerializeMipInfo(const FTextureUpdateContext& Context, FArchive& Ar, int32 MipIndex, int64 MipSize, const FTextureMipInfo& OutMipInfo);
 
 private:
 
diff --git a/Engine/Source/Runtime/Engine/Private/Streaming/Texture2DStreamIn_DDC.cpp b/Engine/Source/Runtime/Engine/Private/Streaming/Texture2DStreamIn_DDC.cpp
--- a/Engine/Source/Runtime/Engine/Private/Streaming/Texture2DStreamIn_DDC.cpp
+++ b/Engine/Source/Runtime/Engine/Private/Streaming/Texture2DStreamIn_DDC.cpp
@@ -66,7 +66,7 @@ void FAbandonedDDCHandleManager::Purge()
 	}
 
 	FDerivedDataCacheInterface& DDCInterface = GetDerivedDataCacheRef();
-	TArray<uint8> Data;
+	TArray64<uint8> Data;
 
 	for (int32 Index = 0; Index < TempHandles.Num(); ++Index)
 	{
@@ -241,7 +241,7 @@ void FTexture2DStreamIn_DDC::DoLoadNewMipsFromDDC(const FContext& Context)
 				if (MipMap.IsPagedToDerivedData())
 				{
 					// The overhead of doing 2 copy of each mip data (from GetSynchronous() and FMemoryReader) in hidden by other texture DDC ops happening at the same time.
-					TArray<uint8> DerivedMipData;
+					TArray64<uint8> DerivedMipData;
 					bool bDDCValid = true;
 
 					const uint32 Handle = DDCHandles[MipIndex];
@@ -257,8 +257,8 @@ void FTexture2DStreamIn_DDC::DoLoadNewMipsFromDDC(const FContext& Context)
 
 					if (bDDCValid)
 					{
-						const int32 ExpectedMipSize = CalcTextureMipMapSize(MipMap.SizeX, MipMap.SizeY, Context.Resource->GetPixelFormat(), 0);
-						FMemoryReader Ar(DerivedMipData, true);
+						const SIZE_T ExpectedMipSize = CalcTextureMipMapSize(MipMap.SizeX, MipMap.SizeY, Context.Resource->GetPixelFormat(), 0);
+						FMemoryReaderView Ar(MakeMemoryView(DerivedMipData), true);
 
 						if (DerivedMipData.Num() == ExpectedMipSize)
 						{
@@ -266,7 +266,8 @@ void FTexture2DStreamIn_DDC::DoLoadNewMipsFromDDC(const FContext& Context)
 						}
 						else
 						{
-							UE_LOG(LogTexture, Error, TEXT("DDC mip size (%d) not as expected (%d)."), MipIndex, ExpectedMipSize);
+							UE_LOG(LogTexture, Error, TEXT("DDC mip size (%" SIZE_T_FMT ") not as expected (%" SIZE_T_FMT ") for mip %d of %s."),
+								static_cast<SIZE_T>(DerivedMipData.Num()), ExpectedMipSize, MipIndex, *Context.Texture->GetPathName());
 							MarkAsCancelled();
 						}
 					}
@@ -328,14 +329,15 @@ void FTexture2DStreamIn_DDC::DoLoadNewMipsFromDDC(const FContext& Context)
 
 					if (MipResult)
 					{
-						const int32 ExpectedMipSize = CalcTextureMipMapSize(MipMap.SizeX, MipMap.SizeY, Context.Resource->GetPixelFormat(), 0);
+						const SIZE_T ExpectedMipSize = CalcTextureMipMapSize(MipMap.SizeX, MipMap.SizeY, Context.Resource->GetPixelFormat(), 0);
 						if (MipResult.GetSize() == ExpectedMipSize)
 						{
 							FMemory::Memcpy(MipData[MipIndex], MipResult.GetData(), MipResult.GetSize());
 						}
 						else
 						{
-							UE_LOG(LogTexture, Error, TEXT("DDC mip size (%d) not as expected (%d) for mip %d of %s."), static_cast<int32>(MipResult.GetSize()), ExpectedMipSize, MipIndex, *Context.Texture->GetPathName());
+							UE_LOG(LogTexture, Error, TEXT("DDC mip size (%" SIZE_T_FMT ") not as expected (%" SIZE_T_FMT ") for mip %d of %s."),
+								static_cast<SIZE_T>(MipResult.GetSize()), ExpectedMipSize, MipIndex, *Context.Texture->GetPathName());
 							MarkAsCancelled();
 						}
 					}
diff --git a/Engine/Source/Runtime/Engine/Private/Texture2D.cpp b/Engine/Source/Runtime/Engine/Private/Texture2D.cpp
--- a/Engine/Source/Runtime/Engine/Private/Texture2D.cpp
+++ b/Engine/Source/Runtime/Engine/Private/Texture2D.cpp
@@ -162,19 +162,19 @@ void FTexture2DMipMap::Serialize(FArchive& Ar, UObject* Owner, int32 MipIdx)
 }
 
 #if WITH_EDITORONLY_DATA
-uint32 FTexture2DMipMap::StoreInDerivedDataCache(const FString& InDerivedDataKey, const FStringView& TextureName, bool bReplaceExistingDDC)
+int64 FTexture2DMipMap::StoreInDerivedDataCache(const FString& InDerivedDataKey, const FStringView& TextureName, bool bReplaceExistingDDC)
 {
-	int32 BulkDataSizeInBytes = BulkData.GetBulkDataSize();
+	int64 BulkDataSizeInBytes = BulkData.GetBulkDataSize();
 	check(BulkDataSizeInBytes > 0);
 
-	TArray<uint8> DerivedData;
-	FMemoryWriter Ar(DerivedData, /*bIsPersistent=*/ true);
+	TArray64<uint8> DerivedData;
+	FMemoryWriter64 Ar(DerivedData, /*bIsPersistent=*/ true);
 	{
 		void* BulkMipData = BulkData.Lock(LOCK_READ_ONLY);
 		Ar.Serialize(BulkMipData, BulkDataSizeInBytes);
 		BulkData.Unlock();
 	}
-	const uint32 Result = DerivedData.Num();
+	const int64 Result = DerivedData.Num();
 	GetDerivedDataCacheRef().Put(*InDerivedDataKey, DerivedData, TextureName, bReplaceExistingDDC);
 	bPagedToDerivedData = true;
 	BulkData.RemoveBulkData();
diff --git a/Engine/Source/Runtime/Engine/Private/TextureDerivedData.cpp b/Engine/Source/Runtime/Engine/Private/TextureDerivedData.cpp
--- a/Engine/Source/Runtime/Engine/Private/TextureDerivedData.cpp
+++ b/Engine/Source/Runtime/Engine/Private/TextureDerivedData.cpp
@@ -992,11 +992,11 @@ static void GetBuildSettingsPerFormat(
  * @param bForceAllMipsToBeInlined - Whether to store all mips in the main DDC. Relates to how the texture resources get initialized (not supporting streaming).
  * @return number of bytes put to the DDC (total, including all mips)
  */
-uint32 PutDerivedDataInCache(FTexturePlatformData* DerivedData, const FString& DerivedDataKeySuffix, const FStringView& TextureName, bool bForceAllMipsToBeInlined, bool bReplaceExistingDDC)
+int64 PutDerivedDataInCache(FTexturePlatformData* DerivedData, const FString& DerivedDataKeySuffix, const FStringView& TextureName, bool bForceAllMipsToBeInlined, bool bReplaceExistingDDC)
 {
-	TArray<uint8> RawDerivedData;
+	TArray64<uint8> RawDerivedData;
 	FString DerivedDataKey;
-	uint32 TotalBytesPut = 0;
+	int64 TotalBytesPut = 0;
 
 	// Build the key with which to cache derived data.
 	GetTextureDerivedDataKeyFromSuffix(DerivedDataKeySuffix, DerivedDataKey);
@@ -1074,7 +1074,7 @@ uint32 PutDerivedDataInCache(FTexturePlatformData* DerivedData, const FString& D
 
 	// Store derived data.
 	// At this point we've stored all the non-inline data in the DDC, so this will only serialize and store the TexturePlatformData metadata and any inline mips
-	FMemoryWriter Ar(RawDerivedData, /*bIsPersistent=*/ true);
+	FMemoryWriter64 Ar(RawDerivedData, /*bIsPersistent=*/ true);
 	DerivedData->Serialize(Ar, NULL);
 	TotalBytesPut += RawDerivedData.Num();
 	GetDerivedDataCacheRef().Put(*DerivedDataKey, RawDerivedData, TextureName, bReplaceExistingDerivedDataDDC);
@@ -1359,7 +1359,7 @@ static void BeginLoadDerivedVTChunks(const TArray<FVirtualTextureDataChunk>& Chu
 }
 
 /** Logs a warning that MipSize is correct for the mipmap. */
-static void CheckMipSize(FTexture2DMipMap& Mip, EPixelFormat PixelFormat, int32 MipSize)
+static void CheckMipSize(FTexture2DMipMap& Mip, EPixelFormat PixelFormat, int64 MipSize)
 {
 	// this check is incorrect ; it does not account of platform tiling and padding done on textures
 	// re-enable if fixed
@@ -1369,7 +1369,7 @@ static void CheckMipSize(FTexture2DMipMap& Mip, EPixelFormat PixelFormat, int32
 	if (MipSize != Mip.SizeZ * CalcTextureMipMapSize(Mip.SizeX, Mip.SizeY, PixelFormat, 0))
 	{
 		UE_LOG(LogTexture, Warning,
-			TEXT("%dx%d mip of %s texture has invalid data in the DDC. Got %d bytes, expected %d. Key=%s"),
+			TEXT("%dx%d mip of %s texture has invalid data in the DDC. Got %" INT64_FMT " bytes, expected %" SIZE_T_FMT ". Key=%s"),
 			Mip.SizeX,
 			Mip.SizeY,
 			GPixelFormats[PixelFormat].Name,
@@ -1385,7 +1385,7 @@ static void CheckMipSize(FTexture2DMipMap& Mip, EPixelFormat PixelFormat, int32
 // Retrieve all built texture data in to the associated arrays, and don't return unless there's an error
 // or we have the data.
 //
-static bool FetchAllTextureDataSynchronous(FTexturePlatformData* PlatformData, FStringView DebugContext, TArray<TArray<uint8>>& OutMipData, TArray<TArray<uint8>>& OutVTChunkData)
+static bool FetchAllTextureDataSynchronous(FTexturePlatformData* PlatformData, FStringView DebugContext, TArray<TArray64<uint8>>& OutMipData, TArray<TArray64<uint8>>& OutVTChunkData)
 {
 	OutMipData.Empty();
 	OutVTChunkData.Empty();
@@ -1488,8 +1488,8 @@ static bool FetchAllTextureDataSynchronous(FTexturePlatformData* PlatformData, F
 // it is illustrative.
 //
 static void EstimateOnDiskCompressionForTextureData(
-	TArray<TArray<uint8>> InMipData,
-	TArray<TArray<uint8>> InVTChunkData,
+	TArray<TArray64<uint8>> InMipData,
+	TArray<TArray64<uint8>> InVTChunkData,
 	FOodleDataCompression::ECompressor InOodleCompressor,
 	FOodleDataCompression::ECompressionLevel InOodleCompressionLevel,
 	uint32 InCompressionBlockSize,
@@ -1503,11 +1503,11 @@ static void EstimateOnDiskCompressionForTextureData(
 	// both.
 	//
 	uint64 UncompressedByteCount = 0;
-	for (TArray<uint8>& Mip : InMipData)
+	for (TArray64<uint8>& Mip : InMipData)
 	{
 		UncompressedByteCount += Mip.Num();
 	}
-	for (TArray<uint8>& VTChunk : InVTChunkData)
+	for (TArray64<uint8>& VTChunk : InVTChunkData)
 	{
 		UncompressedByteCount += VTChunk.Num();
 	}
@@ -1522,7 +1522,7 @@ static void EstimateOnDiskCompressionForTextureData(
 
 	int32 MipIndex = 0;
 	int32 VTChunkIndex = 0;
-	int32 CurrentOffsetInContainer = 0;
+	int64 CurrentOffsetInContainer = 0;
 	uint64 CompressedByteCount = 0;
 
 	// Array for compressed data so we don't have to realloc.
@@ -1530,13 +1530,13 @@ static void EstimateOnDiskCompressionForTextureData(
 	Compressed.Reserve(InCompressionBlockSize + 1024);
 
 	// When we cross our input array boundaries, we accumulate in to here.
-	TArray<uint8> ContinuousMemory;
+	TArray64<uint8> ContinuousMemory;
 	for (;;)
 	{
-		TArray<uint8>& CurrentContainer = MipIndex < InMipData.Num() ? InMipData[MipIndex] : InVTChunkData[VTChunkIndex];
+		TArray64<uint8>& CurrentContainer = MipIndex < InMipData.Num() ? InMipData[MipIndex] : InVTChunkData[VTChunkIndex];
 
-		uint32 NeedBytes = InCompressionBlockSize - ContinuousMemory.Num();
-		uint32 CopyBytes = CurrentContainer.Num() - CurrentOffsetInContainer;
+		uint64 NeedBytes = InCompressionBlockSize - ContinuousMemory.Num();
+		uint64 CopyBytes = CurrentContainer.Num() - CurrentOffsetInContainer;
 		if (CopyBytes > NeedBytes)
 		{
 			CopyBytes = NeedBytes;
@@ -1638,8 +1638,8 @@ TFuture<TTuple<uint64, uint64>> FTexturePlatformData::LaunchEstimateOnDiskSizeTa
 	FStringView InDebugContext
 	)
 {
-	TArray<TArray<uint8>> MipData;
-	TArray<TArray<uint8>> VTChunkData;
+	TArray<TArray64<uint8>> MipData;
+	TArray<TArray64<uint8>> VTChunkData;
 	if (FetchAllTextureDataSynchronous(this, InDebugContext, MipData, VTChunkData) == false)
 	{
 		return TFuture<TTuple<uint64, uint64>>();
@@ -1648,8 +1648,8 @@ TFuture<TTuple<uint64, uint64>> FTexturePlatformData::LaunchEstimateOnDiskSizeTa
 	struct FAsyncEstimateState
 	{
 		TPromise<TPair<uint64, uint64>> Promise;
-		TArray<TArray<uint8>> MipData;
-		TArray<TArray<uint8>> VTChunkData;
+		TArray<TArray64<uint8>> MipData;
+		TArray<TArray64<uint8>> VTChunkData;
 		FOodleDataCompression::ECompressor OodleCompressor;
 		FOodleDataCompression::ECompressionLevel OodleCompressionLevel;
 		uint32 CompressionBlockSize;
@@ -1694,7 +1694,7 @@ bool FTexturePlatformData::TryInlineMipData(int32 FirstMipToLoad, FStringView De
 
 	FAsyncMipHandles AsyncHandles;
 	FAsyncVTChunkHandles AsyncVTHandles;
-	TArray<uint8> TempData;
+	TArray64<uint8> TempData;
 	FDerivedDataCacheInterface& DDC = GetDerivedDataCacheRef();
 
 
@@ -1733,7 +1733,7 @@ bool FTexturePlatformData::TryInlineMipData(int32 FirstMipToLoad, FStringView De
 			}
 			if (bLoadedFromDDC)
 			{
-				FMemoryReader Ar(TempData, /*bIsPersistent=*/ true);
+				FMemoryReaderView Ar(MakeMemoryView(TempData), /*bIsPersistent=*/ true);
 
 				Mip.BulkData.Lock(LOCK_READ_WRITE);
 				void* MipData = Mip.BulkData.Realloc(TempData.Num());
@@ -1845,7 +1845,7 @@ bool FTexturePlatformData::TryLoadMips(int32 FirstMipToLoad, void** OutMipData,
 
 #if WITH_EDITOR
 
-	TArray<uint8> TempData;
+	TArray64<uint8> TempData;
 	FAsyncMipHandles AsyncHandles;
 	FDerivedDataCacheInterface& DDC = GetDerivedDataCacheRef();
 	if (!BeginLoadDerivedMips(*this, FirstMipToLoad, DebugContext, AsyncHandles,
@@ -1853,7 +1853,7 @@ bool FTexturePlatformData::TryLoadMips(int32 FirstMipToLoad, void** OutMipData,
 		{
 			FTexture2DMipMap& Mip = Mips[MipIndex];
 
-			const int32 MipSize = static_cast<int32>(MipBuffer.GetSize());
+			const int64 MipSize = static_cast<int64>(MipBuffer.GetSize());
 			CheckMipSize(Mip, PixelFormat, MipSize);
 			NumMipsCached++;
 
@@ -1919,7 +1919,7 @@ bool FTexturePlatformData::TryLoadMips(int32 FirstMipToLoad, void** OutMipData,
 				DDC.WaitAsynchronousCompletion(AsyncHandle);
 				if (DDC.GetAsynchronousResults(AsyncHandle, TempData))
 				{
-					FMemoryReader Ar(TempData, /*bIsPersistent=*/ true);
+					FMemoryReaderView Ar(MakeMemoryView(TempData), /*bIsPersistent=*/ true);
 					CheckMipSize(Mip, PixelFormat, TempData.Num());
 					NumMipsCached++;
 
@@ -1953,9 +1953,9 @@ bool FTexturePlatformData::TryLoadMips(int32 FirstMipToLoad, void** OutMipData,
 		for (int32 MipIndex = FirstMipToLoad; MipIndex < LoadableMips; ++MipIndex)
 		{
 			FTexture2DMipMap& Mip = Mips[MipIndex];
-			UE_LOG(LogTexture, Verbose, TEXT("  Mip %d, BulkDataSize: %d"),
+			UE_LOG(LogTexture, Verbose, TEXT("  Mip %d, BulkDataSize: %" INT64_FMT),
 				MipIndex,
-				(int32)Mip.BulkData.GetBulkDataSize());
+				Mip.BulkData.GetBulkDataSize());
 
 			if (OutMipData && OutMipData[MipIndex - FirstMipToLoad])
 			{
diff --git a/Engine/Source/Runtime/Engine/Private/TextureDerivedDataTask.cpp b/Engine/Source/Runtime/Engine/Private/TextureDerivedDataTask.cpp
--- a/Engine/Source/Runtime/Engine/Private/TextureDerivedDataTask.cpp
+++ b/Engine/Source/Runtime/Engine/Private/TextureDerivedDataTask.cpp
@@ -710,7 +710,7 @@ void FTextureCacheDerivedDataWorker::DoWork()
 	const bool bValidateVirtualTextureCompression = CVarVTValidateCompressionOnLoad.GetValueOnAnyThread() != 0;
 	bool bInvalidVirtualTextureCompression = false;
 
-	TArray<uint8> RawDerivedData;
+	TArray64<uint8> RawDerivedData;
 
 	// Can't have a fetch first if we are rebuilding
 	if (bForceRebuild)
@@ -765,7 +765,7 @@ void FTextureCacheDerivedDataWorker::DoWork()
 		const bool bForDDC = EnumHasAnyFlags(CacheFlags, ETextureCacheFlags::ForDDCBuild);
 
 		BytesCached = RawDerivedData.Num();
-		FMemoryReader Ar(RawDerivedData, /*bIsPersistent=*/ true);
+		FMemoryReaderView Ar(MakeMemoryView(RawDerivedData), /*bIsPersistent=*/ true);
 		DerivedData->Serialize(Ar, NULL);
 		bSucceeded = true;
 		// Load any streaming (not inline) mips that are necessary for our platform.
diff --git a/Engine/Source/Runtime/Engine/Private/TextureDerivedDataTask.h b/Engine/Source/Runtime/Engine/Private/TextureDerivedDataTask.h
--- a/Engine/Source/Runtime/Engine/Private/TextureDerivedDataTask.h
+++ b/Engine/Source/Runtime/Engine/Private/TextureDerivedDataTask.h
@@ -33,7 +33,7 @@ enum
 namespace UE::DerivedData { class FBuildOutput; }
 
 void GetTextureDerivedDataKeySuffix(const UTexture& Texture, const FTextureBuildSettings* BuildSettingsPerLayer, FString& OutKeySuffix);
-uint32 PutDerivedDataInCache(FTexturePlatformData* DerivedData, const FString& DerivedDataKeySuffix, const FStringView& TextureName, bool bForceAllMipsToBeInlined, bool bReplaceExistingDDC);
+int64 PutDerivedDataInCache(FTexturePlatformData* DerivedData, const FString& DerivedDataKeySuffix, const FStringView& TextureName, bool bForceAllMipsToBeInlined, bool bReplaceExistingDDC);
 
 enum class ETextureCacheFlags : uint32
 {
@@ -144,7 +144,7 @@ class FTextureCacheDerivedDataWorker : public FNonAbandonableTask
 	/** Texture cache flags. */
 	ETextureCacheFlags CacheFlags;
 	/** Have many bytes were loaded from DDC or built (for telemetry) */
-	uint32 BytesCached = 0;
+	int64 BytesCached = 0;
 	/** Estimate of the peak amount of memory required to complete this task. */
 	int64 RequiredMemoryEstimate = -1;
 
@@ -175,7 +175,7 @@ public:
 	void Finalize();
 
 	/** Expose bytes cached for telemetry. */
-	uint32 GetBytesCached() const
+	int64 GetBytesCached() const
 	{
 		return BytesCached;
 	}
diff --git a/Engine/Source/Runtime/Engine/Private/VT/VirtualTextureBuiltData.cpp b/Engine/Source/Runtime/Engine/Private/VT/VirtualTextureBuiltData.cpp
--- a/Engine/Source/Runtime/Engine/Private/VT/VirtualTextureBuiltData.cpp
+++ b/Engine/Source/Runtime/Engine/Private/VT/VirtualTextureBuiltData.cpp
@@ -408,7 +408,7 @@ bool FVirtualTextureBuiltData::ValidateData(FStringView const& InDDCDebugContext
 
 	const uint32 TilePixelSize = GetPhysicalTileSize();
 	TArray<uint8> UncompressedResult;
-	TArray<uint8> ChunkDataDDC;
+	TArray64<uint8> ChunkDataDDC;
 
 	const FString TextureName(InDDCDebugContext);
 
@@ -535,16 +535,16 @@ bool FVirtualTextureDataChunk::ShortenKey(const FString& CacheKey, FString& Resu
 	return true;
 }
 
-uint32 FVirtualTextureDataChunk::StoreInDerivedDataCache(const FString& InDerivedDataKey, const FStringView& TextureName, bool bReplaceExistingDDC)
+int64 FVirtualTextureDataChunk::StoreInDerivedDataCache(const FString& InDerivedDataKey, const FStringView& TextureName, bool bReplaceExistingDDC)
 {
 	TRACE_CPUPROFILER_EVENT_SCOPE(FVirtualTextureDataChunk::StoreInDerivedDataCache);
 
-	int32 BulkDataSizeInBytes = BulkData.GetBulkDataSize();
+	int64 BulkDataSizeInBytes = BulkData.GetBulkDataSize();
 	check(BulkDataSizeInBytes > 0);
 
 	{
 		const uint8* BulkChunkData = (uint8*)BulkData.Lock(LOCK_READ_ONLY);
-		GetDerivedDataCacheRef().Put(*InDerivedDataKey, MakeArrayView(BulkChunkData, BulkDataSizeInBytes), TextureName, bReplaceExistingDDC);
+		GetDerivedDataCacheRef().Put(*InDerivedDataKey, TConstArrayView64<uint8>(BulkChunkData, BulkDataSizeInBytes), TextureName, bReplaceExistingDDC);
 		BulkData.Unlock();
 	}
 
diff --git a/Engine/Source/Runtime/Engine/Private/VT/VirtualTextureBuiltData.h b/Engine/Source/Runtime/Engine/Private/VT/VirtualTextureBuiltData.h
--- a/Engine/Source/Runtime/Engine/Private/VT/VirtualTextureBuiltData.h
+++ b/Engine/Source/Runtime/Engine/Private/VT/VirtualTextureBuiltData.h
@@ -66,8 +66,7 @@ struct FVirtualTextureDataChunk
 	bool ShortenKey(const FString& CacheKey, FString& Result);
 	FThreadSafeBool bFileAvailableInVTDDCDache;
 	bool bCorruptDataLoadedFromDDC = false;
-
-	uint32 StoreInDerivedDataCache(const FString& InDerivedDataKey, const FStringView& TextureName, bool bReplaceExistingDDC);
+	int64 StoreInDerivedDataCache(const FString& InDerivedDataKey, const FStringView& TextureName, bool bReplaceExistingDDC);
 #endif // WITH_EDITORONLY_DATA
 };
 
diff --git a/Engine/Source/Runtime/Engine/Private/VT/VirtualTextureChunkDDCCache.cpp b/Engine/Source/Runtime/Engine/Private/VT/VirtualTextureChunkDDCCache.cpp
--- a/Engine/Source/Runtime/Engine/Private/VT/VirtualTextureChunkDDCCache.cpp
+++ b/Engine/Source/Runtime/Engine/Private/VT/VirtualTextureChunkDDCCache.cpp
@@ -75,7 +75,7 @@ public:
 					Header.Version == FVirtualTextureFileHeader::CurrentVersion &&
 					Header.HashSize <= kMaxHashSize)
 				{
-					TArray<uint8> FileContents;
+					TArray64<uint8> FileContents;
 					FileContents.AddUninitialized(Header.HashSize);
 					Ar->Serialize(FileContents.GetData(), Header.HashSize);
 
@@ -104,7 +104,7 @@ public:
 		}
 
 		// Fetch data from DDC
-		TArray<uint8> Results;
+		TArray64<uint8> Results;
 		//TODO(ddebaets) this sync request seems to be blocking here while it uses the job pool. add overload to perform it on this thread?
 		const bool DDCResult = DDC.GetSynchronous(*Chunk->DerivedDataKey, Results, FinalFilename);
 		if (DDCResult == false)
@@ -115,12 +115,12 @@ public:
 
 		if (Results.Num() <= sizeof(FVirtualTextureChunkHeader))
 		{
-			UE_LOG(LogVTDiskCache, Error, TEXT("VT DDC data is too small %d (key: %s)"), Results.Num(), *Chunk->DerivedDataKey);
+			UE_LOG(LogVTDiskCache, Error, TEXT("VT DDC data is too small %" INT64_FMT " (key: %s)"), Results.Num(), *Chunk->DerivedDataKey);
 			return;
 		}
 
 		const uint8* ChunkData = Results.GetData();
-		const int32 ChunkDataSize = Results.Num();
+		const int64 ChunkDataSize = Results.Num();
 
 		FSHAHash Hash;
 		FSHA1::HashBuffer(ChunkData, ChunkDataSize, Hash.Hash);
diff --git a/Engine/Source/Runtime/Engine/Public/TextureResource.h b/Engine/Source/Runtime/Engine/Public/TextureResource.h
--- a/Engine/Source/Runtime/Engine/Public/TextureResource.h
+++ b/Engine/Source/Runtime/Engine/Public/TextureResource.h
@@ -76,7 +76,7 @@ struct FTexture2DMipMap
 	 * Place mip-map data in the derived data cache associated with the provided
 	 * key.
 	 */
-	uint32 StoreInDerivedDataCache(const FString& InDerivedDataKey, const FStringView& TextureName, bool bReplaceExistingDDC);
+	int64 StoreInDerivedDataCache(const FString& InDerivedDataKey, const FStringView& TextureName, bool bReplaceExistingDDC);
 
 #endif // #if WITH_EDITORONLY_DATA
 };
