Various fixes for MetaSound meter functionality
- Add nullptr check when setting bus sends
- Initialize audio analyzer number of frames
- Remove analyzers when tearing down MetaSound audio meter

diff --git a/Engine/Plugins/Runtime/Metasound/Source/MetasoundEditor/Private/MetasoundEditorMeter.cpp b/Engine/Plugins/Runtime/Metasound/Source/MetasoundEditor/Private/MetasoundEditorMeter.cpp
--- a/Engine/Plugins/Runtime/Metasound/Source/MetasoundEditor/Private/MetasoundEditorMeter.cpp
+++ b/Engine/Plugins/Runtime/Metasound/Source/MetasoundEditor/Private/MetasoundEditorMeter.cpp
@@ -94,6 +94,14 @@ namespace Metasound
 		{
 			if (Analyzer.IsValid() && ResultsDelegateHandle.IsValid())
 			{
+				if (ensure(GEditor))
+				{
+					UWorld* EditorWorld = GEditor->GetEditorWorldContext().World();
+					if (ensure(EditorWorld))
+					{
+						Analyzer->StopAnalyzing(EditorWorld);
+					}
+				}
 				Analyzer->OnLatestPerChannelMeterResultsNative.Remove(ResultsDelegateHandle);
 				Analyzer.Reset();
 				ResultsDelegateHandle.Reset();
diff --git a/Engine/Source/Runtime/AudioAnalyzer/Private/AudioAnalyzer.cpp b/Engine/Source/Runtime/AudioAnalyzer/Private/AudioAnalyzer.cpp
--- a/Engine/Source/Runtime/AudioAnalyzer/Private/AudioAnalyzer.cpp
+++ b/Engine/Source/Runtime/AudioAnalyzer/Private/AudioAnalyzer.cpp
@@ -77,6 +77,7 @@ void UAudioAnalyzer::StartAnalyzing(UWorld* InWorld, UAudioBus* AudioBusToAnalyz
 
 	// Get an output patch for the audio bus
 	PatchOutputStrongPtr = MixerDevice->AddPatchForAudioBus_GameThread(AudioBusId);
+	NumFramesPerBufferToAnalyze = MixerDevice->GetNumOutputFrames();
 
 	// Register this audio analyzer with the audio analyzer subsystem
 	// The subsystem will query this analyzer to see if it has enough audio to perform analysis.
diff --git a/Engine/Source/Runtime/Engine/Private/ActiveSound.cpp b/Engine/Source/Runtime/Engine/Private/ActiveSound.cpp
--- a/Engine/Source/Runtime/Engine/Private/ActiveSound.cpp
+++ b/Engine/Source/Runtime/Engine/Private/ActiveSound.cpp
@@ -486,7 +486,10 @@ void FActiveSound::GetBusSends(EBusSendType BusSendType, TArray<FSoundSourceBusS
 			bool bOverridden = false;
 			for (FSoundSourceBusSendInfo& OutSendInfo : OutSends)
 			{
-				if (OutSendInfo.SoundSourceBus == SendInfo.SoundSourceBus || OutSendInfo.AudioBus == SendInfo.AudioBus)
+				const bool bSameSourceBus = (OutSendInfo.SoundSourceBus != nullptr) && (OutSendInfo.SoundSourceBus == SendInfo.SoundSourceBus);
+				const bool bSameAudioBus = (OutSendInfo.AudioBus != nullptr) && (OutSendInfo.AudioBus == SendInfo.AudioBus);
+
+				if (bSameSourceBus || bSameAudioBus)
 				{
 					OutSendInfo.SendLevel = SendInfo.SendLevel;
 					bOverridden = true;
