Support more expression types with new HLSL translator
Extend CVAR to allow forcing all materials to use new translator, for testing

diff --git a/Engine/Source/Runtime/Engine/Classes/Materials/MaterialExpressionFeatureLevelSwitch.h b/Engine/Source/Runtime/Engine/Classes/Materials/MaterialExpressionFeatureLevelSwitch.h
--- a/Engine/Source/Runtime/Engine/Classes/Materials/MaterialExpressionFeatureLevelSwitch.h
+++ b/Engine/Source/Runtime/Engine/Classes/Materials/MaterialExpressionFeatureLevelSwitch.h
@@ -39,6 +39,7 @@ class UMaterialExpressionFeatureLevelSwitch : public UMaterialExpression
 	virtual bool IsResultMaterialAttributes(int32 OutputIndex) override;
 	virtual uint32 GetInputType(int32 InputIndex) override {return MCT_Unknown;}
 	virtual uint32 GetOutputType(int32 OutputIndex) override {return MCT_Unknown;}
+	virtual bool GenerateHLSLExpression(FMaterialHLSLGenerator& Generator, UE::HLSLTree::FScope& Scope, int32 OutputIndex, UE::HLSLTree::FExpression*& OutExpression) override;
 #endif // WITH_EDITOR
 	//~ End UMaterialExpression Interface
 };
diff --git a/Engine/Source/Runtime/Engine/Classes/Materials/MaterialExpressionSceneTexelSize.h b/Engine/Source/Runtime/Engine/Classes/Materials/MaterialExpressionSceneTexelSize.h
--- a/Engine/Source/Runtime/Engine/Classes/Materials/MaterialExpressionSceneTexelSize.h
+++ b/Engine/Source/Runtime/Engine/Classes/Materials/MaterialExpressionSceneTexelSize.h
@@ -17,6 +17,7 @@ class UMaterialExpressionSceneTexelSize : public UMaterialExpression
 #if WITH_EDITOR
 	virtual int32 Compile(class FMaterialCompiler* Compiler, int32 OutputIndex) override;
 	virtual void GetCaption(TArray<FString>& OutCaptions) const override;
+	virtual bool GenerateHLSLExpression(FMaterialHLSLGenerator& Generator, UE::HLSLTree::FScope& Scope, int32 OutputIndex, UE::HLSLTree::FExpression*& OutExpression) override;
 #endif
 	//~ End UMaterialExpression Interface
 };
diff --git a/Engine/Source/Runtime/Engine/Classes/Materials/MaterialExpressionSceneTexture.h b/Engine/Source/Runtime/Engine/Classes/Materials/MaterialExpressionSceneTexture.h
--- a/Engine/Source/Runtime/Engine/Classes/Materials/MaterialExpressionSceneTexture.h
+++ b/Engine/Source/Runtime/Engine/Classes/Materials/MaterialExpressionSceneTexture.h
@@ -31,6 +31,7 @@ class UMaterialExpressionSceneTexture : public UMaterialExpression
 #if WITH_EDITOR
 	virtual int32 Compile(class FMaterialCompiler* Compiler, int32 OutputIndex) override;
 	virtual void GetCaption(TArray<FString>& OutCaptions) const override;
+	virtual bool GenerateHLSLExpression(FMaterialHLSLGenerator& Generator, UE::HLSLTree::FScope& Scope, int32 OutputIndex, UE::HLSLTree::FExpression*& OutExpression) override;
 #endif
 	//~ End UMaterialExpression Interface
 };
diff --git a/Engine/Source/Runtime/Engine/Classes/Materials/MaterialExpressionScreenPosition.h b/Engine/Source/Runtime/Engine/Classes/Materials/MaterialExpressionScreenPosition.h
--- a/Engine/Source/Runtime/Engine/Classes/Materials/MaterialExpressionScreenPosition.h
+++ b/Engine/Source/Runtime/Engine/Classes/Materials/MaterialExpressionScreenPosition.h
@@ -18,6 +18,7 @@ class UMaterialExpressionScreenPosition : public UMaterialExpression
 #if WITH_EDITOR
 	virtual int32 Compile(class FMaterialCompiler* Compiler, int32 OutputIndex) override;
 	virtual void GetCaption(TArray<FString>& OutCaptions) const override;
+	virtual bool GenerateHLSLExpression(FMaterialHLSLGenerator& Generator, UE::HLSLTree::FScope& Scope, int32 OutputIndex, UE::HLSLTree::FExpression*& OutExpression) override;
 #endif
 	//~ End UMaterialExpression Interface
 };
diff --git a/Engine/Source/Runtime/Engine/Classes/Materials/MaterialExpressionViewSize.h b/Engine/Source/Runtime/Engine/Classes/Materials/MaterialExpressionViewSize.h
--- a/Engine/Source/Runtime/Engine/Classes/Materials/MaterialExpressionViewSize.h
+++ b/Engine/Source/Runtime/Engine/Classes/Materials/MaterialExpressionViewSize.h
@@ -17,6 +17,7 @@ class UMaterialExpressionViewSize : public UMaterialExpression
 #if WITH_EDITOR
 	virtual int32 Compile(class FMaterialCompiler* Compiler, int32 OutputIndex) override;
 	virtual void GetCaption(TArray<FString>& OutCaptions) const override;
+	virtual bool GenerateHLSLExpression(FMaterialHLSLGenerator& Generator, UE::HLSLTree::FScope& Scope, int32 OutputIndex, UE::HLSLTree::FExpression*& OutExpression) override;
 #endif
 	//~ End UMaterialExpression Interface
 };
diff --git a/Engine/Source/Runtime/Engine/Private/HLSLTree/HLSLTreeCommon.cpp b/Engine/Source/Runtime/Engine/Private/HLSLTree/HLSLTreeCommon.cpp
--- a/Engine/Source/Runtime/Engine/Private/HLSLTree/HLSLTreeCommon.cpp
+++ b/Engine/Source/Runtime/Engine/Private/HLSLTree/HLSLTreeCommon.cpp
@@ -3,6 +3,7 @@
 #include "HLSLTree/HLSLTreeEmit.h"
 #include "Misc/StringBuilder.h"
 #include "MaterialShared.h"
+#include "MaterialSceneTextureId.h"
 #include "Engine/Texture.h"
 
 namespace UE
@@ -62,6 +63,11 @@ FExternalInputDescription GetExternalInputDescription(EExternalInput Input)
 	case EExternalInput::WorldPosition_Ddx: return FExternalInputDescription(TEXT("WorldPosition_Ddx"), Shader::EValueType::Float3);
 	case EExternalInput::WorldPosition_Ddy: return FExternalInputDescription(TEXT("WorldPosition_Ddx"), Shader::EValueType::Float3);
 
+	case EExternalInput::ViewportUV: return FExternalInputDescription(TEXT("ViewportUV"), Shader::EValueType::Float2);
+	case EExternalInput::PixelPosition: return FExternalInputDescription(TEXT("PixelPosition"), Shader::EValueType::Float2);
+	case EExternalInput::ViewSize: return FExternalInputDescription(TEXT("ViewSize"), Shader::EValueType::Float2);
+	case EExternalInput::RcpViewSize: return FExternalInputDescription(TEXT("RcpViewSize"), Shader::EValueType::Float2);
+
 	case EExternalInput::CameraWorldPosition: return FExternalInputDescription(TEXT("CameraWorldPosition"), Shader::EValueType::Double3, EExternalInput::None, EExternalInput::None, EExternalInput::PrevCameraWorldPosition);
 	case EExternalInput::PreViewTranslation: return FExternalInputDescription(TEXT("PreViewTranslation"), Shader::EValueType::Double3, EExternalInput::None, EExternalInput::None, EExternalInput::PrevPreViewTranslation);
 	case EExternalInput::TangentToWorld: return FExternalInputDescription(TEXT("TangentToWorld"), Shader::EValueType::Float4x4);
@@ -260,6 +266,24 @@ void FExpressionExternalInput::ComputeAnalyticDerivatives(FTree& Tree, FExpressi
 		OutResult.ExpressionDdx = Tree.NewExpression<FExpressionExternalInput>(InputDesc.Ddx);
 		OutResult.ExpressionDdy = Tree.NewExpression<FExpressionExternalInput>(InputDesc.Ddy);
 	}
+	else
+	{
+		switch (InputType)
+		{
+		case EExternalInput::ViewportUV:
+		{
+			// Ddx = float2(RcpViewSize.x, 0.0f)
+			// Ddy = float2(0.0f, RcpViewSize.y)
+			FExpression* RcpViewSize = Tree.NewExpression<FExpressionExternalInput>(EExternalInput::RcpViewSize);
+			FExpression* Constant0 = Tree.NewConstant(0.0f);
+			OutResult.ExpressionDdx = Tree.NewExpression<FExpressionAppend>(Tree.NewExpression<FExpressionSwizzle>(MakeSwizzleMask(true, false, false, false), RcpViewSize), Constant0);
+			OutResult.ExpressionDdy = Tree.NewExpression<FExpressionAppend>(Constant0, Tree.NewExpression<FExpressionSwizzle>(MakeSwizzleMask(false, true, false, false), RcpViewSize));
+			break;
+		}
+		default:
+			break;
+		}
+	}
 }
 
 FExpression* FExpressionExternalInput::ComputePreviousFrame(FTree& Tree, const FRequestedType& RequestedType) const
@@ -314,6 +338,11 @@ void FExpressionExternalInput::EmitValueShader(FEmitContext& Context, FEmitScope
 		case EExternalInput::WorldPosition_Ddx: Code = TEXT("Parameters.WorldPosition_DDX"); break;
 		case EExternalInput::WorldPosition_Ddy: Code = TEXT("Parameters.WorldPosition_DDY"); break;
 
+		case EExternalInput::ViewportUV: Code = TEXT("GetViewportUV(Parameters)"); break;
+		case EExternalInput::PixelPosition: Code = TEXT("GetPixelPosition(Parameters)"); break;
+		case EExternalInput::ViewSize: Code = TEXT("View.ViewSizeAndInvSize.xy"); break;
+		case EExternalInput::RcpViewSize: Code = TEXT("View.ViewSizeAndInvSize.zw"); break;
+
 		case EExternalInput::CameraWorldPosition: Code = TEXT("ResolvedView.WorldCameraOrigin"); break;
 		case EExternalInput::PreViewTranslation: Code = TEXT("ResolvedView.PreViewTranslation"); break;
 		case EExternalInput::TangentToWorld: Code = TEXT("Parameters.TangentToWorld"); break;
@@ -353,6 +382,52 @@ void FExpressionExternalInput::EmitValueShader(FEmitContext& Context, FEmitScope
 	}
 }
 
+bool FExpressionMaterialSceneTexture::PrepareValue(FEmitContext& Context, FEmitScope& Scope, const FRequestedType& RequestedType, FPrepareValueResult& OutResult) const
+{
+	Context.PrepareExpression(TexCoordExpression, Scope, ERequestedType::Vector2);
+
+	Context.MaterialCompilationOutput->bNeedsSceneTextures = true;
+	Context.MaterialCompilationOutput->SetIsSceneTextureUsed((ESceneTextureId)SceneTextureId);
+
+	return OutResult.SetType(Context, RequestedType, EExpressionEvaluation::Shader, Shader::EValueType::Float4);
+}
+
+void FExpressionMaterialSceneTexture::EmitValueShader(FEmitContext& Context, FEmitScope& Scope, const FRequestedType& RequestedType, FEmitValueShaderResult& OutResult) const
+{
+	const bool bSupportedOnMobile = SceneTextureId == PPI_PostProcessInput0 ||
+		SceneTextureId == PPI_CustomDepth ||
+		SceneTextureId == PPI_SceneDepth ||
+		SceneTextureId == PPI_CustomStencil;
+
+	FEmitShaderExpression* EmitTexCoord = nullptr;
+	if (TexCoordExpression)
+	{
+		EmitTexCoord = TexCoordExpression->GetValueShader(Context, Scope, Shader::EValueType::Float2);
+		EmitTexCoord = Context.EmitExpression(Scope, Shader::EValueType::Float2, TEXT("ClampSceneTextureUV(ViewportUVToSceneTextureUV(%, %), %)"), EmitTexCoord, (int)SceneTextureId, (int)SceneTextureId);
+	}
+	else
+	{
+		EmitTexCoord = Context.EmitExpression(Scope, Shader::EValueType::Float2, TEXT("GetDefaultSceneTextureUV(Parameters, %)"), (int)SceneTextureId);
+	}
+
+	FEmitShaderExpression* EmitLookup = nullptr;
+	if (Context.Material->GetFeatureLevel() >= ERHIFeatureLevel::SM5)
+	{
+		EmitLookup = Context.EmitExpression(Scope, Shader::EValueType::Float4, TEXT("SceneTextureLookup(%, %, %)"), EmitTexCoord, (int)SceneTextureId, bFiltered);
+	}
+	else
+	{
+		EmitLookup = Context.EmitExpression(Scope, Shader::EValueType::Float4, TEXT("MobileSceneTextureLookup(Parameters, %, %, %)"), (int)SceneTextureId, EmitTexCoord);
+	}
+
+	if (SceneTextureId == PPI_PostProcessInput0 && Context.Material->GetMaterialDomain() == MD_PostProcess && Context.Material->GetBlendableLocation() != BL_AfterTonemapping)
+	{
+		EmitLookup = Context.EmitExpression(Scope, Shader::EValueType::Float4, TEXT("(float4(View.OneOverPreExposure.xxx, 1) * %)"), EmitLookup);
+	}
+	
+	OutResult.Code = EmitLookup;
+}
+
 bool FExpressionTextureSample::PrepareValue(FEmitContext& Context, FEmitScope& Scope, const FRequestedType& RequestedType, FPrepareValueResult& OutResult) const
 {
 	Context.PrepareExpression(TexCoordExpression, Scope, ERequestedType::Vector2);
diff --git a/Engine/Source/Runtime/Engine/Private/HLSLTree/HLSLTreeEmit.cpp b/Engine/Source/Runtime/Engine/Private/HLSLTree/HLSLTreeEmit.cpp
--- a/Engine/Source/Runtime/Engine/Private/HLSLTree/HLSLTreeEmit.cpp
+++ b/Engine/Source/Runtime/Engine/Private/HLSLTree/HLSLTreeEmit.cpp
@@ -610,6 +610,7 @@ int32 InternalFormatString(FStringBuilderBase* OutString, FEmitShaderDependencie
 				case EFormatArgType::ShaderValue: FormatArg_ShaderValue(Arg.ShaderValue, OutDependencies, *OutString); break;
 				case EFormatArgType::String: OutString->Append(Arg.String); break;
 				case EFormatArgType::Int: OutString->Appendf(TEXT("%d"), Arg.Int); break;
+				case EFormatArgType::Bool: OutString->Append(Arg.Bool ? TEXT("true") : TEXT("false")); break;
 				default:
 					checkNoEntry();
 					break;
diff --git a/Engine/Source/Runtime/Engine/Private/Materials/Material.cpp b/Engine/Source/Runtime/Engine/Private/Materials/Material.cpp
--- a/Engine/Source/Runtime/Engine/Private/Materials/Material.cpp
+++ b/Engine/Source/Runtime/Engine/Private/Materials/Material.cpp
@@ -136,7 +136,10 @@ static TAutoConsoleVariable<int32> CVarMaterialEnableControlFlow(
 static TAutoConsoleVariable<int32> CVarMaterialEnableNewHLSLGenerator(
 	TEXT("r.MaterialEnableNewHLSLGenerator"),
 	0,
-	TEXT("Enables the new (WIP) material HLSL generator.\n"),
+	TEXT("Enables the new (WIP) material HLSL generator.\n")
+	TEXT("0 - Don't allow\n")
+	TEXT("1 - Allow if enabled by material\n")
+	TEXT("2 - Force all materials to use new generator\n"),
 	ECVF_RenderThreadSafe | ECVF_ReadOnly);
 
 bool Engine_IsStrataEnabled()
@@ -3411,11 +3414,12 @@ bool UMaterial::IsUsingControlFlow() const
 
 bool UMaterial::IsUsingNewHLSLGenerator() const
 {
+	const int CVarValue = CVarMaterialEnableNewHLSLGenerator.GetValueOnAnyThread();
 	if (bEnableNewHLSLGenerator)
 	{
-		return CVarMaterialEnableNewHLSLGenerator.GetValueOnAnyThread() != 0;
+		return CVarValue != 0;
 	}
-	return false;
+	return CVarValue == 2;
 }
 
 #if WITH_EDITOR
diff --git a/Engine/Source/Runtime/Engine/Private/Materials/MaterialExpressionHLSL.cpp b/Engine/Source/Runtime/Engine/Private/Materials/MaterialExpressionHLSL.cpp
--- a/Engine/Source/Runtime/Engine/Private/Materials/MaterialExpressionHLSL.cpp
+++ b/Engine/Source/Runtime/Engine/Private/Materials/MaterialExpressionHLSL.cpp
@@ -26,16 +26,21 @@
 #include "Materials/MaterialExpressionCameraPositionWS.h"
 #include "Materials/MaterialExpressionTime.h"
 #include "Materials/MaterialExpressionDeltaTime.h"
+#include "Materials/MaterialExpressionScreenPosition.h"
+#include "Materials/MaterialExpressionSceneTexelSize.h"
+#include "Materials/MaterialExpressionViewSize.h"
 #include "Materials/MaterialExpressionPanner.h"
 #include "Materials/MaterialExpressionTextureCoordinate.h"
 #include "Materials/MaterialExpressionTextureSample.h"
 #include "Materials/MaterialExpressionTextureSampleParameter.h"
 #include "Materials/MaterialExpressionTextureObject.h"
 #include "Materials/MaterialExpressionTextureObjectParameter.h"
+#include "Materials/MaterialExpressionSceneTexture.h"
 #include "Materials/MaterialExpressionFunctionInput.h"
 #include "Materials/MaterialExpressionFunctionOutput.h"
 #include "Materials/MaterialExpressionMaterialFunctionCall.h"
 #include "Materials/MaterialExpressionStaticSwitch.h"
+#include "Materials/MaterialExpressionFeatureLevelSwitch.h"
 #include "Materials/MaterialExpressionGetLocal.h"
 #include "Materials/MaterialExpressionOneMinus.h"
 #include "Materials/MaterialExpressionAdd.h"
@@ -125,6 +130,28 @@ bool UMaterialExpressionStaticSwitch::GenerateHLSLExpression(FMaterialHLSLGenera
 	return true;
 }
 
+bool UMaterialExpressionFeatureLevelSwitch::GenerateHLSLExpression(FMaterialHLSLGenerator& Generator, UE::HLSLTree::FScope& Scope, int32 OutputIndex, UE::HLSLTree::FExpression*& OutExpression)
+{
+	const ERHIFeatureLevel::Type FeatureLevelToCompile = Generator.GetCompileTarget().FeatureLevel;
+	check(FeatureLevelToCompile < UE_ARRAY_COUNT(Inputs));
+	FExpressionInput& FeatureInput = Inputs[FeatureLevelToCompile];
+
+	if (!Default.GetTracedInput().Expression)
+	{
+		return Generator.GetErrors().AddError(TEXT("Feature Level switch missing default input"));
+	}
+
+	if (FeatureInput.GetTracedInput().Expression)
+	{
+		OutExpression = FeatureInput.AcquireHLSLExpression(Generator, Scope);
+	}
+	else
+	{
+		OutExpression = Default.AcquireHLSLExpression(Generator, Scope);
+	}
+	return OutExpression != nullptr;
+}
+
 bool UMaterialExpressionGetLocal::GenerateHLSLExpression(FMaterialHLSLGenerator& Generator, UE::HLSLTree::FScope& Scope, int32 OutputIndex, UE::HLSLTree::FExpression*& OutExpression)
 {
 	OutExpression = Generator.GetTree().AcquireLocal(Scope, LocalName);
@@ -202,7 +229,7 @@ bool UMaterialExpressionTime::GenerateHLSLExpression(FMaterialHLSLGenerator& Gen
 	}
 
 	EExternalInput InputType = bIgnorePause ? EExternalInput::RealTime : EExternalInput::GameTime;
-	OutExpression = Generator.GetTree().NewExpression<UE::HLSLTree::FExpressionExternalInput>(InputType);
+	OutExpression = Generator.GetTree().NewExpression<FExpressionExternalInput>(InputType);
 	if (bOverride_Period)
 	{
 		OutExpression = Generator.GetTree().NewFmod(OutExpression, Generator.NewConstant(Period));
@@ -216,6 +243,31 @@ bool UMaterialExpressionDeltaTime::GenerateHLSLExpression(FMaterialHLSLGenerator
 	return true;
 }
 
+bool UMaterialExpressionScreenPosition::GenerateHLSLExpression(FMaterialHLSLGenerator& Generator, UE::HLSLTree::FScope& Scope, int32 OutputIndex, UE::HLSLTree::FExpression*& OutExpression)
+{
+	using namespace UE::HLSLTree;
+	const EExternalInput InputType = (OutputIndex == 1) ? EExternalInput::PixelPosition : EExternalInput::ViewportUV;
+	OutExpression = Generator.GetTree().NewExpression<FExpressionExternalInput>(InputType);
+	return true;
+}
+
+bool UMaterialExpressionSceneTexelSize::GenerateHLSLExpression(FMaterialHLSLGenerator& Generator, UE::HLSLTree::FScope& Scope, int32 OutputIndex, UE::HLSLTree::FExpression*& OutExpression)
+{
+	using namespace UE::HLSLTree;
+
+	// To make sure any material that were correctly handling BufferUV != ViewportUV, we just lie to material
+	// to make it believe ViewSize == BufferSize, so they are still compatible with SceneTextureLookup().
+	OutExpression = Generator.GetTree().NewExpression<FExpressionExternalInput>(EExternalInput::RcpViewSize);
+	return true;
+}
+
+bool UMaterialExpressionViewSize::GenerateHLSLExpression(FMaterialHLSLGenerator& Generator, UE::HLSLTree::FScope& Scope, int32 OutputIndex, UE::HLSLTree::FExpression*& OutExpression)
+{
+	using namespace UE::HLSLTree;
+	OutExpression = Generator.GetTree().NewExpression<FExpressionExternalInput>(EExternalInput::ViewSize);
+	return true;
+}
+
 bool UMaterialExpressionPanner::GenerateHLSLExpression(FMaterialHLSLGenerator& Generator, UE::HLSLTree::FScope& Scope, int32 OutputIndex, UE::HLSLTree::FExpression*& OutExpression)
 {
 	using namespace UE::HLSLTree;
@@ -324,6 +376,29 @@ bool UMaterialExpressionTextureSampleParameter::GenerateHLSLExpression(FMaterial
 	return GenerateHLSLExpressionBase(Generator, Scope, TextureDeclaration, OutExpression);
 }
 
+bool UMaterialExpressionSceneTexture::GenerateHLSLExpression(FMaterialHLSLGenerator& Generator, UE::HLSLTree::FScope& Scope, int32 OutputIndex, UE::HLSLTree::FExpression*& OutExpression)
+{
+	using namespace UE::HLSLTree;
+
+	if (OutputIndex == 0)
+	{
+		FExpression* ExpressionTexCoord = nullptr;
+		if (Coordinates.GetTracedInput().Expression)
+		{
+			ExpressionTexCoord = Coordinates.AcquireHLSLExpression(Generator, Scope);
+		}
+		OutExpression = Generator.GetTree().NewExpression<FExpressionMaterialSceneTexture>(ExpressionTexCoord, SceneTextureId, bFiltered);
+		return true;
+	}
+	else if (OutputIndex == 1 || OutputIndex == 2)
+	{
+		//return Compiler->GetSceneTextureViewSize(SceneTextureId, /* InvProperty = */ OutputIndex == 2);
+		return false; // TODO
+	}
+
+	return Generator.GetErrors().AddError(TEXT("Invalid input parameter"));
+}
+
 bool UMaterialExpressionOneMinus::GenerateHLSLExpression(FMaterialHLSLGenerator& Generator, UE::HLSLTree::FScope& Scope, int32 OutputIndex, UE::HLSLTree::FExpression*& OutExpression)
 {
 	UE::HLSLTree::FExpression* InputExpression = Input.AcquireHLSLExpression(Generator, Scope);
diff --git a/Engine/Source/Runtime/Engine/Public/HLSLTree/HLSLTreeCommon.h b/Engine/Source/Runtime/Engine/Public/HLSLTree/HLSLTreeCommon.h
--- a/Engine/Source/Runtime/Engine/Public/HLSLTree/HLSLTreeCommon.h
+++ b/Engine/Source/Runtime/Engine/Public/HLSLTree/HLSLTreeCommon.h
@@ -86,6 +86,11 @@ enum class EExternalInput : uint8
 	WorldPosition_Ddx,
 	WorldPosition_Ddy,
 
+	ViewportUV,
+	PixelPosition,
+	ViewSize,
+	RcpViewSize,
+
 	CameraWorldPosition,
 	PreViewTranslation,
 	TangentToWorld,
@@ -168,6 +173,23 @@ public:
 	virtual void EmitValueShader(FEmitContext& Context, FEmitScope& Scope, const FRequestedType& RequestedType, FEmitValueShaderResult& OutResult) const override;
 };
 
+class FExpressionMaterialSceneTexture : public FExpression
+{
+public:
+	FExpressionMaterialSceneTexture(FExpression* InTexCoordExpression, uint32 InSceneTextureId, bool bInFiltered)
+		: TexCoordExpression(InTexCoordExpression)
+		, SceneTextureId(InSceneTextureId)
+		, bFiltered(bInFiltered)
+	{}
+
+	FExpression* TexCoordExpression;
+	uint32 SceneTextureId;
+	bool bFiltered;
+
+	virtual bool PrepareValue(FEmitContext& Context, FEmitScope& Scope, const FRequestedType& RequestedType, FPrepareValueResult& OutResult) const override;
+	virtual void EmitValueShader(FEmitContext& Context, FEmitScope& Scope, const FRequestedType& RequestedType, FEmitValueShaderResult& OutResult) const override;
+};
+
 class FExpressionTextureSample : public FExpression
 {
 public:
diff --git a/Engine/Source/Runtime/Engine/Public/HLSLTree/HLSLTreeEmit.h b/Engine/Source/Runtime/Engine/Public/HLSLTree/HLSLTreeEmit.h
--- a/Engine/Source/Runtime/Engine/Public/HLSLTree/HLSLTreeEmit.h
+++ b/Engine/Source/Runtime/Engine/Public/HLSLTree/HLSLTreeEmit.h
@@ -112,6 +112,7 @@ enum class EFormatArgType : uint8
 	ShaderValue,
 	String,
 	Int,
+	Bool,
 };
 
 struct FFormatArgVariant
@@ -120,6 +121,7 @@ struct FFormatArgVariant
 	FFormatArgVariant(FEmitShaderExpression* InValue) : Type(EFormatArgType::ShaderValue), ShaderValue(InValue) { check(InValue); }
 	FFormatArgVariant(const TCHAR* InValue) : Type(EFormatArgType::String), String(InValue) { check(InValue); }
 	FFormatArgVariant(int32 InValue) : Type(EFormatArgType::Int), Int(InValue) {}
+	FFormatArgVariant(bool InValue) : Type(EFormatArgType::Bool), Bool(InValue) {}
 
 	EFormatArgType Type = EFormatArgType::Void;
 	union
@@ -127,6 +129,7 @@ struct FFormatArgVariant
 		FEmitShaderExpression* ShaderValue;
 		const TCHAR* String;
 		int32 Int;
+		bool Bool;
 	};
 };
 
