Strata improvement post review.
Updated and test vertical operator thickenss.

diff --git a/Engine/Shaders/Private/Strata/Strata.ush b/Engine/Shaders/Private/Strata/Strata.ush
--- a/Engine/Shaders/Private/Strata/Strata.ush
+++ b/Engine/Shaders/Private/Strata/Strata.ush
@@ -592,24 +592,26 @@ bool StrataHasGreyScaleWeight(in float3 Weight)
 struct FStrataOperator
 {
 	// STRATA_TODO: pack in a single uint
-	uint TopIndex;
-	uint MaxDistanceFromLeaves;
-	uint Type;
-	int	AIndex;
-	int BIndex;
+	uint ParentIndex;
+	uint MaxDistanceFromLeaves;			// The largest depth distance from any tree leaf representing a BSDF
+	uint Type;							// The type of operator
+	int	LeftIndex;
+	int RightIndex;
 	float Weight;
 
-	float  Coverage;
-	float3 Transmittance;	// This is the transmittance of the matter. If coverage is 0.25, then 25% of the light will be affected by the transmittance and 75% of light won't be.
+	float  Coverage;					// The summarised coverage of all the matter (from the sub tree) represented by that operator
+	float3 Transmittance;				// This is the transmittance of the matter itself. If coverage is 0.25, then 25% of the light will be affected by the transmittance and 75% of light won't be.
 
-	float VerticalTopCoverage;
-	float3 VerticalTopTransmittance;
+	float VerticalTopCoverage;			// Only used by vertical layering operators. It represent the coverage of all the matter of the top layer (from the left sub tree).
+	float3 VerticalTopTransmittance;	// Only used by vertical layering operators. It represent the transmittance of matter of the top layer (from the left sub tree).
 };
 
 FStrataOperator GetInitialisedStrataOperator()
 {
 	FStrataOperator StrataOp = (FStrataOperator)0;
-	StrataOp.TopIndex = -1;
+	StrataOp.ParentIndex = -1;
+	StrataOp.LeftIndex = -1;
+	StrataOp.RightIndex = -1;
 	return StrataOp;
 }
 
@@ -786,8 +788,6 @@ FStrataData GetStrataVolumeFogCloudBSDF(float3 Albedo, float3 Extinction, float3
 	StrataData.InlinedBSDF.Coverage = 1.0f;
 #endif
 
-	// STRATA_TODO BSDF OP
-
 	return StrataData;
 }
 
@@ -803,8 +803,6 @@ FStrataData GetStrataUnlitBSDF(float3 Emissive, float3 TransmittanceColor)
 	StrataData.InlinedBSDF.Coverage = 1.0f;
 #endif
 
-	// STRATA_TODO BSDF OP
-
 	return StrataData;
 }
 
@@ -835,8 +833,6 @@ FStrataData GetStrataHairBSDF(float3 BaseColor, float Scatter, float Specular, f
 	StrataData.InlinedBSDF.Coverage = 1.0f;
 #endif
 
-	// STRATA_TODO BSDF OP
-
 	return StrataData;
 }
 
@@ -865,8 +861,6 @@ FStrataData GetStrataSingleLayerWaterBSDF(
 	StrataData.InlinedBSDF.Coverage = 1.0f;
 #endif
 
-	// STRATA_TODO BSDF OP
-
 	return StrataData;
 }
 
@@ -1104,11 +1098,11 @@ FStrataData StrataAdd(FStrataData A, FStrataData B, inout FStrataTree StrataTree
 	StrataTree.Operators[OperatorIndex]							= GetInitialisedStrataOperator();
 	StrataTree.Operators[OperatorIndex].MaxDistanceFromLeaves	= MaxDistanceFromLeaves;
 	StrataTree.Operators[OperatorIndex].Type					= STRATA_OPERATOR_ADD;
-	StrataTree.Operators[OperatorIndex].AIndex					= A.OperatorIndex;
-	StrataTree.Operators[OperatorIndex].BIndex					= B.OperatorIndex;
+	StrataTree.Operators[OperatorIndex].LeftIndex				= A.OperatorIndex;
+	StrataTree.Operators[OperatorIndex].RightIndex				= B.OperatorIndex;
 
-	StrataTree.Operators[A.OperatorIndex].TopIndex				= OperatorIndex;
-	StrataTree.Operators[B.OperatorIndex].TopIndex				= OperatorIndex;
+	StrataTree.Operators[A.OperatorIndex].ParentIndex			= OperatorIndex;
+	StrataTree.Operators[B.OperatorIndex].ParentIndex			= OperatorIndex;
 
 	return StrataData;
 }
@@ -1126,9 +1120,9 @@ FStrataData StrataWeight(FStrataData A, float Weight, inout FStrataTree StrataTr
 	StrataTree.Operators[OperatorIndex].MaxDistanceFromLeaves	= MaxDistanceFromLeaves;
 	StrataTree.Operators[OperatorIndex].Type					= STRATA_OPERATOR_WEIGHT;
 	StrataTree.Operators[OperatorIndex].Weight					= Weight;
-	StrataTree.Operators[OperatorIndex].AIndex					= A.OperatorIndex;
+	StrataTree.Operators[OperatorIndex].LeftIndex				= A.OperatorIndex;
 
-	StrataTree.Operators[A.OperatorIndex].TopIndex				= OperatorIndex;
+	StrataTree.Operators[A.OperatorIndex].ParentIndex				= OperatorIndex;
 
 	return StrataData;
 }
@@ -1147,11 +1141,11 @@ FStrataData StrataHorizontalMixing(FStrataData Background, FStrataData Foregroun
 	StrataTree.Operators[OperatorIndex].MaxDistanceFromLeaves	= MaxDistanceFromLeaves;
 	StrataTree.Operators[OperatorIndex].Type					= STRATA_OPERATOR_HORIZONTAL;
 	StrataTree.Operators[OperatorIndex].Weight					= Mix;
-	StrataTree.Operators[OperatorIndex].AIndex					= Background.OperatorIndex;
-	StrataTree.Operators[OperatorIndex].BIndex					= Foreground.OperatorIndex;
+	StrataTree.Operators[OperatorIndex].LeftIndex				= Background.OperatorIndex;
+	StrataTree.Operators[OperatorIndex].RightIndex				= Foreground.OperatorIndex;
 
-	StrataTree.Operators[Background.OperatorIndex].TopIndex		= OperatorIndex;
-	StrataTree.Operators[Foreground.OperatorIndex].TopIndex		= OperatorIndex;
+	StrataTree.Operators[Background.OperatorIndex].ParentIndex	= OperatorIndex;
+	StrataTree.Operators[Foreground.OperatorIndex].ParentIndex	= OperatorIndex;
 
 	return StrataData;
 }
@@ -1166,11 +1160,11 @@ FStrataData StrataVerticalLayering(FStrataData Top, FStrataData Base, inout FStr
 	StrataTree.Operators[OperatorIndex]							= GetInitialisedStrataOperator();
 	StrataTree.Operators[OperatorIndex].MaxDistanceFromLeaves	= MaxDistanceFromLeaves;
 	StrataTree.Operators[OperatorIndex].Type					= STRATA_OPERATOR_VERTICAL;
-	StrataTree.Operators[OperatorIndex].AIndex					= Top.OperatorIndex;
-	StrataTree.Operators[OperatorIndex].BIndex					= Base.OperatorIndex;
+	StrataTree.Operators[OperatorIndex].LeftIndex				= Top.OperatorIndex;
+	StrataTree.Operators[OperatorIndex].RightIndex				= Base.OperatorIndex;
 
-	StrataTree.Operators[Top.OperatorIndex].TopIndex			= OperatorIndex;
-	StrataTree.Operators[Base.OperatorIndex].TopIndex			= OperatorIndex;
+	StrataTree.Operators[Top.OperatorIndex].ParentIndex			= OperatorIndex;
+	StrataTree.Operators[Base.OperatorIndex].ParentIndex		= OperatorIndex;
 
 	return StrataData;
 }
@@ -1186,7 +1180,7 @@ FStrataData PromoteParameterBlendedBSDFToOperator(FStrataData StrataData, inout
 	StrataTree.Operators[OperatorIndex]							= GetInitialisedStrataOperator();
 	StrataTree.Operators[OperatorIndex].MaxDistanceFromLeaves	= 0;
 	StrataTree.Operators[OperatorIndex].Type					= STRATA_OPERATOR_BSDF;
-	StrataTree.Operators[OperatorIndex].AIndex					= BSDFIndex;
+	StrataTree.Operators[OperatorIndex].LeftIndex				= BSDFIndex;
 
 	StrataTree.BSDFs[BSDFIndex].OperatorIndex					= OperatorIndex;
 	StrataTree.BSDFs[BSDFIndex].bIsBottom						= bIsBottom;
@@ -1449,7 +1443,7 @@ FStrataBSDF StrataVerticalLayeringParameterBlending(in FStrataBSDF TopBSDF, in F
 	const float  TopFactor				= Info.SurfaceTop / max(1e-5, Info.Coverage);
 	
 	// Add both layers thickness
-//	BSDF_SETTHICKNESSCM(ResultBSDF,		BSDF_GETTHICKNESSCM(BaseBSDF)												+ SafeTopCoverage * BSDF_GETTHICKNESSCM(TopBSDF));
+	BSDF_SETTHICKNESSCM(ResultBSDF,		BSDF_GETTHICKNESSCM(BaseBSDF)												+ SafeTopCoverage * BSDF_GETTHICKNESSCM(TopBSDF));
 
 	// Now add bottom layer component weights by throughput to the top layer.
 	// Those parameters will contribute to added luminance, in this case we add them together and saturate to not go out of the safe range
@@ -1498,9 +1492,6 @@ FStrataBSDF StrataVerticalLayeringParameterBlending(in FStrataBSDF TopBSDF, in F
 	const float3 NewMFP = Info2.TransmittanceOnlyBottom * BaseBSDF.TmpDMFP + Info2.TransmittanceOnlyTop * TopBSDF.TmpDMFP + Info2.TransmittanceTopAndBottom * min(TopBSDF.TmpDMFP, BaseBSDF.TmpDMFP);
 #endif
 
-	BSDF_SETTHICKNESSCM(ResultBSDF, CENTIMETER_TO_METER);
-
-	//ResultBSDF.TmpDMFP =				lerp(TopBSDF.TmpDMFP,							BaseBSDF.TmpDMFP,			TopToBottomLerp);
 	ResultBSDF.TmpDMFP					= NewMFP;
 	SLAB_SSSDMFP(ResultBSDF)			= NewMFP;
 
@@ -2309,8 +2300,6 @@ void PackStrataOut(
 	bStrataSubsurfaceEnable = false;
 	EmissiveColor = 0.0f;
 
-	const float OpaqueBSDFThroughput = 0.0f;
-	const float FullThroughput = 1.0f;
 	const float FullyRough = 1.0f;
 
 	// While packing Strata layer data, Classification/SSS/TopLayer data are extracted & stored for dedicated passes
@@ -2399,9 +2388,6 @@ void PackStrataOut(
 
 						bStrataSubsurfaceEnable = bStrataSubsurfaceEnable || (BSDF_GETHASSSS(BSDF) && BSDF.Coverage > 0.0f && !bIsSimpleVolume);
 
-						// By default, without coverage, the diffuse BSDF is completely opaque
-						float3 BSDFThroughput = OpaqueBSDFThroughput;
-
 						if (bIsSimpleVolume)
 						{
 							EnableSlabBSDFSimpleVolumetric(BSDF);
diff --git a/Engine/Shaders/Private/Strata/StrataTree.ush b/Engine/Shaders/Private/Strata/StrataTree.ush
--- a/Engine/Shaders/Private/Strata/StrataTree.ush
+++ b/Engine/Shaders/Private/Strata/StrataTree.ush
@@ -55,8 +55,8 @@ void UpdateSingleBSDFOperatorCoverageTransmittance(
 void UpdateSingleOperatorCoverageTransmittance(inout FStrataTree StrataTree, int OpIndex)
 {
 #define Op StrataTree.Operators[OpIndex]
-#define OpA StrataTree.Operators[StrataTree.Operators[OpIndex].AIndex]
-#define OpB StrataTree.Operators[StrataTree.Operators[OpIndex].BIndex]
+#define OpA StrataTree.Operators[StrataTree.Operators[OpIndex].LeftIndex]
+#define OpB StrataTree.Operators[StrataTree.Operators[OpIndex].RightIndex]
 
 	switch (Op.Type)
 	{
@@ -114,8 +114,12 @@ void UpdateSingleOperatorCoverageTransmittance(inout FStrataTree StrataTree, int
 
 	case STRATA_OPERATOR_ADD:
 	{
+		const float SafeABSDFCoverage = saturate(OpA.Coverage);
+		const float SafeBBSDFCoverage = saturate(OpB.Coverage);
+		const float AMixFactor = SafeABSDFCoverage / (1e-10 + SafeABSDFCoverage + SafeBBSDFCoverage);
+
 		Op.Coverage = saturate(OpA.Coverage + OpB.Coverage);
-		Op.Transmittance = saturate(OpA.Transmittance + OpB.Transmittance);
+		Op.Transmittance = lerp(OpB.Transmittance, OpA.Transmittance, AMixFactor);
 		break;
 	}
 
diff --git a/Engine/Source/Runtime/Engine/Private/Materials/HLSLMaterialTranslator.cpp b/Engine/Source/Runtime/Engine/Private/Materials/HLSLMaterialTranslator.cpp
--- a/Engine/Source/Runtime/Engine/Private/Materials/HLSLMaterialTranslator.cpp
+++ b/Engine/Source/Runtime/Engine/Private/Materials/HLSLMaterialTranslator.cpp
@@ -260,7 +260,7 @@ FHLSLMaterialTranslator::FHLSLMaterialTranslator(FMaterial* InMaterial,
 	}
 
 	StrataMaterialRootOperator = nullptr;
-	StrataMaterialExpressionRegisteredOperators.Reset(STRATA_MAX_OPERATOR_COUNT);
+	StrataMaterialExpressionRegisteredOperators.Reserve(STRATA_MAX_OPERATOR_COUNT);
 	StrataMaterialExpressionToOperatorIndex.Reserve(STRATA_MAX_OPERATOR_COUNT);
 	StrataMaterialBSDFCount = 0;
 	bStrataUsesConversionFromLegacy = false;
@@ -665,11 +665,6 @@ bool FHLSLMaterialTranslator::Translate()
 		UMaterialExpression* FrontMaterialExpr = FrontMaterialInput ? FrontMaterialInput->GetTracedInput().Expression : nullptr;
 		if (bStrataEnabled && FrontMaterialExpr)
 		{
-			if (Material->GetAssetName() == FString("ColorCoat_Strata"))
-			{
-				int i = 0;
-				i = 2;
-			}
 			FrontMaterialExpr->StrataGenerateMaterialTopologyTree(this, nullptr, 0);
 			StrataGenerateDerivedMaterialOperatorData();
 		}
@@ -1414,12 +1409,12 @@ bool FHLSLMaterialTranslator::Translate()
 									{
 									case STRATA_OPERATOR_VERTICAL:
 									{
-										ResourcesString += FString::Printf(TEXT("\t UpdateBSDFWeightAfterOperatorVisit(StrataTree, StrataTree.BSDFs[%d], %d /*Op index*/, %d /*PreviousIsInputA*/); // VERTICAL\n"), BSDFIndex, CurrentOperator.Index, CurrentOperator.AIndex == PreviousOperatorIndex ? 1 : 0);
+										ResourcesString += FString::Printf(TEXT("\t UpdateBSDFWeightAfterOperatorVisit(StrataTree, StrataTree.BSDFs[%d], %d /*Op index*/, %d /*PreviousIsInputA*/); // VERTICAL\n"), BSDFIndex, CurrentOperator.Index, CurrentOperator.LeftIndex == PreviousOperatorIndex ? 1 : 0);
 										break;
 									}
 									case STRATA_OPERATOR_HORIZONTAL:
 									{
-										ResourcesString += FString::Printf(TEXT("\t UpdateBSDFWeightAfterOperatorVisit(StrataTree, StrataTree.BSDFs[%d], %d /*Op index*/, %d /*PreviousIsInputA*/); // HORIZONTAL\n"), BSDFIndex, CurrentOperator.Index, CurrentOperator.AIndex == PreviousOperatorIndex ? 1 : 0);
+										ResourcesString += FString::Printf(TEXT("\t UpdateBSDFWeightAfterOperatorVisit(StrataTree, StrataTree.BSDFs[%d], %d /*Op index*/, %d /*PreviousIsInputA*/); // HORIZONTAL\n"), BSDFIndex, CurrentOperator.Index, CurrentOperator.LeftIndex == PreviousOperatorIndex ? 1 : 0);
 										break;
 									}
 
@@ -1441,9 +1436,9 @@ bool FHLSLMaterialTranslator::Translate()
 									}
 									}
 
-									if (CurrentOperator.TopIndex != INDEX_NONE)
+									if (CurrentOperator.ParentIndex != INDEX_NONE)
 									{
-										WalkOperatorsUp(StrataMaterialExpressionRegisteredOperators[CurrentOperator.TopIndex], CurrentOperator.Index);
+										WalkOperatorsUp(StrataMaterialExpressionRegisteredOperators[CurrentOperator.ParentIndex], CurrentOperator.Index);
 									}
 								};
 
@@ -1451,9 +1446,9 @@ bool FHLSLMaterialTranslator::Translate()
 								FStrataOperator& BSDFOperator = StrataMaterialExpressionRegisteredOperators[BSDFOperatorIndex];
 
 								// Start visiting node up from the BSDF leaf only if it has a parent.
-								if (BSDFOperator.TopIndex != INDEX_NONE)
+								if (BSDFOperator.ParentIndex != INDEX_NONE)
 								{
-									WalkOperatorsUp(StrataMaterialExpressionRegisteredOperators[BSDFOperator.TopIndex], BSDFOperator.Index);
+									WalkOperatorsUp(StrataMaterialExpressionRegisteredOperators[BSDFOperator.ParentIndex], BSDFOperator.Index);
 								}
 							}
 						}
@@ -9578,9 +9573,9 @@ FStrataOperator& FHLSLMaterialTranslator::StrataCompilationRegisterOperator(int3
 	NewOperator.OperatorType = OperatorType;
 	NewOperator.bNodeRequestParameterBlending = bUseParameterBlending;
 	NewOperator.Index = NewOperatorIndex;
-	NewOperator.TopIndex = ParentOperatorIndex != nullptr ? *ParentOperatorIndex : INDEX_NONE;
-	NewOperator.AIndex   = INDEX_NONE;
-	NewOperator.BIndex   = INDEX_NONE;
+	NewOperator.ParentIndex = ParentOperatorIndex != nullptr ? *ParentOperatorIndex : INDEX_NONE;
+	NewOperator.LeftIndex   = INDEX_NONE;
+	NewOperator.RightIndex   = INDEX_NONE;
 
 	NewOperator.BSDFIndex = INDEX_NONE;	// Allocated later to be ableto account for inline
 
@@ -9606,7 +9601,7 @@ void FHLSLMaterialTranslator::StrataGenerateDerivedMaterialOperatorData()
 	int32 RootIndex = INDEX_NONE;
 	for (auto& It : StrataMaterialExpressionRegisteredOperators)
 	{
-		if (It.TopIndex == INDEX_NONE)
+		if (It.ParentIndex == INDEX_NONE)
 		{
 			check(RootIndex == INDEX_NONE);	// There can only be one
 			RootIndex = It.Index;
@@ -9636,20 +9631,20 @@ void FHLSLMaterialTranslator::StrataGenerateDerivedMaterialOperatorData()
 			{
 			case STRATA_OPERATOR_VERTICAL:
 			{
-				WalkOperators(StrataMaterialExpressionRegisteredOperators[CurrentOperator.AIndex], bUseParameterBlending);
-				WalkOperators(StrataMaterialExpressionRegisteredOperators[CurrentOperator.BIndex], bUseParameterBlending);
+				WalkOperators(StrataMaterialExpressionRegisteredOperators[CurrentOperator.LeftIndex], bUseParameterBlending);
+				WalkOperators(StrataMaterialExpressionRegisteredOperators[CurrentOperator.RightIndex], bUseParameterBlending);
 				break;
 			}
 			case STRATA_OPERATOR_HORIZONTAL:
 			case STRATA_OPERATOR_ADD:
 			{
-				WalkOperators(StrataMaterialExpressionRegisteredOperators[CurrentOperator.AIndex], bUseParameterBlending);
-				WalkOperators(StrataMaterialExpressionRegisteredOperators[CurrentOperator.BIndex], bUseParameterBlending);
+				WalkOperators(StrataMaterialExpressionRegisteredOperators[CurrentOperator.LeftIndex], bUseParameterBlending);
+				WalkOperators(StrataMaterialExpressionRegisteredOperators[CurrentOperator.RightIndex], bUseParameterBlending);
 				break;
 			}
 			case STRATA_OPERATOR_WEIGHT:
 			{
-				WalkOperators(StrataMaterialExpressionRegisteredOperators[CurrentOperator.AIndex], bUseParameterBlending);
+				WalkOperators(StrataMaterialExpressionRegisteredOperators[CurrentOperator.LeftIndex], bUseParameterBlending);
 				break;
 			}
 			case STRATA_OPERATOR_BSDF:
@@ -9668,8 +9663,8 @@ void FHLSLMaterialTranslator::StrataGenerateDerivedMaterialOperatorData()
 			{
 				CurrentOperator.OperatorType = STRATA_OPERATOR_BSDF;
 				CurrentOperator.BSDFIndex = StrataMaterialBSDFCount++;
-				CurrentOperator.AIndex = INDEX_NONE;
-				CurrentOperator.BIndex = INDEX_NONE;
+				CurrentOperator.LeftIndex = INDEX_NONE;
+				CurrentOperator.RightIndex = INDEX_NONE;
 			}
 		};
 
@@ -9697,7 +9692,7 @@ void FHLSLMaterialTranslator::StrataGenerateDerivedMaterialOperatorData()
 		// Operators without any child
 		case STRATA_OPERATOR_BSDF:
 		{
-			check(It.AIndex == INDEX_NONE && It.BIndex == INDEX_NONE && It.BSDFIndex != INDEX_NONE);
+			check(It.LeftIndex == INDEX_NONE && It.RightIndex == INDEX_NONE && It.BSDFIndex != INDEX_NONE);
 			break;
 		}
 
@@ -9706,14 +9701,14 @@ void FHLSLMaterialTranslator::StrataGenerateDerivedMaterialOperatorData()
 		case STRATA_OPERATOR_VERTICAL:
 		case STRATA_OPERATOR_ADD:
 		{
-			check(It.BIndex != INDEX_NONE);
+			check(It.RightIndex != INDEX_NONE);
 		}
 		// Fallthrough
 
 		// Operators with a single child
 		case STRATA_OPERATOR_WEIGHT:
 		{
-			check(It.AIndex != INDEX_NONE);
+			check(It.LeftIndex != INDEX_NONE);
 		}
 		// Fallthrough
 		}
@@ -9729,7 +9724,7 @@ void FHLSLMaterialTranslator::StrataGenerateDerivedMaterialOperatorData()
 			{
 			case STRATA_OPERATOR_WEIGHT:
 			{
-				CurrentOperator.MaxDistanceFromLeaves = StrataMaterialExpressionRegisteredOperators[CurrentOperator.AIndex].MaxDistanceFromLeaves + 1;
+				CurrentOperator.MaxDistanceFromLeaves = StrataMaterialExpressionRegisteredOperators[CurrentOperator.LeftIndex].MaxDistanceFromLeaves + 1;
 				break;
 			}
 			case STRATA_OPERATOR_VERTICAL:
@@ -9737,8 +9732,8 @@ void FHLSLMaterialTranslator::StrataGenerateDerivedMaterialOperatorData()
 			case STRATA_OPERATOR_ADD:
 			{
 				CurrentOperator.MaxDistanceFromLeaves = FMath::Max(
-					StrataMaterialExpressionRegisteredOperators[CurrentOperator.AIndex].MaxDistanceFromLeaves,
-					StrataMaterialExpressionRegisteredOperators[CurrentOperator.BIndex].MaxDistanceFromLeaves) + 1;
+					StrataMaterialExpressionRegisteredOperators[CurrentOperator.LeftIndex].MaxDistanceFromLeaves,
+					StrataMaterialExpressionRegisteredOperators[CurrentOperator.RightIndex].MaxDistanceFromLeaves) + 1;
 				break;
 			}
 			case STRATA_OPERATOR_BSDF:
@@ -9748,9 +9743,9 @@ void FHLSLMaterialTranslator::StrataGenerateDerivedMaterialOperatorData()
 			}
 			}
 
-			if (CurrentOperator.TopIndex != INDEX_NONE)
+			if (CurrentOperator.ParentIndex != INDEX_NONE)
 			{
-				WalkOperatorsToRoot(StrataMaterialExpressionRegisteredOperators[CurrentOperator.TopIndex]);
+				WalkOperatorsToRoot(StrataMaterialExpressionRegisteredOperators[CurrentOperator.ParentIndex]);
 			}
 		};
 
@@ -9783,23 +9778,23 @@ void FHLSLMaterialTranslator::StrataGenerateDerivedMaterialOperatorData()
 			case STRATA_OPERATOR_VERTICAL:
 			{
 				VOpTopBranchCountTaken++;
-				WalkOperators(StrataMaterialExpressionRegisteredOperators[CurrentOperator.AIndex]);
+				WalkOperators(StrataMaterialExpressionRegisteredOperators[CurrentOperator.LeftIndex]);
 				VOpTopBranchCountTaken--;
 				VOpBottomBranchCountTaken++;
-				WalkOperators(StrataMaterialExpressionRegisteredOperators[CurrentOperator.BIndex]);
+				WalkOperators(StrataMaterialExpressionRegisteredOperators[CurrentOperator.RightIndex]);
 				VOpBottomBranchCountTaken--;
 				break;
 			}
 			case STRATA_OPERATOR_HORIZONTAL:
 			case STRATA_OPERATOR_ADD:
 			{
-				WalkOperators(StrataMaterialExpressionRegisteredOperators[CurrentOperator.AIndex]);
-				WalkOperators(StrataMaterialExpressionRegisteredOperators[CurrentOperator.BIndex]);
+				WalkOperators(StrataMaterialExpressionRegisteredOperators[CurrentOperator.LeftIndex]);
+				WalkOperators(StrataMaterialExpressionRegisteredOperators[CurrentOperator.RightIndex]);
 				break;
 			}
 			case STRATA_OPERATOR_WEIGHT:
 			{
-				WalkOperators(StrataMaterialExpressionRegisteredOperators[CurrentOperator.AIndex]);
+				WalkOperators(StrataMaterialExpressionRegisteredOperators[CurrentOperator.LeftIndex]);
 				break;
 			}
 			case STRATA_OPERATOR_BSDF:
diff --git a/Engine/Source/Runtime/Engine/Private/Materials/MaterialExpressions.cpp b/Engine/Source/Runtime/Engine/Private/Materials/MaterialExpressions.cpp
--- a/Engine/Source/Runtime/Engine/Private/Materials/MaterialExpressions.cpp
+++ b/Engine/Source/Runtime/Engine/Private/Materials/MaterialExpressions.cpp
@@ -6892,11 +6892,11 @@ FStrataOperator* UMaterialExpressionBlendMaterialAttributes::StrataGenerateMater
 	UMaterialExpression* ChildBExpression = B.GetTracedInput().Expression;
 	if (ChildAExpression)
 	{
-		StrataOperator.AIndex = ChildAExpression->StrataGenerateMaterialTopologyTree(Compiler, this, A.OutputIndex)->Index;
+		StrataOperator.LeftIndex = ChildAExpression->StrataGenerateMaterialTopologyTree(Compiler, this, A.OutputIndex)->Index;
 	}
 	if (ChildBExpression)
 	{
-		StrataOperator.BIndex = ChildBExpression->StrataGenerateMaterialTopologyTree(Compiler, this, B.OutputIndex)->Index;
+		StrataOperator.RightIndex = ChildBExpression->StrataGenerateMaterialTopologyTree(Compiler, this, B.OutputIndex)->Index;
 	}
 
 	return &StrataOperator;
@@ -20698,8 +20698,8 @@ FStrataOperator* UMaterialExpressionStrataLegacyConversion::StrataGenerateMateri
 	// with fake addresses for BSDF defined in shader code.
 	UMaterialExpression* ChildAExpression = reinterpret_cast<UMaterialExpression*>(0);
 	UMaterialExpression* ChildBExpression = reinterpret_cast<UMaterialExpression*>(1);
-	StrataOperator.AIndex = Compiler->StrataCompilationRegisterOperator(STRATA_OPERATOR_BSDF_LEGACY, ChildAExpression, this).Index;
-	StrataOperator.BIndex = Compiler->StrataCompilationRegisterOperator(STRATA_OPERATOR_BSDF_LEGACY, ChildBExpression, this).Index;
+	StrataOperator.LeftIndex = Compiler->StrataCompilationRegisterOperator(STRATA_OPERATOR_BSDF_LEGACY, ChildAExpression, this).Index;
+	StrataOperator.RightIndex = Compiler->StrataCompilationRegisterOperator(STRATA_OPERATOR_BSDF_LEGACY, ChildBExpression, this).Index;
 
 	return &StrataOperator;
 }
@@ -21799,11 +21799,11 @@ FStrataOperator* UMaterialExpressionStrataHorizontalMixing::StrataGenerateMateri
 	UMaterialExpression* ChildBExpression = Foreground.GetTracedInput().Expression;
 	if (ChildAExpression)
 	{
-		StrataOperator.AIndex = ChildAExpression->StrataGenerateMaterialTopologyTree(Compiler, this, Background.OutputIndex)->Index;
+		StrataOperator.LeftIndex = ChildAExpression->StrataGenerateMaterialTopologyTree(Compiler, this, Background.OutputIndex)->Index;
 	}
 	if (ChildBExpression)
 	{
-		StrataOperator.BIndex = ChildBExpression->StrataGenerateMaterialTopologyTree(Compiler, this, Foreground.OutputIndex)->Index;
+		StrataOperator.RightIndex = ChildBExpression->StrataGenerateMaterialTopologyTree(Compiler, this, Foreground.OutputIndex)->Index;
 	}
 
 	return &StrataOperator;
@@ -21961,11 +21961,11 @@ FStrataOperator* UMaterialExpressionStrataVerticalLayering::StrataGenerateMateri
 	UMaterialExpression* ChildBExpression = Base.GetTracedInput().Expression;
 	if (ChildAExpression)
 	{
-		StrataOperator.AIndex = ChildAExpression->StrataGenerateMaterialTopologyTree(Compiler, this, Top.OutputIndex)->Index;
+		StrataOperator.LeftIndex = ChildAExpression->StrataGenerateMaterialTopologyTree(Compiler, this, Top.OutputIndex)->Index;
 	}
 	if (ChildBExpression)
 	{
-		StrataOperator.BIndex = ChildBExpression->StrataGenerateMaterialTopologyTree(Compiler, this, Base.OutputIndex)->Index;
+		StrataOperator.RightIndex = ChildBExpression->StrataGenerateMaterialTopologyTree(Compiler, this, Base.OutputIndex)->Index;
 	}
 
 	return &StrataOperator;
@@ -22122,11 +22122,11 @@ FStrataOperator* UMaterialExpressionStrataAdd::StrataGenerateMaterialTopologyTre
 	UMaterialExpression* ChildBExpression = B.GetTracedInput().Expression;
 	if (ChildAExpression)
 	{
-		StrataOperator.AIndex = ChildAExpression->StrataGenerateMaterialTopologyTree(Compiler, this, A.OutputIndex)->Index;
+		StrataOperator.LeftIndex = ChildAExpression->StrataGenerateMaterialTopologyTree(Compiler, this, A.OutputIndex)->Index;
 	}
 	if (ChildBExpression)
 	{
-		StrataOperator.BIndex = ChildBExpression->StrataGenerateMaterialTopologyTree(Compiler, this, B.OutputIndex)->Index;
+		StrataOperator.RightIndex = ChildBExpression->StrataGenerateMaterialTopologyTree(Compiler, this, B.OutputIndex)->Index;
 	}
 
 	return &StrataOperator;
@@ -22226,7 +22226,7 @@ FStrataOperator* UMaterialExpressionStrataWeight::StrataGenerateMaterialTopology
 	UMaterialExpression* ChildAExpression = A.GetTracedInput().Expression;
 	if (ChildAExpression)
 	{
-		StrataOperator.AIndex = ChildAExpression->StrataGenerateMaterialTopologyTree(Compiler, this, A.OutputIndex)->Index;
+		StrataOperator.LeftIndex = ChildAExpression->StrataGenerateMaterialTopologyTree(Compiler, this, A.OutputIndex)->Index;
 	}
 
 	return &StrataOperator;
diff --git a/Engine/Source/Runtime/Engine/Public/MaterialCompiler.h b/Engine/Source/Runtime/Engine/Public/MaterialCompiler.h
--- a/Engine/Source/Runtime/Engine/Public/MaterialCompiler.h
+++ b/Engine/Source/Runtime/Engine/Public/MaterialCompiler.h
@@ -74,9 +74,9 @@ struct FStrataOperator
 	bool bNodeRequestParameterBlending;
 
 	int32 Index;		// Index into the array of operators
-	int32 TopIndex;		// Parent operator index
-	int32 AIndex;		// Left child operator index
-	int32 BIndex;		// Right child operator index
+	int32 ParentIndex;	// Parent operator index
+	int32 LeftIndex;	// Left child operator index
+	int32 RightIndex;	// Right child operator index
 
 	int32 BSDFIndex;	// Index in the array of BSDF if a BSDF operator
 
@@ -93,9 +93,9 @@ struct FStrataOperator
 		OperatorType = INDEX_NONE;
 		bNodeRequestParameterBlending = false;
 		Index = INDEX_NONE;
-		TopIndex = INDEX_NONE;
-		AIndex = INDEX_NONE;
-		BIndex = INDEX_NONE;
+		ParentIndex = INDEX_NONE;
+		LeftIndex = INDEX_NONE;
+		RightIndex = INDEX_NONE;
 		BSDFIndex = INDEX_NONE;
 		MaxDistanceFromLeaves = 0;
 		LayerDepth = 0;
