PostProcess: Fix to allow HighlightsMin value greater than 1.0

diff --git a/Engine/Shaders/Private/PostProcessCombineLUTs.usf b/Engine/Shaders/Private/PostProcessCombineLUTs.usf
--- a/Engine/Shaders/Private/PostProcessCombineLUTs.usf
+++ b/Engine/Shaders/Private/PostProcessCombineLUTs.usf
@@ -92,6 +92,7 @@ float4 ColorOffsetHighlights;
 
 float  ColorCorrectionShadowsMax;
 float  ColorCorrectionHighlightsMin;
+float  ColorCorrectionHighlightsMax;
 
 float3 ColorCorrect( float3 WorkingColor,
 	float4 ColorSaturation,
@@ -130,7 +131,7 @@ float3 ColorCorrectAll( float3 WorkingColor )
 		ColorGammaHighlights*ColorGamma, 
 		ColorGainHighlights*ColorGain, 
 		ColorOffsetHighlights+ColorOffset);
-	float CCWeightHighlights = smoothstep(ColorCorrectionHighlightsMin, 1, Luma);
+	float CCWeightHighlights = smoothstep(ColorCorrectionHighlightsMin, ColorCorrectionHighlightsMax, Luma);
 
 	// Midtone CC
 	float3 CCColorMidtones = ColorCorrect(WorkingColor, 
diff --git a/Engine/Shaders/Private/PostProcessTonemap.usf b/Engine/Shaders/Private/PostProcessTonemap.usf
--- a/Engine/Shaders/Private/PostProcessTonemap.usf
+++ b/Engine/Shaders/Private/PostProcessTonemap.usf
@@ -83,6 +83,7 @@ float FilmGrainIntensityMidtones;
 float FilmGrainIntensityHighlights;
 float FilmGrainShadowsMax;
 float FilmGrainHighlightsMin;
+float FilmGrainHighlightsMax;
 Texture2D<half3> FilmGrainTexture;
 SamplerState FilmGrainSampler;
 float4 ScreenPosToFilmGrainTextureUV;
@@ -150,7 +151,7 @@ float ComputeFilmGrainIntensity(float3 LinearColorRGBGamut)
 	float Luma = dot(LinearColorRGBGamut,  mul(AP1_2_sRGB, AP1_RGB2Y));
 
 	float CCWeightShadows = 1 - smoothstep(0, FilmGrainShadowsMax, Luma);
-	float CCWeightHighlights = smoothstep(FilmGrainHighlightsMin, 1, Luma);
+	float CCWeightHighlights = smoothstep(FilmGrainHighlightsMin, FilmGrainHighlightsMax, Luma);
 	float CCWeightMidtones = 1 - CCWeightShadows - CCWeightHighlights;
 	
 	// Blend Shadow, Midtone and Highlight CCs
diff --git a/Engine/Source/Runtime/Engine/Classes/Engine/Scene.h b/Engine/Source/Runtime/Engine/Classes/Engine/Scene.h
--- a/Engine/Source/Runtime/Engine/Classes/Engine/Scene.h
+++ b/Engine/Source/Runtime/Engine/Classes/Engine/Scene.h
@@ -216,11 +216,14 @@ struct FColorGradingSettings
 	UPROPERTY(Interp, BlueprintReadWrite, Category = "Color Grading", meta = (UIMin = "-1.0", UIMax = "1.0", DisplayName = "HighlightsMin"))
 	float HighlightsMin;
 
+	UPROPERTY(Interp, BlueprintReadWrite, Category = "Color Grading", meta = (UIMin = "1.0", UIMax = "10.0", DisplayName = "HighlightsMax"))
+	float HighlightsMax;
 
 	FColorGradingSettings()
 	{
 		ShadowsMax = 0.09f;
 		HighlightsMin = 0.5f;
+		HighlightsMax = 1.0f;
 	}
 
 	/* Exports to post process settings with overrides. */
@@ -745,6 +748,8 @@ struct FPostProcessSettings
 	uint8 bOverride_ColorCorrectionShadowsMax : 1;
 	UPROPERTY(EditAnywhere, BlueprintReadWrite, Category = Overrides, meta = (PinHiddenByDefault, InlineEditConditionToggle))
 	uint8 bOverride_ColorCorrectionHighlightsMin : 1;
+	UPROPERTY(EditAnywhere, BlueprintReadWrite, Category = Overrides, meta = (PinHiddenByDefault, InlineEditConditionToggle))
+	uint8 bOverride_ColorCorrectionHighlightsMax : 1;
 
 	UPROPERTY(EditAnywhere, BlueprintReadWrite, Category = Overrides, meta = (PinHiddenByDefault, InlineEditConditionToggle))
 	uint8 bOverride_BlueCorrection : 1;
@@ -970,6 +975,9 @@ struct FPostProcessSettings
 	UPROPERTY(EditAnywhere, BlueprintReadWrite, Category = Overrides, meta = (PinHiddenByDefault, InlineEditConditionToggle))
 	uint8 bOverride_FilmGrainHighlightsMin : 1;
 
+	UPROPERTY(EditAnywhere, BlueprintReadWrite, Category = Overrides, meta = (PinHiddenByDefault, InlineEditConditionToggle))
+	uint8 bOverride_FilmGrainHighlightsMax : 1;
+
 	UPROPERTY(EditAnywhere, BlueprintReadWrite, Category=Overrides, meta=(PinHiddenByDefault, InlineEditConditionToggle))
 	uint8 bOverride_FilmGrainTexelSize:1;
 
@@ -1345,6 +1353,9 @@ struct FPostProcessSettings
 	FVector4 ColorOffsetHighlights;
 	UPROPERTY(interp, BlueprintReadWrite, Category = "Color Grading|Highlights", meta = (UIMin = "-1.0", UIMax = "1.0", editcondition = "bOverride_ColorCorrectionHighlightsMin", DisplayName = "HighlightsMin"))
 	float ColorCorrectionHighlightsMin;
+	/** Should be set larger than HighlightsMin. Default is 1.0, for backwards compatibility */
+	UPROPERTY(interp, BlueprintReadWrite, Category = "Color Grading|Highlights", meta = (UIMin = "1.0", UIMax = "10.0", editcondition = "bOverride_ColorCorrectionHighlightsMax", DisplayName = "HighlightsMax"))
+	float ColorCorrectionHighlightsMax;
 
 	UPROPERTY(interp, BlueprintReadWrite, Category = "Color Grading|Shadows", meta = (UIMin = "-1.0", UIMax = "1.0", editcondition = "bOverride_ColorCorrectionShadowsMax", DisplayName = "ShadowsMax"))
 	float ColorCorrectionShadowsMax;
@@ -1839,6 +1850,10 @@ struct FPostProcessSettings
 	UPROPERTY(interp, BlueprintReadWrite, Category = "Film Grain", meta = (UIMin = "0.0", UIMax = "1.0", editcondition = "bOverride_FilmGrainHighlightsMin"))
 	float FilmGrainHighlightsMin;
 
+	/** Should be set larger than HighlightsMin. Default is 1.0, for backwards compatibility */
+	UPROPERTY(interp, BlueprintReadWrite, Category = "Film Grain", meta = (UIMin = "1.0", UIMax = "10.0", editcondition = "bOverride_FilmGrainHighlightsMax"))
+	float FilmGrainHighlightsMax;
+
 	/** Size of texel of FilmGrainTexture on screen. */
 	UPROPERTY(interp, BlueprintReadWrite, Category = "Film Grain", meta = (UIMin = "0.0", UIMax = "4.0", editcondition = "bOverride_FilmGrainTexelSize"))
 	float FilmGrainTexelSize;
diff --git a/Engine/Source/Runtime/Engine/Private/Scene.cpp b/Engine/Source/Runtime/Engine/Private/Scene.cpp
--- a/Engine/Source/Runtime/Engine/Private/Scene.cpp
+++ b/Engine/Source/Runtime/Engine/Private/Scene.cpp
@@ -34,6 +34,7 @@ void FColorGradingSettings::ExportToPostProcessSettings(FPostProcessSettings* Ou
 
 	OutPostProcessSettings->bOverride_ColorCorrectionShadowsMax = true;
 	OutPostProcessSettings->bOverride_ColorCorrectionHighlightsMin = true;
+	OutPostProcessSettings->bOverride_ColorCorrectionHighlightsMax = true;
 
 	OutPostProcessSettings->ColorSaturation = Global.Saturation;
 	OutPostProcessSettings->ColorContrast = Global.Contrast;
@@ -61,6 +62,7 @@ void FColorGradingSettings::ExportToPostProcessSettings(FPostProcessSettings* Ou
 
 	OutPostProcessSettings->ColorCorrectionShadowsMax = ShadowsMax;
 	OutPostProcessSettings->ColorCorrectionHighlightsMin = HighlightsMin;
+	OutPostProcessSettings->ColorCorrectionHighlightsMax = HighlightsMax;
 }
 
 void FFilmStockSettings::ExportToPostProcessSettings(FPostProcessSettings* OutPostProcessSettings) const
@@ -346,6 +348,7 @@ static void DoPostProcessSettingsSanityCheck()
 		RenameMap.Add(TEXT("FilmGrainIntensityHighlights"), TEXT(""));
 		RenameMap.Add(TEXT("FilmGrainShadowsMax"), TEXT(""));
 		RenameMap.Add(TEXT("FilmGrainHighlightsMin"), TEXT(""));
+		RenameMap.Add(TEXT("FilmGrainHighlightsMax"), TEXT(""));
 		RenameMap.Add(TEXT("FilmGrainTexelSize"), TEXT(""));
 		RenameMap.Add(TEXT("FilmGrainTexture"), TEXT(""));
 		VerifyPostProcessingProperties(TEXT("Film"),
@@ -392,6 +395,7 @@ FPostProcessSettings::FPostProcessSettings()
 
 	ColorCorrectionShadowsMax = 0.09f;
 	ColorCorrectionHighlightsMin = 0.5f;
+	ColorCorrectionHighlightsMax = 1.0f;
 
 	BlueCorrection = 0.6f;
 	ExpandGamut = 1.0f;
@@ -498,6 +502,7 @@ FPostProcessSettings::FPostProcessSettings()
 	FilmGrainIntensityHighlights = 1.0f;
 	FilmGrainShadowsMax = 0.09f;
 	FilmGrainHighlightsMin = 0.5f;
+	FilmGrainHighlightsMax = 1.0f;
 	FilmGrainTexelSize = 1.0f;
 
 	// next value might get overwritten by r.DefaultFeature.AmbientOcclusion
@@ -633,6 +638,7 @@ FPostProcessSettings::FPostProcessSettings(const FPostProcessSettings& Settings)
 	, bOverride_ColorOffsetHighlights(Settings.bOverride_ColorOffsetHighlights)
 	, bOverride_ColorCorrectionShadowsMax(Settings.bOverride_ColorCorrectionShadowsMax)
 	, bOverride_ColorCorrectionHighlightsMin(Settings.bOverride_ColorCorrectionHighlightsMin)
+	, bOverride_ColorCorrectionHighlightsMax(Settings.bOverride_ColorCorrectionHighlightsMax)
 	, bOverride_BlueCorrection(Settings.bOverride_BlueCorrection)
 	, bOverride_ExpandGamut(Settings.bOverride_ExpandGamut)
 	, bOverride_ToneCurveAmount(Settings.bOverride_ToneCurveAmount)
@@ -710,6 +716,7 @@ FPostProcessSettings::FPostProcessSettings(const FPostProcessSettings& Settings)
 	, bOverride_FilmGrainIntensityHighlights(Settings.bOverride_FilmGrainIntensityHighlights)
 	, bOverride_FilmGrainShadowsMax(Settings.bOverride_FilmGrainShadowsMax)
 	, bOverride_FilmGrainHighlightsMin(Settings.bOverride_FilmGrainHighlightsMin)
+	, bOverride_FilmGrainHighlightsMax(Settings.bOverride_FilmGrainHighlightsMax)
 	, bOverride_FilmGrainTexelSize(Settings.bOverride_FilmGrainTexelSize)
 	, bOverride_FilmGrainTexture(Settings.bOverride_FilmGrainTexture)
 	, bOverride_AmbientOcclusionIntensity(Settings.bOverride_AmbientOcclusionIntensity)
@@ -819,6 +826,7 @@ FPostProcessSettings::FPostProcessSettings(const FPostProcessSettings& Settings)
 	, ColorGainHighlights(Settings.ColorGainHighlights)
 	, ColorOffsetHighlights(Settings.ColorOffsetHighlights)
 	, ColorCorrectionHighlightsMin(Settings.ColorCorrectionHighlightsMin)
+	, ColorCorrectionHighlightsMax(Settings.ColorCorrectionHighlightsMax)
 	, ColorCorrectionShadowsMax(Settings.ColorCorrectionShadowsMax)
 	, BlueCorrection(Settings.BlueCorrection)
 	, ExpandGamut(Settings.ExpandGamut)
@@ -922,6 +930,7 @@ FPostProcessSettings::FPostProcessSettings(const FPostProcessSettings& Settings)
 	, FilmGrainIntensityHighlights(Settings.FilmGrainIntensityHighlights)
 	, FilmGrainShadowsMax(Settings.FilmGrainShadowsMax)
 	, FilmGrainHighlightsMin(Settings.FilmGrainHighlightsMin)
+	, FilmGrainHighlightsMax(Settings.FilmGrainHighlightsMax)
 	, FilmGrainTexelSize(Settings.FilmGrainTexelSize)
 	, FilmGrainTexture(Settings.FilmGrainTexture)
 	, AmbientOcclusionIntensity(Settings.AmbientOcclusionIntensity)
diff --git a/Engine/Source/Runtime/Engine/Private/SceneView.cpp b/Engine/Source/Runtime/Engine/Private/SceneView.cpp
--- a/Engine/Source/Runtime/Engine/Private/SceneView.cpp
+++ b/Engine/Source/Runtime/Engine/Private/SceneView.cpp
@@ -1393,6 +1393,7 @@ void FSceneView::OverridePostProcessSettings(const FPostProcessSettings& Src, fl
 
 		LERP_PP(ColorCorrectionShadowsMax);
 		LERP_PP(ColorCorrectionHighlightsMin);
+		LERP_PP(ColorCorrectionHighlightsMax);
 
 		LERP_PP(BlueCorrection);
 		LERP_PP(ExpandGamut);
@@ -1459,6 +1460,7 @@ void FSceneView::OverridePostProcessSettings(const FPostProcessSettings& Src, fl
 		LERP_PP(FilmGrainIntensityHighlights);
 		LERP_PP(FilmGrainShadowsMax);
 		LERP_PP(FilmGrainHighlightsMin);
+		LERP_PP(FilmGrainHighlightsMax);
 		LERP_PP(FilmGrainTexelSize);
 		LERP_PP(AmbientOcclusionIntensity);
 		LERP_PP(AmbientOcclusionStaticFraction);
diff --git a/Engine/Source/Runtime/Renderer/Private/PostProcess/PostProcessCombineLUTs.cpp b/Engine/Source/Runtime/Renderer/Private/PostProcess/PostProcessCombineLUTs.cpp
--- a/Engine/Source/Runtime/Renderer/Private/PostProcess/PostProcessCombineLUTs.cpp
+++ b/Engine/Source/Runtime/Renderer/Private/PostProcess/PostProcessCombineLUTs.cpp
@@ -120,6 +120,7 @@ BEGIN_SHADER_PARAMETER_STRUCT(FCombineLUTParameters, )
 	SHADER_PARAMETER(float, WhiteTint)
 	SHADER_PARAMETER(float, ColorCorrectionShadowsMax)
 	SHADER_PARAMETER(float, ColorCorrectionHighlightsMin)
+	SHADER_PARAMETER(float, ColorCorrectionHighlightsMax)
 	SHADER_PARAMETER(float, BlueCorrection)
 	SHADER_PARAMETER(float, ExpandGamut)
 	SHADER_PARAMETER(float, ToneCurveAmount)
@@ -203,6 +204,7 @@ void GetCombineLUTParameters(
 
 	Parameters.ColorCorrectionShadowsMax = Settings.ColorCorrectionShadowsMax;
 	Parameters.ColorCorrectionHighlightsMin = Settings.ColorCorrectionHighlightsMin;
+	Parameters.ColorCorrectionHighlightsMax = Settings.ColorCorrectionHighlightsMax;
 
 	Parameters.BlueCorrection = Settings.BlueCorrection;
 	Parameters.ExpandGamut = Settings.ExpandGamut;
diff --git a/Engine/Source/Runtime/Renderer/Private/PostProcess/PostProcessTonemap.cpp b/Engine/Source/Runtime/Renderer/Private/PostProcess/PostProcessTonemap.cpp
--- a/Engine/Source/Runtime/Renderer/Private/PostProcess/PostProcessTonemap.cpp
+++ b/Engine/Source/Runtime/Renderer/Private/PostProcess/PostProcessTonemap.cpp
@@ -327,6 +327,7 @@ BEGIN_SHADER_PARAMETER_STRUCT(FFilmGrainParameters, )
 	SHADER_PARAMETER(float, FilmGrainIntensityHighlights)
 	SHADER_PARAMETER(float, FilmGrainShadowsMax)
 	SHADER_PARAMETER(float, FilmGrainHighlightsMin)
+	SHADER_PARAMETER(float, FilmGrainHighlightsMax)
 	SHADER_PARAMETER_RDG_TEXTURE(Texture2D, FilmGrainTexture)
 	SHADER_PARAMETER_SAMPLER(SamplerState, FilmGrainSampler)
 	SHADER_PARAMETER(FScreenTransform, ScreenPosToFilmGrainTextureUV)
@@ -736,6 +737,7 @@ FScreenPassTexture AddTonemapPass(FRDGBuilder& GraphBuilder, const FViewInfo& Vi
 		CommonParameters.FilmGrain.FilmGrainIntensityHighlights = View.FinalPostProcessSettings.FilmGrainIntensity * View.FinalPostProcessSettings.FilmGrainIntensityHighlights;
 		CommonParameters.FilmGrain.FilmGrainShadowsMax = View.FinalPostProcessSettings.FilmGrainShadowsMax;
 		CommonParameters.FilmGrain.FilmGrainHighlightsMin = View.FinalPostProcessSettings.FilmGrainHighlightsMin;
+		CommonParameters.FilmGrain.FilmGrainHighlightsMax = View.FinalPostProcessSettings.FilmGrainHighlightsMax;
 
 		CommonParameters.FilmGrain.FilmGrainTexture = FilmGrainTexture;
 		CommonParameters.FilmGrain.FilmGrainSampler = TStaticSamplerState<SF_Bilinear, AM_Wrap, AM_Wrap>::GetRHI();
@@ -769,6 +771,7 @@ FScreenPassTexture AddTonemapPass(FRDGBuilder& GraphBuilder, const FViewInfo& Vi
 		CommonParameters.FilmGrain.FilmGrainIntensityHighlights = 0.0f;
 		CommonParameters.FilmGrain.FilmGrainShadowsMax = 0.09f;
 		CommonParameters.FilmGrain.FilmGrainHighlightsMin = 0.5f;
+		CommonParameters.FilmGrain.FilmGrainHighlightsMax = 1.0f;
 		CommonParameters.FilmGrain.FilmGrainTexture = GSystemTextures.GetWhiteDummy(GraphBuilder);
 		CommonParameters.FilmGrain.FilmGrainSampler = TStaticSamplerState<SF_Bilinear, AM_Wrap, AM_Wrap>::GetRHI();
 		CommonParameters.FilmGrain.ScreenPosToFilmGrainTextureUV = FScreenTransform::ScreenPosToViewportUV;
