Fix current first LOD index when using cooked data in editor builds and UseStaticMeshMinLODPerQualityLevels is true.

#rnx

diff --git a/Engine/Source/Runtime/Engine/Private/StaticMesh.cpp b/Engine/Source/Runtime/Engine/Private/StaticMesh.cpp
--- a/Engine/Source/Runtime/Engine/Private/StaticMesh.cpp
+++ b/Engine/Source/Runtime/Engine/Private/StaticMesh.cpp
@@ -1638,7 +1638,7 @@ void FStaticMeshRenderData::Serialize(FArchive& Ar, UStaticMesh* Owner, bool bCo
 #if WITH_EDITOR
 	if (bCooked && Ar.IsLoading())
 	{
-		CurrentFirstLODIdx = 0;
+		CurrentFirstLODIdx = Owner->GetMinLODIdx();
 	}
 	else
 #endif
