Adding option to compile out subobject rep keys

diff --git a/Engine/Source/Runtime/Engine/Classes/Engine/ActorChannel.h b/Engine/Source/Runtime/Engine/Classes/Engine/ActorChannel.h
--- a/Engine/Source/Runtime/Engine/Classes/Engine/ActorChannel.h
+++ b/Engine/Source/Runtime/Engine/Classes/Engine/ActorChannel.h
@@ -10,6 +10,10 @@
 #include "Net/DataReplication.h"
 #include "ActorChannel.generated.h"
 
+#ifndef NET_ENABLE_SUBOBJECT_REPKEYS
+#define NET_ENABLE_SUBOBJECT_REPKEYS 1
+#endif // NET_ENABLE_SUBOBJECT_REPKEYS
+
 class AActor;
 class FInBunch;
 class FNetFieldExportGroup;
@@ -286,6 +290,7 @@ public:
 		return WroteSomething;
 	}
 
+#if NET_ENABLE_SUBOBJECT_REPKEYS
 	// Static size for SubobjectRepKeyMap. Allows us to resuse arrays and avoid dyanmic memory allocations
 	static const int32 SubobjectRepKeyBufferSize = 64;
 
@@ -309,6 +314,7 @@ public:
 	// Returns true if the given ObjID is not up to date with RepKey
 	// this implicitly 'writes' the RepKey to the current out bunch.
 	bool KeyNeedsToReplicate(int32 ObjID, int32 RepKey);
+#endif // NET_ENABLE_SUBOBJECT_REPKEYS
 	
 	// --------------------------------
 
diff --git a/Engine/Source/Runtime/Engine/Private/DataChannel.cpp b/Engine/Source/Runtime/Engine/Private/DataChannel.cpp
--- a/Engine/Source/Runtime/Engine/Private/DataChannel.cpp
+++ b/Engine/Source/Runtime/Engine/Private/DataChannel.cpp
@@ -2355,6 +2355,7 @@ void UActorChannel::ReceivedNak( int32 NakPacketId )
 		CompIt.Value()->ReceivedNak(NakPacketId);
 	}
 
+#if NET_ENABLE_SUBOBJECT_REPKEYS
 	// Reset any subobject RepKeys that were sent on this packetId
 	FPacketRepKeyInfo * Info = SubobjectNakMap.Find(NakPacketId % SubobjectRepKeyBufferSize);
 	if (Info)
@@ -2369,6 +2370,7 @@ void UActorChannel::ReceivedNak( int32 NakPacketId )
 			}
 		}
 	}
+#endif // NET_ENABLE_SUBOBJECT_REPKEYS
 }
 
 void UActorChannel::SetChannelActor(AActor* InActor, ESetChannelActorFlags Flags)
@@ -3366,6 +3368,7 @@ int64 UActorChannel::ReplicateActor()
 				RepComp.Value()->PostSendBunch(PacketRange, Bunch.bReliable);
 			}
 
+#if NET_ENABLE_SUBOBJECT_REPKEYS
 			// If there were any subobject keys pending, add them to the NakMap
 			if (PendingObjKeys.Num() > 0)
 			{
@@ -3394,6 +3397,7 @@ int64 UActorChannel::ReplicateActor()
 					}
 				}
 			}
+#endif // NET_ENABLE_SUBOBJECT_REPKEYS
 
 			if (Actor->bNetTemporary)
 			{
@@ -3404,7 +3408,9 @@ int64 UActorChannel::ReplicateActor()
 		NumBitsWrote = Bunch.GetNumBits();
 	}
 
+#if NET_ENABLE_SUBOBJECT_REPKEYS
 	PendingObjKeys.Empty();
+#endif // NET_ENABLE_SUBOBJECT_REPKEYS
 
 	// If we evaluated everything, mark LastUpdateTime, even if nothing changed.
 	LastUpdateTime = Connection->Driver->GetElapsedTime();
@@ -3548,6 +3554,7 @@ void UActorChannel::Serialize(FArchive& Ar)
 			}
 		);
 
+#if NET_ENABLE_SUBOBJECT_REPKEYS
 		GRANULAR_NETWORK_MEMORY_TRACKING_TRACK("SubobjectRepKeyMap", SubobjectRepKeyMap.CountBytes(Ar));
 
 		GRANULAR_NETWORK_MEMORY_TRACKING_TRACK("SubobjectNakMap",
@@ -3559,6 +3566,7 @@ void UActorChannel::Serialize(FArchive& Ar)
 		);
 
 		GRANULAR_NETWORK_MEMORY_TRACKING_TRACK("PendingObjKeys", PendingObjKeys.CountBytes(Ar));
+#endif // NET_ENABLE_SUBOBJECT_REPKEYS
 
 		GRANULAR_NETWORK_MEMORY_TRACKING_TRACK("QueuedBunchObjectReferences", QueuedBunchObjectReferences.CountBytes(Ar));
 	}
@@ -4425,6 +4433,7 @@ bool UActorChannel::ObjectHasReplicator(const TWeakObjectPtr<UObject>& Obj) cons
 	return ReplicatorRefPtr != nullptr && Obj == ReplicatorRefPtr->Get().GetWeakObjectPtr();
 }
 
+#if NET_ENABLE_SUBOBJECT_REPKEYS
 bool UActorChannel::KeyNeedsToReplicate(int32 ObjID, int32 RepKey)
 {
 	int32 &MapKey = SubobjectRepKeyMap.FindOrAdd(ObjID);
@@ -4435,6 +4444,7 @@ bool UActorChannel::KeyNeedsToReplicate(int32 ObjID, int32 RepKey)
 	PendingObjKeys.Add(ObjID);
 	return true;
 }
+#endif // NET_ENABLE_SUBOBJECT_REPKEYS
 
 void UActorChannel::AddedToChannelPool()
 {
@@ -4466,9 +4476,12 @@ void UActorChannel::AddedToChannelPool()
 	bBlockChannelFailure = false;
 #endif
 	QueuedCloseReason = EChannelCloseReason::Destroyed;
+
+#if NET_ENABLE_SUBOBJECT_REPKEYS
 	SubobjectRepKeyMap.Empty();
 	SubobjectNakMap.Empty();
 	PendingObjKeys.Empty();
+#endif // NET_ENABLE_SUBOBJECT_REPKEYS
 }
 
 bool UActorChannel::ReplicateSubobject(UObject* Obj, FOutBunch& Bunch, const FReplicationFlags& RepFlags)
