SceneComponentNeedsLoadForTarget: use the new AllowsEditorObjects when deciding whether to load an object, rather than HasEditorOnlyData, which now refers only to properties within an object. This is an important distinction for CookedEditor, which has editoronly objects but not editoronly data.
#rnx

diff --git a/Engine/Source/Runtime/Engine/Private/Components/SceneComponent.cpp b/Engine/Source/Runtime/Engine/Private/Components/SceneComponent.cpp
--- a/Engine/Source/Runtime/Engine/Private/Components/SceneComponent.cpp
+++ b/Engine/Source/Runtime/Engine/Private/Components/SceneComponent.cpp
@@ -469,7 +469,7 @@ static bool SceneComponentNeedsLoadForTarget(USceneComponent const* SceneCompone
 		}
 	}
 
-	return TargetPlatform->HasEditorOnlyData() || !SceneComponentObject->IsEditorOnly();
+	return TargetPlatform->AllowsEditorObjects() || !SceneComponentObject->IsEditorOnly();
 }
 
 static bool CheckDescendantsAreAlsoCulledForTarget(USceneComponent const* SceneComponentObject, const ITargetPlatform* TargetPlatform)
