Crash at UAnimDataModel::FindFloatCurve
#fix verify whether or not data model is valid before trying to retrieve float curve

diff --git a/Engine/Source/Runtime/Engine/Private/Animation/AnimMontage.cpp b/Engine/Source/Runtime/Engine/Private/Animation/AnimMontage.cpp
--- a/Engine/Source/Runtime/Engine/Private/Animation/AnimMontage.cpp
+++ b/Engine/Source/Runtime/Engine/Private/Animation/AnimMontage.cpp
@@ -2892,14 +2892,17 @@ void UAnimMontage::BakeTimeStretchCurve()
 
 	// See if Montage is hosting a curve named 'TimeStretchCurveName'
 	const FFloatCurve* TimeStretchFloatCurve = nullptr;
-	if (const USkeleton* MySkeleton = GetSkeleton())
-	{
-		if (const FSmartNameMapping* CurveNameMapping = MySkeleton->GetSmartNameContainer(USkeleton::AnimCurveMappingName))
+	if (ShouldDataModelBeValid())
+	{		
+		if (const USkeleton* MySkeleton = GetSkeleton())
 		{
-			const USkeleton::AnimCurveUID CurveUID = CurveNameMapping->FindUID(TimeStretchCurveName);
-			if (CurveUID != SmartName::MaxUID)
+			if (const FSmartNameMapping* CurveNameMapping = MySkeleton->GetSmartNameContainer(USkeleton::AnimCurveMappingName))
 			{
-				TimeStretchFloatCurve = GetDataModel()->FindFloatCurve(FAnimationCurveIdentifier(CurveUID, ERawCurveTrackTypes::RCT_Float));
+				const USkeleton::AnimCurveUID CurveUID = CurveNameMapping->FindUID(TimeStretchCurveName);
+				if (CurveUID != SmartName::MaxUID)
+				{
+					TimeStretchFloatCurve = GetDataModel()->FindFloatCurve(FAnimationCurveIdentifier(CurveUID, ERawCurveTrackTypes::RCT_Float));
+				}
 			}
 		}
 	}
