Mark up some more ParallelFors with names for Insights
behavior nop

diff --git a/Engine/Source/Runtime/Core/Private/GenericPlatform/GenericPlatformFile.cpp b/Engine/Source/Runtime/Core/Private/GenericPlatform/GenericPlatformFile.cpp
--- a/Engine/Source/Runtime/Core/Private/GenericPlatform/GenericPlatformFile.cpp
+++ b/Engine/Source/Runtime/Core/Private/GenericPlatform/GenericPlatformFile.cpp
@@ -593,8 +593,8 @@ bool IPlatformFile::IterateDirectoryRecursively(const TCHAR* Directory, FDirecto
 	while (bResult && DirectoriesToVisitNext.Num() > 0)
 	{
 		TArray<FString> DirectoriesToVisit = MoveTemp(DirectoriesToVisitNext);
-		ParallelFor(
-			DirectoriesToVisit.Num(),
+		ParallelFor( TEXT("IterateDirectoryRecursively.PF"),
+			DirectoriesToVisit.Num(),1,
 			[this, &DirectoriesToVisit, &Recurse, &bResult](int32 Index)
 			{
 				if (bResult && !IterateDirectory(*DirectoriesToVisit[Index], Recurse))
diff --git a/Engine/Source/Runtime/Engine/Private/ShaderCompiler/ShaderCompiler.cpp b/Engine/Source/Runtime/Engine/Private/ShaderCompiler/ShaderCompiler.cpp
--- a/Engine/Source/Runtime/Engine/Private/ShaderCompiler/ShaderCompiler.cpp
+++ b/Engine/Source/Runtime/Engine/Private/ShaderCompiler/ShaderCompiler.cpp
@@ -523,8 +523,8 @@ void FShaderCompileJobCollection::SubmitJobs(const TArray<FShaderCommonCompileJo
 			// Just precompute the InputHash for each job in multiple-thread.
 			if (ShaderCompiler::IsJobCacheEnabled())
 			{
-				TRACE_CPUPROFILER_EVENT_SCOPE(GetInputHash);
-				ParallelFor(InJobs.Num(), [&InJobs](int32 Index) { InJobs[Index]->GetInputHash(); });
+				TRACE_CPUPROFILER_EVENT_SCOPE(ShaderCompiler.GetInputHash);
+				ParallelFor( TEXT("ShaderCompiler.GetInputHash.PF"), InJobs.Num(),1, [&InJobs](int32 Index) { InJobs[Index]->GetInputHash(); });
 			}
 
 			FWriteScopeLock Locker(Lock);
@@ -2202,7 +2202,7 @@ void FShaderCompileThreadRunnable::PushCompletedJobsToManager()
 
 void FShaderCompileThreadRunnable::WriteNewTasks()
 {
-	TRACE_CPUPROFILER_EVENT_SCOPE(FShaderCompileThreadRunnable::WriteNewTasks);
+	TRACE_CPUPROFILER_EVENT_SCOPE(ShaderCompiler.WriteNewTasks);
 
 	// first, a quick check if anything is needed just to avoid hammering the task graph
 	bool bHasTasksToWrite = false;
@@ -2229,7 +2229,7 @@ void FShaderCompileThreadRunnable::WriteNewTasks()
 		// Only write tasks once
 		if (!CurrentWorkerInfo.bIssuedTasksToWorker && CurrentWorkerInfo.QueuedJobs.Num() > 0)
 		{
-			TRACE_CPUPROFILER_EVENT_SCOPE(FShaderCompileThreadRunnable::WriteNewTasksForWorker);
+			TRACE_CPUPROFILER_EVENT_SCOPE(ShaderCompiler.WriteNewTasksForWorker);
 			CurrentWorkerInfo.bIssuedTasksToWorker = true;
 
 			const FString WorkingDirectory = Manager->AbsoluteShaderBaseWorkingDirectory + FString::FromInt(WorkerIndex);
@@ -2312,7 +2312,7 @@ void FShaderCompileThreadRunnable::WriteNewTasks()
 
 	if (bParallelizeIO)
 	{
-		ParallelFor(WorkerInfos.Num(), LoopBody, EParallelForFlags::Unbalanced);
+		ParallelFor( TEXT("ShaderCompiler.WriteNewTasks.PF"), WorkerInfos.Num(),1, LoopBody, EParallelForFlags::Unbalanced);
 	}
 	else
 	{
@@ -2423,7 +2423,7 @@ bool FShaderCompileThreadRunnable::LaunchWorkersIfNeeded()
 
 int32 FShaderCompileThreadRunnable::ReadAvailableResults()
 {
-	TRACE_CPUPROFILER_EVENT_SCOPE(FShaderCompileThreadRunnable::ReadAvailableResults);
+	TRACE_CPUPROFILER_EVENT_SCOPE(ShaderCompiler.ReadAvailableResults);
 	int32 NumProcessed = 0;
 
 	// first, a quick check if anything is needed just to avoid hammering the task graph
@@ -2495,7 +2495,7 @@ int32 FShaderCompileThreadRunnable::ReadAvailableResults()
 
 	if (bParallelizeIO)
 	{
-		ParallelFor(WorkerInfos.Num(), LoopBody, EParallelForFlags::Unbalanced);
+		ParallelFor( TEXT("ShaderCompiler.ReadAvailableResults.PF"),WorkerInfos.Num(),1, LoopBody, EParallelForFlags::Unbalanced);
 	}
 	else 
 	{
diff --git a/Engine/Source/Runtime/Projects/Private/PluginManager.cpp b/Engine/Source/Runtime/Projects/Private/PluginManager.cpp
--- a/Engine/Source/Runtime/Projects/Private/PluginManager.cpp
+++ b/Engine/Source/Runtime/Projects/Private/PluginManager.cpp
@@ -519,8 +519,8 @@ void FPluginManager::ReadPluginsInDirectory(const FString& PluginsDirectory, con
 		TArray<FLoadContext> Contexts;
 		Contexts.SetNum(FileNames.Num());
 
-		ParallelFor(
-			FileNames.Num(),
+		ParallelFor(TEXT("ReadPluginsInDirectory.PF"),
+			FileNames.Num(),1,
 			[&Contexts, &FileNames](int32 Index)
 			{
 				FLoadContext& Context = Contexts[Index];
