Fix for failed TraceSourceFiltering editor tests

diff --git a/Engine/Source/Runtime/Engine/Private/TraceFilter.cpp b/Engine/Source/Runtime/Engine/Private/TraceFilter.cpp
--- a/Engine/Source/Runtime/Engine/Private/TraceFilter.cpp
+++ b/Engine/Source/Runtime/Engine/Private/TraceFilter.cpp
@@ -284,7 +284,8 @@ template<>
 bool ENGINE_API FTraceFilter::IsObjectTraceable</*bForceThreadSafe = */ false>(const UObject* InObject)
 {
 	check(GObjectFilterAnnotations.IsLocked());
-	return GObjectFilterAnnotations.GetAnnotationMap().Find(InObject) != nullptr;
+	// Object not found in the AnnotationMap means that it is at the default value, which is bIsTraceable == true
+	return GObjectFilterAnnotations.GetAnnotationMap().Find(InObject) == nullptr;
 }
 
 template<>
