Removing unneeded Scope event.

diff --git a/Engine/Source/Runtime/Engine/Private/Components/SkinnedMeshComponent.cpp b/Engine/Source/Runtime/Engine/Private/Components/SkinnedMeshComponent.cpp
--- a/Engine/Source/Runtime/Engine/Private/Components/SkinnedMeshComponent.cpp
+++ b/Engine/Source/Runtime/Engine/Private/Components/SkinnedMeshComponent.cpp
@@ -3133,8 +3133,6 @@ bool USkinnedMeshComponent::UpdateLODStatus()
 
 bool USkinnedMeshComponent::UpdateLODStatus_Internal(int32 InMasterPoseComponentPredictedLODLevel, bool bRequestedByMasterPoseComponent)
 {
-	SCOPED_NAMED_EVENT(USkinnedMeshComponent_UpdateLODStatus, FColor::Red);
-
 	// Don't update LOD status for slave component unless it explicitly ignores its master component's LOD or if update is recursively requested by its master.
 	// This is because when UpdateLODStatus is called on master component, it updates the slave component LOD,
 	// therefore if slave component also calls UpdateLODStatus, it could overturn the result and cause LOD to be out of sync from its master.
