VulkanRHI: Fix validation error "requires an image view of type VK_IMAGE_VIEW_TYPE_2D_ARRAY but got VK_IMAGE_VIEW_TYPE_2D" in fallback for SSProfilesPreIntegratedTexture (goes with 18979942)

diff --git a/Engine/Source/Runtime/Engine/Private/Rendering/SubsurfaceProfile.cpp b/Engine/Source/Runtime/Engine/Private/Rendering/SubsurfaceProfile.cpp
--- a/Engine/Source/Runtime/Engine/Private/Rendering/SubsurfaceProfile.cpp
+++ b/Engine/Source/Runtime/Engine/Private/Rendering/SubsurfaceProfile.cpp
@@ -661,7 +661,7 @@ FRHITexture* GetSubsurfaceProfileTextureWithFallback()
 
 FRHITexture* GetSSProfilesPreIntegratedTextureWithFallback()
 {
-	return GSSProfilesPreIntegratedTexture ? GSSProfilesPreIntegratedTexture->GetRHI() : static_cast<FRHITexture*>(GBlackTexture->TextureRHI);
+	return GSSProfilesPreIntegratedTexture ? GSSProfilesPreIntegratedTexture->GetRHI() : static_cast<FRHITexture*>(GBlackArrayTexture->TextureRHI);
 }
 
 void UpdateSubsurfaceProfileTexture(FRDGBuilder& GraphBuilder, EShaderPlatform ShaderPlatform)
