Fixed build (caused by bad merge).

diff --git a/Engine/Source/Runtime/Engine/Private/WorldPartition/WorldPartitionConverter.cpp b/Engine/Source/Runtime/Engine/Private/WorldPartition/WorldPartitionConverter.cpp
--- a/Engine/Source/Runtime/Engine/Private/WorldPartition/WorldPartitionConverter.cpp
+++ b/Engine/Source/Runtime/Engine/Private/WorldPartition/WorldPartitionConverter.cpp
@@ -260,9 +260,9 @@ bool FWorldPartitionConverter::PrepareStreamingLevelForConversion(ULevelStreamin
 	{
 		for (AActor* Actor : Level->Actors)
 		{
-			if (Actor)
+			if (Actor && Actor->CanChangeIsSpatiallyLoadedFlag())
 			{
-				Actor->SetGridPlacement(EActorGridPlacement::AlwaysLoaded);
+				Actor->SetIsSpatiallyLoaded(false);
 			}
 		}
 	}
@@ -349,7 +349,7 @@ void FWorldPartitionConverter::FixupSoftObjectPaths(UPackage* OuterPackage)
 			Object->Serialize(FixupSerializer);
 		}
 		return true;
-	}, true, RF_NoFlags, EInternalObjectFlags::PendingKill);
+	}, true, RF_NoFlags, EInternalObjectFlags::Garbage);
 }
 
 #endif
\ No newline at end of file
diff --git a/Engine/Source/Runtime/Engine/Public/WorldPartition/HLOD/HLODLayer.h b/Engine/Source/Runtime/Engine/Public/WorldPartition/HLOD/HLODLayer.h
--- a/Engine/Source/Runtime/Engine/Public/WorldPartition/HLOD/HLODLayer.h
+++ b/Engine/Source/Runtime/Engine/Public/WorldPartition/HLOD/HLODLayer.h
@@ -53,7 +53,7 @@ public:
 	float GetLoadingRange() const { return !bIsSpatiallyLoaded ? WORLD_MAX : LoadingRange; }
 	const TSoftObjectPtr<UHLODLayer>& GetParentLayer() const;
 	bool IsSpatiallyLoaded() const { return bIsSpatiallyLoaded; }
-	void SetIsAlwaysLoaded(bool bInIsAlwaysLoaded) { bAlwaysLoaded = bInIsAlwaysLoaded; }
+	void SetIsSpatiallyLoaded(bool bInIsSpatiallyLoaded) { bIsSpatiallyLoaded = bInIsSpatiallyLoaded; }
 
 	bool DoesRequireWarmup() const;
 
