UWorld::Rename: Rename top level objects (ex: HLOD StaticMesh outered to Actor Package) when updating actor packages

#rnx

diff --git a/Engine/Source/Runtime/Engine/Private/World.cpp b/Engine/Source/Runtime/Engine/Private/World.cpp
--- a/Engine/Source/Runtime/Engine/Private/World.cpp
+++ b/Engine/Source/Runtime/Engine/Private/World.cpp
@@ -687,8 +687,26 @@ bool UWorld::Rename(const TCHAR* InName, UObject* NewOuter, ERenameFlags Flags)
 		// this will leave dirty empty actor packages being which will be cleaned up though SaveAll, although SaveCurrentLevel won't pick them up
 		for (AActor* Actor : CopyActors)
 		{
+			UPackage* ExternalPackage = Actor->GetPackage();
 			Actor->SetPackageExternal(false);
+			
+			TArray<UObject*> DependantObjects;
+			ForEachObjectWithPackage(ExternalPackage, [&DependantObjects](UObject* Object)
+			{
+				if (!Cast<UMetaData>(Object))
+				{
+					DependantObjects.Add(Object);
+				}
+				return true;
+			}, false);
+						
 			Actor->SetPackageExternal(true);
+
+			// Move dependant objects into the new actor package
+			for (UObject* DependantObject : DependantObjects)
+			{
+				DependantObject->Rename(nullptr, Actor->GetExternalPackage(), REN_NonTransactional | REN_DontCreateRedirectors | REN_ForceNoResetLoaders | REN_DoNotDirty);
+			}
 		}
 
 		if (PersistentLevel->IsUsingExternalObjects())
