Fix cook error

diff --git a/Engine/Source/Runtime/Engine/Private/Materials/HLSLMaterialTranslator.cpp b/Engine/Source/Runtime/Engine/Private/Materials/HLSLMaterialTranslator.cpp
--- a/Engine/Source/Runtime/Engine/Private/Materials/HLSLMaterialTranslator.cpp
+++ b/Engine/Source/Runtime/Engine/Private/Materials/HLSLMaterialTranslator.cpp
@@ -1463,11 +1463,11 @@ bool FHLSLMaterialTranslator::Translate()
 			ResourcesString += "// No Strata material provided, or using legacy material conversion \r\n";
 			
 			// Adde default strata functions
-			ResourcesString += "#if TEMPLATE_USES_STRATA\n;";
+			ResourcesString += "#if TEMPLATE_USES_STRATA\n";
 			ResourcesString += "void UpdateAllBSDFsOperatorCoverageTransmittance(inout FStrataPixelHeader StrataPixelHeader, inout FStrataTree StrataTree, bool bRoughDiffuseEnabled, FStrataAddressing NullStrataAddressing, float3 V, float3 L) {}\n";
 			ResourcesString += "void UpdateAllOperatorsCoverageTransmittance(inout FStrataTree StrataTree) {}\n";
 			ResourcesString += "void UpdateAllBSDFWeightAfterOperatorVisit(inout FStrataTree StrataTree) {}\n";
-			ResourcesString += "#endif\n;";
+			ResourcesString += "#endif\n";
 		}
 
 		LoadShaderSourceFileChecked(TEXT("/Engine/Private/MaterialTemplate.ush"), GetShaderPlatform(), MaterialTemplate);
