Strata - removed front material from get/set material attribute.

diff --git a/Engine/Source/Editor/DetailCustomizations/Private/MaterialAttributePropertyDetails.cpp b/Engine/Source/Editor/DetailCustomizations/Private/MaterialAttributePropertyDetails.cpp
--- a/Engine/Source/Editor/DetailCustomizations/Private/MaterialAttributePropertyDetails.cpp
+++ b/Engine/Source/Editor/DetailCustomizations/Private/MaterialAttributePropertyDetails.cpp
@@ -22,10 +22,19 @@ void FMaterialAttributePropertyDetails::CustomizeDetails(IDetailLayoutBuilder& D
 	// Populate combo boxes with material property list
 	FMaterialAttributeDefinitionMap::GetAttributeNameToIDList(AttributeNameToIDList);
 
+	static const FString ExcludedFrontMaterialName = FString("FrontMaterial");
+
 	AttributeDisplayNameList.Empty(AttributeNameToIDList.Num());
 	for (const TPair<FString, FGuid>& NameGUIDPair : AttributeNameToIDList)
 	{
-		AttributeDisplayNameList.Add(MakeShareable(new FString(NameGUIDPair.Key)));
+		// We remove unwanted elements from the list of material attributes that can be used in the material layer workflow.
+		// We do not want to blend strata BSDF which could increase the complexity of the material (BSDF / Slab count) for each blend operation.
+		// Only parameters are allowed to finally be transform into Strata at the end of the chain.
+		// This filtering is ran only when constructing the UI element.
+		if (NameGUIDPair.Key != ExcludedFrontMaterialName)
+		{
+			AttributeDisplayNameList.Add(MakeShareable(new FString(NameGUIDPair.Key)));
+		}
 	}
 
 	// Fetch root property we're dealing with
diff --git a/Engine/Source/Runtime/Engine/Private/Materials/MaterialExpressions.cpp b/Engine/Source/Runtime/Engine/Private/Materials/MaterialExpressions.cpp
--- a/Engine/Source/Runtime/Engine/Private/Materials/MaterialExpressions.cpp
+++ b/Engine/Source/Runtime/Engine/Private/Materials/MaterialExpressions.cpp
@@ -645,25 +645,8 @@ int32 CompileShadingModelBlendFunction(FMaterialCompiler* Compiler, const int32
 
 int32 CompileStrataBlendFunction(FMaterialCompiler* Compiler, const int32 A, const int32 B, const int32 Alpha)
 {
-	if (A == INDEX_NONE || B == INDEX_NONE || Alpha == INDEX_NONE)
-	{
-		return INDEX_NONE;
-	}
-	uint32 MaxDistanceToLeaves = 0; // STRATA_TODO: fix me
-
-	int32 OutputCodeChunk = Compiler->StrataHorizontalMixing(A, B, Alpha, 0, MaxDistanceToLeaves); // A is background (fully visible when Alpha=0) and B is foreground (fully visible when Alpha=1)
-
-#if WITH_EDITOR
-	// Also register the horizontal mixing operation with the compiler.
-	if (!Compiler->StrataCompilationInfoContainsCodeChunk(A) || !Compiler->StrataCompilationInfoContainsCodeChunk(B))
-	{
-		return Compiler->Errorf(TEXT("Could not find A or B code chunk in CompileStrataBlendFunction"));
-	}
-	FStrataMaterialCompilationInfo StrataInfo = StrataCompilationInfoHorizontalMixing(Compiler, Compiler->GetStrataCompilationInfo(A), Compiler->GetStrataCompilationInfo(B));
-	Compiler->StrataCompilationInfoRegisterCodeChunk(OutputCodeChunk, StrataInfo);
-#endif
-
-	return OutputCodeChunk;
+	check(false);	// we should never blend Strata data from material layers.
+	return INDEX_NONE;
 }
 
 #if WITH_EDITOR
diff --git a/Engine/Source/Runtime/Engine/Private/Materials/MaterialHLSLEmitter.cpp b/Engine/Source/Runtime/Engine/Private/Materials/MaterialHLSLEmitter.cpp
--- a/Engine/Source/Runtime/Engine/Private/Materials/MaterialHLSLEmitter.cpp
+++ b/Engine/Source/Runtime/Engine/Private/Materials/MaterialHLSLEmitter.cpp
@@ -362,7 +362,6 @@ static FString GenerateMaterialTemplateHLSL(EShaderPlatform ShaderPlatform,
 
 			if (PropertyIndex == MP_FrontMaterial)
 			{
-				// TODO - Strata
 				EvaluateMaterialAttributesCode += FString::Printf("    PixelMaterialInputs.%s = GetInitialisedStrataData();" LINE_TERMINATOR, *PropertyName);
 			}
 			else if (PropertyIndex == MP_SubsurfaceColor)
