Output some more information about the material topology and strata tree in the compilation stat window.

diff --git a/Engine/Source/Runtime/Engine/Private/Materials/HLSLMaterialTranslator.cpp b/Engine/Source/Runtime/Engine/Private/Materials/HLSLMaterialTranslator.cpp
--- a/Engine/Source/Runtime/Engine/Private/Materials/HLSLMaterialTranslator.cpp
+++ b/Engine/Source/Runtime/Engine/Private/Materials/HLSLMaterialTranslator.cpp
@@ -1985,9 +1985,64 @@ void FHLSLMaterialTranslator::GetMaterialEnvironment(EShaderPlatform InPlatform,
 				}
 			}
 
-			StrataMaterialDescription += FString::Printf(TEXT("------------------\r\n"));
+			StrataMaterialDescription += FString::Printf(TEXT("----------- STRATA TREE -----------\r\n"));
 			StrataMaterialDescription += FString::Printf(TEXT("Graph maximum distance to leaves %u\r\n"), RootMaximumDistanceToLeaves);
-			StrataMaterialDescription += FString::Printf(TEXT("------------------\r\n"));
+			// Debug print operators according to depth from root.
+			{
+				auto OperatorTypeToStr = [&](int32 OperatorType)
+				{
+					switch (OperatorType)
+					{
+					case STRATA_OPERATOR_WEIGHT:
+					{
+						static const FString Name = TEXT("WEIGHT ");
+						return Name;
+					}
+					case STRATA_OPERATOR_VERTICAL:
+					{
+						static const FString Name = TEXT("VERTICA");
+						return Name;
+					}
+					case STRATA_OPERATOR_HORIZONTAL:
+					{
+						static const FString Name = TEXT("HORIZON");
+						return Name;
+					}
+					case STRATA_OPERATOR_ADD:
+					{
+						static const FString Name = TEXT("ADD    ");
+						return Name;
+					}
+					case STRATA_OPERATOR_BSDF:
+					{
+						static const FString Name = TEXT("BSDF   ");
+						return Name;
+					}
+					case STRATA_OPERATOR_BSDF_LEGACY:
+					{
+						static const FString Name = TEXT("BSDFLEG");
+						return Name;
+					}
+					}
+					static const FString Name = TEXT("UNKNOWN");
+					return Name;
+				};
+
+				check(StrataMaterialRootOperator);
+				for (int32 DistanceToLeaves = RootMaximumDistanceToLeaves; DistanceToLeaves >= 0; --DistanceToLeaves)
+				{
+					StrataMaterialDescription += FString::Printf(TEXT("----- DistanceFromLeaves = %d -----\r\n"), DistanceToLeaves);
+					for (auto& It : StrataMaterialExpressionRegisteredOperators)
+					{
+						if (!It.IsDiscarded() && It.MaxDistanceFromLeaves == DistanceToLeaves)
+						{
+							StrataMaterialDescription += FString::Printf(TEXT("\tIdx=%d Op=%s ParentIdx=%d LeftIndex=%d RightIndex=%d BSDFIdx=%d LayerDepth=%d IsTop=%d IsBot=%d\r\n"),
+								It.Index, *OperatorTypeToStr(It.OperatorType), It.ParentIndex, It.LeftIndex, It.RightIndex, It.BSDFIndex, It.LayerDepth, It.bIsTop, It.bIsBottom);
+						}
+					}
+				}
+			}
+			StrataMaterialDescription += FString::Printf(TEXT("-----------------------------------\r\n"));
 
 			StrataMaterialDescription += FString::Printf(TEXT("Byte Per Pixel Budget      %u\r\n"), StrataBytePerPixel);
 			StrataMaterialDescription += FString::Printf(TEXT("Result.bFitInMemoryBudget  %s\r\n"), StrataMaterialAnalysis.bFitInMemoryBudget ? TEXT("YES") : TEXT("NO"));
