Add texture types to Shader::EValueType, texture values are now represented by FExpressions (rather than separate texture desc type) in new HLSL translator
Makes it cleaner to pass textures to functions or custom HLSL
Also add more misc functionality to new HLSL translator

diff --git a/Engine/Shaders/Private/MaterialTemplate.ush b/Engine/Shaders/Private/MaterialTemplate.ush
--- a/Engine/Shaders/Private/MaterialTemplate.ush
+++ b/Engine/Shaders/Private/MaterialTemplate.ush
@@ -19,6 +19,7 @@
 #include "/Engine/Private/HairStrands/HairCardsAttributeCommon.ush"
 #include "/Engine/Private/HairStrands/HairStrandsAttributeCommon.ush"
 #include "/Engine/Private/DeferredShadingCommon.ush"
+#include "/Engine/Private/MaterialTexture.ush"
 
 // Update this GUID to force all material shaders to recompile
 // Merge conflicts on this line should be resolved by generating a new GUID
@@ -1706,73 +1707,6 @@ MaterialFloat2 GetParticleMacroUV(FMaterialPixelParameters Parameters)
 	return (Parameters.ScreenPosition.xy / Parameters.ScreenPosition.w - Parameters.Particle.MacroUV.xy) * Parameters.Particle.MacroUV.zw + MaterialFloat2(.5, .5);
 }
 
-MaterialFloat4 ProcessMaterialColorTextureLookup(MaterialFloat4 TextureValue)
-{
-	return TextureValue;
-}
-
-MaterialFloat4 ProcessMaterialVirtualColorTextureLookup(MaterialFloat4 TextureValue)
-{
-	TextureValue = ProcessMaterialColorTextureLookup(TextureValue);
-#if FEATURE_LEVEL == FEATURE_LEVEL_ES3_1
-	// on mobile all VT physical spaces use linear color formats, do sRGB to Linear conversion in the shader
-	TextureValue = MaterialFloat4(TextureValue.rgb*TextureValue.rgb, TextureValue.a);
-#endif
-	return TextureValue;
-}
-
-MaterialFloat4 ProcessMaterialExternalTextureLookup(MaterialFloat4 TextureValue)
-{
-#if COMPILER_GLSL_ES3_1
-	return MaterialFloat4(pow(TextureValue.rgb, 2.2f), TextureValue.a);
-#else
-	return ProcessMaterialColorTextureLookup(TextureValue);
-#endif
-}
-
-MaterialFloat4 ProcessMaterialLinearColorTextureLookup(MaterialFloat4 TextureValue)
-{
-	return TextureValue;
-}
-
-MaterialFloat4 ProcessMaterialGreyscaleTextureLookup(MaterialFloat4 TextureValue)
-{
-	// Sampling a greyscale texture in D3D9 gives: (G,G,G)
-	// Sampling a greyscale texture in D3D11 gives: (G,0,0)
-	// This replication reproduces the D3D9 behavior in all cases.
-	MaterialFloat GrayValue = TextureValue.r;
-#if (COMPILER_GLSL_ES3_1 || VULKAN_PROFILE) // OpenGLES3.1, Vulkan3.1 do not support sRGB sampling from R8
-	#if MOBILE_EMULATION
-	if( ResolvedView.MobilePreviewMode > 0.5f )
-	{
-		// undo HW srgb->lin
-		GrayValue = pow(GrayValue, 1.0f/2.2f); // TODO: replace with a more accurate lin -> sRGB conversion.
-	}
-	#endif
-	// sRGB read approximation (in highp if possible)
-	float LinValue = GrayValue;
-	LinValue *= LinValue;
-	return MaterialFloat4(LinValue.rrrr);
-#endif 
-	return GrayValue.rrrr;
-}
-
-MaterialFloat4 ProcessMaterialLinearGreyscaleTextureLookup(MaterialFloat4 TextureValue)
-{
-	// Sampling a greyscale texture in D3D9 gives: (G,G,G)
-	// Sampling a greyscale texture in D3D11 gives: (G,0,0)
-	// This replication reproduces the D3D9 behavior in all cases.
-	return TextureValue.rrrr;
-}
-
-MaterialFloat4 ProcessMaterialAlphaTextureLookup(MaterialFloat4 TextureValue)
-{
-	// Sampling a single channel texture in D3D9 gives: (G,G,G)
-	// Sampling a single channel texture in D3D11 gives: (G,0,0)
-	// This replication reproduces the D3D9 behavior in all cases.
-	return TextureValue.rrrr;
-}
-
 /** Accesses a shared material sampler or falls back if independent samplers are not supported. */
 SamplerState GetMaterialSharedSampler(SamplerState TextureSampler, SamplerState SharedSampler)
 {
diff --git a/Engine/Shaders/Private/MaterialTexture.ush b/Engine/Shaders/Private/MaterialTexture.ush
new file mode 100644
--- /dev/null
+++ b/Engine/Shaders/Private/MaterialTexture.ush
@@ -0,0 +1,136 @@
+// Copyright Epic Games, Inc. All Rights Reserved.
+
+#pragma once
+
+#define SAMPLERTYPE_Color 0
+#define SAMPLERTYPE_Grayscale 1
+#define SAMPLERTYPE_Alpha 2
+#define SAMPLERTYPE_Normal 3
+#define SAMPLERTYPE_Masks 4
+#define SAMPLERTYPE_DistanceFieldFont 5
+#define SAMPLERTYPE_LinearColor 6
+#define SAMPLERTYPE_LinearGrayscale 7
+#define SAMPLERTYPE_Data 8
+#define SAMPLERTYPE_External 9
+
+#define SAMPLERTYPE_VirtualColor 10
+#define SAMPLERTYPE_VirtualGrayscale 11
+#define SAMPLERTYPE_VirtualAlpha 12
+#define SAMPLERTYPE_VirtualNormal 13
+#define SAMPLERTYPE_VirtualMasks 14
+#define SAMPLERTYPE_VirtualLinearColor 15
+#define SAMPLERTYPE_VirtualLinearGrayscale 16
+
+struct FTexture2D
+{
+	Texture2D Texture;
+	SamplerState Sampler;
+	int SamplerType; // EMaterialSamplerType
+};
+
+FTexture2D MakeTexture2D(Texture2D Texture, SamplerState Sampler, int SamplerType)
+{
+	FTexture2D Result;
+	Result.Texture = Texture;
+	Result.Sampler = Sampler;
+	Result.SamplerType = SamplerType;
+	return Result;
+}
+
+MaterialFloat4 ProcessMaterialColorTextureLookup(MaterialFloat4 TextureValue)
+{
+	return TextureValue;
+}
+
+MaterialFloat4 ProcessMaterialVirtualColorTextureLookup(MaterialFloat4 TextureValue)
+{
+	TextureValue = ProcessMaterialColorTextureLookup(TextureValue);
+#if FEATURE_LEVEL == FEATURE_LEVEL_ES3_1
+	// on mobile all VT physical spaces use linear color formats, do sRGB to Linear conversion in the shader
+	TextureValue = MaterialFloat4(TextureValue.rgb*TextureValue.rgb, TextureValue.a);
+#endif
+	return TextureValue;
+}
+
+MaterialFloat4 ProcessMaterialExternalTextureLookup(MaterialFloat4 TextureValue)
+{
+#if COMPILER_GLSL_ES3_1
+	return MaterialFloat4(pow(TextureValue.rgb, 2.2f), TextureValue.a);
+#else
+	return ProcessMaterialColorTextureLookup(TextureValue);
+#endif
+}
+
+MaterialFloat4 ProcessMaterialLinearColorTextureLookup(MaterialFloat4 TextureValue)
+{
+	return TextureValue;
+}
+
+MaterialFloat4 ProcessMaterialGreyscaleTextureLookup(MaterialFloat4 TextureValue)
+{
+	// Sampling a greyscale texture in D3D9 gives: (G,G,G)
+	// Sampling a greyscale texture in D3D11 gives: (G,0,0)
+	// This replication reproduces the D3D9 behavior in all cases.
+	MaterialFloat GrayValue = TextureValue.r;
+#if (COMPILER_GLSL_ES3_1 || VULKAN_PROFILE) // OpenGLES3.1, Vulkan3.1 do not support sRGB sampling from R8
+	#if MOBILE_EMULATION
+	if( ResolvedView.MobilePreviewMode > 0.5f )
+	{
+		// undo HW srgb->lin
+		GrayValue = pow(GrayValue, 1.0f/2.2f); // TODO: replace with a more accurate lin -> sRGB conversion.
+	}
+	#endif
+	// sRGB read approximation (in highp if possible)
+	float LinValue = GrayValue;
+	LinValue *= LinValue;
+	return MaterialFloat4(LinValue.rrrr);
+#endif 
+	return GrayValue.rrrr;
+}
+
+MaterialFloat4 ProcessMaterialLinearGreyscaleTextureLookup(MaterialFloat4 TextureValue)
+{
+	// Sampling a greyscale texture in D3D9 gives: (G,G,G)
+	// Sampling a greyscale texture in D3D11 gives: (G,0,0)
+	// This replication reproduces the D3D9 behavior in all cases.
+	return TextureValue.rrrr;
+}
+
+MaterialFloat4 ProcessMaterialAlphaTextureLookup(MaterialFloat4 TextureValue)
+{
+	// Sampling a single channel texture in D3D9 gives: (G,G,G)
+	// Sampling a single channel texture in D3D11 gives: (G,0,0)
+	// This replication reproduces the D3D9 behavior in all cases.
+	return TextureValue.rrrr;
+}
+
+float4 ApplyMaterialSamplerType(float4 Value, int SamplerType)
+{
+	switch(SamplerType)
+	{
+	case SAMPLERTYPE_External: return ProcessMaterialExternalTextureLookup(Value);
+	case SAMPLERTYPE_Color: return ProcessMaterialColorTextureLookup(Value);
+	case SAMPLERTYPE_VirtualColor: return ProcessMaterialVirtualColorTextureLookup(Value);
+
+	case SAMPLERTYPE_LinearColor:
+	case SAMPLERTYPE_VirtualLinearColor: return ProcessMaterialLinearColorTextureLookup(Value);
+
+	case SAMPLERTYPE_Alpha:
+	case SAMPLERTYPE_VirtualAlpha:
+	case SAMPLERTYPE_DistanceFieldFont: return ProcessMaterialAlphaTextureLookup(Value);
+
+	case SAMPLERTYPE_Grayscale:
+	case SAMPLERTYPE_VirtualGrayscale: return ProcessMaterialGreyscaleTextureLookup(Value);
+
+	case SAMPLERTYPE_LinearGrayscale:
+	case SAMPLERTYPE_VirtualLinearGrayscale: return ProcessMaterialLinearGreyscaleTextureLookup(Value);
+
+	case SAMPLERTYPE_Normal:
+	case SAMPLERTYPE_VirtualNormal: return UnpackNormalMap(Value);
+
+	case SAMPLERTYPE_Masks:
+	case SAMPLERTYPE_VirtualMasks:
+	case SAMPLERTYPE_Data:
+	default: return Value;
+	}
+}
diff --git a/Engine/Source/Runtime/Engine/Classes/Materials/MaterialExpression.h b/Engine/Source/Runtime/Engine/Classes/Materials/MaterialExpression.h
--- a/Engine/Source/Runtime/Engine/Classes/Materials/MaterialExpression.h
+++ b/Engine/Source/Runtime/Engine/Classes/Materials/MaterialExpression.h
@@ -29,7 +29,6 @@ namespace HLSLTree
 class FScope;
 class FStatement;
 class FExpression;
-class FTextureParameterDeclaration;
 }
 }
 
@@ -307,7 +306,6 @@ class ENGINE_API UMaterialExpression : public UObject
 	 */
 	virtual bool GenerateHLSLStatements(FMaterialHLSLGenerator& Generator, UE::HLSLTree::FScope& Scope);
 	virtual bool GenerateHLSLExpression(FMaterialHLSLGenerator& Generator, UE::HLSLTree::FScope& Scope, int32 OutputIndex, UE::HLSLTree::FExpression* &OutExpression);
-	virtual bool GenerateHLSLTexture(FMaterialHLSLGenerator& Generator, UE::HLSLTree::FScope& Scope, int32 OutputIndex, UE::HLSLTree::FTextureParameterDeclaration*& OutTexture);
 
 #endif // WITH_EDITOR
 
diff --git a/Engine/Source/Runtime/Engine/Classes/Materials/MaterialExpressionAbs.h b/Engine/Source/Runtime/Engine/Classes/Materials/MaterialExpressionAbs.h
--- a/Engine/Source/Runtime/Engine/Classes/Materials/MaterialExpressionAbs.h
+++ b/Engine/Source/Runtime/Engine/Classes/Materials/MaterialExpressionAbs.h
@@ -27,6 +27,7 @@ class UMaterialExpressionAbs : public UMaterialExpression
 #if WITH_EDITOR
 	virtual int32 Compile(class FMaterialCompiler* Compiler, int32 OutputIndex) override;
 	virtual void GetCaption(TArray<FString>& OutCaptions) const override;
+	virtual bool GenerateHLSLExpression(FMaterialHLSLGenerator& Generator, UE::HLSLTree::FScope& Scope, int32 OutputIndex, UE::HLSLTree::FExpression*& OutExpression) override;
 #endif
 	//~ End UMaterialExpression Interface
 
diff --git a/Engine/Source/Runtime/Engine/Classes/Materials/MaterialExpressionCeil.h b/Engine/Source/Runtime/Engine/Classes/Materials/MaterialExpressionCeil.h
--- a/Engine/Source/Runtime/Engine/Classes/Materials/MaterialExpressionCeil.h
+++ b/Engine/Source/Runtime/Engine/Classes/Materials/MaterialExpressionCeil.h
@@ -22,6 +22,7 @@ class UMaterialExpressionCeil : public UMaterialExpression
 #if WITH_EDITOR
 	virtual int32 Compile(class FMaterialCompiler* Compiler, int32 OutputIndex) override;
 	virtual void GetCaption(TArray<FString>& OutCaptions) const override;
+	virtual bool GenerateHLSLExpression(FMaterialHLSLGenerator& Generator, UE::HLSLTree::FScope& Scope, int32 OutputIndex, UE::HLSLTree::FExpression*& OutExpression) override;
 #endif
 	//~ End UMaterialExpression Interface
 };
diff --git a/Engine/Source/Runtime/Engine/Classes/Materials/MaterialExpressionFloor.h b/Engine/Source/Runtime/Engine/Classes/Materials/MaterialExpressionFloor.h
--- a/Engine/Source/Runtime/Engine/Classes/Materials/MaterialExpressionFloor.h
+++ b/Engine/Source/Runtime/Engine/Classes/Materials/MaterialExpressionFloor.h
@@ -22,6 +22,7 @@ class UMaterialExpressionFloor : public UMaterialExpression
 #if WITH_EDITOR
 	virtual int32 Compile(class FMaterialCompiler* Compiler, int32 OutputIndex) override;
 	virtual void GetCaption(TArray<FString>& OutCaptions) const override;
+	virtual bool GenerateHLSLExpression(FMaterialHLSLGenerator& Generator, UE::HLSLTree::FScope& Scope, int32 OutputIndex, UE::HLSLTree::FExpression*& OutExpression) override;
 #endif
 	//~ End UMaterialExpression Interface
 };
diff --git a/Engine/Source/Runtime/Engine/Classes/Materials/MaterialExpressionFontSample.h b/Engine/Source/Runtime/Engine/Classes/Materials/MaterialExpressionFontSample.h
--- a/Engine/Source/Runtime/Engine/Classes/Materials/MaterialExpressionFontSample.h
+++ b/Engine/Source/Runtime/Engine/Classes/Materials/MaterialExpressionFontSample.h
@@ -10,6 +10,14 @@
 
 class UTexture;
 
+namespace UE
+{
+namespace Shader
+{
+struct FTextureValue;
+}
+}
+
 UCLASS(collapsecategories, hidecategories=Object, MinimalAPI)
 class UMaterialExpressionFontSample : public UMaterialExpression
 {
@@ -31,6 +39,8 @@ class UMaterialExpressionFontSample : public UMaterialExpression
 	virtual bool MatchesSearchQuery(const TCHAR* SearchQuery) override;
 	virtual int32 GetWidth() const override;
 	virtual int32 GetLabelPadding() override { return 8; }
+	virtual bool GenerateHLSLExpression(FMaterialHLSLGenerator& Generator, UE::HLSLTree::FScope& Scope, int32 OutputIndex, UE::HLSLTree::FExpression*& OutExpression) override;
+	virtual UE::HLSLTree::FExpression* GenerateHLSLTextureExpression(FMaterialHLSLGenerator& Generator, const UE::Shader::FTextureValue* TextureValue);
 #endif
 
 	/** 
diff --git a/Engine/Source/Runtime/Engine/Classes/Materials/MaterialExpressionFontSampleParameter.h b/Engine/Source/Runtime/Engine/Classes/Materials/MaterialExpressionFontSampleParameter.h
--- a/Engine/Source/Runtime/Engine/Classes/Materials/MaterialExpressionFontSampleParameter.h
+++ b/Engine/Source/Runtime/Engine/Classes/Materials/MaterialExpressionFontSampleParameter.h
@@ -73,6 +73,7 @@ class UMaterialExpressionFontSampleParameter : public UMaterialExpressionFontSam
 		}
 		return false;
 	}
+	virtual UE::HLSLTree::FExpression* GenerateHLSLTextureExpression(FMaterialHLSLGenerator& Generator, const UE::Shader::FTextureValue* TextureValue) override;
 #endif
 	//~ End UMaterialExpression Interface
 	
diff --git a/Engine/Source/Runtime/Engine/Classes/Materials/MaterialExpressionFrac.h b/Engine/Source/Runtime/Engine/Classes/Materials/MaterialExpressionFrac.h
--- a/Engine/Source/Runtime/Engine/Classes/Materials/MaterialExpressionFrac.h
+++ b/Engine/Source/Runtime/Engine/Classes/Materials/MaterialExpressionFrac.h
@@ -22,6 +22,7 @@ class UMaterialExpressionFrac : public UMaterialExpression
 #if WITH_EDITOR
 	virtual int32 Compile(class FMaterialCompiler* Compiler, int32 OutputIndex) override;
 	virtual void GetCaption(TArray<FString>& OutCaptions) const override;
+	virtual bool GenerateHLSLExpression(FMaterialHLSLGenerator& Generator, UE::HLSLTree::FScope& Scope, int32 OutputIndex, UE::HLSLTree::FExpression*& OutExpression) override;
 #endif
 	//~ End UMaterialExpression Interface
 };
diff --git a/Engine/Source/Runtime/Engine/Classes/Materials/MaterialExpressionIf.h b/Engine/Source/Runtime/Engine/Classes/Materials/MaterialExpressionIf.h
--- a/Engine/Source/Runtime/Engine/Classes/Materials/MaterialExpressionIf.h
+++ b/Engine/Source/Runtime/Engine/Classes/Materials/MaterialExpressionIf.h
@@ -46,6 +46,7 @@ class UMaterialExpressionIf : public UMaterialExpression
 	virtual uint32 GetInputType(int32 InputIndex) override;
 	virtual uint32 GetOutputType(int32 InputIndex) override {return MCT_Unknown;}
 	virtual bool IsResultMaterialAttributes(int32 OutputIndex) override;
+	virtual bool GenerateHLSLExpression(FMaterialHLSLGenerator& Generator, UE::HLSLTree::FScope& Scope, int32 OutputIndex, UE::HLSLTree::FExpression*& OutExpression) override;
 #endif
 	//~ End UMaterialExpression Interface
 };
diff --git a/Engine/Source/Runtime/Engine/Classes/Materials/MaterialExpressionPower.h b/Engine/Source/Runtime/Engine/Classes/Materials/MaterialExpressionPower.h
--- a/Engine/Source/Runtime/Engine/Classes/Materials/MaterialExpressionPower.h
+++ b/Engine/Source/Runtime/Engine/Classes/Materials/MaterialExpressionPower.h
@@ -29,6 +29,7 @@ class UMaterialExpressionPower : public UMaterialExpression
 	virtual int32 Compile(class FMaterialCompiler* Compiler, int32 OutputIndex) override;
 	virtual void GetCaption(TArray<FString>& OutCaptions) const override;
 	virtual void GetExpressionToolTip(TArray<FString>& OutToolTip) override;
+	virtual bool GenerateHLSLExpression(FMaterialHLSLGenerator& Generator, UE::HLSLTree::FScope& Scope, int32 OutputIndex, UE::HLSLTree::FExpression*& OutExpression) override;
 #endif
 	//~ End UMaterialExpression Interface
 
diff --git a/Engine/Source/Runtime/Engine/Classes/Materials/MaterialExpressionReroute.h b/Engine/Source/Runtime/Engine/Classes/Materials/MaterialExpressionReroute.h
--- a/Engine/Source/Runtime/Engine/Classes/Materials/MaterialExpressionReroute.h
+++ b/Engine/Source/Runtime/Engine/Classes/Materials/MaterialExpressionReroute.h
@@ -25,6 +25,7 @@ class ENGINE_API UMaterialExpressionReroute : public UMaterialExpressionRerouteB
 
 	virtual FText GetCreationDescription() const override;
 	virtual FText GetCreationName() const override;
+	virtual bool GenerateHLSLExpression(FMaterialHLSLGenerator& Generator, UE::HLSLTree::FScope& Scope, int32 OutputIndex, UE::HLSLTree::FExpression*& OutExpression) override;
 #endif
 	//~ End UMaterialExpression Interface
 
diff --git a/Engine/Source/Runtime/Engine/Classes/Materials/MaterialExpressionRound.h b/Engine/Source/Runtime/Engine/Classes/Materials/MaterialExpressionRound.h
--- a/Engine/Source/Runtime/Engine/Classes/Materials/MaterialExpressionRound.h
+++ b/Engine/Source/Runtime/Engine/Classes/Materials/MaterialExpressionRound.h
@@ -21,6 +21,7 @@ class UMaterialExpressionRound : public UMaterialExpression
 	virtual int32 Compile(class FMaterialCompiler* Compiler, int32 OutputIndex) override;
 	virtual void GetCaption(TArray<FString>& OutCaptions) const override;
 	virtual void GetExpressionToolTip(TArray<FString>& OutToolTip) override;
+	virtual bool GenerateHLSLExpression(FMaterialHLSLGenerator& Generator, UE::HLSLTree::FScope& Scope, int32 OutputIndex, UE::HLSLTree::FExpression*& OutExpression) override;
 #endif
 	//~ End UMaterialExpression Interface
 };
diff --git a/Engine/Source/Runtime/Engine/Classes/Materials/MaterialExpressionSaturate.h b/Engine/Source/Runtime/Engine/Classes/Materials/MaterialExpressionSaturate.h
--- a/Engine/Source/Runtime/Engine/Classes/Materials/MaterialExpressionSaturate.h
+++ b/Engine/Source/Runtime/Engine/Classes/Materials/MaterialExpressionSaturate.h
@@ -21,6 +21,7 @@ class UMaterialExpressionSaturate : public UMaterialExpression
 	virtual int32 Compile(class FMaterialCompiler* Compiler, int32 OutputIndex) override;
 	virtual void GetCaption(TArray<FString>& OutCaptions) const override;
 	virtual void GetExpressionToolTip(TArray<FString>& OutToolTip) override;
+	virtual bool GenerateHLSLExpression(FMaterialHLSLGenerator& Generator, UE::HLSLTree::FScope& Scope, int32 OutputIndex, UE::HLSLTree::FExpression*& OutExpression) override;
 #endif
 	//~ End UMaterialExpression Interface
 };
diff --git a/Engine/Source/Runtime/Engine/Classes/Materials/MaterialExpressionSign.h b/Engine/Source/Runtime/Engine/Classes/Materials/MaterialExpressionSign.h
--- a/Engine/Source/Runtime/Engine/Classes/Materials/MaterialExpressionSign.h
+++ b/Engine/Source/Runtime/Engine/Classes/Materials/MaterialExpressionSign.h
@@ -21,6 +21,7 @@ class UMaterialExpressionSign : public UMaterialExpression
 	virtual int32 Compile(class FMaterialCompiler* Compiler, int32 OutputIndex) override;
 	virtual void GetCaption(TArray<FString>& OutCaptions) const override;
 	virtual void GetExpressionToolTip(TArray<FString>& OutToolTip) override;
+	virtual bool GenerateHLSLExpression(FMaterialHLSLGenerator& Generator, UE::HLSLTree::FScope& Scope, int32 OutputIndex, UE::HLSLTree::FExpression*& OutExpression) override;
 #endif
 	//~ End UMaterialExpression Interface
 };
diff --git a/Engine/Source/Runtime/Engine/Classes/Materials/MaterialExpressionTextureObject.h b/Engine/Source/Runtime/Engine/Classes/Materials/MaterialExpressionTextureObject.h
--- a/Engine/Source/Runtime/Engine/Classes/Materials/MaterialExpressionTextureObject.h
+++ b/Engine/Source/Runtime/Engine/Classes/Materials/MaterialExpressionTextureObject.h
@@ -28,7 +28,7 @@ class UMaterialExpressionTextureObject : public UMaterialExpressionTextureBase
 	virtual int32 CompilePreview(class FMaterialCompiler* Compiler, int32 OutputIndex) override;
 	virtual void GetCaption(TArray<FString>& OutCaptions) const override;
 	virtual uint32 GetOutputType(int32 OutputIndex) override;
-	virtual bool GenerateHLSLTexture(FMaterialHLSLGenerator& Generator, UE::HLSLTree::FScope& Scope, int32 OutputIndex, UE::HLSLTree::FTextureParameterDeclaration*& OutTexture) override;
+	virtual bool GenerateHLSLExpression(FMaterialHLSLGenerator& Generator, UE::HLSLTree::FScope& Scope, int32 OutputIndex, UE::HLSLTree::FExpression*& OutExpression) override;
 	//~ End UMaterialExpression Interface
 #endif // WITH_EDITOR
 };
diff --git a/Engine/Source/Runtime/Engine/Classes/Materials/MaterialExpressionTextureObjectParameter.h b/Engine/Source/Runtime/Engine/Classes/Materials/MaterialExpressionTextureObjectParameter.h
--- a/Engine/Source/Runtime/Engine/Classes/Materials/MaterialExpressionTextureObjectParameter.h
+++ b/Engine/Source/Runtime/Engine/Classes/Materials/MaterialExpressionTextureObjectParameter.h
@@ -30,7 +30,7 @@ class UMaterialExpressionTextureObjectParameter : public UMaterialExpressionText
 	virtual int32 CompilePreview(class FMaterialCompiler* Compiler, int32 OutputIndex) override;
 	virtual uint32 GetOutputType(int32 InputIndex) override {return MCT_Texture;}
 	virtual const TArray<FExpressionInput*> GetInputs() override;
-	virtual bool GenerateHLSLTexture(FMaterialHLSLGenerator& Generator, UE::HLSLTree::FScope& Scope, int32 OutputIndex, UE::HLSLTree::FTextureParameterDeclaration*& OutTexture) override;
+	virtual bool GenerateHLSLExpression(FMaterialHLSLGenerator& Generator, UE::HLSLTree::FScope& Scope, int32 OutputIndex, UE::HLSLTree::FExpression*& OutExpression) override;
 #endif
 	//~ End UMaterialExpression Interface
 };
diff --git a/Engine/Source/Runtime/Engine/Classes/Materials/MaterialExpressionTextureSample.h b/Engine/Source/Runtime/Engine/Classes/Materials/MaterialExpressionTextureSample.h
--- a/Engine/Source/Runtime/Engine/Classes/Materials/MaterialExpressionTextureSample.h
+++ b/Engine/Source/Runtime/Engine/Classes/Materials/MaterialExpressionTextureSample.h
@@ -64,7 +64,7 @@ protected:
 	uint8 bShowTextureInputPin : 1;
 
 #if WITH_EDITOR
-	bool GenerateHLSLExpressionBase(FMaterialHLSLGenerator& Generator, UE::HLSLTree::FScope& Scope, UE::HLSLTree::FTextureParameterDeclaration* TextureDeclaration, UE::HLSLTree::FExpression*& OutExpression);
+	bool GenerateHLSLExpressionBase(FMaterialHLSLGenerator& Generator, UE::HLSLTree::FScope& Scope, UE::HLSLTree::FExpression* TextureExpression, UE::HLSLTree::FExpression*& OutExpression);
 #endif
 
 public:
diff --git a/Engine/Source/Runtime/Engine/Classes/Materials/MaterialExpressionTruncate.h b/Engine/Source/Runtime/Engine/Classes/Materials/MaterialExpressionTruncate.h
--- a/Engine/Source/Runtime/Engine/Classes/Materials/MaterialExpressionTruncate.h
+++ b/Engine/Source/Runtime/Engine/Classes/Materials/MaterialExpressionTruncate.h
@@ -21,6 +21,7 @@ class UMaterialExpressionTruncate : public UMaterialExpression
 	virtual int32 Compile(class FMaterialCompiler* Compiler, int32 OutputIndex) override;
 	virtual void GetCaption(TArray<FString>& OutCaptions) const override;
 	virtual void GetExpressionToolTip(TArray<FString>& OutToolTip) override;
+	virtual bool GenerateHLSLExpression(FMaterialHLSLGenerator& Generator, UE::HLSLTree::FScope& Scope, int32 OutputIndex, UE::HLSLTree::FExpression*& OutExpression) override;
 #endif
 	//~ End UMaterialExpression Interface
 };
diff --git a/Engine/Source/Runtime/Engine/Classes/Materials/MaterialExpressionVertexColor.h b/Engine/Source/Runtime/Engine/Classes/Materials/MaterialExpressionVertexColor.h
--- a/Engine/Source/Runtime/Engine/Classes/Materials/MaterialExpressionVertexColor.h
+++ b/Engine/Source/Runtime/Engine/Classes/Materials/MaterialExpressionVertexColor.h
@@ -19,6 +19,7 @@ class UMaterialExpressionVertexColor : public UMaterialExpression
 	virtual int32 Compile(class FMaterialCompiler* Compiler, int32 OutputIndex) override;
 	virtual void GetCaption(TArray<FString>& OutCaptions) const override;
 	virtual FText GetKeywords() const override { return FText::FromString(TEXT("vertex color")); }
+	virtual bool GenerateHLSLExpression(FMaterialHLSLGenerator& Generator, UE::HLSLTree::FScope& Scope, int32 OutputIndex, UE::HLSLTree::FExpression*& OutExpression) override;
 #endif
 	//~ End UMaterialExpression Interface
 };
diff --git a/Engine/Source/Runtime/Engine/Private/HLSLTree/HLSLTree.cpp b/Engine/Source/Runtime/Engine/Private/HLSLTree/HLSLTree.cpp
--- a/Engine/Source/Runtime/Engine/Private/HLSLTree/HLSLTree.cpp
+++ b/Engine/Source/Runtime/Engine/Private/HLSLTree/HLSLTree.cpp
@@ -93,11 +93,18 @@ FOperationDescription GetOperationDescription(EOperation Op)
 {
 	switch (Op)
 	{
-	// Unary
 	case EOperation::None: return FOperationDescription(TEXT("None"), TEXT(""), 0, Shader::EPreshaderOpcode::Nop); break;
+	// Unary
+	case EOperation::Abs: return FOperationDescription(TEXT("Abs"), TEXT("abs"), 1, Shader::EPreshaderOpcode::Abs); break;
 	case EOperation::Neg: return FOperationDescription(TEXT("Neg"), TEXT("-"), 1, Shader::EPreshaderOpcode::Neg); break;
 	case EOperation::Rcp: return FOperationDescription(TEXT("Rcp"), TEXT("/"), 1, Shader::EPreshaderOpcode::Rcp); break;
 	case EOperation::Frac: return FOperationDescription(TEXT("Frac"), TEXT("frac"), 1, Shader::EPreshaderOpcode::Frac); break;
+	case EOperation::Floor: return FOperationDescription(TEXT("Floor"), TEXT("floor"), 1, Shader::EPreshaderOpcode::Floor); break;
+	case EOperation::Ceil: return FOperationDescription(TEXT("Ceil"), TEXT("ceil"), 1, Shader::EPreshaderOpcode::Ceil); break;
+	case EOperation::Round: return FOperationDescription(TEXT("Round"), TEXT("round"), 1, Shader::EPreshaderOpcode::Round); break;
+	case EOperation::Trunc: return FOperationDescription(TEXT("Trunc"), TEXT("trunc"), 1, Shader::EPreshaderOpcode::Trunc); break;
+	case EOperation::Saturate: return FOperationDescription(TEXT("Saturate"), TEXT("saturate"), 1, Shader::EPreshaderOpcode::Saturate); break;
+	case EOperation::Sign: return FOperationDescription(TEXT("Sign"), TEXT("sign"), 1, Shader::EPreshaderOpcode::Sign); break;
 	case EOperation::Length: return FOperationDescription(TEXT("Length"), TEXT("length"), 1, Shader::EPreshaderOpcode::Length); break;
 	case EOperation::Normalize: return FOperationDescription(TEXT("Normalize"), TEXT("normalize"), 1, Shader::EPreshaderOpcode::Normalize); break;
 	// Binary
@@ -106,11 +113,14 @@ FOperationDescription GetOperationDescription(EOperation Op)
 	case EOperation::Mul: return FOperationDescription(TEXT("Multiply"), TEXT("*"), 2, Shader::EPreshaderOpcode::Mul); break;
 	case EOperation::Div: return FOperationDescription(TEXT("Divide"), TEXT("/"), 2, Shader::EPreshaderOpcode::Div); break;
 	case EOperation::Fmod: return FOperationDescription(TEXT("Fmod"), TEXT("%"), 2, Shader::EPreshaderOpcode::Fmod); break;
+	case EOperation::PowPositiveClamped: return FOperationDescription(TEXT("PowPositiveClamped"), TEXT("PowPositiveClamped"), 2, Shader::EPreshaderOpcode::Nop); break;
 	case EOperation::Dot: return FOperationDescription(TEXT("Dot"), TEXT("dot"), 2, Shader::EPreshaderOpcode::Dot); break;
 	case EOperation::Min: return FOperationDescription(TEXT("Min"), TEXT("min"), 2, Shader::EPreshaderOpcode::Min); break;
 	case EOperation::Max: return FOperationDescription(TEXT("Max"), TEXT("max"), 2, Shader::EPreshaderOpcode::Max); break;
 	case EOperation::Less: return FOperationDescription(TEXT("Less"), TEXT("<"), 2, Shader::EPreshaderOpcode::Less); break;
 	case EOperation::Greater: return FOperationDescription(TEXT("Greater"), TEXT(">"), 2, Shader::EPreshaderOpcode::Greater); break;
+	case EOperation::LessEqual: return FOperationDescription(TEXT("LessEqual"), TEXT("<="), 2, Shader::EPreshaderOpcode::Nop); break;
+	case EOperation::GreaterEqual: return FOperationDescription(TEXT("GreaterEqual"), TEXT(">="), 2, Shader::EPreshaderOpcode::Nop); break;
 	case EOperation::VecMulMatrix3: return FOperationDescription(TEXT("VecMulMatrix3"), TEXT("mul"), 2, Shader::EPreshaderOpcode::Nop); break;
 	case EOperation::VecMulMatrix4: return FOperationDescription(TEXT("VecMulMatrix4"), TEXT("mul"), 2, Shader::EPreshaderOpcode::Nop); break;
 	case EOperation::Matrix3MulVec: return FOperationDescription(TEXT("Matrix3MulVec"), TEXT("mul"), 2, Shader::EPreshaderOpcode::Nop); break;
@@ -970,6 +980,11 @@ FPreparedType MergePreparedTypes(const FPreparedType& Lhs, const FPreparedType&
 
 bool FPrepareValueResult::TryMergePreparedType(FEmitContext& Context, const Shader::FStructType* StructType, Shader::EValueComponentType ComponentType)
 {
+	if (!StructType && ComponentType == Shader::EValueComponentType::Void)
+	{
+		return false;
+	}
+
 	// If we previously had a forwarded value set, reset that and start over
 	if (!PreparedType.IsInitialized())
 	{
@@ -989,10 +1004,6 @@ bool FPrepareValueResult::TryMergePreparedType(FEmitContext& Context, const Shad
 	}
 	else
 	{
-		if (ComponentType == Shader::EValueComponentType::Void)
-		{
-			return false;
-		}
 		PreparedType.ValueComponentType = Shader::CombineComponentTypes(PreparedType.ValueComponentType, ComponentType);
 	}
 
@@ -1579,11 +1590,5 @@ FFunction* FTree::NewFunction()
 	return NewFunction;
 }
 
-FTextureParameterDeclaration* FTree::NewTextureParameterDeclaration(const FName& Name, const FTextureDescription& DefaultValue)
-{
-	FTextureParameterDeclaration* Declaration = NewNode<FTextureParameterDeclaration>(Name, DefaultValue);
-	return Declaration;
-}
-
 } // namespace HLSLTree
 } // namespace UE
diff --git a/Engine/Source/Runtime/Engine/Private/HLSLTree/HLSLTreeCommon.cpp b/Engine/Source/Runtime/Engine/Private/HLSLTree/HLSLTreeCommon.cpp
--- a/Engine/Source/Runtime/Engine/Private/HLSLTree/HLSLTreeCommon.cpp
+++ b/Engine/Source/Runtime/Engine/Private/HLSLTree/HLSLTreeCommon.cpp
@@ -4,6 +4,7 @@
 #include "Misc/StringBuilder.h"
 #include "MaterialShared.h"
 #include "MaterialSceneTextureId.h"
+#include "Engine/BlendableInterface.h" // BL_AfterTonemapping
 #include "Engine/Texture.h"
 
 namespace UE
@@ -50,6 +51,10 @@ FExternalInputDescription GetExternalInputDescription(EExternalInput Input)
 	case EExternalInput::TexCoord6_Ddy: return FExternalInputDescription(TEXT("TexCoord6_Ddy"), Shader::EValueType::Float2);
 	case EExternalInput::TexCoord7_Ddy: return FExternalInputDescription(TEXT("TexCoord7_Ddy"), Shader::EValueType::Float2);
 
+	case EExternalInput::VertexColor: return FExternalInputDescription(TEXT("VertexColor"), Shader::EValueType::Float4, EExternalInput::VertexColor_Ddx, EExternalInput::VertexColor_Ddy);
+	case EExternalInput::VertexColor_Ddx: return FExternalInputDescription(TEXT("VertexColor_Ddx"), Shader::EValueType::Float4);
+	case EExternalInput::VertexColor_Ddy: return FExternalInputDescription(TEXT("VertexColor_Ddy"), Shader::EValueType::Float4);
+
 	case EExternalInput::WorldPosition: return FExternalInputDescription(TEXT("WorldPosition"), Shader::EValueType::Double3, EExternalInput::WorldPosition_Ddx, EExternalInput::WorldPosition_Ddy, EExternalInput::PrevWorldPosition);
 	case EExternalInput::WorldPosition_NoOffsets: return FExternalInputDescription(TEXT("WorldPosition_NoOffsets"), Shader::EValueType::Double3, EExternalInput::WorldPosition_Ddx, EExternalInput::WorldPosition_Ddy, EExternalInput::PrevWorldPosition_NoOffsets);
 	case EExternalInput::TranslatedWorldPosition: return FExternalInputDescription(TEXT("TranslatedWorldPosition"), Shader::EValueType::Float3, EExternalInput::WorldPosition_Ddx, EExternalInput::WorldPosition_Ddy, EExternalInput::PrevTranslatedWorldPosition);
@@ -216,8 +221,75 @@ void FExpressionMaterialParameter::ComputeAnalyticDerivatives(FTree& Tree, FExpr
 
 bool FExpressionMaterialParameter::PrepareValue(FEmitContext& Context, FEmitScope& Scope, const FRequestedType& RequestedType, FPrepareValueResult& OutResult) const
 {
-	const EExpressionEvaluation Evaluation = IsStaticMaterialParameter(ParameterType) ? EExpressionEvaluation::Constant : EExpressionEvaluation::Preshader;
-	return OutResult.SetType(Context, RequestedType, Evaluation, GetShaderValueType(ParameterType));
+	EExpressionEvaluation Evaluation = EExpressionEvaluation::Shader;
+	if (IsStaticMaterialParameter(ParameterType))
+	{
+		Evaluation = EExpressionEvaluation::Constant;
+	}
+	else if (ParameterType == EMaterialParameterType::Scalar ||
+		ParameterType == EMaterialParameterType::Vector ||
+		ParameterType == EMaterialParameterType::DoubleVector)
+	{
+		Evaluation = EExpressionEvaluation::Preshader;
+	}
+	return OutResult.SetType(Context, RequestedType, Evaluation, DefaultValue.Type);
+}
+
+void FExpressionMaterialParameter::EmitValueShader(FEmitContext& Context, FEmitScope& Scope, const FRequestedType& RequestedType, FEmitValueShaderResult& OutResult) const
+{
+	if (ParameterType == EMaterialParameterType::Texture)
+	{
+		const Shader::FTextureValue* TextureValue = DefaultValue.AsTexture();
+		check(TextureValue);
+
+		FMaterialTextureParameterInfo TextureParameterInfo;
+		TextureParameterInfo.ParameterInfo = ParameterName;
+		TextureParameterInfo.TextureIndex = Context.Material->GetReferencedTextures().Find(TextureValue->Texture);
+		TextureParameterInfo.SamplerSource = SSM_FromTextureAsset; // TODO - Is this needed?
+		check(TextureParameterInfo.TextureIndex != INDEX_NONE);
+
+		const Shader::EValueType TextureType = TextureValue->GetType();
+		EMaterialTextureParameterType TextureParameterType = EMaterialTextureParameterType::Count;
+		const TCHAR* ConstructorName = nullptr;
+		const TCHAR* TextureTypeName = nullptr;
+		switch (TextureType)
+		{
+		case Shader::EValueType::Texture2D:
+			TextureParameterType = EMaterialTextureParameterType::Standard2D;
+			ConstructorName = TEXT("MakeTexture2D");
+			TextureTypeName = TEXT("Texture2D");
+			break;
+		case Shader::EValueType::Texture2DArray:
+			TextureParameterType = EMaterialTextureParameterType::Array2D;
+			ConstructorName = TEXT("MakeTexture2DArray");
+			TextureTypeName = TEXT("Texture2DArray");
+			break;
+		case Shader::EValueType::TextureCube:
+			TextureParameterType = EMaterialTextureParameterType::Cube;
+			ConstructorName = TEXT("MakeTextureCube");
+			TextureTypeName = TEXT("TextureCube");
+			break;
+		case Shader::EValueType::TextureCubeArray:
+			TextureParameterType = EMaterialTextureParameterType::ArrayCube;
+			ConstructorName = TEXT("MakeTextureCubeArray");
+			TextureTypeName = TEXT("TextureCubeArray");
+			break;
+		case Shader::EValueType::Texture3D:
+			TextureParameterType = EMaterialTextureParameterType::Volume;
+			ConstructorName = TEXT("MakeTexture3D");
+			TextureTypeName = TEXT("VolumeTexture");
+			break;
+		default:
+			checkNoEntry();
+			break;
+		}
+
+		const int32 ParameterIndex = Context.MaterialCompilationOutput->UniformExpressionSet.FindOrAddTextureParameter(TextureParameterType, TextureParameterInfo);
+
+		TStringBuilder<256> FormattedCode;
+		FormattedCode.Appendf(TEXT("%s(Material.%s_%d, Material.%s_%dSampler, %d)"), ConstructorName, TextureTypeName, ParameterIndex, TextureTypeName, ParameterIndex, (int32)TextureValue->SamplerType);
+		OutResult.Code = Context.EmitExpression(Scope, TextureType, FormattedCode.ToView());
+	}
 }
 
 void FExpressionMaterialParameter::EmitValuePreshader(FEmitContext& Context, FEmitScope& Scope, const FRequestedType& RequestedType, FEmitValuePreshaderResult& OutResult) const
@@ -327,6 +399,9 @@ void FExpressionExternalInput::EmitValueShader(FEmitContext& Context, FEmitScope
 		const TCHAR* Code = nullptr;
 		switch (InputType)
 		{
+		case EExternalInput::VertexColor: Code = TEXT("Parameters.VertexColor"); Context.bUsesVertexColor = (Context.ShaderFrequency != SF_Vertex); break;
+		case EExternalInput::VertexColor_Ddx: Code = TEXT("Parameters.VertexColor_DDX"); Context.bUsesVertexColor = (Context.ShaderFrequency != SF_Vertex); break;
+		case EExternalInput::VertexColor_Ddy: Code = TEXT("Parameters.VertexColor_DDY"); Context.bUsesVertexColor = (Context.ShaderFrequency != SF_Vertex); break;
 		case EExternalInput::WorldPosition: Code = TEXT("GetWorldPosition(Parameters)"); break;
 		case EExternalInput::WorldPosition_NoOffsets: Code = TEXT("GetWorldPosition_NoMaterialOffsets(Parameters)"); break;
 		case EExternalInput::TranslatedWorldPosition: Code = TEXT("GetTranslatedWorldPosition(Parameters)"); break;
@@ -430,7 +505,18 @@ void FExpressionMaterialSceneTexture::EmitValueShader(FEmitContext& Context, FEm
 
 bool FExpressionTextureSample::PrepareValue(FEmitContext& Context, FEmitScope& Scope, const FRequestedType& RequestedType, FPrepareValueResult& OutResult) const
 {
-	Context.PrepareExpression(TexCoordExpression, Scope, ERequestedType::Vector2);
+	const FPreparedType& TextureType = Context.PrepareExpression(TextureExpression, Scope, ERequestedType::Texture);
+	if (!Shader::IsTextureType(TextureType.ValueComponentType))
+	{
+		return Context.Errors->AddError(TEXT("Expected texture"));
+	}
+
+	const FPreparedType& TexCoordType = Context.PrepareExpression(TexCoordExpression, Scope, ERequestedType::Vector2);
+	if (TexCoordType.IsVoid())
+	{
+		return false;
+	}
+
 	Context.PrepareExpression(TexCoordDerivatives.ExpressionDdx, Scope, ERequestedType::Vector2);
 	Context.PrepareExpression(TexCoordDerivatives.ExpressionDdy, Scope, ERequestedType::Vector2);
 
@@ -439,156 +525,88 @@ bool FExpressionTextureSample::PrepareValue(FEmitContext& Context, FEmitScope& S
 
 void FExpressionTextureSample::EmitValueShader(FEmitContext& Context, FEmitScope& Scope, const FRequestedType& RequestedType, FEmitValueShaderResult& OutResult) const
 {
-	const FTextureDescription& Desc = Declaration->Description;
-	const EMaterialValueType MaterialValueType = Desc.Texture->GetMaterialType();
-	EMaterialTextureParameterType TextureType = EMaterialTextureParameterType::Count;
+	FEmitShaderExpression* EmitTexture = TextureExpression->GetValueShader(Context, Scope);
+	const Shader::EValueType TextureType = EmitTexture->Type;
+	check(Shader::IsTextureType(TextureType));
+
 	bool bVirtualTexture = false;
-	const TCHAR* TextureTypeName = nullptr;
 	const TCHAR* SampleFunctionName = nullptr;
-	switch (MaterialValueType)
+	switch (TextureType)
 	{
-	case MCT_Texture2D:
-		TextureType = EMaterialTextureParameterType::Standard2D;
-		TextureTypeName = TEXT("Texture2D");
+	case Shader::EValueType::Texture2D:
 		SampleFunctionName = TEXT("Texture2DSample");
 		break;
-	case MCT_TextureCube:
-		TextureType = EMaterialTextureParameterType::Cube;
-		TextureTypeName = TEXT("TextureCube");
+	case Shader::EValueType::TextureCube:
 		SampleFunctionName = TEXT("TextureCubeSample");
 		break;
-	case MCT_Texture2DArray:
-		TextureType = EMaterialTextureParameterType::Array2D;
-		TextureTypeName = TEXT("Texture2DArray");
+	case Shader::EValueType::Texture2DArray:
 		SampleFunctionName = TEXT("Texture2DArraySample");
 		break;
-	case MCT_VolumeTexture:
-		TextureType = EMaterialTextureParameterType::Volume;
-		TextureTypeName = TEXT("VolumeTexture");
+	case Shader::EValueType::Texture3D:
 		SampleFunctionName = TEXT("Texture3DSample");
 		break;
-	case MCT_TextureExternal:
+	/*case MCT_TextureExternal:
 		// TODO
-		TextureTypeName = TEXT("ExternalTexture");
 		SampleFunctionName = TEXT("TextureExternalSample");
 		break;
 	case MCT_TextureVirtual:
 		// TODO
-		TextureType = EMaterialTextureParameterType::Virtual;
 		SampleFunctionName = TEXT("TextureVirtualSample");
 		bVirtualTexture = true;
-		break;
+		break;*/
 	default:
 		checkNoEntry();
 		break;
 	}
 
-	const TCHAR* SamplerTypeFunction = TEXT("");
-	switch (Desc.SamplerType)
+	const bool AutomaticViewMipBias = false; // TODO
+	TStringBuilder<256> FormattedSampler;
+	switch (SamplerSource)
 	{
-	case SAMPLERTYPE_External:
-		SamplerTypeFunction = TEXT("ProcessMaterialExternalTextureLookup");
-		break;
-	case SAMPLERTYPE_Color:
-		SamplerTypeFunction = TEXT("ProcessMaterialColorTextureLookup");
-		break;
-	case SAMPLERTYPE_VirtualColor:
-		// has a mobile specific workaround
-		SamplerTypeFunction = TEXT("ProcessMaterialVirtualColorTextureLookup");
-		break;
-	case SAMPLERTYPE_LinearColor:
-	case SAMPLERTYPE_VirtualLinearColor:
-		SamplerTypeFunction = TEXT("ProcessMaterialLinearColorTextureLookup");
-		break;
-	case SAMPLERTYPE_Alpha:
-	case SAMPLERTYPE_VirtualAlpha:
-	case SAMPLERTYPE_DistanceFieldFont:
-		SamplerTypeFunction = TEXT("ProcessMaterialAlphaTextureLookup");
-		break;
-	case SAMPLERTYPE_Grayscale:
-	case SAMPLERTYPE_VirtualGrayscale:
-		SamplerTypeFunction = TEXT("ProcessMaterialGreyscaleTextureLookup");
+	case SSM_FromTextureAsset:
+		FormattedSampler.Appendf(TEXT("%s.Sampler"), EmitTexture->Reference);
 		break;
-	case SAMPLERTYPE_LinearGrayscale:
-	case SAMPLERTYPE_VirtualLinearGrayscale:
-		SamplerTypeFunction = TEXT("ProcessMaterialLinearGreyscaleTextureLookup");
+	case SSM_Wrap_WorldGroupSettings:
+		FormattedSampler.Appendf(TEXT("GetMaterialSharedSampler(%s.Sampler,%s)"),
+			EmitTexture->Reference,
+			AutomaticViewMipBias ? TEXT("View.MaterialTextureBilinearWrapedSampler") : TEXT("Material.Wrap_WorldGroupSettings"));
 		break;
-	case SAMPLERTYPE_Normal:
-	case SAMPLERTYPE_VirtualNormal:
-		// Normal maps need to be unpacked in the pixel shader.
-		SamplerTypeFunction = TEXT("UnpackNormalMap");
-		break;
-	case SAMPLERTYPE_Masks:
-	case SAMPLERTYPE_VirtualMasks:
-	case SAMPLERTYPE_Data:
-		SamplerTypeFunction = TEXT("");
+	case SSM_Clamp_WorldGroupSettings:
+		FormattedSampler.Appendf(TEXT("GetMaterialSharedSampler(%s.Sampler,%s)"),
+			EmitTexture->Reference,
+			AutomaticViewMipBias ? TEXT("View.MaterialTextureBilinearClampedSampler") : TEXT("Material.Clamp_WorldGroupSettings"));
 		break;
 	default:
-		check(0);
+		checkNoEntry();
 		break;
 	}
 
-	FMaterialTextureParameterInfo TextureParameterInfo;
-	TextureParameterInfo.ParameterInfo = Declaration->Name;
-	TextureParameterInfo.TextureIndex = Context.Material->GetReferencedTextures().Find(Desc.Texture);
-	TextureParameterInfo.SamplerSource = SamplerSource;
-	check(TextureParameterInfo.TextureIndex != INDEX_NONE);
-
-	const int32 ParameterIndex = Context.MaterialCompilationOutput->UniformExpressionSet.FindOrAddTextureParameter(TextureType, TextureParameterInfo);
-	const FString TextureName = FString::Printf(TEXT("Material.%s_%u"), TextureTypeName, ParameterIndex);
-
-	FString SamplerStateCode;
-	bool AutomaticViewMipBias = false; // TODO
-	bool RequiresManualViewMipBias = AutomaticViewMipBias;
-
-	if (!bVirtualTexture) //VT does not have explict samplers (and always requires manual view mip bias)
-	{
-		if (SamplerSource == SSM_FromTextureAsset)
-		{
-			SamplerStateCode = FString::Printf(TEXT("%sSampler"), *TextureName);
-		}
-		else if (SamplerSource == SSM_Wrap_WorldGroupSettings)
-		{
-			// Use the shared sampler to save sampler slots
-			const TCHAR* SharedSamplerName = AutomaticViewMipBias ? TEXT("View.MaterialTextureBilinearWrapedSampler") : TEXT("Material.Wrap_WorldGroupSettings");
-			SamplerStateCode = FString::Printf(TEXT("GetMaterialSharedSampler(%sSampler,%s)"), *TextureName, SharedSamplerName);
-			RequiresManualViewMipBias = false;
-		}
-		else if (SamplerSource == SSM_Clamp_WorldGroupSettings)
-		{
-			// Use the shared sampler to save sampler slots
-			const TCHAR* SharedSamplerName = AutomaticViewMipBias ? TEXT("View.MaterialTextureBilinearClampedSampler") : TEXT("Material.Clamp_WorldGroupSettings");
-			SamplerStateCode = FString::Printf(TEXT("GetMaterialSharedSampler(%sSampler,%s)"), *TextureName, SharedSamplerName);
-			RequiresManualViewMipBias = false;
-		}
-	}
-
 	FEmitShaderExpression* TexCoordValue = TexCoordExpression->GetValueShader(Context, Scope, Shader::EValueType::Float2);
 	FEmitShaderExpression* TexCoordValueDdx = nullptr;
 	FEmitShaderExpression* TexCoordValueDdy = nullptr;
+	FEmitShaderExpression* TextureResult = nullptr;
 	if (TexCoordDerivatives.IsValid())
 	{
 		TexCoordValueDdx = TexCoordDerivatives.ExpressionDdx->GetValueShader(Context, Scope, Shader::EValueType::Float2);
 		TexCoordValueDdy = TexCoordDerivatives.ExpressionDdy->GetValueShader(Context, Scope, Shader::EValueType::Float2);
-		OutResult.Code = Context.EmitExpression(Scope, Shader::EValueType::Float4, TEXT("%(%Grad(%, %, %, %, %))"),
-			SamplerTypeFunction,
+		TextureResult = Context.EmitExpression(Scope, Shader::EValueType::Float4, TEXT("%Grad(%.Texture, %, %, %, %)"),
 			SampleFunctionName,
-			*TextureName,
-			*SamplerStateCode,
+			EmitTexture,
+			FormattedSampler.ToString(),
 			TexCoordValue,
 			TexCoordValueDdx,
 			TexCoordValueDdy);
 	}
 	else
 	{
-		OutResult.Code = Context.EmitExpression(Scope, Shader::EValueType::Float4, TEXT("%(%(%, %, %))"),
-			SamplerTypeFunction,
+		TextureResult = Context.EmitExpression(Scope, Shader::EValueType::Float4, TEXT("%(%.Texture, %, %)"),
 			SampleFunctionName,
-			*TextureName,
-			*SamplerStateCode,
+			EmitTexture,
+			FormattedSampler.ToString(),
 			TexCoordValue);
 	}
-	
+
+	OutResult.Code = Context.EmitExpression(Scope, Shader::EValueType::Float4, TEXT("ApplyMaterialSamplerType(%, %.SamplerType)"), TextureResult, EmitTexture);
 }
 
 void FExpressionGetStructField::ComputeAnalyticDerivatives(FTree& Tree, FExpressionDerivatives& OutResult) const
@@ -837,6 +855,8 @@ void FExpressionOperation::ComputeAnalyticDerivatives(FTree& Tree, FExpressionDe
 	{
 	case EOperation::Less:
 	case EOperation::Greater:
+	case EOperation::LessEqual:
+	case EOperation::GreaterEqual:
 		OutResult.ExpressionDdx = Tree.NewConstant(0.0f);
 		OutResult.ExpressionDdx = OutResult.ExpressionDdy;
 		break;
@@ -879,6 +899,14 @@ void FExpressionOperation::ComputeAnalyticDerivatives(FTree& Tree, FExpressionDe
 		break;
 	case EOperation::Length:
 	case EOperation::Normalize:
+	case EOperation::Floor:
+	case EOperation::Ceil:
+	case EOperation::Round:
+	case EOperation::Trunc:
+	case EOperation::Sign:
+	case EOperation::Abs:
+	case EOperation::Saturate:
+	case EOperation::PowPositiveClamped:
 		// TODO
 		break;
 	case EOperation::Add:
@@ -1037,7 +1065,7 @@ bool FExpressionOperation::PrepareValue(FEmitContext& Context, FEmitScope& Scope
 			return Context.Errors->AddErrorf(TEXT("Invalid arithmetic between %s and %s"), InputType[0].GetType().GetName(), InputType[1].GetType().GetName());
 		}
 		ResultType = MergePreparedTypes(InputType[0], InputType[1]);
-		if (Op == EOperation::Less || Op == EOperation::Greater)
+		if (Op == EOperation::Less || Op == EOperation::Greater || Op == EOperation::LessEqual || Op == EOperation::GreaterEqual)
 		{
 			ResultType.ValueComponentType = Shader::EValueComponentType::Bool;
 		}
@@ -1095,15 +1123,23 @@ FOperationTypes GetOperationTypes(EOperation Op, TConstArrayView<Shader::EValueT
 	case EOperation::Normalize:
 	case EOperation::Frac:
 	case EOperation::Rcp:
+	case EOperation::Sign:
+	case EOperation::Saturate:
 		Types.ResultType = Shader::MakeNonLWCType(IntermediateType);
 		break;
 	case EOperation::Less:
 	case EOperation::Greater:
+	case EOperation::LessEqual:
+	case EOperation::GreaterEqual:
 		Types.ResultType = Shader::MakeValueType(Shader::EValueComponentType::Bool, NumComponents);
 		break;
 	case EOperation::Fmod:
 		Types.InputType[1] = Types.ResultType = Shader::MakeNonLWCType(IntermediateType);
 		break;
+	case EOperation::PowPositiveClamped:
+		// No LWC support yet
+		Types.InputType[0] = Types.InputType[1] = Types.ResultType = Shader::MakeNonLWCType(IntermediateType);
+		break;
 	case EOperation::Dot:
 		Types.ResultType = Shader::MakeValueType(ComponentType, 1);
 		break;
@@ -1149,6 +1185,16 @@ void FExpressionOperation::EmitValueShader(FEmitContext& Context, FEmitScope& Sc
 
 	switch (Op)
 	{
+	case EOperation::Abs:
+		if (Types.bIsLWC)
+		{
+			OutResult.Code = Context.EmitExpression(Scope, Types.ResultType, TEXT("LWCAbs(%)"), InputValue[0]);
+		}
+		else
+		{
+			OutResult.Code = Context.EmitInlineExpression(Scope, Types.ResultType, TEXT("abs(%)"), InputValue[0]);
+		}
+		break;
 	case EOperation::Neg:
 		if (Types.bIsLWC)
 		{
@@ -1179,6 +1225,67 @@ void FExpressionOperation::EmitValueShader(FEmitContext& Context, FEmitScope& Sc
 			OutResult.Code = Context.EmitExpression(Scope, Types.ResultType, TEXT("frac(%)"), InputValue[0]);
 		}
 		break;
+
+	case EOperation::Floor:
+		if (Types.bIsLWC)
+		{
+			OutResult.Code = Context.EmitExpression(Scope, Types.ResultType, TEXT("LWCFloor(%)"), InputValue[0]);
+		}
+		else
+		{
+			OutResult.Code = Context.EmitExpression(Scope, Types.ResultType, TEXT("floor(%)"), InputValue[0]);
+		}
+		break;
+	case EOperation::Ceil:
+		if (Types.bIsLWC)
+		{
+			OutResult.Code = Context.EmitExpression(Scope, Types.ResultType, TEXT("LWCCeil(%)"), InputValue[0]);
+		}
+		else
+		{
+			OutResult.Code = Context.EmitExpression(Scope, Types.ResultType, TEXT("ceil(%)"), InputValue[0]);
+		}
+		break;
+	case EOperation::Round:
+		if (Types.bIsLWC)
+		{
+			OutResult.Code = Context.EmitExpression(Scope, Types.ResultType, TEXT("LWCRound(%)"), InputValue[0]);
+		}
+		else
+		{
+			OutResult.Code = Context.EmitExpression(Scope, Types.ResultType, TEXT("round(%)"), InputValue[0]);
+		}
+		break;
+	case EOperation::Trunc:
+		if (Types.bIsLWC)
+		{
+			OutResult.Code = Context.EmitExpression(Scope, Types.ResultType, TEXT("LWCTrunc(%)"), InputValue[0]);
+		}
+		else
+		{
+			OutResult.Code = Context.EmitExpression(Scope, Types.ResultType, TEXT("trunc(%)"), InputValue[0]);
+		}
+		break;
+	case EOperation::Saturate:
+		if (Types.bIsLWC)
+		{
+			OutResult.Code = Context.EmitExpression(Scope, Types.ResultType, TEXT("LWCSaturate(%)"), InputValue[0]);
+		}
+		else
+		{
+			OutResult.Code = Context.EmitExpression(Scope, Types.ResultType, TEXT("saturate(%)"), InputValue[0]);
+		}
+		break;
+	case EOperation::Sign:
+		if (Types.bIsLWC)
+		{
+			OutResult.Code = Context.EmitExpression(Scope, Types.ResultType, TEXT("LWCSign(%)"), InputValue[0]);
+		}
+		else
+		{
+			OutResult.Code = Context.EmitExpression(Scope, Types.ResultType, TEXT("sign(%)"), InputValue[0]);
+		}
+		break;
 	case EOperation::Length:
 		if (Types.bIsLWC)
 		{
@@ -1249,6 +1356,9 @@ void FExpressionOperation::EmitValueShader(FEmitContext& Context, FEmitScope& Sc
 			OutResult.Code = Context.EmitExpression(Scope, Types.ResultType, TEXT("fmod(%, %)"), InputValue[0], InputValue[1]);
 		}
 		break;
+	case EOperation::PowPositiveClamped:
+		OutResult.Code = Context.EmitExpression(Scope, Types.ResultType, TEXT("PositiveClampedPow(%, %)"), InputValue[0], InputValue[1]);
+		break;
 	case EOperation::Dot:
 		if (Types.bIsLWC)
 		{
@@ -1299,6 +1409,26 @@ void FExpressionOperation::EmitValueShader(FEmitContext& Context, FEmitScope& Sc
 			OutResult.Code = Context.EmitExpression(Scope, Types.ResultType, TEXT("(% > %)"), InputValue[0], InputValue[1]);
 		}
 		break;
+	case EOperation::LessEqual:
+		if (Types.bIsLWC)
+		{
+			OutResult.Code = Context.EmitExpression(Scope, Types.ResultType, TEXT("LWCLessEqual(%, %)"), InputValue[0], InputValue[1]);
+		}
+		else
+		{
+			OutResult.Code = Context.EmitExpression(Scope, Types.ResultType, TEXT("(% <= %)"), InputValue[0], InputValue[1]);
+		}
+		break;
+	case EOperation::GreaterEqual:
+		if (Types.bIsLWC)
+		{
+			OutResult.Code = Context.EmitExpression(Scope, Types.ResultType, TEXT("LWCGreaterEqual(%, %)"), InputValue[0], InputValue[1]);
+		}
+		else
+		{
+			OutResult.Code = Context.EmitExpression(Scope, Types.ResultType, TEXT("(% >= %)"), InputValue[0], InputValue[1]);
+		}
+		break;
 	case EOperation::VecMulMatrix3:
 		if (Types.bIsLWC)
 		{
diff --git a/Engine/Source/Runtime/Engine/Private/HLSLTree/HLSLTreeEmit.cpp b/Engine/Source/Runtime/Engine/Private/HLSLTree/HLSLTreeEmit.cpp
--- a/Engine/Source/Runtime/Engine/Private/HLSLTree/HLSLTreeEmit.cpp
+++ b/Engine/Source/Runtime/Engine/Private/HLSLTree/HLSLTreeEmit.cpp
@@ -216,22 +216,45 @@ void EmitCustomHLSL(const FEmitCustomHLSL& EmitCustomHLSL, const TCHAR* Paramete
 
 	// Function that wraps custom HLSL, provides the interface expected by custom HLSL
 	// * LWC inputs have both LWC version and non-LWC version
+	// * Texture inputs are split into a Texture/SamplerState pair
 	// * First output is given through return value, additional outputs use inout function parameters
 	OutCode.Appendf(TEXT("%s CustomExpressionInternal%d(%s Parameters"), OutputStruct->Fields[0].Type.GetName(), EmitCustomHLSL.Index, ParametersTypeName);
 	for (const FEmitCustomHLSLInput& Input : EmitCustomHLSL.Inputs)
 	{
 		Shader::FType InputType = Input.Type;
-		if (InputType.IsNumericLWC())
+		if (InputType.IsTexture())
 		{
-			// Add an additional input for LWC type with the LWC-prefix
-			OutCode.Appendf(TEXT(", %s LWC"), InputType.GetName());
+			// Texture parameter followed by a SamplerState parameter
+			const TCHAR* TextureTypeName = nullptr;
+			switch (InputType.ValueType)
+			{
+			case Shader::EValueType::Texture2D: TextureTypeName = TEXT("Texture2D"); break;
+			case Shader::EValueType::Texture2DArray: TextureTypeName = TEXT("Texture2DArray"); break;
+			case Shader::EValueType::TextureCube: TextureTypeName = TEXT("TextureCube"); break;
+			case Shader::EValueType::TextureCubeArray: TextureTypeName = TEXT("TextureCubeArray"); break;
+			case Shader::EValueType::Texture3D: TextureTypeName = TEXT("Texture3D"); break;
+			default: checkNoEntry(); break;
+			}
+			OutCode.Appendf(TEXT(", %s "), TextureTypeName);
 			OutCode.Append(Input.Name);
-			// Regular, unprefixed input uses non-LWC type
-			InputType = InputType.GetNonLWCType();
+			OutCode.Appendf(TEXT(", SamplerState "));
+			OutCode.Append(Input.Name);
+			OutCode.Append(TEXT("Sampler"));
 		}
+		else
+		{
+			if (InputType.IsNumericLWC())
+			{
+				// Add an additional input for LWC type with the LWC-prefix
+				OutCode.Appendf(TEXT(", %s LWC"), InputType.GetName());
+				OutCode.Append(Input.Name);
+				// Regular, unprefixed input uses non-LWC type
+				InputType = InputType.GetNonLWCType();
+			}
 
-		OutCode.Appendf(TEXT(", %s "), InputType.GetName());
-		OutCode.Append(Input.Name);
+			OutCode.Appendf(TEXT(", %s "), InputType.GetName());
+			OutCode.Append(Input.Name);
+		}
 	}
 	for (int32 OutputIndex = 1; OutputIndex < OutputStruct->Fields.Num(); ++OutputIndex)
 	{
@@ -243,7 +266,7 @@ void EmitCustomHLSL(const FEmitCustomHLSL& EmitCustomHLSL, const TCHAR* Paramete
 	OutCode.Append(TEXT("\n}\n"));
 
 	// Function that calls the above wrapper, provides the interface expected by HLSLTree
-	// * No special handling for LWC inputs
+	// * All inputs types are passed through directly
 	// * All outputs are provided through a 'struct' type
 	OutCode.Appendf(TEXT("%s CustomExpression%d(%s Parameters"), OutputStruct->Name, EmitCustomHLSL.Index, ParametersTypeName);
 	for (const FEmitCustomHLSLInput& Input : EmitCustomHLSL.Inputs)
@@ -258,7 +281,14 @@ void EmitCustomHLSL(const FEmitCustomHLSL& EmitCustomHLSL, const TCHAR* Paramete
 	{
 		OutCode.Append(TEXT(", "));
 		OutCode.Append(Input.Name);
-		if (Input.Type.IsNumericLWC())
+		if (Input.Type.IsTexture())
+		{
+			// Pass the Texture and SamplerState to the wrapper
+			OutCode.Append(TEXT(".Texture, "));
+			OutCode.Append(Input.Name);
+			OutCode.Append(TEXT(".Sampler"));
+		}
+		else if (Input.Type.IsNumericLWC())
 		{
 			OutCode.Append(TEXT(", LWCToFloat("));
 			OutCode.Append(Input.Name);
@@ -590,7 +620,7 @@ void MoveToScope(FEmitShaderNode* EmitNode, FEmitScope& Scope)
 
 void FormatArg_ShaderValue(FEmitShaderExpression* ShaderValue, FEmitShaderDependencies& OutDependencies, FStringBuilderBase& OutCode)
 {
-	OutDependencies.Add(ShaderValue);
+	OutDependencies.AddUnique(ShaderValue);
 	OutCode.Append(ShaderValue->Reference);
 }
 
diff --git a/Engine/Source/Runtime/Engine/Private/Materials/MaterialExpressionHLSL.cpp b/Engine/Source/Runtime/Engine/Private/Materials/MaterialExpressionHLSL.cpp
--- a/Engine/Source/Runtime/Engine/Private/Materials/MaterialExpressionHLSL.cpp
+++ b/Engine/Source/Runtime/Engine/Private/Materials/MaterialExpressionHLSL.cpp
@@ -11,6 +11,7 @@
 #include "Materials/MaterialExpression.h"
 #include "Materials/MaterialExpressionExecBegin.h"
 #include "Materials/MaterialExpressionExecEnd.h"
+#include "Materials/MaterialExpressionReroute.h"
 #include "Materials/MaterialExpressionGenericConstant.h"
 #include "Materials/MaterialExpressionConstant.h"
 #include "Materials/MaterialExpressionConstant2Vector.h"
@@ -30,9 +31,12 @@
 #include "Materials/MaterialExpressionSceneTexelSize.h"
 #include "Materials/MaterialExpressionViewSize.h"
 #include "Materials/MaterialExpressionPanner.h"
+#include "Materials/MaterialExpressionVertexColor.h"
 #include "Materials/MaterialExpressionTextureCoordinate.h"
 #include "Materials/MaterialExpressionTextureSample.h"
 #include "Materials/MaterialExpressionTextureSampleParameter.h"
+#include "Materials/MaterialExpressionFontSample.h"
+#include "Materials/MaterialExpressionFontSampleParameter.h"
 #include "Materials/MaterialExpressionTextureObject.h"
 #include "Materials/MaterialExpressionTextureObjectParameter.h"
 #include "Materials/MaterialExpressionSceneTexture.h"
@@ -43,10 +47,19 @@
 #include "Materials/MaterialExpressionFeatureLevelSwitch.h"
 #include "Materials/MaterialExpressionGetLocal.h"
 #include "Materials/MaterialExpressionOneMinus.h"
+#include "Materials/MaterialExpressionAbs.h"
+#include "Materials/MaterialExpressionFrac.h"
+#include "Materials/MaterialExpressionFloor.h"
+#include "Materials/MaterialExpressionCeil.h"
+#include "Materials/MaterialExpressionRound.h"
+#include "Materials/MaterialExpressionTruncate.h"
+#include "Materials/MaterialExpressionSaturate.h"
+#include "Materials/MaterialExpressionSign.h"
 #include "Materials/MaterialExpressionAdd.h"
 #include "Materials/MaterialExpressionSubtract.h"
 #include "Materials/MaterialExpressionMultiply.h"
 #include "Materials/MaterialExpressionDivide.h"
+#include "Materials/MaterialExpressionPower.h"
 #include "Materials/MaterialExpressionDotProduct.h"
 #include "Materials/MaterialExpressionMin.h"
 #include "Materials/MaterialExpressionMax.h"
@@ -62,12 +75,16 @@
 #include "Materials/MaterialExpressionReflectionVectorWS.h"
 #include "Materials/MaterialExpressionTransform.h"
 #include "Materials/MaterialExpressionTransformPosition.h"
+#include "Materials/MaterialExpressionIf.h"
 #include "Materials/MaterialExpressionCustom.h"
 #include "Materials/MaterialExpressionSetLocal.h"
 #include "Materials/MaterialExpressionIfThenElse.h"
 #include "Materials/MaterialExpressionForLoop.h"
 #include "Materials/MaterialExpressionWhileLoop.h"
 #include "Materials/MaterialFunctionInterface.h"
+#include "Engine/Engine.h"
+#include "Engine/Texture2D.h"
+#include "Engine/Font.h"
 
 bool UMaterialExpression::GenerateHLSLExpression(FMaterialHLSLGenerator& Generator, UE::HLSLTree::FScope& Scope, int32 OutputIndex, UE::HLSLTree::FExpression*& OutExpression)
 {
@@ -79,9 +96,10 @@ bool UMaterialExpression::GenerateHLSLStatements(FMaterialHLSLGenerator& Generat
 	return Generator.GetErrors().AddError(TEXT("Node does not support statements"));
 }
 
-bool UMaterialExpression::GenerateHLSLTexture(FMaterialHLSLGenerator& Generator, UE::HLSLTree::FScope& Scope, int32 OutputIndex, UE::HLSLTree::FTextureParameterDeclaration*& OutTexture)
+bool UMaterialExpressionReroute::GenerateHLSLExpression(FMaterialHLSLGenerator& Generator, UE::HLSLTree::FScope& Scope, int32 OutputIndex, UE::HLSLTree::FExpression*& OutExpression)
 {
-	return Generator.GetErrors().AddError(TEXT("Node does not support textures"));
+	OutExpression = Input.TryAcquireHLSLExpression(Generator, Scope);
+	return OutExpression != nullptr;
 }
 
 bool UMaterialExpressionGenericConstant::GenerateHLSLExpression(FMaterialHLSLGenerator& Generator, UE::HLSLTree::FScope& Scope, int32 OutputIndex, UE::HLSLTree::FExpression*& OutExpression)
@@ -293,6 +311,13 @@ bool UMaterialExpressionPanner::GenerateHLSLExpression(FMaterialHLSLGenerator& G
 	return true;
 }
 
+bool UMaterialExpressionVertexColor::GenerateHLSLExpression(FMaterialHLSLGenerator& Generator, UE::HLSLTree::FScope& Scope, int32 OutputIndex, UE::HLSLTree::FExpression*& OutExpression)
+{
+	using namespace UE::HLSLTree;
+	OutExpression = Generator.GetTree().NewExpression<FExpressionExternalInput>(EExternalInput::VertexColor);
+	return true;
+}
+
 bool UMaterialExpressionTextureCoordinate::GenerateHLSLExpression(FMaterialHLSLGenerator& Generator, UE::HLSLTree::FScope& Scope, int32 OutputIndex, UE::HLSLTree::FExpression*& OutExpression)
 {
 	OutExpression = Generator.NewTexCoord(CoordinateIndex);
@@ -316,70 +341,138 @@ bool UMaterialExpressionTextureCoordinate::GenerateHLSLExpression(FMaterialHLSLG
 	return true;
 }
 
-bool UMaterialExpressionTextureObject::GenerateHLSLTexture(FMaterialHLSLGenerator& Generator, UE::HLSLTree::FScope& Scope, int32 OutputIndex, UE::HLSLTree::FTextureParameterDeclaration*& OutTexture)
+bool UMaterialExpressionTextureObject::GenerateHLSLExpression(FMaterialHLSLGenerator& Generator, UE::HLSLTree::FScope& Scope, int32 OutputIndex, UE::HLSLTree::FExpression*& OutExpression)
 {
-	const UE::HLSLTree::FTextureDescription TextureDesc(Texture, SamplerType);
-	OutTexture = Generator.AcquireTextureDeclaration(TextureDesc);
+	using namespace UE::HLSLTree;
+	using namespace UE::Shader;
+	const FTextureValue* TextureValue = Generator.AcquireTextureValue(FTextureValue(Texture, SamplerType));
+	OutExpression = Generator.GetTree().NewExpression<FExpressionMaterialParameter>(EMaterialParameterType::Texture, FName(), TextureValue);
 	return true;
 }
 
-bool UMaterialExpressionTextureObjectParameter::GenerateHLSLTexture(FMaterialHLSLGenerator& Generator, UE::HLSLTree::FScope& Scope, int32 OutputIndex, UE::HLSLTree::FTextureParameterDeclaration*& OutTexture)
+bool UMaterialExpressionTextureObjectParameter::GenerateHLSLExpression(FMaterialHLSLGenerator& Generator, UE::HLSLTree::FScope& Scope, int32 OutputIndex, UE::HLSLTree::FExpression*& OutExpression)
 {
-	const UE::HLSLTree::FTextureDescription TextureDesc(Texture, SamplerType);
-	OutTexture = Generator.AcquireTextureParameterDeclaration(ParameterName, TextureDesc);
+	using namespace UE::HLSLTree;
+	using namespace UE::Shader;
+	const FTextureValue* TextureValue = Generator.AcquireTextureValue(FTextureValue(Texture, SamplerType));
+	OutExpression = Generator.GetTree().NewExpression<FExpressionMaterialParameter>(EMaterialParameterType::Texture, ParameterName, TextureValue);
 	return true;
 }
 
-bool UMaterialExpressionTextureSample::GenerateHLSLExpressionBase(FMaterialHLSLGenerator& Generator, UE::HLSLTree::FScope& Scope, UE::HLSLTree::FTextureParameterDeclaration* TextureDeclaration, UE::HLSLTree::FExpression*& OutExpression)
+bool UMaterialExpressionTextureSample::GenerateHLSLExpressionBase(FMaterialHLSLGenerator& Generator, UE::HLSLTree::FScope& Scope, UE::HLSLTree::FExpression* TextureExpression, UE::HLSLTree::FExpression*& OutExpression)
 {
 	using namespace UE::HLSLTree;
-	if (!TextureDeclaration)
+	if (!TextureExpression)
 	{
 		return Generator.GetErrors().AddError(TEXT("Missing input texture"));
 	}
 
 	FExpression* TexCoordExpression = Coordinates.GetTracedInput().Expression ? Coordinates.TryAcquireHLSLExpression(Generator, Scope) : Generator.NewTexCoord(ConstCoordinate);
 	const FExpressionDerivatives TexCoordDerivatives;// = Generator.GetTree().GetAnalyticDerivatives(TexCoordExpression);
-	OutExpression = Generator.GetTree().NewExpression<FExpressionTextureSample>(TextureDeclaration, TexCoordExpression, TexCoordDerivatives, SamplerSource, MipValueMode);
+	OutExpression = Generator.GetTree().NewExpression<FExpressionTextureSample>(TextureExpression, TexCoordExpression, TexCoordDerivatives, SamplerSource, MipValueMode);
 	return true;
 }
 
 bool UMaterialExpressionTextureSample::GenerateHLSLExpression(FMaterialHLSLGenerator& Generator, UE::HLSLTree::FScope& Scope, int32 OutputIndex, UE::HLSLTree::FExpression*& OutExpression)
 {
-	UE::HLSLTree::FTextureParameterDeclaration* TextureDeclaration = nullptr;
+	using namespace UE::HLSLTree;
+	using namespace UE::Shader;
+	FExpression* TextureExpression = nullptr;
 	if (TextureObject.GetTracedInput().Expression)
 	{
-		TextureDeclaration = TextureObject.AcquireHLSLTexture(Generator, Scope);
+		TextureExpression = TextureObject.AcquireHLSLExpression(Generator, Scope);
 	}
 	else if (Texture)
 	{
-		const UE::HLSLTree::FTextureDescription TextureDesc(Texture, SamplerType);
-		TextureDeclaration = Generator.AcquireTextureDeclaration(TextureDesc);
+		const FTextureValue* TextureValue = Generator.AcquireTextureValue(FTextureValue(Texture, SamplerType));
+		TextureExpression = Generator.GetTree().NewExpression<FExpressionMaterialParameter>(EMaterialParameterType::Texture, FName(), TextureValue);
 	}
 
-	return GenerateHLSLExpressionBase(Generator, Scope, TextureDeclaration, OutExpression);
+	return GenerateHLSLExpressionBase(Generator, Scope, TextureExpression, OutExpression);
 }
 
 bool UMaterialExpressionTextureSampleParameter::GenerateHLSLExpression(FMaterialHLSLGenerator& Generator, UE::HLSLTree::FScope& Scope, int32 OutputIndex, UE::HLSLTree::FExpression*& OutExpression)
 {
-	UE::HLSLTree::FTextureParameterDeclaration* TextureDeclaration = nullptr;
+	using namespace UE::HLSLTree;
+	using namespace UE::Shader;
+	FExpression* TextureExpression = nullptr;
 	if (ParameterName.IsNone() && TextureObject.GetTracedInput().Expression)
 	{
-		TextureDeclaration = TextureObject.AcquireHLSLTexture(Generator, Scope);
+		TextureExpression = TextureObject.AcquireHLSLExpression(Generator, Scope);
 	}
 	else if (Texture)
 	{
-		const UE::HLSLTree::FTextureDescription TextureDesc(Texture, SamplerType);
-		TextureDeclaration = Generator.AcquireTextureParameterDeclaration(ParameterName, TextureDesc);
+		const FTextureValue* TextureValue = Generator.AcquireTextureValue(FTextureValue(Texture, SamplerType));
+		TextureExpression = Generator.GetTree().NewExpression<FExpressionMaterialParameter>(EMaterialParameterType::Texture, ParameterName, TextureValue);
 	}
 
-	return GenerateHLSLExpressionBase(Generator, Scope, TextureDeclaration, OutExpression);
+	return GenerateHLSLExpressionBase(Generator, Scope, TextureExpression, OutExpression);
 }
 
-bool UMaterialExpressionSceneTexture::GenerateHLSLExpression(FMaterialHLSLGenerator& Generator, UE::HLSLTree::FScope& Scope, int32 OutputIndex, UE::HLSLTree::FExpression*& OutExpression)
+UE::HLSLTree::FExpression* UMaterialExpressionFontSample::GenerateHLSLTextureExpression(FMaterialHLSLGenerator& Generator, const UE::Shader::FTextureValue* TextureValue)
+{
+	return Generator.GetTree().NewExpression<UE::HLSLTree::FExpressionMaterialParameter>(EMaterialParameterType::Texture, FName(), TextureValue);
+}
+
+UE::HLSLTree::FExpression* UMaterialExpressionFontSampleParameter::GenerateHLSLTextureExpression(FMaterialHLSLGenerator& Generator, const UE::Shader::FTextureValue* TextureValue)
+{
+	return Generator.GetTree().NewExpression<UE::HLSLTree::FExpressionMaterialParameter>(EMaterialParameterType::Texture, ParameterName, TextureValue);
+}
+
+bool UMaterialExpressionFontSample::GenerateHLSLExpression(FMaterialHLSLGenerator& Generator, UE::HLSLTree::FScope& Scope, int32 OutputIndex, UE::HLSLTree::FExpression*& OutExpression)
 {
 	using namespace UE::HLSLTree;
+	using namespace UE::Shader;
+
+	if (!Font)
+	{
+		return Generator.GetErrors().AddError(TEXT("Missing input Font"));
+	}
+	else if (Font->FontCacheType == EFontCacheType::Runtime)
+	{
+		return Generator.GetErrors().AddErrorf(TEXT("Font '%s' is runtime cached, but only offline cached fonts can be sampled"), *Font->GetName());
+	}
+	else if (!Font->Textures.IsValidIndex(FontTexturePage))
+	{
+		return Generator.GetErrors().AddErrorf(TEXT("Invalid font page %d. Max allowed is %d"), FontTexturePage, Font->Textures.Num());
+	}
+
+	UTexture* Texture = Font->Textures[FontTexturePage];
+	if (!Texture)
+	{
+		UE_LOG(LogMaterial, Log, TEXT("Invalid font texture. Using default texture"));
+		Texture = GEngine->DefaultTexture;
+	}
+	check(Texture);
+
+	EMaterialSamplerType ExpectedSamplerType;
+	if (Texture->CompressionSettings == TC_DistanceFieldFont)
+	{
+		ExpectedSamplerType = SAMPLERTYPE_DistanceFieldFont;
+	}
+	else
+	{
+		ExpectedSamplerType = Texture->SRGB ? SAMPLERTYPE_Color : SAMPLERTYPE_LinearColor;
+	}
+
+	FString SamplerTypeError;
+	if (!UMaterialExpressionTextureBase::VerifySamplerType(Generator.GetCompileTarget().FeatureLevel, Generator.GetCompileTarget().TargetPlatform, Texture, ExpectedSamplerType, SamplerTypeError))
+	{
+		return Generator.GetErrors().AddErrorf(TEXT("%s"), *SamplerTypeError);
+	}
+
+	const FTextureValue* TextureValue = Generator.AcquireTextureValue(FTextureValue(Texture, ExpectedSamplerType));
+	FExpression* TextureExpression = GenerateHLSLTextureExpression(Generator, TextureValue);
+	FExpression* TexCoordExpression = Generator.NewTexCoord(0);
+	FExpressionDerivatives TexCoordDerivatives;
+
+	OutExpression = Generator.GetTree().NewExpression<FExpressionTextureSample>(TextureExpression, TexCoordExpression, TexCoordDerivatives, SSM_FromTextureAsset, TMVM_None);
+	return true;
+}
 
+bool UMaterialExpressionSceneTexture::GenerateHLSLExpression(FMaterialHLSLGenerator& Generator, UE::HLSLTree::FScope& Scope, int32 OutputIndex, UE::HLSLTree::FExpression*& OutExpression)
+{
+	using namespace UE::HLSLTree;
 	if (OutputIndex == 0)
 	{
 		FExpression* ExpressionTexCoord = nullptr;
@@ -410,6 +503,102 @@ bool UMaterialExpressionOneMinus::GenerateHLSLExpression(FMaterialHLSLGenerator&
 	return true;
 }
 
+bool UMaterialExpressionAbs::GenerateHLSLExpression(FMaterialHLSLGenerator& Generator, UE::HLSLTree::FScope& Scope, int32 OutputIndex, UE::HLSLTree::FExpression*& OutExpression)
+{
+	using namespace UE::HLSLTree;
+	FExpression* InputExpression = Input.AcquireHLSLExpression(Generator, Scope);
+	if (!InputExpression)
+	{
+		return false;
+	}
+	OutExpression = Generator.GetTree().NewUnaryOp(EOperation::Abs, InputExpression);
+	return true;
+}
+
+bool UMaterialExpressionFrac::GenerateHLSLExpression(FMaterialHLSLGenerator& Generator, UE::HLSLTree::FScope& Scope, int32 OutputIndex, UE::HLSLTree::FExpression*& OutExpression)
+{
+	using namespace UE::HLSLTree;
+	FExpression* InputExpression = Input.AcquireHLSLExpression(Generator, Scope);
+	if (!InputExpression)
+	{
+		return false;
+	}
+	OutExpression = Generator.GetTree().NewUnaryOp(EOperation::Frac, InputExpression);
+	return true;
+}
+
+bool UMaterialExpressionFloor::GenerateHLSLExpression(FMaterialHLSLGenerator& Generator, UE::HLSLTree::FScope& Scope, int32 OutputIndex, UE::HLSLTree::FExpression*& OutExpression)
+{
+	using namespace UE::HLSLTree;
+	FExpression* InputExpression = Input.AcquireHLSLExpression(Generator, Scope);
+	if (!InputExpression)
+	{
+		return false;
+	}
+	OutExpression = Generator.GetTree().NewUnaryOp(EOperation::Floor, InputExpression);
+	return true;
+}
+
+bool UMaterialExpressionCeil::GenerateHLSLExpression(FMaterialHLSLGenerator& Generator, UE::HLSLTree::FScope& Scope, int32 OutputIndex, UE::HLSLTree::FExpression*& OutExpression)
+{
+	using namespace UE::HLSLTree;
+	FExpression* InputExpression = Input.AcquireHLSLExpression(Generator, Scope);
+	if (!InputExpression)
+	{
+		return false;
+	}
+	OutExpression = Generator.GetTree().NewUnaryOp(EOperation::Ceil, InputExpression);
+	return true;
+}
+
+bool UMaterialExpressionRound::GenerateHLSLExpression(FMaterialHLSLGenerator& Generator, UE::HLSLTree::FScope& Scope, int32 OutputIndex, UE::HLSLTree::FExpression*& OutExpression)
+{
+	using namespace UE::HLSLTree;
+	FExpression* InputExpression = Input.AcquireHLSLExpression(Generator, Scope);
+	if (!InputExpression)
+	{
+		return false;
+	}
+	OutExpression = Generator.GetTree().NewUnaryOp(EOperation::Round, InputExpression);
+	return true;
+}
+
+bool UMaterialExpressionTruncate::GenerateHLSLExpression(FMaterialHLSLGenerator& Generator, UE::HLSLTree::FScope& Scope, int32 OutputIndex, UE::HLSLTree::FExpression*& OutExpression)
+{
+	using namespace UE::HLSLTree;
+	FExpression* InputExpression = Input.AcquireHLSLExpression(Generator, Scope);
+	if (!InputExpression)
+	{
+		return false;
+	}
+	OutExpression = Generator.GetTree().NewUnaryOp(EOperation::Trunc, InputExpression);
+	return true;
+}
+
+bool UMaterialExpressionSaturate::GenerateHLSLExpression(FMaterialHLSLGenerator& Generator, UE::HLSLTree::FScope& Scope, int32 OutputIndex, UE::HLSLTree::FExpression*& OutExpression)
+{
+	using namespace UE::HLSLTree;
+	FExpression* InputExpression = Input.AcquireHLSLExpression(Generator, Scope);
+	if (!InputExpression)
+	{
+		return false;
+	}
+	OutExpression = Generator.GetTree().NewUnaryOp(EOperation::Saturate, InputExpression);
+	return true;
+}
+
+bool UMaterialExpressionSign::GenerateHLSLExpression(FMaterialHLSLGenerator& Generator, UE::HLSLTree::FScope& Scope, int32 OutputIndex, UE::HLSLTree::FExpression*& OutExpression)
+{
+	using namespace UE::HLSLTree;
+	FExpression* InputExpression = Input.AcquireHLSLExpression(Generator, Scope);
+	if (!InputExpression)
+	{
+		return false;
+	}
+	OutExpression = Generator.GetTree().NewUnaryOp(EOperation::Sign, InputExpression);
+	return true;
+}
+
 bool UMaterialExpressionBinaryOp::GenerateHLSLExpression(FMaterialHLSLGenerator& Generator, UE::HLSLTree::FScope& Scope, int32 OutputIndex, UE::HLSLTree::FExpression*& OutExpression)
 {
 	UE::HLSLTree::FExpression* Lhs = A.AcquireHLSLExpressionOrConstant(Generator, Scope, ConstA);
@@ -471,6 +660,19 @@ bool UMaterialExpressionDivide::GenerateHLSLExpression(FMaterialHLSLGenerator& G
 	return true;
 }
 
+bool UMaterialExpressionPower::GenerateHLSLExpression(FMaterialHLSLGenerator& Generator, UE::HLSLTree::FScope& Scope, int32 OutputIndex, UE::HLSLTree::FExpression*& OutExpression)
+{
+	using namespace UE::HLSLTree;
+	FExpression* Lhs = Base.AcquireHLSLExpression(Generator, Scope);
+	FExpression* Rhs = Exponent.AcquireHLSLExpressionOrConstant(Generator, Scope, ConstExponent);
+	if (!Lhs || !Rhs)
+	{
+		return false;
+	}
+	OutExpression = Generator.GetTree().NewBinaryOp(EOperation::PowPositiveClamped, Lhs, Rhs);
+	return true;
+}
+
 bool UMaterialExpressionDotProduct::GenerateHLSLExpression(FMaterialHLSLGenerator& Generator, UE::HLSLTree::FScope& Scope, int32 OutputIndex, UE::HLSLTree::FExpression*& OutExpression)
 {
 	UE::HLSLTree::FExpression* Lhs = A.AcquireHLSLExpression(Generator, Scope);
@@ -917,6 +1119,35 @@ bool UMaterialExpressionTransformPosition::GenerateHLSLExpression(FMaterialHLSLG
 	return true;
 }
 
+bool UMaterialExpressionIf::GenerateHLSLExpression(FMaterialHLSLGenerator& Generator, UE::HLSLTree::FScope& Scope, int32 OutputIndex, UE::HLSLTree::FExpression*& OutExpression)
+{
+	using namespace UE::HLSLTree;
+	using namespace UE::Shader;
+
+	FExpression* ExpressionA = A.AcquireHLSLExpression(Generator, Scope);
+	FExpression* ExpressionB = A.AcquireHLSLExpressionOrConstant(Generator, Scope, ConstB);
+	FExpression* ExpressionAGreaterThanB = AGreaterThanB.AcquireHLSLExpression(Generator, Scope);
+	FExpression* ExpressionALessThanB = ALessThanB.AcquireHLSLExpression(Generator, Scope);
+	if (!ExpressionA || !ExpressionB || !ExpressionAGreaterThanB || !ExpressionALessThanB)
+	{
+		return false;
+	}
+
+	UE::HLSLTree::FTree& Tree = Generator.GetTree();
+	FExpression* ExpressionCondAGreaterEqualB = Tree.NewGreaterEqual(ExpressionA, ExpressionB);
+	OutExpression = Tree.NewExpression<FExpressionSelect>(ExpressionCondAGreaterEqualB, ExpressionAGreaterThanB, ExpressionALessThanB);
+
+	FExpression* ExpressionAEqualsB = AEqualsB.TryAcquireHLSLExpression(Generator, Scope);
+	if (ExpressionAEqualsB)
+	{
+		FExpression* ExpressionThreshold = Generator.NewConstant(EqualsThreshold);
+		FExpression* ExpressionCondANotEqualsB = Tree.NewGreater(Tree.NewAbs(Tree.NewSub(ExpressionA, ExpressionB)), ExpressionThreshold);
+		OutExpression = Tree.NewExpression<FExpressionSelect>(ExpressionCondANotEqualsB, OutExpression, ExpressionAEqualsB);
+	}
+
+	return true;
+}
+
 static UE::Shader::FType GetCustomOutputType(const FMaterialHLSLGenerator& Generator, ECustomMaterialOutputType Type)
 {
 	using namespace UE::Shader;
diff --git a/Engine/Source/Runtime/Engine/Private/Materials/MaterialHLSLEmitter.cpp b/Engine/Source/Runtime/Engine/Private/Materials/MaterialHLSLEmitter.cpp
--- a/Engine/Source/Runtime/Engine/Private/Materials/MaterialHLSLEmitter.cpp
+++ b/Engine/Source/Runtime/Engine/Private/Materials/MaterialHLSLEmitter.cpp
@@ -391,6 +391,7 @@ static FString GenerateMaterialTemplateHLSL(EShaderPlatform ShaderPlatform,
 
 static void GetMaterialEnvironment(EShaderPlatform InPlatform,
 	const FMaterial& InMaterial,
+	const UE::HLSLTree::FEmitContext& EmitContext,
 	const FMaterialCompilationOutput& MaterialCompilationOutput,
 	FShaderCompilerEnvironment& OutEnvironment)
 {
@@ -486,7 +487,7 @@ static void GetMaterialEnvironment(EShaderPlatform InPlatform,
 	// @todo MetalMRT: Remove this hack and implement proper atmospheric-fog solution for Metal MRT...
 	OutEnvironment.SetDefine(TEXT("MATERIAL_ATMOSPHERIC_FOG"), false);// !IsMetalMRTPlatform(InPlatform) ? bUsesAtmosphericFog : 0);
 	OutEnvironment.SetDefine(TEXT("MATERIAL_SKY_ATMOSPHERE"), false);// bUsesSkyAtmosphere);
-	OutEnvironment.SetDefine(TEXT("INTERPOLATE_VERTEX_COLOR"), false);// bUsesVertexColor);
+	OutEnvironment.SetDefine(TEXT("INTERPOLATE_VERTEX_COLOR"), EmitContext.bUsesVertexColor);
 	OutEnvironment.SetDefine(TEXT("NEEDS_PARTICLE_COLOR"), false);// bUsesParticleColor);
 	OutEnvironment.SetDefine(TEXT("NEEDS_PARTICLE_LOCAL_TO_WORLD"), false);// bUsesParticleLocalToWorld);
 	OutEnvironment.SetDefine(TEXT("NEEDS_PARTICLE_WORLD_TO_LOCAL"), false);// bUsesParticleWorldToLocal);
@@ -933,7 +934,7 @@ bool MaterialEmitHLSL(const FMaterialCompileTargetParameters& InCompilerTarget,
 
 	OutMaterialEnvironment = new FSharedShaderCompilerEnvironment();
 	OutMaterialEnvironment->TargetPlatform = InCompilerTarget.TargetPlatform;
-	GetMaterialEnvironment(InCompilerTarget.ShaderPlatform, InOutMaterial, OutCompilationOutput, *OutMaterialEnvironment);
+	GetMaterialEnvironment(InCompilerTarget.ShaderPlatform, InOutMaterial, EmitContext, OutCompilationOutput, *OutMaterialEnvironment);
 	OutMaterialEnvironment->IncludeVirtualPathToContentsMap.Add(TEXT("/Engine/Generated/Material.ush"), MoveTemp(MaterialTemplateSource));
 
 	return true;
diff --git a/Engine/Source/Runtime/Engine/Private/Materials/MaterialHLSLGenerator.cpp b/Engine/Source/Runtime/Engine/Private/Materials/MaterialHLSLGenerator.cpp
--- a/Engine/Source/Runtime/Engine/Private/Materials/MaterialHLSLGenerator.cpp
+++ b/Engine/Source/Runtime/Engine/Private/Materials/MaterialHLSLGenerator.cpp
@@ -359,38 +359,32 @@ UE::HLSLTree::FExpression* FMaterialHLSLGenerator::NewSwizzle(const UE::HLSLTree
 	return HLSLTree->NewExpression<UE::HLSLTree::FExpressionSwizzle>(Params, Input);
 }
 
-UE::HLSLTree::FTextureParameterDeclaration* FMaterialHLSLGenerator::AcquireTextureDeclaration(const UE::HLSLTree::FTextureDescription& Value)
+const UE::Shader::FTextureValue* FMaterialHLSLGenerator::AcquireTextureValue(const UE::Shader::FTextureValue& InValue)
 {
+	using namespace UE::Shader;
 	FString SamplerTypeError;
-	if (!UMaterialExpressionTextureBase::VerifySamplerType(CompileTarget.FeatureLevel, CompileTarget.TargetPlatform, Value.Texture, Value.SamplerType, SamplerTypeError))
+	if (!UMaterialExpressionTextureBase::VerifySamplerType(CompileTarget.FeatureLevel, CompileTarget.TargetPlatform, InValue.Texture, InValue.SamplerType, SamplerTypeError))
 	{
 		Errors.AddError(SamplerTypeError);
 		return nullptr;
 	}
 
-	UE::HLSLTree::FTextureParameterDeclaration*& Declaration = TextureDeclarationMap.FindOrAdd(Value);
-	if (!Declaration)
-	{
-		Declaration = HLSLTree->NewTextureParameterDeclaration(FName(), Value);
-	}
-	return Declaration;
-}
+	FXxHash64Builder Hasher;
+	Hasher.Update(&InValue.Texture, sizeof(InValue.Texture));
+	Hasher.Update(&InValue.SamplerType, sizeof(InValue.SamplerType));
+	const FXxHash64 Hash = Hasher.Finalize();
 
-UE::HLSLTree::FTextureParameterDeclaration* FMaterialHLSLGenerator::AcquireTextureParameterDeclaration(const FName& Name, const UE::HLSLTree::FTextureDescription& DefaultValue)
-{
-	FString SamplerTypeError;
-	if (!UMaterialExpressionTextureBase::VerifySamplerType(CompileTarget.FeatureLevel, CompileTarget.TargetPlatform, DefaultValue.Texture, DefaultValue.SamplerType, SamplerTypeError))
+	FTextureValue const* const* PrevValue = TextureValueMap.Find(Hash);
+	if (PrevValue)
 	{
-		Errors.AddError(SamplerTypeError);
-		return nullptr;
+		check(*PrevValue);
+		check(**PrevValue == InValue);
+		return *PrevValue;
 	}
 
-	UE::HLSLTree::FTextureParameterDeclaration*& Declaration = TextureParameterDeclarationMap.FindOrAdd(Name);
-	if (!Declaration)
-	{
-		Declaration = HLSLTree->NewTextureParameterDeclaration(Name, DefaultValue);
-	}
-	return Declaration;
+	FTextureValue* Value = new(HLSLTree->GetAllocator()) FTextureValue(InValue);
+	TextureValueMap.Add(Hash, Value);
+	return Value;
 }
 
 UE::HLSLTree::FExpression* FMaterialHLSLGenerator::AcquireExpression(UE::HLSLTree::FScope& Scope, UMaterialExpression* MaterialExpression, int32 OutputIndex)
@@ -468,20 +462,6 @@ UE::HLSLTree::FExpression* FMaterialHLSLGenerator::AcquireFunctionInputExpressio
 	return InputExpression;
 }
 
-UE::HLSLTree::FTextureParameterDeclaration* FMaterialHLSLGenerator::AcquireTextureDeclaration(UE::HLSLTree::FScope& Scope, UMaterialExpression* MaterialExpression, int32 OutputIndex)
-{
-	using namespace UE::HLSLTree;
-	FOwnerScope TreeOwnerScope(GetTree(), MaterialExpression);
-	FOwnerScope ErrorOwnerScope(Errors, MaterialExpression);
-
-	FTextureParameterDeclaration* TextureDeclaration = nullptr;
-	if (MaterialExpression->GenerateHLSLTexture(*this, Scope, OutputIndex, TextureDeclaration))
-	{
-		return TextureDeclaration;
-	}
-	return nullptr;
-}
-
 bool FMaterialHLSLGenerator::GenerateStatements(UE::HLSLTree::FScope& Scope, UMaterialExpression* MaterialExpression)
 {
 	using namespace UE::HLSLTree;
diff --git a/Engine/Source/Runtime/Engine/Private/Materials/MaterialShared.cpp b/Engine/Source/Runtime/Engine/Private/Materials/MaterialShared.cpp
--- a/Engine/Source/Runtime/Engine/Private/Materials/MaterialShared.cpp
+++ b/Engine/Source/Runtime/Engine/Private/Materials/MaterialShared.cpp
@@ -313,23 +313,6 @@ UE::HLSLTree::FExpression* FExpressionInput::AcquireHLSLExpressionOrConstant(FMa
 	return TryAcquireHLSLExpression(Generator, Scope);
 }
 
-UE::HLSLTree::FTextureParameterDeclaration* FExpressionInput::AcquireHLSLTexture(FMaterialHLSLGenerator& Generator, UE::HLSLTree::FScope& Scope) const
-{
-	UE::HLSLTree::FTextureParameterDeclaration* Result = nullptr;
-	const FExpressionInput TracedInput = GetTracedInput();
-	if (!TracedInput.Expression)
-	{
-		Generator.GetErrors().AddError(TEXT("Missing input"));
-	}
-	else if (Expression)
-	{
-		Expression->ValidateState();
-		Result = Generator.AcquireTextureDeclaration(Scope, Expression, OutputIndex);
-	}
-
-	return Result;
-}
-
 void FExpressionInput::Connect( int32 InOutputIndex, class UMaterialExpression* InExpression )
 {
 	InExpression->ConnectExpression(this, InOutputIndex);
diff --git a/Engine/Source/Runtime/Engine/Private/Shader/Preshader.cpp b/Engine/Source/Runtime/Engine/Private/Shader/Preshader.cpp
--- a/Engine/Source/Runtime/Engine/Private/Shader/Preshader.cpp
+++ b/Engine/Source/Runtime/Engine/Private/Shader/Preshader.cpp
@@ -168,6 +168,8 @@ void FPreshaderData::WriteValue(const FValue& Value)
 {
 	const int32 NumComponents = Value.Type.GetNumComponents();
 
+	check(!Value.Type.IsTexture());
+
 	WriteType(Value.Type);
 	for (int32 Index = 0; Index < NumComponents; ++Index)
 	{
diff --git a/Engine/Source/Runtime/Engine/Private/Shader/ShaderValue.cpp b/Engine/Source/Runtime/Engine/Private/Shader/ShaderValue.cpp
--- a/Engine/Source/Runtime/Engine/Private/Shader/ShaderValue.cpp
+++ b/Engine/Source/Runtime/Engine/Private/Shader/ShaderValue.cpp
@@ -4,6 +4,7 @@
 #include "Hash/xxhash.h"
 #include "Misc/MemStackUtility.h"
 #include "Misc/LargeWorldRenderPosition.h"
+#include "Engine/Texture.h"
 
 namespace UE
 {
@@ -33,7 +34,11 @@ FType FType::GetDerivativeType() const
 
 	check(ValueType != EValueType::Struct);
 	const FValueTypeDescription TypeDesc = GetValueTypeDescription(ValueType);
-	return MakeValueType(EValueComponentType::Float, TypeDesc.NumComponents);
+	if (IsNumericType(TypeDesc.ComponentType))
+	{
+		return MakeValueType(EValueComponentType::Float, TypeDesc.NumComponents);
+	}
+	return EValueType::Void;
 }
 
 int32 FType::GetNumComponents() const
@@ -107,6 +112,11 @@ bool FType::Merge(const FType& OtherType)
 		const FValueTypeDescription TypeDesc = GetValueTypeDescription(ValueType);
 		const FValueTypeDescription OtherTypeDesc = GetValueTypeDescription(OtherType);
 		const EValueComponentType ComponentType = CombineComponentTypes(TypeDesc.ComponentType, OtherTypeDesc.ComponentType);
+		if (ComponentType == EValueComponentType::Void)
+		{
+			return false;
+		}
+
 		const int8 NumComponents = FMath::Max(TypeDesc.NumComponents, OtherTypeDesc.NumComponents);
 		ValueType = MakeValueType(ComponentType, NumComponents);
 	}
@@ -272,8 +282,39 @@ void FormatComponent_Double(double Value, int32 NumComponents, EValueStringForma
 
 } // namespace Private
 
+EValueType FTextureValue::GetType() const
+{
+	if (Texture)
+	{
+		const EMaterialValueType MaterialType = Texture->GetMaterialType();
+		switch (MaterialType)
+		{
+		case MCT_Texture2D: return EValueType::Texture2D;
+		case MCT_Texture2DArray: return EValueType::Texture2DArray;
+		case MCT_TextureCube: return EValueType::TextureCube;
+		case MCT_TextureCubeArray: return EValueType::TextureCubeArray;
+		case MCT_VolumeTexture: return EValueType::Texture3D;
+		default: checkNoEntry(); break;
+		}
+	}
+	return EValueType::Void;
+}
+
+FValue::FValue(const FTextureValue* InValue) : Type(EValueType::Void)
+{
+	if (InValue)
+	{
+		Type = InValue->GetType();
+		if (Type != EValueType::Void)
+		{
+			Component.Add(InValue);
+		}
+	}
+}
+
 FValue FValue::FromMemoryImage(EValueType Type, const void* Data, uint32* OutSizeInBytes)
 {
+	check(IsNumericType(Type));
 	const FValueTypeDescription TypeDesc = GetValueTypeDescription(Type);
 
 	FValue Result(TypeDesc.ComponentType, TypeDesc.NumComponents);
@@ -296,7 +337,7 @@ FValue FValue::FromMemoryImage(EValueType Type, const void* Data, uint32* OutSiz
 
 FMemoryImageValue FValue::AsMemoryImage() const
 {
-	check(!Type.IsStruct());
+	check(Type.IsNumeric());
 	const FValueTypeDescription TypeDesc = GetValueTypeDescription(Type);
 
 	FMemoryImageValue Result;
@@ -375,15 +416,29 @@ bool FValue::AsBoolScalar() const
 	return false;
 }
 
+const FTextureValue* FValue::AsTexture() const
+{
+	if (Type.IsTexture() && Component.Num() > 0)
+	{
+		return Component[0].Texture;
+	}
+	return nullptr;
+}
+
 FValueComponentTypeDescription GetValueComponentTypeDescription(EValueComponentType Type)
 {
 	switch (Type)
 	{
 	case EValueComponentType::Void: return FValueComponentTypeDescription(TEXT("void"), 0u, EComponentBound::Zero, EComponentBound::Zero);
-	case EValueComponentType::Float: return FValueComponentTypeDescription(TEXT("float"), 4u, EComponentBound::NegFloatMax, EComponentBound::FloatMax);
-	case EValueComponentType::Double: return FValueComponentTypeDescription(TEXT("double"), 8u, EComponentBound::NegDoubleMax, EComponentBound::DoubleMax);
-	case EValueComponentType::Int: return FValueComponentTypeDescription(TEXT("int"), 4u, EComponentBound::IntMin, EComponentBound::IntMax);
+	case EValueComponentType::Float: return FValueComponentTypeDescription(TEXT("float"), sizeof(float), EComponentBound::NegFloatMax, EComponentBound::FloatMax);
+	case EValueComponentType::Double: return FValueComponentTypeDescription(TEXT("double"), sizeof(double), EComponentBound::NegDoubleMax, EComponentBound::DoubleMax);
+	case EValueComponentType::Int: return FValueComponentTypeDescription(TEXT("int"), sizeof(int32), EComponentBound::IntMin, EComponentBound::IntMax);
 	case EValueComponentType::Bool: return FValueComponentTypeDescription(TEXT("bool"), 1u, EComponentBound::Zero, EComponentBound::One);
+	case EValueComponentType::Texture2D: return FValueComponentTypeDescription(TEXT("Texture2D"), sizeof(void*), EComponentBound::Zero, EComponentBound::Zero);
+	case EValueComponentType::Texture2DArray: return FValueComponentTypeDescription(TEXT("Texture2DArray"), sizeof(void*), EComponentBound::Zero, EComponentBound::Zero);
+	case EValueComponentType::TextureCube: return FValueComponentTypeDescription(TEXT("TextureCube"), sizeof(void*), EComponentBound::Zero, EComponentBound::Zero);
+	case EValueComponentType::TextureCubeArray: return FValueComponentTypeDescription(TEXT("TextureCubeArray"), sizeof(void*), EComponentBound::Zero, EComponentBound::Zero);
+	case EValueComponentType::Texture3D: return FValueComponentTypeDescription(TEXT("Texture3D"), sizeof(void*), EComponentBound::Zero, EComponentBound::Zero);
 	default: checkNoEntry() return FValueComponentTypeDescription();
 	}
 }
@@ -410,10 +465,14 @@ EValueComponentType CombineComponentTypes(EValueComponentType Lhs, EValueCompone
 	{
 		return EValueComponentType::Float;
 	}
-	else
+	else if(IsNumericType(Lhs) && IsNumericType(Rhs))
 	{
 		return EValueComponentType::Int;
 	}
+	else
+	{
+		return EValueComponentType::Void;
+	}
 }
 
 const TCHAR* FValueComponent::ToString(EValueComponentType Type, FStringBuilderBase& OutString) const
@@ -577,6 +636,11 @@ FValueTypeDescription GetValueTypeDescription(EValueType Type)
 	case EValueType::Double4x4: return FValueTypeDescription(TEXT("FLWCMatrix"), EValueComponentType::Double, 16);
 	case EValueType::DoubleInverse4x4: return FValueTypeDescription(TEXT("FLWCInverseMatrix"), EValueComponentType::Double, 16);
 	case EValueType::Struct: return FValueTypeDescription(TEXT("struct"), EValueComponentType::Void, 0);
+	case EValueType::Texture2D: return FValueTypeDescription(TEXT("FTexture2D"), EValueComponentType::Texture2D, 1);
+	case EValueType::Texture2DArray: return FValueTypeDescription(TEXT("FTexture2DArray"), EValueComponentType::Texture2DArray, 1);
+	case EValueType::TextureCube: return FValueTypeDescription(TEXT("FTextureCube"), EValueComponentType::TextureCube, 1);
+	case EValueType::TextureCubeArray: return FValueTypeDescription(TEXT("FTextureCubeArray"), EValueComponentType::TextureCubeArray, 1);
+	case EValueType::Texture3D: return FValueTypeDescription(TEXT("FTexture3D"), EValueComponentType::Texture3D, 1);
 	default: checkNoEntry(); return FValueTypeDescription(TEXT("<INVALID>"), EValueComponentType::Void, 0);
 	}
 }
@@ -600,6 +664,7 @@ EValueType MakeValueType(EValueComponentType ComponentType, int32 NumComponents)
 		case 16: return EValueType::Float4x4;
 		default: break;
 		}
+		break;
 	case EValueComponentType::Double:
 		switch (NumComponents)
 		{
@@ -610,6 +675,7 @@ EValueType MakeValueType(EValueComponentType ComponentType, int32 NumComponents)
 		case 16: return EValueType::Double4x4;
 		default: break;
 		}
+		break;
 	case EValueComponentType::Int:
 		switch (NumComponents)
 		{
@@ -619,6 +685,7 @@ EValueType MakeValueType(EValueComponentType ComponentType, int32 NumComponents)
 		case 4: return EValueType::Int4;
 		default: break;
 		}
+		break;
 	case EValueComponentType::Bool:
 		switch (NumComponents)
 		{
@@ -628,6 +695,12 @@ EValueType MakeValueType(EValueComponentType ComponentType, int32 NumComponents)
 		case 4: return EValueType::Bool4;
 		default: break;
 		}
+		break;
+	case EValueComponentType::Texture2D: check(NumComponents == 1); return EValueType::Texture2D; break;
+	case EValueComponentType::Texture2DArray: check(NumComponents == 1); return EValueType::Texture2DArray; break;
+	case EValueComponentType::TextureCube: check(NumComponents == 1); return EValueType::TextureCube; break;
+	case EValueComponentType::TextureCubeArray: check(NumComponents == 1); return EValueType::TextureCubeArray; break;
+	case EValueComponentType::Texture3D: check(NumComponents == 1); return EValueType::Texture3D; break;
 	default:
 		break;
 	}
diff --git a/Engine/Source/Runtime/Engine/Public/HLSLTree/HLSLTree.h b/Engine/Source/Runtime/Engine/Public/HLSLTree/HLSLTree.h
--- a/Engine/Source/Runtime/Engine/Public/HLSLTree/HLSLTree.h
+++ b/Engine/Source/Runtime/Engine/Public/HLSLTree/HLSLTree.h
@@ -174,6 +174,7 @@ enum class ERequestedType : uint8
 {
 	Void = 0u,
 	Scalar = 1u,
+	Texture = 1u, // Currently no different than scalar, just makes code more clear
 	Vector2 = 2u,
 	Vector3 = 3u,
 	Vector4 = 4u,
@@ -455,18 +456,6 @@ public:
 	TArray<FExpression*, TInlineAllocator<8>> OutputExpressions;
 };
 
-/**
- * Represents an HLSL texture parameter.
- */
-class FTextureParameterDeclaration final : public FNode
-{
-public:
-	FTextureParameterDeclaration(const FName& InName, const FTextureDescription& InDescription) : Name(InName), Description(InDescription) {}
-
-	FName Name;
-	FTextureDescription Description;
-};
-
 /**
  * Represents an HLSL scope.  A scope contains a single statement, along with any expressions required by that statement
  */
@@ -577,6 +566,7 @@ public:
 	FExpression* NewUnaryOp(EOperation Op, FExpression* Input);
 	FExpression* NewBinaryOp(EOperation Op, FExpression* Lhs, FExpression* Rhs);
 
+	FExpression* NewAbs(FExpression* Input) { return NewUnaryOp(EOperation::Abs, Input); }
 	FExpression* NewNeg(FExpression* Input) { return NewUnaryOp(EOperation::Neg, Input); }
 	FExpression* NewRcp(FExpression* Input) { return NewUnaryOp(EOperation::Rcp, Input); }
 	FExpression* NewFrac(FExpression* Input) { return NewUnaryOp(EOperation::Frac, Input); }
@@ -593,8 +583,8 @@ public:
 	FExpression* NewMax(FExpression* Lhs, FExpression* Rhs) { return NewBinaryOp(EOperation::Max, Lhs, Rhs); }
 	FExpression* NewLess(FExpression* Lhs, FExpression* Rhs) { return NewBinaryOp(EOperation::Less, Lhs, Rhs); }
 	FExpression* NewGreater(FExpression* Lhs, FExpression* Rhs) { return NewBinaryOp(EOperation::Greater, Lhs, Rhs); }
-
-	FTextureParameterDeclaration* NewTextureParameterDeclaration(const FName& Name, const FTextureDescription& DefaultValue);
+	FExpression* NewLessEqual(FExpression* Lhs, FExpression* Rhs) { return NewBinaryOp(EOperation::LessEqual, Lhs, Rhs); }
+	FExpression* NewGreaterEqual(FExpression* Lhs, FExpression* Rhs) { return NewBinaryOp(EOperation::GreaterEqual, Lhs, Rhs); }
 
 private:
 	template<typename T, typename... ArgTypes>
diff --git a/Engine/Source/Runtime/Engine/Public/HLSLTree/HLSLTreeCommon.h b/Engine/Source/Runtime/Engine/Public/HLSLTree/HLSLTreeCommon.h
--- a/Engine/Source/Runtime/Engine/Public/HLSLTree/HLSLTreeCommon.h
+++ b/Engine/Source/Runtime/Engine/Public/HLSLTree/HLSLTreeCommon.h
@@ -39,6 +39,7 @@ public:
 
 	virtual void ComputeAnalyticDerivatives(FTree& Tree, FExpressionDerivatives& OutResult) const override;
 	virtual bool PrepareValue(FEmitContext& Context, FEmitScope& Scope, const FRequestedType& RequestedType, FPrepareValueResult& OutResult) const override;
+	virtual void EmitValueShader(FEmitContext& Context, FEmitScope& Scope, const FRequestedType& RequestedType, FEmitValueShaderResult& OutResult) const override;
 	virtual void EmitValuePreshader(FEmitContext& Context, FEmitScope& Scope, const FRequestedType& RequestedType, FEmitValuePreshaderResult& OutResult) const override;
 };
 
@@ -73,6 +74,10 @@ enum class EExternalInput : uint8
 	TexCoord6_Ddy,
 	TexCoord7_Ddy,
 
+	VertexColor,
+	VertexColor_Ddx,
+	VertexColor_Ddy,
+
 	WorldPosition,
 	WorldPosition_NoOffsets,
 	TranslatedWorldPosition,
@@ -193,15 +198,19 @@ public:
 class FExpressionTextureSample : public FExpression
 {
 public:
-	FExpressionTextureSample(FTextureParameterDeclaration* InDeclaration, FExpression* InTexCoordExpression, const FExpressionDerivatives& InTexCoordDerivatives, ESamplerSourceMode InSamplerSource, ETextureMipValueMode InMipValueMode)
-		: Declaration(InDeclaration)
+	FExpressionTextureSample(FExpression* InTextureExpression,
+		FExpression* InTexCoordExpression,
+		const FExpressionDerivatives& InTexCoordDerivatives,
+		ESamplerSourceMode InSamplerSource,
+		ETextureMipValueMode InMipValueMode)
+		: TextureExpression(InTextureExpression)
 		, TexCoordExpression(InTexCoordExpression)
 		, TexCoordDerivatives(InTexCoordDerivatives)
 		, SamplerSource(InSamplerSource)
 		, MipValueMode(InMipValueMode)
 	{}
 
-	FTextureParameterDeclaration* Declaration;
+	FExpression* TextureExpression;
 	FExpression* TexCoordExpression;
 	FExpressionDerivatives TexCoordDerivatives;
 	ESamplerSourceMode SamplerSource;
diff --git a/Engine/Source/Runtime/Engine/Public/HLSLTree/HLSLTreeEmit.h b/Engine/Source/Runtime/Engine/Public/HLSLTree/HLSLTreeEmit.h
--- a/Engine/Source/Runtime/Engine/Public/HLSLTree/HLSLTreeEmit.h
+++ b/Engine/Source/Runtime/Engine/Public/HLSLTree/HLSLTreeEmit.h
@@ -431,6 +431,7 @@ public:
 	uint32 CurrentBoolUniformOffset = ~0u;
 	uint32 CurrentNumBoolComponents = 32u;
 	bool bReadMaterialNormal = false;
+	bool bUsesVertexColor = false;
 	uint32 TexCoordMask[SF_NumFrequencies] = { 0u };
 
 	int32 NumExpressionLocals = 0;
diff --git a/Engine/Source/Runtime/Engine/Public/HLSLTree/HLSLTreeTypes.h b/Engine/Source/Runtime/Engine/Public/HLSLTree/HLSLTreeTypes.h
--- a/Engine/Source/Runtime/Engine/Public/HLSLTree/HLSLTreeTypes.h
+++ b/Engine/Source/Runtime/Engine/Public/HLSLTree/HLSLTreeTypes.h
@@ -21,7 +21,6 @@ namespace HLSLTree
 class FNode;
 class FStructType;
 class FExpression;
-class FTextureParameterDeclaration;
 class FStatement;
 class FScope;
 class FTree;
@@ -64,9 +63,16 @@ enum class EOperation : uint8
 	None,
 
 	// Unary Ops
+	Abs,
 	Neg,
 	Rcp,
 	Frac,
+	Floor,
+	Ceil,
+	Round,
+	Trunc,
+	Saturate,
+	Sign,
 	Length,
 	Normalize,
 
@@ -76,11 +82,14 @@ enum class EOperation : uint8
 	Mul,
 	Div,
 	Fmod,
+	PowPositiveClamped,
 	Dot,
 	Min,
 	Max,
 	Less,
 	Greater,
+	LessEqual,
+	GreaterEqual,
 
 	VecMulMatrix3,
 	VecMulMatrix4,
@@ -101,32 +110,6 @@ struct FOperationDescription
 
 FOperationDescription GetOperationDescription(EOperation Op);
 
-struct FTextureDescription
-{
-	FTextureDescription()
-		: Texture(nullptr), SamplerType(SAMPLERTYPE_Color)
-	{}
-
-	FTextureDescription(UTexture* InTexture, EMaterialSamplerType InSamplerType)
-		: Texture(InTexture), SamplerType(InSamplerType)
-	{}
-
-	UTexture* Texture;
-	EMaterialSamplerType SamplerType;
-};
-inline bool operator==(const FTextureDescription& Lhs, const FTextureDescription& Rhs)
-{
-	return Lhs.Texture == Rhs.Texture && Lhs.SamplerType == Rhs.SamplerType;
-}
-inline bool operator!=(const FTextureDescription& Lhs, const FTextureDescription& Rhs)
-{
-	return !operator==(Lhs, Rhs);
-}
-inline uint32 GetTypeHash(const FTextureDescription& Ref)
-{
-	return HashCombine(GetTypeHash(Ref.Texture), GetTypeHash(Ref.SamplerType));
-}
-
 struct FCustomHLSLInput
 {
 	FCustomHLSLInput() = default;
diff --git a/Engine/Source/Runtime/Engine/Public/MaterialExpressionIO.h b/Engine/Source/Runtime/Engine/Public/MaterialExpressionIO.h
--- a/Engine/Source/Runtime/Engine/Public/MaterialExpressionIO.h
+++ b/Engine/Source/Runtime/Engine/Public/MaterialExpressionIO.h
@@ -14,7 +14,6 @@ namespace HLSLTree
 {
 class FScope;
 class FExpression;
-class FTextureParameterDeclaration;
 }
 namespace Shader
 {
@@ -81,7 +80,6 @@ struct FExpressionInput
 	ENGINE_API UE::HLSLTree::FExpression* TryAcquireHLSLExpression(FMaterialHLSLGenerator& Generator, UE::HLSLTree::FScope& Scope) const;
 	ENGINE_API UE::HLSLTree::FExpression* AcquireHLSLExpression(FMaterialHLSLGenerator& Generator, UE::HLSLTree::FScope& Scope) const;
 	ENGINE_API UE::HLSLTree::FExpression* AcquireHLSLExpressionOrConstant(FMaterialHLSLGenerator& Generator, UE::HLSLTree::FScope& Scope, const UE::Shader::FValue& ConstantValue) const;
-	ENGINE_API UE::HLSLTree::FTextureParameterDeclaration* AcquireHLSLTexture(FMaterialHLSLGenerator& Generator, UE::HLSLTree::FScope& Scope) const;
 #endif // WITH_EDITOR
 
 	/**
diff --git a/Engine/Source/Runtime/Engine/Public/MaterialHLSLGenerator.h b/Engine/Source/Runtime/Engine/Public/MaterialHLSLGenerator.h
--- a/Engine/Source/Runtime/Engine/Public/MaterialHLSLGenerator.h
+++ b/Engine/Source/Runtime/Engine/Public/MaterialHLSLGenerator.h
@@ -93,11 +93,7 @@ public:
 	UE::HLSLTree::FExpression* NewTexCoord(int32 Index);
 	UE::HLSLTree::FExpression* NewSwizzle(const UE::HLSLTree::FSwizzleParameters& Params, UE::HLSLTree::FExpression* Input);
 
-	/** Returns a declaration to access the given texture, with no parameter */
-	UE::HLSLTree::FTextureParameterDeclaration* AcquireTextureDeclaration(const UE::HLSLTree::FTextureDescription& Value);
-
-	/** Returns a declaration to access the given texture parameter */
-	UE::HLSLTree::FTextureParameterDeclaration* AcquireTextureParameterDeclaration(const FName& Name, const UE::HLSLTree::FTextureDescription& DefaultValue);
+	const UE::Shader::FTextureValue* AcquireTextureValue(const UE::Shader::FTextureValue& Value);
 
 	/**
 	 * Returns the appropriate HLSLNode representing the given UMaterialExpression.
@@ -107,7 +103,6 @@ public:
 	UE::HLSLTree::FExpression* AcquireExpression(UE::HLSLTree::FScope& Scope, UMaterialExpression* MaterialExpression, int32 OutputIndex);
 	UE::HLSLTree::FExpression* AcquireFunctionInputExpression(UE::HLSLTree::FScope& Scope, const UMaterialExpressionFunctionInput* MaterialExpression);
 
-	UE::HLSLTree::FTextureParameterDeclaration* AcquireTextureDeclaration(UE::HLSLTree::FScope& Scope, UMaterialExpression* MaterialExpression, int32 OutputIndex);
 	bool GenerateStatements(UE::HLSLTree::FScope& Scope, UMaterialExpression* MaterialExpression);
 
 	UE::HLSLTree::FExpression* GenerateFunctionCall(UE::HLSLTree::FScope& Scope, UMaterialFunctionInterface* Function, TArrayView<const FFunctionExpressionInput> ConnectedInputs, int32 OutputIndex);
@@ -199,8 +194,7 @@ private:
 
 	TArray<FFunctionCallEntry*, TInlineAllocator<8>> FunctionCallStack;
 	TArray<UE::HLSLTree::FScope*> JoinedScopeStack;
-	TMap<UE::HLSLTree::FTextureDescription, UE::HLSLTree::FTextureParameterDeclaration*> TextureDeclarationMap;
-	TMap<FName, UE::HLSLTree::FTextureParameterDeclaration*> TextureParameterDeclarationMap;
+	TMap<FXxHash64, const UE::Shader::FTextureValue*> TextureValueMap;
 	TMap<FSHAHash, FFunctionCallEntry*> FunctionCallMap;
 	TMap<UMaterialFunctionInterface*, UE::HLSLTree::FFunction*> FunctionMap;
 	TMap<UMaterialExpression*, FStatementEntry> StatementMap;
diff --git a/Engine/Source/Runtime/Engine/Public/Shader/ShaderTypes.h b/Engine/Source/Runtime/Engine/Public/Shader/ShaderTypes.h
--- a/Engine/Source/Runtime/Engine/Public/Shader/ShaderTypes.h
+++ b/Engine/Source/Runtime/Engine/Public/Shader/ShaderTypes.h
@@ -2,11 +2,13 @@
 #pragma once
 
 #include "CoreMinimal.h"
+#include "Engine/EngineTypes.h"
 #include "Containers/ArrayView.h"
 #include "Containers/StringView.h"
 #include "Serialization/MemoryLayout.h"
 
 class FMemStackBase;
+class UTexture;
 
 namespace UE
 {
@@ -59,6 +61,13 @@ enum class EValueComponentType : uint8
 	Double,
 	Int,
 	Bool,
+
+	FirstTexture,
+	Texture2D = FirstTexture,
+	Texture2DArray,
+	TextureCube,
+	TextureCubeArray,
+	Texture3D,
 };
 
 struct FValueComponentTypeDescription
@@ -74,6 +83,8 @@ struct FValueComponentTypeDescription
 FValueComponentTypeDescription GetValueComponentTypeDescription(EValueComponentType Type);
 inline const TCHAR* GetComponentTypeName(EValueComponentType Type) { return GetValueComponentTypeDescription(Type).Name; }
 inline uint32 GetComponentTypeSizeInBytes(EValueComponentType Type) { return GetValueComponentTypeDescription(Type).SizeInBytes; }
+inline bool IsNumericType(EValueComponentType Type) { return Type != EValueComponentType::Void && (uint8)Type < (uint8)EValueComponentType::FirstTexture; }
+inline bool IsTextureType(EValueComponentType Type) { return (uint8)Type >= (uint8)EValueComponentType::FirstTexture; }
 inline bool IsComponentTypeWithinBounds(EValueComponentType Type, FComponentBounds Bounds) { return IsWithinBounds(GetValueComponentTypeDescription(Type).Bounds, Bounds); }
 
 EValueComponentType CombineComponentTypes(EValueComponentType Lhs, EValueComponentType Rhs);
@@ -114,6 +125,12 @@ enum class EValueType : uint8
 	DoubleInverse4x4,
 
 	Struct,
+
+	Texture2D,
+	Texture2DArray,
+	TextureCube,
+	TextureCubeArray,
+	Texture3D,
 };
 
 struct FValueTypeDescription
@@ -135,6 +152,9 @@ EValueType MakeNonLWCType(EValueType Type);
 EValueType MakeArithmeticResultType(EValueType Lhs, EValueType Rhs, FString& OutErrorMessage);
 EValueType MakeComparisonResultType(EValueType Lhs, EValueType Rhs, FString& OutErrorMessage);
 
+inline bool IsNumericType(EValueType Type) { return IsNumericType(GetValueTypeDescription(Type).ComponentType); }
+inline bool IsTextureType(EValueType Type) { return IsTextureType(GetValueTypeDescription(Type).ComponentType); }
+
 struct FType
 {
 	FType() : StructType(nullptr), ValueType(EValueType::Void) {}
@@ -146,7 +166,8 @@ struct FType
 	FType GetNonLWCType() const { return IsNumericLWC() ? FType(MakeNonLWCType(ValueType)) : *this; }
 	bool IsVoid() const { return ValueType == EValueType::Void; }
 	bool IsStruct() const { return ValueType == EValueType::Struct; }
-	bool IsNumeric() const { return !IsVoid() && !IsStruct(); }
+	bool IsNumeric() const { return !IsStruct() && IsNumericType(ValueType); }
+	bool IsTexture() const { return !IsStruct() && IsTextureType(ValueType); }
 	bool IsNumericLWC() const { return IsNumeric() && IsLWCType(ValueType); }
 	int32 GetNumComponents() const;
 	int32 GetNumFlatFields() const;
@@ -281,13 +302,37 @@ struct FMemoryImageValue
 	uint32 Size;
 };
 
+struct FTextureValue
+{
+	FTextureValue() = default;
+	FTextureValue(UTexture* InTexture, EMaterialSamplerType InSamplerType = SAMPLERTYPE_Color)
+		: Texture(InTexture)
+		, SamplerType(InSamplerType)
+	{}
+
+	EValueType GetType() const;
+
+	UTexture* Texture = nullptr;
+	EMaterialSamplerType SamplerType = SAMPLERTYPE_Color;
+};
+
+inline bool operator==(const FTextureValue& Lhs, const FTextureValue& Rhs)
+{
+	return Lhs.Texture == Rhs.Texture && Lhs.SamplerType == Rhs.SamplerType;
+}
+inline bool operator!=(const FTextureValue& Lhs, const FTextureValue& Rhs)
+{
+	return !operator==(Lhs, Rhs);
+}
+
 union FValueComponent
 {
 	FValueComponent() : Packed(0u) {}
-	FValueComponent(float v) : Packed(0u) { Float = v; }
-	FValueComponent(double v) : Packed(0u) { Double = v; }
-	FValueComponent(int32 v) : Packed(0u) { Int = v; }
-	FValueComponent(bool v) : Packed(0u) { Bool = v ? 1 : 0; }
+	FValueComponent(float InValue) : Packed(0u) { Float = InValue; }
+	FValueComponent(double InValue) : Packed(0u) { Double = InValue; }
+	FValueComponent(int32 InValue) : Packed(0u) { Int = InValue; }
+	FValueComponent(bool InValue) : Packed(0u) { Bool = InValue ? 1 : 0; }
+	FValueComponent(const FTextureValue* InValue) : Packed(0u) { Texture = InValue; }
 
 	// 'Bool' is stored as uint8 to avoid changing on different compilers
 	bool AsBool() const { return Bool != 0u; }
@@ -295,6 +340,7 @@ union FValueComponent
 	const TCHAR* ToString(EValueComponentType Type, FStringBuilderBase& OutString) const;
 
 	uint64 Packed;
+	const FTextureValue* Texture;
 	double Double;
 	float Float;
 	int32 Int;
@@ -430,6 +476,8 @@ struct FValue
 		Component.Add(v);
 	}
 
+	FValue(const FTextureValue* InValue);
+
 	inline const FType& GetType() const { return Type; }
 	inline const FValueComponent& GetComponent(int32 i) const
 	{
@@ -450,6 +498,7 @@ struct FValue
 	FDoubleValue AsDouble() const;
 	FIntValue AsInt() const;
 	FBoolValue AsBool() const;
+	const FTextureValue* AsTexture() const;
 
 	FLinearColor AsLinearColor() const;
 	FVector4d AsVector4d() const;
