[Backout] - CL19027463
Original CL Desc
-----------------------------------------------------------------
Add texture types to Shader::EValueType, texture values are now represented by FExpressions (rather than separate texture desc type) in new HLSL translator
Makes it cleaner to pass textures to functions or custom HLSL
Also add more misc functionality to new HLSL translator

diff --git a/Engine/Shaders/Private/MaterialTemplate.ush b/Engine/Shaders/Private/MaterialTemplate.ush
--- a/Engine/Shaders/Private/MaterialTemplate.ush
+++ b/Engine/Shaders/Private/MaterialTemplate.ush
@@ -19,7 +19,6 @@
 #include "/Engine/Private/HairStrands/HairCardsAttributeCommon.ush"
 #include "/Engine/Private/HairStrands/HairStrandsAttributeCommon.ush"
 #include "/Engine/Private/DeferredShadingCommon.ush"
-#include "/Engine/Private/MaterialTexture.ush"
 
 // Update this GUID to force all material shaders to recompile
 // Merge conflicts on this line should be resolved by generating a new GUID
@@ -1707,6 +1706,73 @@ MaterialFloat2 GetParticleMacroUV(FMaterialPixelParameters Parameters)
 	return (Parameters.ScreenPosition.xy / Parameters.ScreenPosition.w - Parameters.Particle.MacroUV.xy) * Parameters.Particle.MacroUV.zw + MaterialFloat2(.5, .5);
 }
 
+MaterialFloat4 ProcessMaterialColorTextureLookup(MaterialFloat4 TextureValue)
+{
+	return TextureValue;
+}
+
+MaterialFloat4 ProcessMaterialVirtualColorTextureLookup(MaterialFloat4 TextureValue)
+{
+	TextureValue = ProcessMaterialColorTextureLookup(TextureValue);
+#if FEATURE_LEVEL == FEATURE_LEVEL_ES3_1
+	// on mobile all VT physical spaces use linear color formats, do sRGB to Linear conversion in the shader
+	TextureValue = MaterialFloat4(TextureValue.rgb*TextureValue.rgb, TextureValue.a);
+#endif
+	return TextureValue;
+}
+
+MaterialFloat4 ProcessMaterialExternalTextureLookup(MaterialFloat4 TextureValue)
+{
+#if COMPILER_GLSL_ES3_1
+	return MaterialFloat4(pow(TextureValue.rgb, 2.2f), TextureValue.a);
+#else
+	return ProcessMaterialColorTextureLookup(TextureValue);
+#endif
+}
+
+MaterialFloat4 ProcessMaterialLinearColorTextureLookup(MaterialFloat4 TextureValue)
+{
+	return TextureValue;
+}
+
+MaterialFloat4 ProcessMaterialGreyscaleTextureLookup(MaterialFloat4 TextureValue)
+{
+	// Sampling a greyscale texture in D3D9 gives: (G,G,G)
+	// Sampling a greyscale texture in D3D11 gives: (G,0,0)
+	// This replication reproduces the D3D9 behavior in all cases.
+	MaterialFloat GrayValue = TextureValue.r;
+#if (COMPILER_GLSL_ES3_1 || VULKAN_PROFILE) // OpenGLES3.1, Vulkan3.1 do not support sRGB sampling from R8
+	#if MOBILE_EMULATION
+	if( ResolvedView.MobilePreviewMode > 0.5f )
+	{
+		// undo HW srgb->lin
+		GrayValue = pow(GrayValue, 1.0f/2.2f); // TODO: replace with a more accurate lin -> sRGB conversion.
+	}
+	#endif
+	// sRGB read approximation (in highp if possible)
+	float LinValue = GrayValue;
+	LinValue *= LinValue;
+	return MaterialFloat4(LinValue.rrrr);
+#endif 
+	return GrayValue.rrrr;
+}
+
+MaterialFloat4 ProcessMaterialLinearGreyscaleTextureLookup(MaterialFloat4 TextureValue)
+{
+	// Sampling a greyscale texture in D3D9 gives: (G,G,G)
+	// Sampling a greyscale texture in D3D11 gives: (G,0,0)
+	// This replication reproduces the D3D9 behavior in all cases.
+	return TextureValue.rrrr;
+}
+
+MaterialFloat4 ProcessMaterialAlphaTextureLookup(MaterialFloat4 TextureValue)
+{
+	// Sampling a single channel texture in D3D9 gives: (G,G,G)
+	// Sampling a single channel texture in D3D11 gives: (G,0,0)
+	// This replication reproduces the D3D9 behavior in all cases.
+	return TextureValue.rrrr;
+}
+
 /** Accesses a shared material sampler or falls back if independent samplers are not supported. */
 SamplerState GetMaterialSharedSampler(SamplerState TextureSampler, SamplerState SharedSampler)
 {
diff --git a/Engine/Shaders/Private/MaterialTexture.ush b/Engine/Shaders/Private/MaterialTexture.ush
deleted file mode 100644
diff --git a/Engine/Source/Runtime/Engine/Classes/Materials/MaterialExpression.h b/Engine/Source/Runtime/Engine/Classes/Materials/MaterialExpression.h
--- a/Engine/Source/Runtime/Engine/Classes/Materials/MaterialExpression.h
+++ b/Engine/Source/Runtime/Engine/Classes/Materials/MaterialExpression.h
@@ -29,6 +29,7 @@ namespace HLSLTree
 class FScope;
 class FStatement;
 class FExpression;
+class FTextureParameterDeclaration;
 }
 }
 
@@ -306,6 +307,7 @@ class ENGINE_API UMaterialExpression : public UObject
 	 */
 	virtual bool GenerateHLSLStatements(FMaterialHLSLGenerator& Generator, UE::HLSLTree::FScope& Scope);
 	virtual bool GenerateHLSLExpression(FMaterialHLSLGenerator& Generator, UE::HLSLTree::FScope& Scope, int32 OutputIndex, UE::HLSLTree::FExpression* &OutExpression);
+	virtual bool GenerateHLSLTexture(FMaterialHLSLGenerator& Generator, UE::HLSLTree::FScope& Scope, int32 OutputIndex, UE::HLSLTree::FTextureParameterDeclaration*& OutTexture);
 
 #endif // WITH_EDITOR
 
diff --git a/Engine/Source/Runtime/Engine/Classes/Materials/MaterialExpressionAbs.h b/Engine/Source/Runtime/Engine/Classes/Materials/MaterialExpressionAbs.h
--- a/Engine/Source/Runtime/Engine/Classes/Materials/MaterialExpressionAbs.h
+++ b/Engine/Source/Runtime/Engine/Classes/Materials/MaterialExpressionAbs.h
@@ -27,7 +27,6 @@ class UMaterialExpressionAbs : public UMaterialExpression
 #if WITH_EDITOR
 	virtual int32 Compile(class FMaterialCompiler* Compiler, int32 OutputIndex) override;
 	virtual void GetCaption(TArray<FString>& OutCaptions) const override;
-	virtual bool GenerateHLSLExpression(FMaterialHLSLGenerator& Generator, UE::HLSLTree::FScope& Scope, int32 OutputIndex, UE::HLSLTree::FExpression*& OutExpression) override;
 #endif
 	//~ End UMaterialExpression Interface
 
diff --git a/Engine/Source/Runtime/Engine/Classes/Materials/MaterialExpressionCeil.h b/Engine/Source/Runtime/Engine/Classes/Materials/MaterialExpressionCeil.h
--- a/Engine/Source/Runtime/Engine/Classes/Materials/MaterialExpressionCeil.h
+++ b/Engine/Source/Runtime/Engine/Classes/Materials/MaterialExpressionCeil.h
@@ -22,7 +22,6 @@ class UMaterialExpressionCeil : public UMaterialExpression
 #if WITH_EDITOR
 	virtual int32 Compile(class FMaterialCompiler* Compiler, int32 OutputIndex) override;
 	virtual void GetCaption(TArray<FString>& OutCaptions) const override;
-	virtual bool GenerateHLSLExpression(FMaterialHLSLGenerator& Generator, UE::HLSLTree::FScope& Scope, int32 OutputIndex, UE::HLSLTree::FExpression*& OutExpression) override;
 #endif
 	//~ End UMaterialExpression Interface
 };
diff --git a/Engine/Source/Runtime/Engine/Classes/Materials/MaterialExpressionFloor.h b/Engine/Source/Runtime/Engine/Classes/Materials/MaterialExpressionFloor.h
--- a/Engine/Source/Runtime/Engine/Classes/Materials/MaterialExpressionFloor.h
+++ b/Engine/Source/Runtime/Engine/Classes/Materials/MaterialExpressionFloor.h
@@ -22,7 +22,6 @@ class UMaterialExpressionFloor : public UMaterialExpression
 #if WITH_EDITOR
 	virtual int32 Compile(class FMaterialCompiler* Compiler, int32 OutputIndex) override;
 	virtual void GetCaption(TArray<FString>& OutCaptions) const override;
-	virtual bool GenerateHLSLExpression(FMaterialHLSLGenerator& Generator, UE::HLSLTree::FScope& Scope, int32 OutputIndex, UE::HLSLTree::FExpression*& OutExpression) override;
 #endif
 	//~ End UMaterialExpression Interface
 };
diff --git a/Engine/Source/Runtime/Engine/Classes/Materials/MaterialExpressionFontSample.h b/Engine/Source/Runtime/Engine/Classes/Materials/MaterialExpressionFontSample.h
--- a/Engine/Source/Runtime/Engine/Classes/Materials/MaterialExpressionFontSample.h
+++ b/Engine/Source/Runtime/Engine/Classes/Materials/MaterialExpressionFontSample.h
@@ -10,14 +10,6 @@
 
 class UTexture;
 
-namespace UE
-{
-namespace Shader
-{
-struct FTextureValue;
-}
-}
-
 UCLASS(collapsecategories, hidecategories=Object, MinimalAPI)
 class UMaterialExpressionFontSample : public UMaterialExpression
 {
@@ -39,8 +31,6 @@ class UMaterialExpressionFontSample : public UMaterialExpression
 	virtual bool MatchesSearchQuery(const TCHAR* SearchQuery) override;
 	virtual int32 GetWidth() const override;
 	virtual int32 GetLabelPadding() override { return 8; }
-	virtual bool GenerateHLSLExpression(FMaterialHLSLGenerator& Generator, UE::HLSLTree::FScope& Scope, int32 OutputIndex, UE::HLSLTree::FExpression*& OutExpression) override;
-	virtual UE::HLSLTree::FExpression* GenerateHLSLTextureExpression(FMaterialHLSLGenerator& Generator, const UE::Shader::FTextureValue* TextureValue);
 #endif
 
 	/** 
diff --git a/Engine/Source/Runtime/Engine/Classes/Materials/MaterialExpressionFontSampleParameter.h b/Engine/Source/Runtime/Engine/Classes/Materials/MaterialExpressionFontSampleParameter.h
--- a/Engine/Source/Runtime/Engine/Classes/Materials/MaterialExpressionFontSampleParameter.h
+++ b/Engine/Source/Runtime/Engine/Classes/Materials/MaterialExpressionFontSampleParameter.h
@@ -73,7 +73,6 @@ class UMaterialExpressionFontSampleParameter : public UMaterialExpressionFontSam
 		}
 		return false;
 	}
-	virtual UE::HLSLTree::FExpression* GenerateHLSLTextureExpression(FMaterialHLSLGenerator& Generator, const UE::Shader::FTextureValue* TextureValue) override;
 #endif
 	//~ End UMaterialExpression Interface
 	
diff --git a/Engine/Source/Runtime/Engine/Classes/Materials/MaterialExpressionFrac.h b/Engine/Source/Runtime/Engine/Classes/Materials/MaterialExpressionFrac.h
--- a/Engine/Source/Runtime/Engine/Classes/Materials/MaterialExpressionFrac.h
+++ b/Engine/Source/Runtime/Engine/Classes/Materials/MaterialExpressionFrac.h
@@ -22,7 +22,6 @@ class UMaterialExpressionFrac : public UMaterialExpression
 #if WITH_EDITOR
 	virtual int32 Compile(class FMaterialCompiler* Compiler, int32 OutputIndex) override;
 	virtual void GetCaption(TArray<FString>& OutCaptions) const override;
-	virtual bool GenerateHLSLExpression(FMaterialHLSLGenerator& Generator, UE::HLSLTree::FScope& Scope, int32 OutputIndex, UE::HLSLTree::FExpression*& OutExpression) override;
 #endif
 	//~ End UMaterialExpression Interface
 };
diff --git a/Engine/Source/Runtime/Engine/Classes/Materials/MaterialExpressionIf.h b/Engine/Source/Runtime/Engine/Classes/Materials/MaterialExpressionIf.h
--- a/Engine/Source/Runtime/Engine/Classes/Materials/MaterialExpressionIf.h
+++ b/Engine/Source/Runtime/Engine/Classes/Materials/MaterialExpressionIf.h
@@ -46,7 +46,6 @@ class UMaterialExpressionIf : public UMaterialExpression
 	virtual uint32 GetInputType(int32 InputIndex) override;
 	virtual uint32 GetOutputType(int32 InputIndex) override {return MCT_Unknown;}
 	virtual bool IsResultMaterialAttributes(int32 OutputIndex) override;
-	virtual bool GenerateHLSLExpression(FMaterialHLSLGenerator& Generator, UE::HLSLTree::FScope& Scope, int32 OutputIndex, UE::HLSLTree::FExpression*& OutExpression) override;
 #endif
 	//~ End UMaterialExpression Interface
 };
diff --git a/Engine/Source/Runtime/Engine/Classes/Materials/MaterialExpressionPower.h b/Engine/Source/Runtime/Engine/Classes/Materials/MaterialExpressionPower.h
--- a/Engine/Source/Runtime/Engine/Classes/Materials/MaterialExpressionPower.h
+++ b/Engine/Source/Runtime/Engine/Classes/Materials/MaterialExpressionPower.h
@@ -29,7 +29,6 @@ class UMaterialExpressionPower : public UMaterialExpression
 	virtual int32 Compile(class FMaterialCompiler* Compiler, int32 OutputIndex) override;
 	virtual void GetCaption(TArray<FString>& OutCaptions) const override;
 	virtual void GetExpressionToolTip(TArray<FString>& OutToolTip) override;
-	virtual bool GenerateHLSLExpression(FMaterialHLSLGenerator& Generator, UE::HLSLTree::FScope& Scope, int32 OutputIndex, UE::HLSLTree::FExpression*& OutExpression) override;
 #endif
 	//~ End UMaterialExpression Interface
 
diff --git a/Engine/Source/Runtime/Engine/Classes/Materials/MaterialExpressionReroute.h b/Engine/Source/Runtime/Engine/Classes/Materials/MaterialExpressionReroute.h
--- a/Engine/Source/Runtime/Engine/Classes/Materials/MaterialExpressionReroute.h
+++ b/Engine/Source/Runtime/Engine/Classes/Materials/MaterialExpressionReroute.h
@@ -25,7 +25,6 @@ class ENGINE_API UMaterialExpressionReroute : public UMaterialExpressionRerouteB
 
 	virtual FText GetCreationDescription() const override;
 	virtual FText GetCreationName() const override;
-	virtual bool GenerateHLSLExpression(FMaterialHLSLGenerator& Generator, UE::HLSLTree::FScope& Scope, int32 OutputIndex, UE::HLSLTree::FExpression*& OutExpression) override;
 #endif
 	//~ End UMaterialExpression Interface
 
diff --git a/Engine/Source/Runtime/Engine/Classes/Materials/MaterialExpressionRound.h b/Engine/Source/Runtime/Engine/Classes/Materials/MaterialExpressionRound.h
--- a/Engine/Source/Runtime/Engine/Classes/Materials/MaterialExpressionRound.h
+++ b/Engine/Source/Runtime/Engine/Classes/Materials/MaterialExpressionRound.h
@@ -21,7 +21,6 @@ class UMaterialExpressionRound : public UMaterialExpression
 	virtual int32 Compile(class FMaterialCompiler* Compiler, int32 OutputIndex) override;
 	virtual void GetCaption(TArray<FString>& OutCaptions) const override;
 	virtual void GetExpressionToolTip(TArray<FString>& OutToolTip) override;
-	virtual bool GenerateHLSLExpression(FMaterialHLSLGenerator& Generator, UE::HLSLTree::FScope& Scope, int32 OutputIndex, UE::HLSLTree::FExpression*& OutExpression) override;
 #endif
 	//~ End UMaterialExpression Interface
 };
diff --git a/Engine/Source/Runtime/Engine/Classes/Materials/MaterialExpressionSaturate.h b/Engine/Source/Runtime/Engine/Classes/Materials/MaterialExpressionSaturate.h
--- a/Engine/Source/Runtime/Engine/Classes/Materials/MaterialExpressionSaturate.h
+++ b/Engine/Source/Runtime/Engine/Classes/Materials/MaterialExpressionSaturate.h
@@ -21,7 +21,6 @@ class UMaterialExpressionSaturate : public UMaterialExpression
 	virtual int32 Compile(class FMaterialCompiler* Compiler, int32 OutputIndex) override;
 	virtual void GetCaption(TArray<FString>& OutCaptions) const override;
 	virtual void GetExpressionToolTip(TArray<FString>& OutToolTip) override;
-	virtual bool GenerateHLSLExpression(FMaterialHLSLGenerator& Generator, UE::HLSLTree::FScope& Scope, int32 OutputIndex, UE::HLSLTree::FExpression*& OutExpression) override;
 #endif
 	//~ End UMaterialExpression Interface
 };
diff --git a/Engine/Source/Runtime/Engine/Classes/Materials/MaterialExpressionSign.h b/Engine/Source/Runtime/Engine/Classes/Materials/MaterialExpressionSign.h
--- a/Engine/Source/Runtime/Engine/Classes/Materials/MaterialExpressionSign.h
+++ b/Engine/Source/Runtime/Engine/Classes/Materials/MaterialExpressionSign.h
@@ -21,7 +21,6 @@ class UMaterialExpressionSign : public UMaterialExpression
 	virtual int32 Compile(class FMaterialCompiler* Compiler, int32 OutputIndex) override;
 	virtual void GetCaption(TArray<FString>& OutCaptions) const override;
 	virtual void GetExpressionToolTip(TArray<FString>& OutToolTip) override;
-	virtual bool GenerateHLSLExpression(FMaterialHLSLGenerator& Generator, UE::HLSLTree::FScope& Scope, int32 OutputIndex, UE::HLSLTree::FExpression*& OutExpression) override;
 #endif
 	//~ End UMaterialExpression Interface
 };
diff --git a/Engine/Source/Runtime/Engine/Classes/Materials/MaterialExpressionTextureObject.h b/Engine/Source/Runtime/Engine/Classes/Materials/MaterialExpressionTextureObject.h
--- a/Engine/Source/Runtime/Engine/Classes/Materials/MaterialExpressionTextureObject.h
+++ b/Engine/Source/Runtime/Engine/Classes/Materials/MaterialExpressionTextureObject.h
@@ -28,7 +28,7 @@ class UMaterialExpressionTextureObject : public UMaterialExpressionTextureBase
 	virtual int32 CompilePreview(class FMaterialCompiler* Compiler, int32 OutputIndex) override;
 	virtual void GetCaption(TArray<FString>& OutCaptions) const override;
 	virtual uint32 GetOutputType(int32 OutputIndex) override;
-	virtual bool GenerateHLSLExpression(FMaterialHLSLGenerator& Generator, UE::HLSLTree::FScope& Scope, int32 OutputIndex, UE::HLSLTree::FExpression*& OutExpression) override;
+	virtual bool GenerateHLSLTexture(FMaterialHLSLGenerator& Generator, UE::HLSLTree::FScope& Scope, int32 OutputIndex, UE::HLSLTree::FTextureParameterDeclaration*& OutTexture) override;
 	//~ End UMaterialExpression Interface
 #endif // WITH_EDITOR
 };
diff --git a/Engine/Source/Runtime/Engine/Classes/Materials/MaterialExpressionTextureObjectParameter.h b/Engine/Source/Runtime/Engine/Classes/Materials/MaterialExpressionTextureObjectParameter.h
--- a/Engine/Source/Runtime/Engine/Classes/Materials/MaterialExpressionTextureObjectParameter.h
+++ b/Engine/Source/Runtime/Engine/Classes/Materials/MaterialExpressionTextureObjectParameter.h
@@ -30,7 +30,7 @@ class UMaterialExpressionTextureObjectParameter : public UMaterialExpressionText
 	virtual int32 CompilePreview(class FMaterialCompiler* Compiler, int32 OutputIndex) override;
 	virtual uint32 GetOutputType(int32 InputIndex) override {return MCT_Texture;}
 	virtual const TArray<FExpressionInput*> GetInputs() override;
-	virtual bool GenerateHLSLExpression(FMaterialHLSLGenerator& Generator, UE::HLSLTree::FScope& Scope, int32 OutputIndex, UE::HLSLTree::FExpression*& OutExpression) override;
+	virtual bool GenerateHLSLTexture(FMaterialHLSLGenerator& Generator, UE::HLSLTree::FScope& Scope, int32 OutputIndex, UE::HLSLTree::FTextureParameterDeclaration*& OutTexture) override;
 #endif
 	//~ End UMaterialExpression Interface
 };
diff --git a/Engine/Source/Runtime/Engine/Classes/Materials/MaterialExpressionTextureSample.h b/Engine/Source/Runtime/Engine/Classes/Materials/MaterialExpressionTextureSample.h
--- a/Engine/Source/Runtime/Engine/Classes/Materials/MaterialExpressionTextureSample.h
+++ b/Engine/Source/Runtime/Engine/Classes/Materials/MaterialExpressionTextureSample.h
@@ -64,7 +64,7 @@ protected:
 	uint8 bShowTextureInputPin : 1;
 
 #if WITH_EDITOR
-	bool GenerateHLSLExpressionBase(FMaterialHLSLGenerator& Generator, UE::HLSLTree::FScope& Scope, UE::HLSLTree::FExpression* TextureExpression, UE::HLSLTree::FExpression*& OutExpression);
+	bool GenerateHLSLExpressionBase(FMaterialHLSLGenerator& Generator, UE::HLSLTree::FScope& Scope, UE::HLSLTree::FTextureParameterDeclaration* TextureDeclaration, UE::HLSLTree::FExpression*& OutExpression);
 #endif
 
 public:
diff --git a/Engine/Source/Runtime/Engine/Classes/Materials/MaterialExpressionTruncate.h b/Engine/Source/Runtime/Engine/Classes/Materials/MaterialExpressionTruncate.h
--- a/Engine/Source/Runtime/Engine/Classes/Materials/MaterialExpressionTruncate.h
+++ b/Engine/Source/Runtime/Engine/Classes/Materials/MaterialExpressionTruncate.h
@@ -21,7 +21,6 @@ class UMaterialExpressionTruncate : public UMaterialExpression
 	virtual int32 Compile(class FMaterialCompiler* Compiler, int32 OutputIndex) override;
 	virtual void GetCaption(TArray<FString>& OutCaptions) const override;
 	virtual void GetExpressionToolTip(TArray<FString>& OutToolTip) override;
-	virtual bool GenerateHLSLExpression(FMaterialHLSLGenerator& Generator, UE::HLSLTree::FScope& Scope, int32 OutputIndex, UE::HLSLTree::FExpression*& OutExpression) override;
 #endif
 	//~ End UMaterialExpression Interface
 };
diff --git a/Engine/Source/Runtime/Engine/Classes/Materials/MaterialExpressionVertexColor.h b/Engine/Source/Runtime/Engine/Classes/Materials/MaterialExpressionVertexColor.h
--- a/Engine/Source/Runtime/Engine/Classes/Materials/MaterialExpressionVertexColor.h
+++ b/Engine/Source/Runtime/Engine/Classes/Materials/MaterialExpressionVertexColor.h
@@ -19,7 +19,6 @@ class UMaterialExpressionVertexColor : public UMaterialExpression
 	virtual int32 Compile(class FMaterialCompiler* Compiler, int32 OutputIndex) override;
 	virtual void GetCaption(TArray<FString>& OutCaptions) const override;
 	virtual FText GetKeywords() const override { return FText::FromString(TEXT("vertex color")); }
-	virtual bool GenerateHLSLExpression(FMaterialHLSLGenerator& Generator, UE::HLSLTree::FScope& Scope, int32 OutputIndex, UE::HLSLTree::FExpression*& OutExpression) override;
 #endif
 	//~ End UMaterialExpression Interface
 };
diff --git a/Engine/Source/Runtime/Engine/Private/HLSLTree/HLSLTree.cpp b/Engine/Source/Runtime/Engine/Private/HLSLTree/HLSLTree.cpp
--- a/Engine/Source/Runtime/Engine/Private/HLSLTree/HLSLTree.cpp
+++ b/Engine/Source/Runtime/Engine/Private/HLSLTree/HLSLTree.cpp
@@ -93,18 +93,11 @@ FOperationDescription GetOperationDescription(EOperation Op)
 {
 	switch (Op)
 	{
-	case EOperation::None: return FOperationDescription(TEXT("None"), TEXT(""), 0, Shader::EPreshaderOpcode::Nop); break;
 	// Unary
-	case EOperation::Abs: return FOperationDescription(TEXT("Abs"), TEXT("abs"), 1, Shader::EPreshaderOpcode::Abs); break;
+	case EOperation::None: return FOperationDescription(TEXT("None"), TEXT(""), 0, Shader::EPreshaderOpcode::Nop); break;
 	case EOperation::Neg: return FOperationDescription(TEXT("Neg"), TEXT("-"), 1, Shader::EPreshaderOpcode::Neg); break;
 	case EOperation::Rcp: return FOperationDescription(TEXT("Rcp"), TEXT("/"), 1, Shader::EPreshaderOpcode::Rcp); break;
 	case EOperation::Frac: return FOperationDescription(TEXT("Frac"), TEXT("frac"), 1, Shader::EPreshaderOpcode::Frac); break;
-	case EOperation::Floor: return FOperationDescription(TEXT("Floor"), TEXT("floor"), 1, Shader::EPreshaderOpcode::Floor); break;
-	case EOperation::Ceil: return FOperationDescription(TEXT("Ceil"), TEXT("ceil"), 1, Shader::EPreshaderOpcode::Ceil); break;
-	case EOperation::Round: return FOperationDescription(TEXT("Round"), TEXT("round"), 1, Shader::EPreshaderOpcode::Round); break;
-	case EOperation::Trunc: return FOperationDescription(TEXT("Trunc"), TEXT("trunc"), 1, Shader::EPreshaderOpcode::Trunc); break;
-	case EOperation::Saturate: return FOperationDescription(TEXT("Saturate"), TEXT("saturate"), 1, Shader::EPreshaderOpcode::Saturate); break;
-	case EOperation::Sign: return FOperationDescription(TEXT("Sign"), TEXT("sign"), 1, Shader::EPreshaderOpcode::Sign); break;
 	case EOperation::Length: return FOperationDescription(TEXT("Length"), TEXT("length"), 1, Shader::EPreshaderOpcode::Length); break;
 	case EOperation::Normalize: return FOperationDescription(TEXT("Normalize"), TEXT("normalize"), 1, Shader::EPreshaderOpcode::Normalize); break;
 	// Binary
@@ -113,14 +106,11 @@ FOperationDescription GetOperationDescription(EOperation Op)
 	case EOperation::Mul: return FOperationDescription(TEXT("Multiply"), TEXT("*"), 2, Shader::EPreshaderOpcode::Mul); break;
 	case EOperation::Div: return FOperationDescription(TEXT("Divide"), TEXT("/"), 2, Shader::EPreshaderOpcode::Div); break;
 	case EOperation::Fmod: return FOperationDescription(TEXT("Fmod"), TEXT("%"), 2, Shader::EPreshaderOpcode::Fmod); break;
-	case EOperation::PowPositiveClamped: return FOperationDescription(TEXT("PowPositiveClamped"), TEXT("PowPositiveClamped"), 2, Shader::EPreshaderOpcode::Nop); break;
 	case EOperation::Dot: return FOperationDescription(TEXT("Dot"), TEXT("dot"), 2, Shader::EPreshaderOpcode::Dot); break;
 	case EOperation::Min: return FOperationDescription(TEXT("Min"), TEXT("min"), 2, Shader::EPreshaderOpcode::Min); break;
 	case EOperation::Max: return FOperationDescription(TEXT("Max"), TEXT("max"), 2, Shader::EPreshaderOpcode::Max); break;
 	case EOperation::Less: return FOperationDescription(TEXT("Less"), TEXT("<"), 2, Shader::EPreshaderOpcode::Less); break;
 	case EOperation::Greater: return FOperationDescription(TEXT("Greater"), TEXT(">"), 2, Shader::EPreshaderOpcode::Greater); break;
-	case EOperation::LessEqual: return FOperationDescription(TEXT("LessEqual"), TEXT("<="), 2, Shader::EPreshaderOpcode::Nop); break;
-	case EOperation::GreaterEqual: return FOperationDescription(TEXT("GreaterEqual"), TEXT(">="), 2, Shader::EPreshaderOpcode::Nop); break;
 	case EOperation::VecMulMatrix3: return FOperationDescription(TEXT("VecMulMatrix3"), TEXT("mul"), 2, Shader::EPreshaderOpcode::Nop); break;
 	case EOperation::VecMulMatrix4: return FOperationDescription(TEXT("VecMulMatrix4"), TEXT("mul"), 2, Shader::EPreshaderOpcode::Nop); break;
 	case EOperation::Matrix3MulVec: return FOperationDescription(TEXT("Matrix3MulVec"), TEXT("mul"), 2, Shader::EPreshaderOpcode::Nop); break;
@@ -980,11 +970,6 @@ FPreparedType MergePreparedTypes(const FPreparedType& Lhs, const FPreparedType&
 
 bool FPrepareValueResult::TryMergePreparedType(FEmitContext& Context, const Shader::FStructType* StructType, Shader::EValueComponentType ComponentType)
 {
-	if (!StructType && ComponentType == Shader::EValueComponentType::Void)
-	{
-		return false;
-	}
-
 	// If we previously had a forwarded value set, reset that and start over
 	if (!PreparedType.IsInitialized())
 	{
@@ -1004,6 +989,10 @@ bool FPrepareValueResult::TryMergePreparedType(FEmitContext& Context, const Shad
 	}
 	else
 	{
+		if (ComponentType == Shader::EValueComponentType::Void)
+		{
+			return false;
+		}
 		PreparedType.ValueComponentType = Shader::CombineComponentTypes(PreparedType.ValueComponentType, ComponentType);
 	}
 
@@ -1590,5 +1579,11 @@ FFunction* FTree::NewFunction()
 	return NewFunction;
 }
 
+FTextureParameterDeclaration* FTree::NewTextureParameterDeclaration(const FName& Name, const FTextureDescription& DefaultValue)
+{
+	FTextureParameterDeclaration* Declaration = NewNode<FTextureParameterDeclaration>(Name, DefaultValue);
+	return Declaration;
+}
+
 } // namespace HLSLTree
 } // namespace UE
diff --git a/Engine/Source/Runtime/Engine/Private/HLSLTree/HLSLTreeCommon.cpp b/Engine/Source/Runtime/Engine/Private/HLSLTree/HLSLTreeCommon.cpp
--- a/Engine/Source/Runtime/Engine/Private/HLSLTree/HLSLTreeCommon.cpp
+++ b/Engine/Source/Runtime/Engine/Private/HLSLTree/HLSLTreeCommon.cpp
@@ -4,7 +4,6 @@
 #include "Misc/StringBuilder.h"
 #include "MaterialShared.h"
 #include "MaterialSceneTextureId.h"
-#include "Engine/BlendableInterface.h" // BL_AfterTonemapping
 #include "Engine/Texture.h"
 
 namespace UE
@@ -51,10 +50,6 @@ FExternalInputDescription GetExternalInputDescription(EExternalInput Input)
 	case EExternalInput::TexCoord6_Ddy: return FExternalInputDescription(TEXT("TexCoord6_Ddy"), Shader::EValueType::Float2);
 	case EExternalInput::TexCoord7_Ddy: return FExternalInputDescription(TEXT("TexCoord7_Ddy"), Shader::EValueType::Float2);
 
-	case EExternalInput::VertexColor: return FExternalInputDescription(TEXT("VertexColor"), Shader::EValueType::Float4, EExternalInput::VertexColor_Ddx, EExternalInput::VertexColor_Ddy);
-	case EExternalInput::VertexColor_Ddx: return FExternalInputDescription(TEXT("VertexColor_Ddx"), Shader::EValueType::Float4);
-	case EExternalInput::VertexColor_Ddy: return FExternalInputDescription(TEXT("VertexColor_Ddy"), Shader::EValueType::Float4);
-
 	case EExternalInput::WorldPosition: return FExternalInputDescription(TEXT("WorldPosition"), Shader::EValueType::Double3, EExternalInput::WorldPosition_Ddx, EExternalInput::WorldPosition_Ddy, EExternalInput::PrevWorldPosition);
 	case EExternalInput::WorldPosition_NoOffsets: return FExternalInputDescription(TEXT("WorldPosition_NoOffsets"), Shader::EValueType::Double3, EExternalInput::WorldPosition_Ddx, EExternalInput::WorldPosition_Ddy, EExternalInput::PrevWorldPosition_NoOffsets);
 	case EExternalInput::TranslatedWorldPosition: return FExternalInputDescription(TEXT("TranslatedWorldPosition"), Shader::EValueType::Float3, EExternalInput::WorldPosition_Ddx, EExternalInput::WorldPosition_Ddy, EExternalInput::PrevTranslatedWorldPosition);
@@ -221,75 +216,8 @@ void FExpressionMaterialParameter::ComputeAnalyticDerivatives(FTree& Tree, FExpr
 
 bool FExpressionMaterialParameter::PrepareValue(FEmitContext& Context, FEmitScope& Scope, const FRequestedType& RequestedType, FPrepareValueResult& OutResult) const
 {
-	EExpressionEvaluation Evaluation = EExpressionEvaluation::Shader;
-	if (IsStaticMaterialParameter(ParameterType))
-	{
-		Evaluation = EExpressionEvaluation::Constant;
-	}
-	else if (ParameterType == EMaterialParameterType::Scalar ||
-		ParameterType == EMaterialParameterType::Vector ||
-		ParameterType == EMaterialParameterType::DoubleVector)
-	{
-		Evaluation = EExpressionEvaluation::Preshader;
-	}
-	return OutResult.SetType(Context, RequestedType, Evaluation, DefaultValue.Type);
-}
-
-void FExpressionMaterialParameter::EmitValueShader(FEmitContext& Context, FEmitScope& Scope, const FRequestedType& RequestedType, FEmitValueShaderResult& OutResult) const
-{
-	if (ParameterType == EMaterialParameterType::Texture)
-	{
-		const Shader::FTextureValue* TextureValue = DefaultValue.AsTexture();
-		check(TextureValue);
-
-		FMaterialTextureParameterInfo TextureParameterInfo;
-		TextureParameterInfo.ParameterInfo = ParameterName;
-		TextureParameterInfo.TextureIndex = Context.Material->GetReferencedTextures().Find(TextureValue->Texture);
-		TextureParameterInfo.SamplerSource = SSM_FromTextureAsset; // TODO - Is this needed?
-		check(TextureParameterInfo.TextureIndex != INDEX_NONE);
-
-		const Shader::EValueType TextureType = TextureValue->GetType();
-		EMaterialTextureParameterType TextureParameterType = EMaterialTextureParameterType::Count;
-		const TCHAR* ConstructorName = nullptr;
-		const TCHAR* TextureTypeName = nullptr;
-		switch (TextureType)
-		{
-		case Shader::EValueType::Texture2D:
-			TextureParameterType = EMaterialTextureParameterType::Standard2D;
-			ConstructorName = TEXT("MakeTexture2D");
-			TextureTypeName = TEXT("Texture2D");
-			break;
-		case Shader::EValueType::Texture2DArray:
-			TextureParameterType = EMaterialTextureParameterType::Array2D;
-			ConstructorName = TEXT("MakeTexture2DArray");
-			TextureTypeName = TEXT("Texture2DArray");
-			break;
-		case Shader::EValueType::TextureCube:
-			TextureParameterType = EMaterialTextureParameterType::Cube;
-			ConstructorName = TEXT("MakeTextureCube");
-			TextureTypeName = TEXT("TextureCube");
-			break;
-		case Shader::EValueType::TextureCubeArray:
-			TextureParameterType = EMaterialTextureParameterType::ArrayCube;
-			ConstructorName = TEXT("MakeTextureCubeArray");
-			TextureTypeName = TEXT("TextureCubeArray");
-			break;
-		case Shader::EValueType::Texture3D:
-			TextureParameterType = EMaterialTextureParameterType::Volume;
-			ConstructorName = TEXT("MakeTexture3D");
-			TextureTypeName = TEXT("VolumeTexture");
-			break;
-		default:
-			checkNoEntry();
-			break;
-		}
-
-		const int32 ParameterIndex = Context.MaterialCompilationOutput->UniformExpressionSet.FindOrAddTextureParameter(TextureParameterType, TextureParameterInfo);
-
-		TStringBuilder<256> FormattedCode;
-		FormattedCode.Appendf(TEXT("%s(Material.%s_%d, Material.%s_%dSampler, %d)"), ConstructorName, TextureTypeName, ParameterIndex, TextureTypeName, ParameterIndex, (int32)TextureValue->SamplerType);
-		OutResult.Code = Context.EmitExpression(Scope, TextureType, FormattedCode.ToView());
-	}
+	const EExpressionEvaluation Evaluation = IsStaticMaterialParameter(ParameterType) ? EExpressionEvaluation::Constant : EExpressionEvaluation::Preshader;
+	return OutResult.SetType(Context, RequestedType, Evaluation, GetShaderValueType(ParameterType));
 }
 
 void FExpressionMaterialParameter::EmitValuePreshader(FEmitContext& Context, FEmitScope& Scope, const FRequestedType& RequestedType, FEmitValuePreshaderResult& OutResult) const
@@ -399,9 +327,6 @@ void FExpressionExternalInput::EmitValueShader(FEmitContext& Context, FEmitScope
 		const TCHAR* Code = nullptr;
 		switch (InputType)
 		{
-		case EExternalInput::VertexColor: Code = TEXT("Parameters.VertexColor"); Context.bUsesVertexColor = (Context.ShaderFrequency != SF_Vertex); break;
-		case EExternalInput::VertexColor_Ddx: Code = TEXT("Parameters.VertexColor_DDX"); Context.bUsesVertexColor = (Context.ShaderFrequency != SF_Vertex); break;
-		case EExternalInput::VertexColor_Ddy: Code = TEXT("Parameters.VertexColor_DDY"); Context.bUsesVertexColor = (Context.ShaderFrequency != SF_Vertex); break;
 		case EExternalInput::WorldPosition: Code = TEXT("GetWorldPosition(Parameters)"); break;
 		case EExternalInput::WorldPosition_NoOffsets: Code = TEXT("GetWorldPosition_NoMaterialOffsets(Parameters)"); break;
 		case EExternalInput::TranslatedWorldPosition: Code = TEXT("GetTranslatedWorldPosition(Parameters)"); break;
@@ -505,18 +430,7 @@ void FExpressionMaterialSceneTexture::EmitValueShader(FEmitContext& Context, FEm
 
 bool FExpressionTextureSample::PrepareValue(FEmitContext& Context, FEmitScope& Scope, const FRequestedType& RequestedType, FPrepareValueResult& OutResult) const
 {
-	const FPreparedType& TextureType = Context.PrepareExpression(TextureExpression, Scope, ERequestedType::Texture);
-	if (!Shader::IsTextureType(TextureType.ValueComponentType))
-	{
-		return Context.Errors->AddError(TEXT("Expected texture"));
-	}
-
-	const FPreparedType& TexCoordType = Context.PrepareExpression(TexCoordExpression, Scope, ERequestedType::Vector2);
-	if (TexCoordType.IsVoid())
-	{
-		return false;
-	}
-
+	Context.PrepareExpression(TexCoordExpression, Scope, ERequestedType::Vector2);
 	Context.PrepareExpression(TexCoordDerivatives.ExpressionDdx, Scope, ERequestedType::Vector2);
 	Context.PrepareExpression(TexCoordDerivatives.ExpressionDdy, Scope, ERequestedType::Vector2);
 
@@ -525,88 +439,156 @@ bool FExpressionTextureSample::PrepareValue(FEmitContext& Context, FEmitScope& S
 
 void FExpressionTextureSample::EmitValueShader(FEmitContext& Context, FEmitScope& Scope, const FRequestedType& RequestedType, FEmitValueShaderResult& OutResult) const
 {
-	FEmitShaderExpression* EmitTexture = TextureExpression->GetValueShader(Context, Scope);
-	const Shader::EValueType TextureType = EmitTexture->Type;
-	check(Shader::IsTextureType(TextureType));
-
+	const FTextureDescription& Desc = Declaration->Description;
+	const EMaterialValueType MaterialValueType = Desc.Texture->GetMaterialType();
+	EMaterialTextureParameterType TextureType = EMaterialTextureParameterType::Count;
 	bool bVirtualTexture = false;
+	const TCHAR* TextureTypeName = nullptr;
 	const TCHAR* SampleFunctionName = nullptr;
-	switch (TextureType)
+	switch (MaterialValueType)
 	{
-	case Shader::EValueType::Texture2D:
+	case MCT_Texture2D:
+		TextureType = EMaterialTextureParameterType::Standard2D;
+		TextureTypeName = TEXT("Texture2D");
 		SampleFunctionName = TEXT("Texture2DSample");
 		break;
-	case Shader::EValueType::TextureCube:
+	case MCT_TextureCube:
+		TextureType = EMaterialTextureParameterType::Cube;
+		TextureTypeName = TEXT("TextureCube");
 		SampleFunctionName = TEXT("TextureCubeSample");
 		break;
-	case Shader::EValueType::Texture2DArray:
+	case MCT_Texture2DArray:
+		TextureType = EMaterialTextureParameterType::Array2D;
+		TextureTypeName = TEXT("Texture2DArray");
 		SampleFunctionName = TEXT("Texture2DArraySample");
 		break;
-	case Shader::EValueType::Texture3D:
+	case MCT_VolumeTexture:
+		TextureType = EMaterialTextureParameterType::Volume;
+		TextureTypeName = TEXT("VolumeTexture");
 		SampleFunctionName = TEXT("Texture3DSample");
 		break;
-	/*case MCT_TextureExternal:
+	case MCT_TextureExternal:
 		// TODO
+		TextureTypeName = TEXT("ExternalTexture");
 		SampleFunctionName = TEXT("TextureExternalSample");
 		break;
 	case MCT_TextureVirtual:
 		// TODO
+		TextureType = EMaterialTextureParameterType::Virtual;
 		SampleFunctionName = TEXT("TextureVirtualSample");
 		bVirtualTexture = true;
-		break;*/
+		break;
 	default:
 		checkNoEntry();
 		break;
 	}
 
-	const bool AutomaticViewMipBias = false; // TODO
-	TStringBuilder<256> FormattedSampler;
-	switch (SamplerSource)
+	const TCHAR* SamplerTypeFunction = TEXT("");
+	switch (Desc.SamplerType)
 	{
-	case SSM_FromTextureAsset:
-		FormattedSampler.Appendf(TEXT("%s.Sampler"), EmitTexture->Reference);
+	case SAMPLERTYPE_External:
+		SamplerTypeFunction = TEXT("ProcessMaterialExternalTextureLookup");
+		break;
+	case SAMPLERTYPE_Color:
+		SamplerTypeFunction = TEXT("ProcessMaterialColorTextureLookup");
+		break;
+	case SAMPLERTYPE_VirtualColor:
+		// has a mobile specific workaround
+		SamplerTypeFunction = TEXT("ProcessMaterialVirtualColorTextureLookup");
+		break;
+	case SAMPLERTYPE_LinearColor:
+	case SAMPLERTYPE_VirtualLinearColor:
+		SamplerTypeFunction = TEXT("ProcessMaterialLinearColorTextureLookup");
+		break;
+	case SAMPLERTYPE_Alpha:
+	case SAMPLERTYPE_VirtualAlpha:
+	case SAMPLERTYPE_DistanceFieldFont:
+		SamplerTypeFunction = TEXT("ProcessMaterialAlphaTextureLookup");
+		break;
+	case SAMPLERTYPE_Grayscale:
+	case SAMPLERTYPE_VirtualGrayscale:
+		SamplerTypeFunction = TEXT("ProcessMaterialGreyscaleTextureLookup");
 		break;
-	case SSM_Wrap_WorldGroupSettings:
-		FormattedSampler.Appendf(TEXT("GetMaterialSharedSampler(%s.Sampler,%s)"),
-			EmitTexture->Reference,
-			AutomaticViewMipBias ? TEXT("View.MaterialTextureBilinearWrapedSampler") : TEXT("Material.Wrap_WorldGroupSettings"));
+	case SAMPLERTYPE_LinearGrayscale:
+	case SAMPLERTYPE_VirtualLinearGrayscale:
+		SamplerTypeFunction = TEXT("ProcessMaterialLinearGreyscaleTextureLookup");
 		break;
-	case SSM_Clamp_WorldGroupSettings:
-		FormattedSampler.Appendf(TEXT("GetMaterialSharedSampler(%s.Sampler,%s)"),
-			EmitTexture->Reference,
-			AutomaticViewMipBias ? TEXT("View.MaterialTextureBilinearClampedSampler") : TEXT("Material.Clamp_WorldGroupSettings"));
+	case SAMPLERTYPE_Normal:
+	case SAMPLERTYPE_VirtualNormal:
+		// Normal maps need to be unpacked in the pixel shader.
+		SamplerTypeFunction = TEXT("UnpackNormalMap");
+		break;
+	case SAMPLERTYPE_Masks:
+	case SAMPLERTYPE_VirtualMasks:
+	case SAMPLERTYPE_Data:
+		SamplerTypeFunction = TEXT("");
 		break;
 	default:
-		checkNoEntry();
+		check(0);
 		break;
 	}
 
+	FMaterialTextureParameterInfo TextureParameterInfo;
+	TextureParameterInfo.ParameterInfo = Declaration->Name;
+	TextureParameterInfo.TextureIndex = Context.Material->GetReferencedTextures().Find(Desc.Texture);
+	TextureParameterInfo.SamplerSource = SamplerSource;
+	check(TextureParameterInfo.TextureIndex != INDEX_NONE);
+
+	const int32 ParameterIndex = Context.MaterialCompilationOutput->UniformExpressionSet.FindOrAddTextureParameter(TextureType, TextureParameterInfo);
+	const FString TextureName = FString::Printf(TEXT("Material.%s_%u"), TextureTypeName, ParameterIndex);
+
+	FString SamplerStateCode;
+	bool AutomaticViewMipBias = false; // TODO
+	bool RequiresManualViewMipBias = AutomaticViewMipBias;
+
+	if (!bVirtualTexture) //VT does not have explict samplers (and always requires manual view mip bias)
+	{
+		if (SamplerSource == SSM_FromTextureAsset)
+		{
+			SamplerStateCode = FString::Printf(TEXT("%sSampler"), *TextureName);
+		}
+		else if (SamplerSource == SSM_Wrap_WorldGroupSettings)
+		{
+			// Use the shared sampler to save sampler slots
+			const TCHAR* SharedSamplerName = AutomaticViewMipBias ? TEXT("View.MaterialTextureBilinearWrapedSampler") : TEXT("Material.Wrap_WorldGroupSettings");
+			SamplerStateCode = FString::Printf(TEXT("GetMaterialSharedSampler(%sSampler,%s)"), *TextureName, SharedSamplerName);
+			RequiresManualViewMipBias = false;
+		}
+		else if (SamplerSource == SSM_Clamp_WorldGroupSettings)
+		{
+			// Use the shared sampler to save sampler slots
+			const TCHAR* SharedSamplerName = AutomaticViewMipBias ? TEXT("View.MaterialTextureBilinearClampedSampler") : TEXT("Material.Clamp_WorldGroupSettings");
+			SamplerStateCode = FString::Printf(TEXT("GetMaterialSharedSampler(%sSampler,%s)"), *TextureName, SharedSamplerName);
+			RequiresManualViewMipBias = false;
+		}
+	}
+
 	FEmitShaderExpression* TexCoordValue = TexCoordExpression->GetValueShader(Context, Scope, Shader::EValueType::Float2);
 	FEmitShaderExpression* TexCoordValueDdx = nullptr;
 	FEmitShaderExpression* TexCoordValueDdy = nullptr;
-	FEmitShaderExpression* TextureResult = nullptr;
 	if (TexCoordDerivatives.IsValid())
 	{
 		TexCoordValueDdx = TexCoordDerivatives.ExpressionDdx->GetValueShader(Context, Scope, Shader::EValueType::Float2);
 		TexCoordValueDdy = TexCoordDerivatives.ExpressionDdy->GetValueShader(Context, Scope, Shader::EValueType::Float2);
-		TextureResult = Context.EmitExpression(Scope, Shader::EValueType::Float4, TEXT("%Grad(%.Texture, %, %, %, %)"),
+		OutResult.Code = Context.EmitExpression(Scope, Shader::EValueType::Float4, TEXT("%(%Grad(%, %, %, %, %))"),
+			SamplerTypeFunction,
 			SampleFunctionName,
-			EmitTexture,
-			FormattedSampler.ToString(),
+			*TextureName,
+			*SamplerStateCode,
 			TexCoordValue,
 			TexCoordValueDdx,
 			TexCoordValueDdy);
 	}
 	else
 	{
-		TextureResult = Context.EmitExpression(Scope, Shader::EValueType::Float4, TEXT("%(%.Texture, %, %)"),
+		OutResult.Code = Context.EmitExpression(Scope, Shader::EValueType::Float4, TEXT("%(%(%, %, %))"),
+			SamplerTypeFunction,
 			SampleFunctionName,
-			EmitTexture,
-			FormattedSampler.ToString(),
+			*TextureName,
+			*SamplerStateCode,
 			TexCoordValue);
 	}
-
-	OutResult.Code = Context.EmitExpression(Scope, Shader::EValueType::Float4, TEXT("ApplyMaterialSamplerType(%, %.SamplerType)"), TextureResult, EmitTexture);
+	
 }
 
 void FExpressionGetStructField::ComputeAnalyticDerivatives(FTree& Tree, FExpressionDerivatives& OutResult) const
@@ -855,8 +837,6 @@ void FExpressionOperation::ComputeAnalyticDerivatives(FTree& Tree, FExpressionDe
 	{
 	case EOperation::Less:
 	case EOperation::Greater:
-	case EOperation::LessEqual:
-	case EOperation::GreaterEqual:
 		OutResult.ExpressionDdx = Tree.NewConstant(0.0f);
 		OutResult.ExpressionDdx = OutResult.ExpressionDdy;
 		break;
@@ -899,14 +879,6 @@ void FExpressionOperation::ComputeAnalyticDerivatives(FTree& Tree, FExpressionDe
 		break;
 	case EOperation::Length:
 	case EOperation::Normalize:
-	case EOperation::Floor:
-	case EOperation::Ceil:
-	case EOperation::Round:
-	case EOperation::Trunc:
-	case EOperation::Sign:
-	case EOperation::Abs:
-	case EOperation::Saturate:
-	case EOperation::PowPositiveClamped:
 		// TODO
 		break;
 	case EOperation::Add:
@@ -1065,7 +1037,7 @@ bool FExpressionOperation::PrepareValue(FEmitContext& Context, FEmitScope& Scope
 			return Context.Errors->AddErrorf(TEXT("Invalid arithmetic between %s and %s"), InputType[0].GetType().GetName(), InputType[1].GetType().GetName());
 		}
 		ResultType = MergePreparedTypes(InputType[0], InputType[1]);
-		if (Op == EOperation::Less || Op == EOperation::Greater || Op == EOperation::LessEqual || Op == EOperation::GreaterEqual)
+		if (Op == EOperation::Less || Op == EOperation::Greater)
 		{
 			ResultType.ValueComponentType = Shader::EValueComponentType::Bool;
 		}
@@ -1123,23 +1095,15 @@ FOperationTypes GetOperationTypes(EOperation Op, TConstArrayView<Shader::EValueT
 	case EOperation::Normalize:
 	case EOperation::Frac:
 	case EOperation::Rcp:
-	case EOperation::Sign:
-	case EOperation::Saturate:
 		Types.ResultType = Shader::MakeNonLWCType(IntermediateType);
 		break;
 	case EOperation::Less:
 	case EOperation::Greater:
-	case EOperation::LessEqual:
-	case EOperation::GreaterEqual:
 		Types.ResultType = Shader::MakeValueType(Shader::EValueComponentType::Bool, NumComponents);
 		break;
 	case EOperation::Fmod:
 		Types.InputType[1] = Types.ResultType = Shader::MakeNonLWCType(IntermediateType);
 		break;
-	case EOperation::PowPositiveClamped:
-		// No LWC support yet
-		Types.InputType[0] = Types.InputType[1] = Types.ResultType = Shader::MakeNonLWCType(IntermediateType);
-		break;
 	case EOperation::Dot:
 		Types.ResultType = Shader::MakeValueType(ComponentType, 1);
 		break;
@@ -1185,16 +1149,6 @@ void FExpressionOperation::EmitValueShader(FEmitContext& Context, FEmitScope& Sc
 
 	switch (Op)
 	{
-	case EOperation::Abs:
-		if (Types.bIsLWC)
-		{
-			OutResult.Code = Context.EmitExpression(Scope, Types.ResultType, TEXT("LWCAbs(%)"), InputValue[0]);
-		}
-		else
-		{
-			OutResult.Code = Context.EmitInlineExpression(Scope, Types.ResultType, TEXT("abs(%)"), InputValue[0]);
-		}
-		break;
 	case EOperation::Neg:
 		if (Types.bIsLWC)
 		{
@@ -1225,67 +1179,6 @@ void FExpressionOperation::EmitValueShader(FEmitContext& Context, FEmitScope& Sc
 			OutResult.Code = Context.EmitExpression(Scope, Types.ResultType, TEXT("frac(%)"), InputValue[0]);
 		}
 		break;
-
-	case EOperation::Floor:
-		if (Types.bIsLWC)
-		{
-			OutResult.Code = Context.EmitExpression(Scope, Types.ResultType, TEXT("LWCFloor(%)"), InputValue[0]);
-		}
-		else
-		{
-			OutResult.Code = Context.EmitExpression(Scope, Types.ResultType, TEXT("floor(%)"), InputValue[0]);
-		}
-		break;
-	case EOperation::Ceil:
-		if (Types.bIsLWC)
-		{
-			OutResult.Code = Context.EmitExpression(Scope, Types.ResultType, TEXT("LWCCeil(%)"), InputValue[0]);
-		}
-		else
-		{
-			OutResult.Code = Context.EmitExpression(Scope, Types.ResultType, TEXT("ceil(%)"), InputValue[0]);
-		}
-		break;
-	case EOperation::Round:
-		if (Types.bIsLWC)
-		{
-			OutResult.Code = Context.EmitExpression(Scope, Types.ResultType, TEXT("LWCRound(%)"), InputValue[0]);
-		}
-		else
-		{
-			OutResult.Code = Context.EmitExpression(Scope, Types.ResultType, TEXT("round(%)"), InputValue[0]);
-		}
-		break;
-	case EOperation::Trunc:
-		if (Types.bIsLWC)
-		{
-			OutResult.Code = Context.EmitExpression(Scope, Types.ResultType, TEXT("LWCTrunc(%)"), InputValue[0]);
-		}
-		else
-		{
-			OutResult.Code = Context.EmitExpression(Scope, Types.ResultType, TEXT("trunc(%)"), InputValue[0]);
-		}
-		break;
-	case EOperation::Saturate:
-		if (Types.bIsLWC)
-		{
-			OutResult.Code = Context.EmitExpression(Scope, Types.ResultType, TEXT("LWCSaturate(%)"), InputValue[0]);
-		}
-		else
-		{
-			OutResult.Code = Context.EmitExpression(Scope, Types.ResultType, TEXT("saturate(%)"), InputValue[0]);
-		}
-		break;
-	case EOperation::Sign:
-		if (Types.bIsLWC)
-		{
-			OutResult.Code = Context.EmitExpression(Scope, Types.ResultType, TEXT("LWCSign(%)"), InputValue[0]);
-		}
-		else
-		{
-			OutResult.Code = Context.EmitExpression(Scope, Types.ResultType, TEXT("sign(%)"), InputValue[0]);
-		}
-		break;
 	case EOperation::Length:
 		if (Types.bIsLWC)
 		{
@@ -1356,9 +1249,6 @@ void FExpressionOperation::EmitValueShader(FEmitContext& Context, FEmitScope& Sc
 			OutResult.Code = Context.EmitExpression(Scope, Types.ResultType, TEXT("fmod(%, %)"), InputValue[0], InputValue[1]);
 		}
 		break;
-	case EOperation::PowPositiveClamped:
-		OutResult.Code = Context.EmitExpression(Scope, Types.ResultType, TEXT("PositiveClampedPow(%, %)"), InputValue[0], InputValue[1]);
-		break;
 	case EOperation::Dot:
 		if (Types.bIsLWC)
 		{
@@ -1409,26 +1299,6 @@ void FExpressionOperation::EmitValueShader(FEmitContext& Context, FEmitScope& Sc
 			OutResult.Code = Context.EmitExpression(Scope, Types.ResultType, TEXT("(% > %)"), InputValue[0], InputValue[1]);
 		}
 		break;
-	case EOperation::LessEqual:
-		if (Types.bIsLWC)
-		{
-			OutResult.Code = Context.EmitExpression(Scope, Types.ResultType, TEXT("LWCLessEqual(%, %)"), InputValue[0], InputValue[1]);
-		}
-		else
-		{
-			OutResult.Code = Context.EmitExpression(Scope, Types.ResultType, TEXT("(% <= %)"), InputValue[0], InputValue[1]);
-		}
-		break;
-	case EOperation::GreaterEqual:
-		if (Types.bIsLWC)
-		{
-			OutResult.Code = Context.EmitExpression(Scope, Types.ResultType, TEXT("LWCGreaterEqual(%, %)"), InputValue[0], InputValue[1]);
-		}
-		else
-		{
-			OutResult.Code = Context.EmitExpression(Scope, Types.ResultType, TEXT("(% >= %)"), InputValue[0], InputValue[1]);
-		}
-		break;
 	case EOperation::VecMulMatrix3:
 		if (Types.bIsLWC)
 		{
diff --git a/Engine/Source/Runtime/Engine/Private/HLSLTree/HLSLTreeEmit.cpp b/Engine/Source/Runtime/Engine/Private/HLSLTree/HLSLTreeEmit.cpp
--- a/Engine/Source/Runtime/Engine/Private/HLSLTree/HLSLTreeEmit.cpp
+++ b/Engine/Source/Runtime/Engine/Private/HLSLTree/HLSLTreeEmit.cpp
@@ -216,45 +216,22 @@ void EmitCustomHLSL(const FEmitCustomHLSL& EmitCustomHLSL, const TCHAR* Paramete
 
 	// Function that wraps custom HLSL, provides the interface expected by custom HLSL
 	// * LWC inputs have both LWC version and non-LWC version
-	// * Texture inputs are split into a Texture/SamplerState pair
 	// * First output is given through return value, additional outputs use inout function parameters
 	OutCode.Appendf(TEXT("%s CustomExpressionInternal%d(%s Parameters"), OutputStruct->Fields[0].Type.GetName(), EmitCustomHLSL.Index, ParametersTypeName);
 	for (const FEmitCustomHLSLInput& Input : EmitCustomHLSL.Inputs)
 	{
 		Shader::FType InputType = Input.Type;
-		if (InputType.IsTexture())
+		if (InputType.IsNumericLWC())
 		{
-			// Texture parameter followed by a SamplerState parameter
-			const TCHAR* TextureTypeName = nullptr;
-			switch (InputType.ValueType)
-			{
-			case Shader::EValueType::Texture2D: TextureTypeName = TEXT("Texture2D"); break;
-			case Shader::EValueType::Texture2DArray: TextureTypeName = TEXT("Texture2DArray"); break;
-			case Shader::EValueType::TextureCube: TextureTypeName = TEXT("TextureCube"); break;
-			case Shader::EValueType::TextureCubeArray: TextureTypeName = TEXT("TextureCubeArray"); break;
-			case Shader::EValueType::Texture3D: TextureTypeName = TEXT("Texture3D"); break;
-			default: checkNoEntry(); break;
-			}
-			OutCode.Appendf(TEXT(", %s "), TextureTypeName);
+			// Add an additional input for LWC type with the LWC-prefix
+			OutCode.Appendf(TEXT(", %s LWC"), InputType.GetName());
 			OutCode.Append(Input.Name);
-			OutCode.Appendf(TEXT(", SamplerState "));
-			OutCode.Append(Input.Name);
-			OutCode.Append(TEXT("Sampler"));
+			// Regular, unprefixed input uses non-LWC type
+			InputType = InputType.GetNonLWCType();
 		}
-		else
-		{
-			if (InputType.IsNumericLWC())
-			{
-				// Add an additional input for LWC type with the LWC-prefix
-				OutCode.Appendf(TEXT(", %s LWC"), InputType.GetName());
-				OutCode.Append(Input.Name);
-				// Regular, unprefixed input uses non-LWC type
-				InputType = InputType.GetNonLWCType();
-			}
 
-			OutCode.Appendf(TEXT(", %s "), InputType.GetName());
-			OutCode.Append(Input.Name);
-		}
+		OutCode.Appendf(TEXT(", %s "), InputType.GetName());
+		OutCode.Append(Input.Name);
 	}
 	for (int32 OutputIndex = 1; OutputIndex < OutputStruct->Fields.Num(); ++OutputIndex)
 	{
@@ -266,7 +243,7 @@ void EmitCustomHLSL(const FEmitCustomHLSL& EmitCustomHLSL, const TCHAR* Paramete
 	OutCode.Append(TEXT("\n}\n"));
 
 	// Function that calls the above wrapper, provides the interface expected by HLSLTree
-	// * All inputs types are passed through directly
+	// * No special handling for LWC inputs
 	// * All outputs are provided through a 'struct' type
 	OutCode.Appendf(TEXT("%s CustomExpression%d(%s Parameters"), OutputStruct->Name, EmitCustomHLSL.Index, ParametersTypeName);
 	for (const FEmitCustomHLSLInput& Input : EmitCustomHLSL.Inputs)
@@ -281,14 +258,7 @@ void EmitCustomHLSL(const FEmitCustomHLSL& EmitCustomHLSL, const TCHAR* Paramete
 	{
 		OutCode.Append(TEXT(", "));
 		OutCode.Append(Input.Name);
-		if (Input.Type.IsTexture())
-		{
-			// Pass the Texture and SamplerState to the wrapper
-			OutCode.Append(TEXT(".Texture, "));
-			OutCode.Append(Input.Name);
-			OutCode.Append(TEXT(".Sampler"));
-		}
-		else if (Input.Type.IsNumericLWC())
+		if (Input.Type.IsNumericLWC())
 		{
 			OutCode.Append(TEXT(", LWCToFloat("));
 			OutCode.Append(Input.Name);
@@ -620,7 +590,7 @@ void MoveToScope(FEmitShaderNode* EmitNode, FEmitScope& Scope)
 
 void FormatArg_ShaderValue(FEmitShaderExpression* ShaderValue, FEmitShaderDependencies& OutDependencies, FStringBuilderBase& OutCode)
 {
-	OutDependencies.AddUnique(ShaderValue);
+	OutDependencies.Add(ShaderValue);
 	OutCode.Append(ShaderValue->Reference);
 }
 
diff --git a/Engine/Source/Runtime/Engine/Private/Materials/MaterialExpressionHLSL.cpp b/Engine/Source/Runtime/Engine/Private/Materials/MaterialExpressionHLSL.cpp
--- a/Engine/Source/Runtime/Engine/Private/Materials/MaterialExpressionHLSL.cpp
+++ b/Engine/Source/Runtime/Engine/Private/Materials/MaterialExpressionHLSL.cpp
@@ -11,7 +11,6 @@
 #include "Materials/MaterialExpression.h"
 #include "Materials/MaterialExpressionExecBegin.h"
 #include "Materials/MaterialExpressionExecEnd.h"
-#include "Materials/MaterialExpressionReroute.h"
 #include "Materials/MaterialExpressionGenericConstant.h"
 #include "Materials/MaterialExpressionConstant.h"
 #include "Materials/MaterialExpressionConstant2Vector.h"
@@ -31,12 +30,9 @@
 #include "Materials/MaterialExpressionSceneTexelSize.h"
 #include "Materials/MaterialExpressionViewSize.h"
 #include "Materials/MaterialExpressionPanner.h"
-#include "Materials/MaterialExpressionVertexColor.h"
 #include "Materials/MaterialExpressionTextureCoordinate.h"
 #include "Materials/MaterialExpressionTextureSample.h"
 #include "Materials/MaterialExpressionTextureSampleParameter.h"
-#include "Materials/MaterialExpressionFontSample.h"
-#include "Materials/MaterialExpressionFontSampleParameter.h"
 #include "Materials/MaterialExpressionTextureObject.h"
 #include "Materials/MaterialExpressionTextureObjectParameter.h"
 #include "Materials/MaterialExpressionSceneTexture.h"
@@ -47,19 +43,10 @@
 #include "Materials/MaterialExpressionFeatureLevelSwitch.h"
 #include "Materials/MaterialExpressionGetLocal.h"
 #include "Materials/MaterialExpressionOneMinus.h"
-#include "Materials/MaterialExpressionAbs.h"
-#include "Materials/MaterialExpressionFrac.h"
-#include "Materials/MaterialExpressionFloor.h"
-#include "Materials/MaterialExpressionCeil.h"
-#include "Materials/MaterialExpressionRound.h"
-#include "Materials/MaterialExpressionTruncate.h"
-#include "Materials/MaterialExpressionSaturate.h"
-#include "Materials/MaterialExpressionSign.h"
 #include "Materials/MaterialExpressionAdd.h"
 #include "Materials/MaterialExpressionSubtract.h"
 #include "Materials/MaterialExpressionMultiply.h"
 #include "Materials/MaterialExpressionDivide.h"
-#include "Materials/MaterialExpressionPower.h"
 #include "Materials/MaterialExpressionDotProduct.h"
 #include "Materials/MaterialExpressionMin.h"
 #include "Materials/MaterialExpressionMax.h"
@@ -75,16 +62,12 @@
 #include "Materials/MaterialExpressionReflectionVectorWS.h"
 #include "Materials/MaterialExpressionTransform.h"
 #include "Materials/MaterialExpressionTransformPosition.h"
-#include "Materials/MaterialExpressionIf.h"
 #include "Materials/MaterialExpressionCustom.h"
 #include "Materials/MaterialExpressionSetLocal.h"
 #include "Materials/MaterialExpressionIfThenElse.h"
 #include "Materials/MaterialExpressionForLoop.h"
 #include "Materials/MaterialExpressionWhileLoop.h"
 #include "Materials/MaterialFunctionInterface.h"
-#include "Engine/Engine.h"
-#include "Engine/Texture2D.h"
-#include "Engine/Font.h"
 
 bool UMaterialExpression::GenerateHLSLExpression(FMaterialHLSLGenerator& Generator, UE::HLSLTree::FScope& Scope, int32 OutputIndex, UE::HLSLTree::FExpression*& OutExpression)
 {
@@ -96,10 +79,9 @@ bool UMaterialExpression::GenerateHLSLStatements(FMaterialHLSLGenerator& Generat
 	return Generator.GetErrors().AddError(TEXT("Node does not support statements"));
 }
 
-bool UMaterialExpressionReroute::GenerateHLSLExpression(FMaterialHLSLGenerator& Generator, UE::HLSLTree::FScope& Scope, int32 OutputIndex, UE::HLSLTree::FExpression*& OutExpression)
+bool UMaterialExpression::GenerateHLSLTexture(FMaterialHLSLGenerator& Generator, UE::HLSLTree::FScope& Scope, int32 OutputIndex, UE::HLSLTree::FTextureParameterDeclaration*& OutTexture)
 {
-	OutExpression = Input.TryAcquireHLSLExpression(Generator, Scope);
-	return OutExpression != nullptr;
+	return Generator.GetErrors().AddError(TEXT("Node does not support textures"));
 }
 
 bool UMaterialExpressionGenericConstant::GenerateHLSLExpression(FMaterialHLSLGenerator& Generator, UE::HLSLTree::FScope& Scope, int32 OutputIndex, UE::HLSLTree::FExpression*& OutExpression)
@@ -311,13 +293,6 @@ bool UMaterialExpressionPanner::GenerateHLSLExpression(FMaterialHLSLGenerator& G
 	return true;
 }
 
-bool UMaterialExpressionVertexColor::GenerateHLSLExpression(FMaterialHLSLGenerator& Generator, UE::HLSLTree::FScope& Scope, int32 OutputIndex, UE::HLSLTree::FExpression*& OutExpression)
-{
-	using namespace UE::HLSLTree;
-	OutExpression = Generator.GetTree().NewExpression<FExpressionExternalInput>(EExternalInput::VertexColor);
-	return true;
-}
-
 bool UMaterialExpressionTextureCoordinate::GenerateHLSLExpression(FMaterialHLSLGenerator& Generator, UE::HLSLTree::FScope& Scope, int32 OutputIndex, UE::HLSLTree::FExpression*& OutExpression)
 {
 	OutExpression = Generator.NewTexCoord(CoordinateIndex);
@@ -341,138 +316,70 @@ bool UMaterialExpressionTextureCoordinate::GenerateHLSLExpression(FMaterialHLSLG
 	return true;
 }
 
-bool UMaterialExpressionTextureObject::GenerateHLSLExpression(FMaterialHLSLGenerator& Generator, UE::HLSLTree::FScope& Scope, int32 OutputIndex, UE::HLSLTree::FExpression*& OutExpression)
+bool UMaterialExpressionTextureObject::GenerateHLSLTexture(FMaterialHLSLGenerator& Generator, UE::HLSLTree::FScope& Scope, int32 OutputIndex, UE::HLSLTree::FTextureParameterDeclaration*& OutTexture)
 {
-	using namespace UE::HLSLTree;
-	using namespace UE::Shader;
-	const FTextureValue* TextureValue = Generator.AcquireTextureValue(FTextureValue(Texture, SamplerType));
-	OutExpression = Generator.GetTree().NewExpression<FExpressionMaterialParameter>(EMaterialParameterType::Texture, FName(), TextureValue);
+	const UE::HLSLTree::FTextureDescription TextureDesc(Texture, SamplerType);
+	OutTexture = Generator.AcquireTextureDeclaration(TextureDesc);
 	return true;
 }
 
-bool UMaterialExpressionTextureObjectParameter::GenerateHLSLExpression(FMaterialHLSLGenerator& Generator, UE::HLSLTree::FScope& Scope, int32 OutputIndex, UE::HLSLTree::FExpression*& OutExpression)
+bool UMaterialExpressionTextureObjectParameter::GenerateHLSLTexture(FMaterialHLSLGenerator& Generator, UE::HLSLTree::FScope& Scope, int32 OutputIndex, UE::HLSLTree::FTextureParameterDeclaration*& OutTexture)
 {
-	using namespace UE::HLSLTree;
-	using namespace UE::Shader;
-	const FTextureValue* TextureValue = Generator.AcquireTextureValue(FTextureValue(Texture, SamplerType));
-	OutExpression = Generator.GetTree().NewExpression<FExpressionMaterialParameter>(EMaterialParameterType::Texture, ParameterName, TextureValue);
+	const UE::HLSLTree::FTextureDescription TextureDesc(Texture, SamplerType);
+	OutTexture = Generator.AcquireTextureParameterDeclaration(ParameterName, TextureDesc);
 	return true;
 }
 
-bool UMaterialExpressionTextureSample::GenerateHLSLExpressionBase(FMaterialHLSLGenerator& Generator, UE::HLSLTree::FScope& Scope, UE::HLSLTree::FExpression* TextureExpression, UE::HLSLTree::FExpression*& OutExpression)
+bool UMaterialExpressionTextureSample::GenerateHLSLExpressionBase(FMaterialHLSLGenerator& Generator, UE::HLSLTree::FScope& Scope, UE::HLSLTree::FTextureParameterDeclaration* TextureDeclaration, UE::HLSLTree::FExpression*& OutExpression)
 {
 	using namespace UE::HLSLTree;
-	if (!TextureExpression)
+	if (!TextureDeclaration)
 	{
 		return Generator.GetErrors().AddError(TEXT("Missing input texture"));
 	}
 
 	FExpression* TexCoordExpression = Coordinates.GetTracedInput().Expression ? Coordinates.TryAcquireHLSLExpression(Generator, Scope) : Generator.NewTexCoord(ConstCoordinate);
 	const FExpressionDerivatives TexCoordDerivatives;// = Generator.GetTree().GetAnalyticDerivatives(TexCoordExpression);
-	OutExpression = Generator.GetTree().NewExpression<FExpressionTextureSample>(TextureExpression, TexCoordExpression, TexCoordDerivatives, SamplerSource, MipValueMode);
+	OutExpression = Generator.GetTree().NewExpression<FExpressionTextureSample>(TextureDeclaration, TexCoordExpression, TexCoordDerivatives, SamplerSource, MipValueMode);
 	return true;
 }
 
 bool UMaterialExpressionTextureSample::GenerateHLSLExpression(FMaterialHLSLGenerator& Generator, UE::HLSLTree::FScope& Scope, int32 OutputIndex, UE::HLSLTree::FExpression*& OutExpression)
 {
-	using namespace UE::HLSLTree;
-	using namespace UE::Shader;
-	FExpression* TextureExpression = nullptr;
+	UE::HLSLTree::FTextureParameterDeclaration* TextureDeclaration = nullptr;
 	if (TextureObject.GetTracedInput().Expression)
 	{
-		TextureExpression = TextureObject.AcquireHLSLExpression(Generator, Scope);
+		TextureDeclaration = TextureObject.AcquireHLSLTexture(Generator, Scope);
 	}
 	else if (Texture)
 	{
-		const FTextureValue* TextureValue = Generator.AcquireTextureValue(FTextureValue(Texture, SamplerType));
-		TextureExpression = Generator.GetTree().NewExpression<FExpressionMaterialParameter>(EMaterialParameterType::Texture, FName(), TextureValue);
+		const UE::HLSLTree::FTextureDescription TextureDesc(Texture, SamplerType);
+		TextureDeclaration = Generator.AcquireTextureDeclaration(TextureDesc);
 	}
 
-	return GenerateHLSLExpressionBase(Generator, Scope, TextureExpression, OutExpression);
+	return GenerateHLSLExpressionBase(Generator, Scope, TextureDeclaration, OutExpression);
 }
 
 bool UMaterialExpressionTextureSampleParameter::GenerateHLSLExpression(FMaterialHLSLGenerator& Generator, UE::HLSLTree::FScope& Scope, int32 OutputIndex, UE::HLSLTree::FExpression*& OutExpression)
 {
-	using namespace UE::HLSLTree;
-	using namespace UE::Shader;
-	FExpression* TextureExpression = nullptr;
+	UE::HLSLTree::FTextureParameterDeclaration* TextureDeclaration = nullptr;
 	if (ParameterName.IsNone() && TextureObject.GetTracedInput().Expression)
 	{
-		TextureExpression = TextureObject.AcquireHLSLExpression(Generator, Scope);
+		TextureDeclaration = TextureObject.AcquireHLSLTexture(Generator, Scope);
 	}
 	else if (Texture)
 	{
-		const FTextureValue* TextureValue = Generator.AcquireTextureValue(FTextureValue(Texture, SamplerType));
-		TextureExpression = Generator.GetTree().NewExpression<FExpressionMaterialParameter>(EMaterialParameterType::Texture, ParameterName, TextureValue);
-	}
-
-	return GenerateHLSLExpressionBase(Generator, Scope, TextureExpression, OutExpression);
-}
-
-UE::HLSLTree::FExpression* UMaterialExpressionFontSample::GenerateHLSLTextureExpression(FMaterialHLSLGenerator& Generator, const UE::Shader::FTextureValue* TextureValue)
-{
-	return Generator.GetTree().NewExpression<UE::HLSLTree::FExpressionMaterialParameter>(EMaterialParameterType::Texture, FName(), TextureValue);
-}
-
-UE::HLSLTree::FExpression* UMaterialExpressionFontSampleParameter::GenerateHLSLTextureExpression(FMaterialHLSLGenerator& Generator, const UE::Shader::FTextureValue* TextureValue)
-{
-	return Generator.GetTree().NewExpression<UE::HLSLTree::FExpressionMaterialParameter>(EMaterialParameterType::Texture, ParameterName, TextureValue);
-}
-
-bool UMaterialExpressionFontSample::GenerateHLSLExpression(FMaterialHLSLGenerator& Generator, UE::HLSLTree::FScope& Scope, int32 OutputIndex, UE::HLSLTree::FExpression*& OutExpression)
-{
-	using namespace UE::HLSLTree;
-	using namespace UE::Shader;
-
-	if (!Font)
-	{
-		return Generator.GetErrors().AddError(TEXT("Missing input Font"));
-	}
-	else if (Font->FontCacheType == EFontCacheType::Runtime)
-	{
-		return Generator.GetErrors().AddErrorf(TEXT("Font '%s' is runtime cached, but only offline cached fonts can be sampled"), *Font->GetName());
-	}
-	else if (!Font->Textures.IsValidIndex(FontTexturePage))
-	{
-		return Generator.GetErrors().AddErrorf(TEXT("Invalid font page %d. Max allowed is %d"), FontTexturePage, Font->Textures.Num());
-	}
-
-	UTexture* Texture = Font->Textures[FontTexturePage];
-	if (!Texture)
-	{
-		UE_LOG(LogMaterial, Log, TEXT("Invalid font texture. Using default texture"));
-		Texture = GEngine->DefaultTexture;
+		const UE::HLSLTree::FTextureDescription TextureDesc(Texture, SamplerType);
+		TextureDeclaration = Generator.AcquireTextureParameterDeclaration(ParameterName, TextureDesc);
 	}
-	check(Texture);
 
-	EMaterialSamplerType ExpectedSamplerType;
-	if (Texture->CompressionSettings == TC_DistanceFieldFont)
-	{
-		ExpectedSamplerType = SAMPLERTYPE_DistanceFieldFont;
-	}
-	else
-	{
-		ExpectedSamplerType = Texture->SRGB ? SAMPLERTYPE_Color : SAMPLERTYPE_LinearColor;
-	}
-
-	FString SamplerTypeError;
-	if (!UMaterialExpressionTextureBase::VerifySamplerType(Generator.GetCompileTarget().FeatureLevel, Generator.GetCompileTarget().TargetPlatform, Texture, ExpectedSamplerType, SamplerTypeError))
-	{
-		return Generator.GetErrors().AddErrorf(TEXT("%s"), *SamplerTypeError);
-	}
-
-	const FTextureValue* TextureValue = Generator.AcquireTextureValue(FTextureValue(Texture, ExpectedSamplerType));
-	FExpression* TextureExpression = GenerateHLSLTextureExpression(Generator, TextureValue);
-	FExpression* TexCoordExpression = Generator.NewTexCoord(0);
-	FExpressionDerivatives TexCoordDerivatives;
-
-	OutExpression = Generator.GetTree().NewExpression<FExpressionTextureSample>(TextureExpression, TexCoordExpression, TexCoordDerivatives, SSM_FromTextureAsset, TMVM_None);
-	return true;
+	return GenerateHLSLExpressionBase(Generator, Scope, TextureDeclaration, OutExpression);
 }
 
 bool UMaterialExpressionSceneTexture::GenerateHLSLExpression(FMaterialHLSLGenerator& Generator, UE::HLSLTree::FScope& Scope, int32 OutputIndex, UE::HLSLTree::FExpression*& OutExpression)
 {
 	using namespace UE::HLSLTree;
+
 	if (OutputIndex == 0)
 	{
 		FExpression* ExpressionTexCoord = nullptr;
@@ -503,102 +410,6 @@ bool UMaterialExpressionOneMinus::GenerateHLSLExpression(FMaterialHLSLGenerator&
 	return true;
 }
 
-bool UMaterialExpressionAbs::GenerateHLSLExpression(FMaterialHLSLGenerator& Generator, UE::HLSLTree::FScope& Scope, int32 OutputIndex, UE::HLSLTree::FExpression*& OutExpression)
-{
-	using namespace UE::HLSLTree;
-	FExpression* InputExpression = Input.AcquireHLSLExpression(Generator, Scope);
-	if (!InputExpression)
-	{
-		return false;
-	}
-	OutExpression = Generator.GetTree().NewUnaryOp(EOperation::Abs, InputExpression);
-	return true;
-}
-
-bool UMaterialExpressionFrac::GenerateHLSLExpression(FMaterialHLSLGenerator& Generator, UE::HLSLTree::FScope& Scope, int32 OutputIndex, UE::HLSLTree::FExpression*& OutExpression)
-{
-	using namespace UE::HLSLTree;
-	FExpression* InputExpression = Input.AcquireHLSLExpression(Generator, Scope);
-	if (!InputExpression)
-	{
-		return false;
-	}
-	OutExpression = Generator.GetTree().NewUnaryOp(EOperation::Frac, InputExpression);
-	return true;
-}
-
-bool UMaterialExpressionFloor::GenerateHLSLExpression(FMaterialHLSLGenerator& Generator, UE::HLSLTree::FScope& Scope, int32 OutputIndex, UE::HLSLTree::FExpression*& OutExpression)
-{
-	using namespace UE::HLSLTree;
-	FExpression* InputExpression = Input.AcquireHLSLExpression(Generator, Scope);
-	if (!InputExpression)
-	{
-		return false;
-	}
-	OutExpression = Generator.GetTree().NewUnaryOp(EOperation::Floor, InputExpression);
-	return true;
-}
-
-bool UMaterialExpressionCeil::GenerateHLSLExpression(FMaterialHLSLGenerator& Generator, UE::HLSLTree::FScope& Scope, int32 OutputIndex, UE::HLSLTree::FExpression*& OutExpression)
-{
-	using namespace UE::HLSLTree;
-	FExpression* InputExpression = Input.AcquireHLSLExpression(Generator, Scope);
-	if (!InputExpression)
-	{
-		return false;
-	}
-	OutExpression = Generator.GetTree().NewUnaryOp(EOperation::Ceil, InputExpression);
-	return true;
-}
-
-bool UMaterialExpressionRound::GenerateHLSLExpression(FMaterialHLSLGenerator& Generator, UE::HLSLTree::FScope& Scope, int32 OutputIndex, UE::HLSLTree::FExpression*& OutExpression)
-{
-	using namespace UE::HLSLTree;
-	FExpression* InputExpression = Input.AcquireHLSLExpression(Generator, Scope);
-	if (!InputExpression)
-	{
-		return false;
-	}
-	OutExpression = Generator.GetTree().NewUnaryOp(EOperation::Round, InputExpression);
-	return true;
-}
-
-bool UMaterialExpressionTruncate::GenerateHLSLExpression(FMaterialHLSLGenerator& Generator, UE::HLSLTree::FScope& Scope, int32 OutputIndex, UE::HLSLTree::FExpression*& OutExpression)
-{
-	using namespace UE::HLSLTree;
-	FExpression* InputExpression = Input.AcquireHLSLExpression(Generator, Scope);
-	if (!InputExpression)
-	{
-		return false;
-	}
-	OutExpression = Generator.GetTree().NewUnaryOp(EOperation::Trunc, InputExpression);
-	return true;
-}
-
-bool UMaterialExpressionSaturate::GenerateHLSLExpression(FMaterialHLSLGenerator& Generator, UE::HLSLTree::FScope& Scope, int32 OutputIndex, UE::HLSLTree::FExpression*& OutExpression)
-{
-	using namespace UE::HLSLTree;
-	FExpression* InputExpression = Input.AcquireHLSLExpression(Generator, Scope);
-	if (!InputExpression)
-	{
-		return false;
-	}
-	OutExpression = Generator.GetTree().NewUnaryOp(EOperation::Saturate, InputExpression);
-	return true;
-}
-
-bool UMaterialExpressionSign::GenerateHLSLExpression(FMaterialHLSLGenerator& Generator, UE::HLSLTree::FScope& Scope, int32 OutputIndex, UE::HLSLTree::FExpression*& OutExpression)
-{
-	using namespace UE::HLSLTree;
-	FExpression* InputExpression = Input.AcquireHLSLExpression(Generator, Scope);
-	if (!InputExpression)
-	{
-		return false;
-	}
-	OutExpression = Generator.GetTree().NewUnaryOp(EOperation::Sign, InputExpression);
-	return true;
-}
-
 bool UMaterialExpressionBinaryOp::GenerateHLSLExpression(FMaterialHLSLGenerator& Generator, UE::HLSLTree::FScope& Scope, int32 OutputIndex, UE::HLSLTree::FExpression*& OutExpression)
 {
 	UE::HLSLTree::FExpression* Lhs = A.AcquireHLSLExpressionOrConstant(Generator, Scope, ConstA);
@@ -660,19 +471,6 @@ bool UMaterialExpressionDivide::GenerateHLSLExpression(FMaterialHLSLGenerator& G
 	return true;
 }
 
-bool UMaterialExpressionPower::GenerateHLSLExpression(FMaterialHLSLGenerator& Generator, UE::HLSLTree::FScope& Scope, int32 OutputIndex, UE::HLSLTree::FExpression*& OutExpression)
-{
-	using namespace UE::HLSLTree;
-	FExpression* Lhs = Base.AcquireHLSLExpression(Generator, Scope);
-	FExpression* Rhs = Exponent.AcquireHLSLExpressionOrConstant(Generator, Scope, ConstExponent);
-	if (!Lhs || !Rhs)
-	{
-		return false;
-	}
-	OutExpression = Generator.GetTree().NewBinaryOp(EOperation::PowPositiveClamped, Lhs, Rhs);
-	return true;
-}
-
 bool UMaterialExpressionDotProduct::GenerateHLSLExpression(FMaterialHLSLGenerator& Generator, UE::HLSLTree::FScope& Scope, int32 OutputIndex, UE::HLSLTree::FExpression*& OutExpression)
 {
 	UE::HLSLTree::FExpression* Lhs = A.AcquireHLSLExpression(Generator, Scope);
@@ -1119,35 +917,6 @@ bool UMaterialExpressionTransformPosition::GenerateHLSLExpression(FMaterialHLSLG
 	return true;
 }
 
-bool UMaterialExpressionIf::GenerateHLSLExpression(FMaterialHLSLGenerator& Generator, UE::HLSLTree::FScope& Scope, int32 OutputIndex, UE::HLSLTree::FExpression*& OutExpression)
-{
-	using namespace UE::HLSLTree;
-	using namespace UE::Shader;
-
-	FExpression* ExpressionA = A.AcquireHLSLExpression(Generator, Scope);
-	FExpression* ExpressionB = A.AcquireHLSLExpressionOrConstant(Generator, Scope, ConstB);
-	FExpression* ExpressionAGreaterThanB = AGreaterThanB.AcquireHLSLExpression(Generator, Scope);
-	FExpression* ExpressionALessThanB = ALessThanB.AcquireHLSLExpression(Generator, Scope);
-	if (!ExpressionA || !ExpressionB || !ExpressionAGreaterThanB || !ExpressionALessThanB)
-	{
-		return false;
-	}
-
-	UE::HLSLTree::FTree& Tree = Generator.GetTree();
-	FExpression* ExpressionCondAGreaterEqualB = Tree.NewGreaterEqual(ExpressionA, ExpressionB);
-	OutExpression = Tree.NewExpression<FExpressionSelect>(ExpressionCondAGreaterEqualB, ExpressionAGreaterThanB, ExpressionALessThanB);
-
-	FExpression* ExpressionAEqualsB = AEqualsB.TryAcquireHLSLExpression(Generator, Scope);
-	if (ExpressionAEqualsB)
-	{
-		FExpression* ExpressionThreshold = Generator.NewConstant(EqualsThreshold);
-		FExpression* ExpressionCondANotEqualsB = Tree.NewGreater(Tree.NewAbs(Tree.NewSub(ExpressionA, ExpressionB)), ExpressionThreshold);
-		OutExpression = Tree.NewExpression<FExpressionSelect>(ExpressionCondANotEqualsB, OutExpression, ExpressionAEqualsB);
-	}
-
-	return true;
-}
-
 static UE::Shader::FType GetCustomOutputType(const FMaterialHLSLGenerator& Generator, ECustomMaterialOutputType Type)
 {
 	using namespace UE::Shader;
diff --git a/Engine/Source/Runtime/Engine/Private/Materials/MaterialHLSLEmitter.cpp b/Engine/Source/Runtime/Engine/Private/Materials/MaterialHLSLEmitter.cpp
--- a/Engine/Source/Runtime/Engine/Private/Materials/MaterialHLSLEmitter.cpp
+++ b/Engine/Source/Runtime/Engine/Private/Materials/MaterialHLSLEmitter.cpp
@@ -391,7 +391,6 @@ static FString GenerateMaterialTemplateHLSL(EShaderPlatform ShaderPlatform,
 
 static void GetMaterialEnvironment(EShaderPlatform InPlatform,
 	const FMaterial& InMaterial,
-	const UE::HLSLTree::FEmitContext& EmitContext,
 	const FMaterialCompilationOutput& MaterialCompilationOutput,
 	FShaderCompilerEnvironment& OutEnvironment)
 {
@@ -487,7 +486,7 @@ static void GetMaterialEnvironment(EShaderPlatform InPlatform,
 	// @todo MetalMRT: Remove this hack and implement proper atmospheric-fog solution for Metal MRT...
 	OutEnvironment.SetDefine(TEXT("MATERIAL_ATMOSPHERIC_FOG"), false);// !IsMetalMRTPlatform(InPlatform) ? bUsesAtmosphericFog : 0);
 	OutEnvironment.SetDefine(TEXT("MATERIAL_SKY_ATMOSPHERE"), false);// bUsesSkyAtmosphere);
-	OutEnvironment.SetDefine(TEXT("INTERPOLATE_VERTEX_COLOR"), EmitContext.bUsesVertexColor);
+	OutEnvironment.SetDefine(TEXT("INTERPOLATE_VERTEX_COLOR"), false);// bUsesVertexColor);
 	OutEnvironment.SetDefine(TEXT("NEEDS_PARTICLE_COLOR"), false);// bUsesParticleColor);
 	OutEnvironment.SetDefine(TEXT("NEEDS_PARTICLE_LOCAL_TO_WORLD"), false);// bUsesParticleLocalToWorld);
 	OutEnvironment.SetDefine(TEXT("NEEDS_PARTICLE_WORLD_TO_LOCAL"), false);// bUsesParticleWorldToLocal);
@@ -934,7 +933,7 @@ bool MaterialEmitHLSL(const FMaterialCompileTargetParameters& InCompilerTarget,
 
 	OutMaterialEnvironment = new FSharedShaderCompilerEnvironment();
 	OutMaterialEnvironment->TargetPlatform = InCompilerTarget.TargetPlatform;
-	GetMaterialEnvironment(InCompilerTarget.ShaderPlatform, InOutMaterial, EmitContext, OutCompilationOutput, *OutMaterialEnvironment);
+	GetMaterialEnvironment(InCompilerTarget.ShaderPlatform, InOutMaterial, OutCompilationOutput, *OutMaterialEnvironment);
 	OutMaterialEnvironment->IncludeVirtualPathToContentsMap.Add(TEXT("/Engine/Generated/Material.ush"), MoveTemp(MaterialTemplateSource));
 
 	return true;
diff --git a/Engine/Source/Runtime/Engine/Private/Materials/MaterialHLSLGenerator.cpp b/Engine/Source/Runtime/Engine/Private/Materials/MaterialHLSLGenerator.cpp
--- a/Engine/Source/Runtime/Engine/Private/Materials/MaterialHLSLGenerator.cpp
+++ b/Engine/Source/Runtime/Engine/Private/Materials/MaterialHLSLGenerator.cpp
@@ -359,32 +359,38 @@ UE::HLSLTree::FExpression* FMaterialHLSLGenerator::NewSwizzle(const UE::HLSLTree
 	return HLSLTree->NewExpression<UE::HLSLTree::FExpressionSwizzle>(Params, Input);
 }
 
-const UE::Shader::FTextureValue* FMaterialHLSLGenerator::AcquireTextureValue(const UE::Shader::FTextureValue& InValue)
+UE::HLSLTree::FTextureParameterDeclaration* FMaterialHLSLGenerator::AcquireTextureDeclaration(const UE::HLSLTree::FTextureDescription& Value)
 {
-	using namespace UE::Shader;
 	FString SamplerTypeError;
-	if (!UMaterialExpressionTextureBase::VerifySamplerType(CompileTarget.FeatureLevel, CompileTarget.TargetPlatform, InValue.Texture, InValue.SamplerType, SamplerTypeError))
+	if (!UMaterialExpressionTextureBase::VerifySamplerType(CompileTarget.FeatureLevel, CompileTarget.TargetPlatform, Value.Texture, Value.SamplerType, SamplerTypeError))
 	{
 		Errors.AddError(SamplerTypeError);
 		return nullptr;
 	}
 
-	FXxHash64Builder Hasher;
-	Hasher.Update(&InValue.Texture, sizeof(InValue.Texture));
-	Hasher.Update(&InValue.SamplerType, sizeof(InValue.SamplerType));
-	const FXxHash64 Hash = Hasher.Finalize();
+	UE::HLSLTree::FTextureParameterDeclaration*& Declaration = TextureDeclarationMap.FindOrAdd(Value);
+	if (!Declaration)
+	{
+		Declaration = HLSLTree->NewTextureParameterDeclaration(FName(), Value);
+	}
+	return Declaration;
+}
 
-	FTextureValue const* const* PrevValue = TextureValueMap.Find(Hash);
-	if (PrevValue)
+UE::HLSLTree::FTextureParameterDeclaration* FMaterialHLSLGenerator::AcquireTextureParameterDeclaration(const FName& Name, const UE::HLSLTree::FTextureDescription& DefaultValue)
+{
+	FString SamplerTypeError;
+	if (!UMaterialExpressionTextureBase::VerifySamplerType(CompileTarget.FeatureLevel, CompileTarget.TargetPlatform, DefaultValue.Texture, DefaultValue.SamplerType, SamplerTypeError))
 	{
-		check(*PrevValue);
-		check(**PrevValue == InValue);
-		return *PrevValue;
+		Errors.AddError(SamplerTypeError);
+		return nullptr;
 	}
 
-	FTextureValue* Value = new(HLSLTree->GetAllocator()) FTextureValue(InValue);
-	TextureValueMap.Add(Hash, Value);
-	return Value;
+	UE::HLSLTree::FTextureParameterDeclaration*& Declaration = TextureParameterDeclarationMap.FindOrAdd(Name);
+	if (!Declaration)
+	{
+		Declaration = HLSLTree->NewTextureParameterDeclaration(Name, DefaultValue);
+	}
+	return Declaration;
 }
 
 UE::HLSLTree::FExpression* FMaterialHLSLGenerator::AcquireExpression(UE::HLSLTree::FScope& Scope, UMaterialExpression* MaterialExpression, int32 OutputIndex)
@@ -462,6 +468,20 @@ UE::HLSLTree::FExpression* FMaterialHLSLGenerator::AcquireFunctionInputExpressio
 	return InputExpression;
 }
 
+UE::HLSLTree::FTextureParameterDeclaration* FMaterialHLSLGenerator::AcquireTextureDeclaration(UE::HLSLTree::FScope& Scope, UMaterialExpression* MaterialExpression, int32 OutputIndex)
+{
+	using namespace UE::HLSLTree;
+	FOwnerScope TreeOwnerScope(GetTree(), MaterialExpression);
+	FOwnerScope ErrorOwnerScope(Errors, MaterialExpression);
+
+	FTextureParameterDeclaration* TextureDeclaration = nullptr;
+	if (MaterialExpression->GenerateHLSLTexture(*this, Scope, OutputIndex, TextureDeclaration))
+	{
+		return TextureDeclaration;
+	}
+	return nullptr;
+}
+
 bool FMaterialHLSLGenerator::GenerateStatements(UE::HLSLTree::FScope& Scope, UMaterialExpression* MaterialExpression)
 {
 	using namespace UE::HLSLTree;
diff --git a/Engine/Source/Runtime/Engine/Private/Materials/MaterialShared.cpp b/Engine/Source/Runtime/Engine/Private/Materials/MaterialShared.cpp
--- a/Engine/Source/Runtime/Engine/Private/Materials/MaterialShared.cpp
+++ b/Engine/Source/Runtime/Engine/Private/Materials/MaterialShared.cpp
@@ -313,6 +313,23 @@ UE::HLSLTree::FExpression* FExpressionInput::AcquireHLSLExpressionOrConstant(FMa
 	return TryAcquireHLSLExpression(Generator, Scope);
 }
 
+UE::HLSLTree::FTextureParameterDeclaration* FExpressionInput::AcquireHLSLTexture(FMaterialHLSLGenerator& Generator, UE::HLSLTree::FScope& Scope) const
+{
+	UE::HLSLTree::FTextureParameterDeclaration* Result = nullptr;
+	const FExpressionInput TracedInput = GetTracedInput();
+	if (!TracedInput.Expression)
+	{
+		Generator.GetErrors().AddError(TEXT("Missing input"));
+	}
+	else if (Expression)
+	{
+		Expression->ValidateState();
+		Result = Generator.AcquireTextureDeclaration(Scope, Expression, OutputIndex);
+	}
+
+	return Result;
+}
+
 void FExpressionInput::Connect( int32 InOutputIndex, class UMaterialExpression* InExpression )
 {
 	InExpression->ConnectExpression(this, InOutputIndex);
diff --git a/Engine/Source/Runtime/Engine/Private/Shader/Preshader.cpp b/Engine/Source/Runtime/Engine/Private/Shader/Preshader.cpp
--- a/Engine/Source/Runtime/Engine/Private/Shader/Preshader.cpp
+++ b/Engine/Source/Runtime/Engine/Private/Shader/Preshader.cpp
@@ -168,8 +168,6 @@ void FPreshaderData::WriteValue(const FValue& Value)
 {
 	const int32 NumComponents = Value.Type.GetNumComponents();
 
-	check(!Value.Type.IsTexture());
-
 	WriteType(Value.Type);
 	for (int32 Index = 0; Index < NumComponents; ++Index)
 	{
diff --git a/Engine/Source/Runtime/Engine/Private/Shader/ShaderValue.cpp b/Engine/Source/Runtime/Engine/Private/Shader/ShaderValue.cpp
--- a/Engine/Source/Runtime/Engine/Private/Shader/ShaderValue.cpp
+++ b/Engine/Source/Runtime/Engine/Private/Shader/ShaderValue.cpp
@@ -4,7 +4,6 @@
 #include "Hash/xxhash.h"
 #include "Misc/MemStackUtility.h"
 #include "Misc/LargeWorldRenderPosition.h"
-#include "Engine/Texture.h"
 
 namespace UE
 {
@@ -34,11 +33,7 @@ FType FType::GetDerivativeType() const
 
 	check(ValueType != EValueType::Struct);
 	const FValueTypeDescription TypeDesc = GetValueTypeDescription(ValueType);
-	if (IsNumericType(TypeDesc.ComponentType))
-	{
-		return MakeValueType(EValueComponentType::Float, TypeDesc.NumComponents);
-	}
-	return EValueType::Void;
+	return MakeValueType(EValueComponentType::Float, TypeDesc.NumComponents);
 }
 
 int32 FType::GetNumComponents() const
@@ -112,11 +107,6 @@ bool FType::Merge(const FType& OtherType)
 		const FValueTypeDescription TypeDesc = GetValueTypeDescription(ValueType);
 		const FValueTypeDescription OtherTypeDesc = GetValueTypeDescription(OtherType);
 		const EValueComponentType ComponentType = CombineComponentTypes(TypeDesc.ComponentType, OtherTypeDesc.ComponentType);
-		if (ComponentType == EValueComponentType::Void)
-		{
-			return false;
-		}
-
 		const int8 NumComponents = FMath::Max(TypeDesc.NumComponents, OtherTypeDesc.NumComponents);
 		ValueType = MakeValueType(ComponentType, NumComponents);
 	}
@@ -282,39 +272,8 @@ void FormatComponent_Double(double Value, int32 NumComponents, EValueStringForma
 
 } // namespace Private
 
-EValueType FTextureValue::GetType() const
-{
-	if (Texture)
-	{
-		const EMaterialValueType MaterialType = Texture->GetMaterialType();
-		switch (MaterialType)
-		{
-		case MCT_Texture2D: return EValueType::Texture2D;
-		case MCT_Texture2DArray: return EValueType::Texture2DArray;
-		case MCT_TextureCube: return EValueType::TextureCube;
-		case MCT_TextureCubeArray: return EValueType::TextureCubeArray;
-		case MCT_VolumeTexture: return EValueType::Texture3D;
-		default: checkNoEntry(); break;
-		}
-	}
-	return EValueType::Void;
-}
-
-FValue::FValue(const FTextureValue* InValue) : Type(EValueType::Void)
-{
-	if (InValue)
-	{
-		Type = InValue->GetType();
-		if (Type != EValueType::Void)
-		{
-			Component.Add(InValue);
-		}
-	}
-}
-
 FValue FValue::FromMemoryImage(EValueType Type, const void* Data, uint32* OutSizeInBytes)
 {
-	check(IsNumericType(Type));
 	const FValueTypeDescription TypeDesc = GetValueTypeDescription(Type);
 
 	FValue Result(TypeDesc.ComponentType, TypeDesc.NumComponents);
@@ -337,7 +296,7 @@ FValue FValue::FromMemoryImage(EValueType Type, const void* Data, uint32* OutSiz
 
 FMemoryImageValue FValue::AsMemoryImage() const
 {
-	check(Type.IsNumeric());
+	check(!Type.IsStruct());
 	const FValueTypeDescription TypeDesc = GetValueTypeDescription(Type);
 
 	FMemoryImageValue Result;
@@ -416,29 +375,15 @@ bool FValue::AsBoolScalar() const
 	return false;
 }
 
-const FTextureValue* FValue::AsTexture() const
-{
-	if (Type.IsTexture() && Component.Num() > 0)
-	{
-		return Component[0].Texture;
-	}
-	return nullptr;
-}
-
 FValueComponentTypeDescription GetValueComponentTypeDescription(EValueComponentType Type)
 {
 	switch (Type)
 	{
 	case EValueComponentType::Void: return FValueComponentTypeDescription(TEXT("void"), 0u, EComponentBound::Zero, EComponentBound::Zero);
-	case EValueComponentType::Float: return FValueComponentTypeDescription(TEXT("float"), sizeof(float), EComponentBound::NegFloatMax, EComponentBound::FloatMax);
-	case EValueComponentType::Double: return FValueComponentTypeDescription(TEXT("double"), sizeof(double), EComponentBound::NegDoubleMax, EComponentBound::DoubleMax);
-	case EValueComponentType::Int: return FValueComponentTypeDescription(TEXT("int"), sizeof(int32), EComponentBound::IntMin, EComponentBound::IntMax);
+	case EValueComponentType::Float: return FValueComponentTypeDescription(TEXT("float"), 4u, EComponentBound::NegFloatMax, EComponentBound::FloatMax);
+	case EValueComponentType::Double: return FValueComponentTypeDescription(TEXT("double"), 8u, EComponentBound::NegDoubleMax, EComponentBound::DoubleMax);
+	case EValueComponentType::Int: return FValueComponentTypeDescription(TEXT("int"), 4u, EComponentBound::IntMin, EComponentBound::IntMax);
 	case EValueComponentType::Bool: return FValueComponentTypeDescription(TEXT("bool"), 1u, EComponentBound::Zero, EComponentBound::One);
-	case EValueComponentType::Texture2D: return FValueComponentTypeDescription(TEXT("Texture2D"), sizeof(void*), EComponentBound::Zero, EComponentBound::Zero);
-	case EValueComponentType::Texture2DArray: return FValueComponentTypeDescription(TEXT("Texture2DArray"), sizeof(void*), EComponentBound::Zero, EComponentBound::Zero);
-	case EValueComponentType::TextureCube: return FValueComponentTypeDescription(TEXT("TextureCube"), sizeof(void*), EComponentBound::Zero, EComponentBound::Zero);
-	case EValueComponentType::TextureCubeArray: return FValueComponentTypeDescription(TEXT("TextureCubeArray"), sizeof(void*), EComponentBound::Zero, EComponentBound::Zero);
-	case EValueComponentType::Texture3D: return FValueComponentTypeDescription(TEXT("Texture3D"), sizeof(void*), EComponentBound::Zero, EComponentBound::Zero);
 	default: checkNoEntry() return FValueComponentTypeDescription();
 	}
 }
@@ -465,13 +410,9 @@ EValueComponentType CombineComponentTypes(EValueComponentType Lhs, EValueCompone
 	{
 		return EValueComponentType::Float;
 	}
-	else if(IsNumericType(Lhs) && IsNumericType(Rhs))
-	{
-		return EValueComponentType::Int;
-	}
 	else
 	{
-		return EValueComponentType::Void;
+		return EValueComponentType::Int;
 	}
 }
 
@@ -636,11 +577,6 @@ FValueTypeDescription GetValueTypeDescription(EValueType Type)
 	case EValueType::Double4x4: return FValueTypeDescription(TEXT("FLWCMatrix"), EValueComponentType::Double, 16);
 	case EValueType::DoubleInverse4x4: return FValueTypeDescription(TEXT("FLWCInverseMatrix"), EValueComponentType::Double, 16);
 	case EValueType::Struct: return FValueTypeDescription(TEXT("struct"), EValueComponentType::Void, 0);
-	case EValueType::Texture2D: return FValueTypeDescription(TEXT("FTexture2D"), EValueComponentType::Texture2D, 1);
-	case EValueType::Texture2DArray: return FValueTypeDescription(TEXT("FTexture2DArray"), EValueComponentType::Texture2DArray, 1);
-	case EValueType::TextureCube: return FValueTypeDescription(TEXT("FTextureCube"), EValueComponentType::TextureCube, 1);
-	case EValueType::TextureCubeArray: return FValueTypeDescription(TEXT("FTextureCubeArray"), EValueComponentType::TextureCubeArray, 1);
-	case EValueType::Texture3D: return FValueTypeDescription(TEXT("FTexture3D"), EValueComponentType::Texture3D, 1);
 	default: checkNoEntry(); return FValueTypeDescription(TEXT("<INVALID>"), EValueComponentType::Void, 0);
 	}
 }
@@ -664,7 +600,6 @@ EValueType MakeValueType(EValueComponentType ComponentType, int32 NumComponents)
 		case 16: return EValueType::Float4x4;
 		default: break;
 		}
-		break;
 	case EValueComponentType::Double:
 		switch (NumComponents)
 		{
@@ -675,7 +610,6 @@ EValueType MakeValueType(EValueComponentType ComponentType, int32 NumComponents)
 		case 16: return EValueType::Double4x4;
 		default: break;
 		}
-		break;
 	case EValueComponentType::Int:
 		switch (NumComponents)
 		{
@@ -685,7 +619,6 @@ EValueType MakeValueType(EValueComponentType ComponentType, int32 NumComponents)
 		case 4: return EValueType::Int4;
 		default: break;
 		}
-		break;
 	case EValueComponentType::Bool:
 		switch (NumComponents)
 		{
@@ -695,12 +628,6 @@ EValueType MakeValueType(EValueComponentType ComponentType, int32 NumComponents)
 		case 4: return EValueType::Bool4;
 		default: break;
 		}
-		break;
-	case EValueComponentType::Texture2D: check(NumComponents == 1); return EValueType::Texture2D; break;
-	case EValueComponentType::Texture2DArray: check(NumComponents == 1); return EValueType::Texture2DArray; break;
-	case EValueComponentType::TextureCube: check(NumComponents == 1); return EValueType::TextureCube; break;
-	case EValueComponentType::TextureCubeArray: check(NumComponents == 1); return EValueType::TextureCubeArray; break;
-	case EValueComponentType::Texture3D: check(NumComponents == 1); return EValueType::Texture3D; break;
 	default:
 		break;
 	}
diff --git a/Engine/Source/Runtime/Engine/Public/HLSLTree/HLSLTree.h b/Engine/Source/Runtime/Engine/Public/HLSLTree/HLSLTree.h
--- a/Engine/Source/Runtime/Engine/Public/HLSLTree/HLSLTree.h
+++ b/Engine/Source/Runtime/Engine/Public/HLSLTree/HLSLTree.h
@@ -174,7 +174,6 @@ enum class ERequestedType : uint8
 {
 	Void = 0u,
 	Scalar = 1u,
-	Texture = 1u, // Currently no different than scalar, just makes code more clear
 	Vector2 = 2u,
 	Vector3 = 3u,
 	Vector4 = 4u,
@@ -456,6 +455,18 @@ public:
 	TArray<FExpression*, TInlineAllocator<8>> OutputExpressions;
 };
 
+/**
+ * Represents an HLSL texture parameter.
+ */
+class FTextureParameterDeclaration final : public FNode
+{
+public:
+	FTextureParameterDeclaration(const FName& InName, const FTextureDescription& InDescription) : Name(InName), Description(InDescription) {}
+
+	FName Name;
+	FTextureDescription Description;
+};
+
 /**
  * Represents an HLSL scope.  A scope contains a single statement, along with any expressions required by that statement
  */
@@ -566,7 +577,6 @@ public:
 	FExpression* NewUnaryOp(EOperation Op, FExpression* Input);
 	FExpression* NewBinaryOp(EOperation Op, FExpression* Lhs, FExpression* Rhs);
 
-	FExpression* NewAbs(FExpression* Input) { return NewUnaryOp(EOperation::Abs, Input); }
 	FExpression* NewNeg(FExpression* Input) { return NewUnaryOp(EOperation::Neg, Input); }
 	FExpression* NewRcp(FExpression* Input) { return NewUnaryOp(EOperation::Rcp, Input); }
 	FExpression* NewFrac(FExpression* Input) { return NewUnaryOp(EOperation::Frac, Input); }
@@ -583,8 +593,8 @@ public:
 	FExpression* NewMax(FExpression* Lhs, FExpression* Rhs) { return NewBinaryOp(EOperation::Max, Lhs, Rhs); }
 	FExpression* NewLess(FExpression* Lhs, FExpression* Rhs) { return NewBinaryOp(EOperation::Less, Lhs, Rhs); }
 	FExpression* NewGreater(FExpression* Lhs, FExpression* Rhs) { return NewBinaryOp(EOperation::Greater, Lhs, Rhs); }
-	FExpression* NewLessEqual(FExpression* Lhs, FExpression* Rhs) { return NewBinaryOp(EOperation::LessEqual, Lhs, Rhs); }
-	FExpression* NewGreaterEqual(FExpression* Lhs, FExpression* Rhs) { return NewBinaryOp(EOperation::GreaterEqual, Lhs, Rhs); }
+
+	FTextureParameterDeclaration* NewTextureParameterDeclaration(const FName& Name, const FTextureDescription& DefaultValue);
 
 private:
 	template<typename T, typename... ArgTypes>
diff --git a/Engine/Source/Runtime/Engine/Public/HLSLTree/HLSLTreeCommon.h b/Engine/Source/Runtime/Engine/Public/HLSLTree/HLSLTreeCommon.h
--- a/Engine/Source/Runtime/Engine/Public/HLSLTree/HLSLTreeCommon.h
+++ b/Engine/Source/Runtime/Engine/Public/HLSLTree/HLSLTreeCommon.h
@@ -39,7 +39,6 @@ public:
 
 	virtual void ComputeAnalyticDerivatives(FTree& Tree, FExpressionDerivatives& OutResult) const override;
 	virtual bool PrepareValue(FEmitContext& Context, FEmitScope& Scope, const FRequestedType& RequestedType, FPrepareValueResult& OutResult) const override;
-	virtual void EmitValueShader(FEmitContext& Context, FEmitScope& Scope, const FRequestedType& RequestedType, FEmitValueShaderResult& OutResult) const override;
 	virtual void EmitValuePreshader(FEmitContext& Context, FEmitScope& Scope, const FRequestedType& RequestedType, FEmitValuePreshaderResult& OutResult) const override;
 };
 
@@ -74,10 +73,6 @@ enum class EExternalInput : uint8
 	TexCoord6_Ddy,
 	TexCoord7_Ddy,
 
-	VertexColor,
-	VertexColor_Ddx,
-	VertexColor_Ddy,
-
 	WorldPosition,
 	WorldPosition_NoOffsets,
 	TranslatedWorldPosition,
@@ -198,19 +193,15 @@ public:
 class FExpressionTextureSample : public FExpression
 {
 public:
-	FExpressionTextureSample(FExpression* InTextureExpression,
-		FExpression* InTexCoordExpression,
-		const FExpressionDerivatives& InTexCoordDerivatives,
-		ESamplerSourceMode InSamplerSource,
-		ETextureMipValueMode InMipValueMode)
-		: TextureExpression(InTextureExpression)
+	FExpressionTextureSample(FTextureParameterDeclaration* InDeclaration, FExpression* InTexCoordExpression, const FExpressionDerivatives& InTexCoordDerivatives, ESamplerSourceMode InSamplerSource, ETextureMipValueMode InMipValueMode)
+		: Declaration(InDeclaration)
 		, TexCoordExpression(InTexCoordExpression)
 		, TexCoordDerivatives(InTexCoordDerivatives)
 		, SamplerSource(InSamplerSource)
 		, MipValueMode(InMipValueMode)
 	{}
 
-	FExpression* TextureExpression;
+	FTextureParameterDeclaration* Declaration;
 	FExpression* TexCoordExpression;
 	FExpressionDerivatives TexCoordDerivatives;
 	ESamplerSourceMode SamplerSource;
diff --git a/Engine/Source/Runtime/Engine/Public/HLSLTree/HLSLTreeEmit.h b/Engine/Source/Runtime/Engine/Public/HLSLTree/HLSLTreeEmit.h
--- a/Engine/Source/Runtime/Engine/Public/HLSLTree/HLSLTreeEmit.h
+++ b/Engine/Source/Runtime/Engine/Public/HLSLTree/HLSLTreeEmit.h
@@ -431,7 +431,6 @@ public:
 	uint32 CurrentBoolUniformOffset = ~0u;
 	uint32 CurrentNumBoolComponents = 32u;
 	bool bReadMaterialNormal = false;
-	bool bUsesVertexColor = false;
 	uint32 TexCoordMask[SF_NumFrequencies] = { 0u };
 
 	int32 NumExpressionLocals = 0;
diff --git a/Engine/Source/Runtime/Engine/Public/HLSLTree/HLSLTreeTypes.h b/Engine/Source/Runtime/Engine/Public/HLSLTree/HLSLTreeTypes.h
--- a/Engine/Source/Runtime/Engine/Public/HLSLTree/HLSLTreeTypes.h
+++ b/Engine/Source/Runtime/Engine/Public/HLSLTree/HLSLTreeTypes.h
@@ -21,6 +21,7 @@ namespace HLSLTree
 class FNode;
 class FStructType;
 class FExpression;
+class FTextureParameterDeclaration;
 class FStatement;
 class FScope;
 class FTree;
@@ -63,16 +64,9 @@ enum class EOperation : uint8
 	None,
 
 	// Unary Ops
-	Abs,
 	Neg,
 	Rcp,
 	Frac,
-	Floor,
-	Ceil,
-	Round,
-	Trunc,
-	Saturate,
-	Sign,
 	Length,
 	Normalize,
 
@@ -82,14 +76,11 @@ enum class EOperation : uint8
 	Mul,
 	Div,
 	Fmod,
-	PowPositiveClamped,
 	Dot,
 	Min,
 	Max,
 	Less,
 	Greater,
-	LessEqual,
-	GreaterEqual,
 
 	VecMulMatrix3,
 	VecMulMatrix4,
@@ -110,6 +101,32 @@ struct FOperationDescription
 
 FOperationDescription GetOperationDescription(EOperation Op);
 
+struct FTextureDescription
+{
+	FTextureDescription()
+		: Texture(nullptr), SamplerType(SAMPLERTYPE_Color)
+	{}
+
+	FTextureDescription(UTexture* InTexture, EMaterialSamplerType InSamplerType)
+		: Texture(InTexture), SamplerType(InSamplerType)
+	{}
+
+	UTexture* Texture;
+	EMaterialSamplerType SamplerType;
+};
+inline bool operator==(const FTextureDescription& Lhs, const FTextureDescription& Rhs)
+{
+	return Lhs.Texture == Rhs.Texture && Lhs.SamplerType == Rhs.SamplerType;
+}
+inline bool operator!=(const FTextureDescription& Lhs, const FTextureDescription& Rhs)
+{
+	return !operator==(Lhs, Rhs);
+}
+inline uint32 GetTypeHash(const FTextureDescription& Ref)
+{
+	return HashCombine(GetTypeHash(Ref.Texture), GetTypeHash(Ref.SamplerType));
+}
+
 struct FCustomHLSLInput
 {
 	FCustomHLSLInput() = default;
diff --git a/Engine/Source/Runtime/Engine/Public/MaterialExpressionIO.h b/Engine/Source/Runtime/Engine/Public/MaterialExpressionIO.h
--- a/Engine/Source/Runtime/Engine/Public/MaterialExpressionIO.h
+++ b/Engine/Source/Runtime/Engine/Public/MaterialExpressionIO.h
@@ -14,6 +14,7 @@ namespace HLSLTree
 {
 class FScope;
 class FExpression;
+class FTextureParameterDeclaration;
 }
 namespace Shader
 {
@@ -80,6 +81,7 @@ struct FExpressionInput
 	ENGINE_API UE::HLSLTree::FExpression* TryAcquireHLSLExpression(FMaterialHLSLGenerator& Generator, UE::HLSLTree::FScope& Scope) const;
 	ENGINE_API UE::HLSLTree::FExpression* AcquireHLSLExpression(FMaterialHLSLGenerator& Generator, UE::HLSLTree::FScope& Scope) const;
 	ENGINE_API UE::HLSLTree::FExpression* AcquireHLSLExpressionOrConstant(FMaterialHLSLGenerator& Generator, UE::HLSLTree::FScope& Scope, const UE::Shader::FValue& ConstantValue) const;
+	ENGINE_API UE::HLSLTree::FTextureParameterDeclaration* AcquireHLSLTexture(FMaterialHLSLGenerator& Generator, UE::HLSLTree::FScope& Scope) const;
 #endif // WITH_EDITOR
 
 	/**
diff --git a/Engine/Source/Runtime/Engine/Public/MaterialHLSLGenerator.h b/Engine/Source/Runtime/Engine/Public/MaterialHLSLGenerator.h
--- a/Engine/Source/Runtime/Engine/Public/MaterialHLSLGenerator.h
+++ b/Engine/Source/Runtime/Engine/Public/MaterialHLSLGenerator.h
@@ -93,7 +93,11 @@ public:
 	UE::HLSLTree::FExpression* NewTexCoord(int32 Index);
 	UE::HLSLTree::FExpression* NewSwizzle(const UE::HLSLTree::FSwizzleParameters& Params, UE::HLSLTree::FExpression* Input);
 
-	const UE::Shader::FTextureValue* AcquireTextureValue(const UE::Shader::FTextureValue& Value);
+	/** Returns a declaration to access the given texture, with no parameter */
+	UE::HLSLTree::FTextureParameterDeclaration* AcquireTextureDeclaration(const UE::HLSLTree::FTextureDescription& Value);
+
+	/** Returns a declaration to access the given texture parameter */
+	UE::HLSLTree::FTextureParameterDeclaration* AcquireTextureParameterDeclaration(const FName& Name, const UE::HLSLTree::FTextureDescription& DefaultValue);
 
 	/**
 	 * Returns the appropriate HLSLNode representing the given UMaterialExpression.
@@ -103,6 +107,7 @@ public:
 	UE::HLSLTree::FExpression* AcquireExpression(UE::HLSLTree::FScope& Scope, UMaterialExpression* MaterialExpression, int32 OutputIndex);
 	UE::HLSLTree::FExpression* AcquireFunctionInputExpression(UE::HLSLTree::FScope& Scope, const UMaterialExpressionFunctionInput* MaterialExpression);
 
+	UE::HLSLTree::FTextureParameterDeclaration* AcquireTextureDeclaration(UE::HLSLTree::FScope& Scope, UMaterialExpression* MaterialExpression, int32 OutputIndex);
 	bool GenerateStatements(UE::HLSLTree::FScope& Scope, UMaterialExpression* MaterialExpression);
 
 	UE::HLSLTree::FExpression* GenerateFunctionCall(UE::HLSLTree::FScope& Scope, UMaterialFunctionInterface* Function, TArrayView<const FFunctionExpressionInput> ConnectedInputs, int32 OutputIndex);
@@ -194,7 +199,8 @@ private:
 
 	TArray<FFunctionCallEntry*, TInlineAllocator<8>> FunctionCallStack;
 	TArray<UE::HLSLTree::FScope*> JoinedScopeStack;
-	TMap<FXxHash64, const UE::Shader::FTextureValue*> TextureValueMap;
+	TMap<UE::HLSLTree::FTextureDescription, UE::HLSLTree::FTextureParameterDeclaration*> TextureDeclarationMap;
+	TMap<FName, UE::HLSLTree::FTextureParameterDeclaration*> TextureParameterDeclarationMap;
 	TMap<FSHAHash, FFunctionCallEntry*> FunctionCallMap;
 	TMap<UMaterialFunctionInterface*, UE::HLSLTree::FFunction*> FunctionMap;
 	TMap<UMaterialExpression*, FStatementEntry> StatementMap;
diff --git a/Engine/Source/Runtime/Engine/Public/Shader/ShaderTypes.h b/Engine/Source/Runtime/Engine/Public/Shader/ShaderTypes.h
--- a/Engine/Source/Runtime/Engine/Public/Shader/ShaderTypes.h
+++ b/Engine/Source/Runtime/Engine/Public/Shader/ShaderTypes.h
@@ -2,13 +2,11 @@
 #pragma once
 
 #include "CoreMinimal.h"
-#include "Engine/EngineTypes.h"
 #include "Containers/ArrayView.h"
 #include "Containers/StringView.h"
 #include "Serialization/MemoryLayout.h"
 
 class FMemStackBase;
-class UTexture;
 
 namespace UE
 {
@@ -61,13 +59,6 @@ enum class EValueComponentType : uint8
 	Double,
 	Int,
 	Bool,
-
-	FirstTexture,
-	Texture2D = FirstTexture,
-	Texture2DArray,
-	TextureCube,
-	TextureCubeArray,
-	Texture3D,
 };
 
 struct FValueComponentTypeDescription
@@ -83,8 +74,6 @@ struct FValueComponentTypeDescription
 FValueComponentTypeDescription GetValueComponentTypeDescription(EValueComponentType Type);
 inline const TCHAR* GetComponentTypeName(EValueComponentType Type) { return GetValueComponentTypeDescription(Type).Name; }
 inline uint32 GetComponentTypeSizeInBytes(EValueComponentType Type) { return GetValueComponentTypeDescription(Type).SizeInBytes; }
-inline bool IsNumericType(EValueComponentType Type) { return Type != EValueComponentType::Void && (uint8)Type < (uint8)EValueComponentType::FirstTexture; }
-inline bool IsTextureType(EValueComponentType Type) { return (uint8)Type >= (uint8)EValueComponentType::FirstTexture; }
 inline bool IsComponentTypeWithinBounds(EValueComponentType Type, FComponentBounds Bounds) { return IsWithinBounds(GetValueComponentTypeDescription(Type).Bounds, Bounds); }
 
 EValueComponentType CombineComponentTypes(EValueComponentType Lhs, EValueComponentType Rhs);
@@ -125,12 +114,6 @@ enum class EValueType : uint8
 	DoubleInverse4x4,
 
 	Struct,
-
-	Texture2D,
-	Texture2DArray,
-	TextureCube,
-	TextureCubeArray,
-	Texture3D,
 };
 
 struct FValueTypeDescription
@@ -152,9 +135,6 @@ EValueType MakeNonLWCType(EValueType Type);
 EValueType MakeArithmeticResultType(EValueType Lhs, EValueType Rhs, FString& OutErrorMessage);
 EValueType MakeComparisonResultType(EValueType Lhs, EValueType Rhs, FString& OutErrorMessage);
 
-inline bool IsNumericType(EValueType Type) { return IsNumericType(GetValueTypeDescription(Type).ComponentType); }
-inline bool IsTextureType(EValueType Type) { return IsTextureType(GetValueTypeDescription(Type).ComponentType); }
-
 struct FType
 {
 	FType() : StructType(nullptr), ValueType(EValueType::Void) {}
@@ -166,8 +146,7 @@ struct FType
 	FType GetNonLWCType() const { return IsNumericLWC() ? FType(MakeNonLWCType(ValueType)) : *this; }
 	bool IsVoid() const { return ValueType == EValueType::Void; }
 	bool IsStruct() const { return ValueType == EValueType::Struct; }
-	bool IsNumeric() const { return !IsStruct() && IsNumericType(ValueType); }
-	bool IsTexture() const { return !IsStruct() && IsTextureType(ValueType); }
+	bool IsNumeric() const { return !IsVoid() && !IsStruct(); }
 	bool IsNumericLWC() const { return IsNumeric() && IsLWCType(ValueType); }
 	int32 GetNumComponents() const;
 	int32 GetNumFlatFields() const;
@@ -302,37 +281,13 @@ struct FMemoryImageValue
 	uint32 Size;
 };
 
-struct FTextureValue
-{
-	FTextureValue() = default;
-	FTextureValue(UTexture* InTexture, EMaterialSamplerType InSamplerType = SAMPLERTYPE_Color)
-		: Texture(InTexture)
-		, SamplerType(InSamplerType)
-	{}
-
-	EValueType GetType() const;
-
-	UTexture* Texture = nullptr;
-	EMaterialSamplerType SamplerType = SAMPLERTYPE_Color;
-};
-
-inline bool operator==(const FTextureValue& Lhs, const FTextureValue& Rhs)
-{
-	return Lhs.Texture == Rhs.Texture && Lhs.SamplerType == Rhs.SamplerType;
-}
-inline bool operator!=(const FTextureValue& Lhs, const FTextureValue& Rhs)
-{
-	return !operator==(Lhs, Rhs);
-}
-
 union FValueComponent
 {
 	FValueComponent() : Packed(0u) {}
-	FValueComponent(float InValue) : Packed(0u) { Float = InValue; }
-	FValueComponent(double InValue) : Packed(0u) { Double = InValue; }
-	FValueComponent(int32 InValue) : Packed(0u) { Int = InValue; }
-	FValueComponent(bool InValue) : Packed(0u) { Bool = InValue ? 1 : 0; }
-	FValueComponent(const FTextureValue* InValue) : Packed(0u) { Texture = InValue; }
+	FValueComponent(float v) : Packed(0u) { Float = v; }
+	FValueComponent(double v) : Packed(0u) { Double = v; }
+	FValueComponent(int32 v) : Packed(0u) { Int = v; }
+	FValueComponent(bool v) : Packed(0u) { Bool = v ? 1 : 0; }
 
 	// 'Bool' is stored as uint8 to avoid changing on different compilers
 	bool AsBool() const { return Bool != 0u; }
@@ -340,7 +295,6 @@ union FValueComponent
 	const TCHAR* ToString(EValueComponentType Type, FStringBuilderBase& OutString) const;
 
 	uint64 Packed;
-	const FTextureValue* Texture;
 	double Double;
 	float Float;
 	int32 Int;
@@ -476,8 +430,6 @@ struct FValue
 		Component.Add(v);
 	}
 
-	FValue(const FTextureValue* InValue);
-
 	inline const FType& GetType() const { return Type; }
 	inline const FValueComponent& GetComponent(int32 i) const
 	{
@@ -498,7 +450,6 @@ struct FValue
 	FDoubleValue AsDouble() const;
 	FIntValue AsInt() const;
 	FBoolValue AsBool() const;
-	const FTextureValue* AsTexture() const;
 
 	FLinearColor AsLinearColor() const;
 	FVector4d AsVector4d() const;
