Export a CSV file with number of object instances for Crunchersharp
#tests dump .csv from client and server, open them with CruncherSharp

diff --git a/Engine/Source/Runtime/Engine/Private/UnrealEngine.cpp b/Engine/Source/Runtime/Engine/Private/UnrealEngine.cpp
--- a/Engine/Source/Runtime/Engine/Private/UnrealEngine.cpp
+++ b/Engine/Source/Runtime/Engine/Private/UnrealEngine.cpp
@@ -29,6 +29,7 @@ UnrealEngine.cpp: Implements the UEngine class and helpers.
 #include "Misc/FrameValue.h"
 #include "HAL/Runnable.h"
 #include "Misc/OutputDeviceArchiveWrapper.h"
+#include "Misc/OutputDeviceFile.h"
 #include "Stats/StatsMisc.h"
 #include "Containers/Ticker.h"
 #include "Misc/ConfigCacheIni.h"
@@ -8900,6 +8901,90 @@ bool UEngine::HandleObjCommand( const TCHAR* Cmd, FOutputDevice& Ar )
 		return true;
 
 	}
+	else if (FParse::Command(&Cmd, TEXT("COUNT")))
+	{
+		FSlowHeartBeatScope SuspendHeartBeat;
+		FDisableHitchDetectorScope SuspendGameThreadHitch;
+
+		TMap<FString, int32> ClassMap;
+
+		for( FThreadSafeObjectIterator It; It; ++It )
+		{
+			UClass* Class = It->GetClass();
+
+			// Iterate through containers to find struct properties
+			for (TFieldIterator<FProperty> ItProp(Class, EFieldIterationFlags::IncludeAll); ItProp; ++ItProp)
+			{
+				FProperty* Property = *ItProp;
+				if (FArrayProperty* ArrayProperty = CastField<FArrayProperty>(Property))
+				{
+					if (FStructProperty* StructProperty = CastField<FStructProperty>(ArrayProperty->Inner))
+					{
+						FScriptArrayHelper ArrayHelper(ArrayProperty, ArrayProperty->template ContainerPtrToValuePtr<void>(*It));
+						if (int32 Num = ArrayHelper.Num())
+						{
+							FString StructName;
+							StructProperty->GetCPPMacroType(StructName);
+							ClassMap.FindOrAdd(StructName, 0) += Num;
+						}
+					}
+				}
+				else if (FSetProperty* SetProperty = CastField<FSetProperty>(Property))
+				{
+					if (FStructProperty* StructProperty = CastField<FStructProperty>(SetProperty->ElementProp))
+					{
+						FScriptSetHelper SetHelper(SetProperty, SetProperty->template ContainerPtrToValuePtr<void>(*It));
+						if (int32 Num = SetHelper.Num())
+						{
+							FString StructName;
+							StructProperty->GetCPPMacroType(StructName);
+							ClassMap.FindOrAdd(StructName, 0) += Num;
+						}
+					}
+				}
+				else if (FMapProperty* MapProperty = CastField<FMapProperty>(Property))
+				{
+					FScriptMapHelper MapHelper(MapProperty, MapProperty->template ContainerPtrToValuePtr<void>(*It));
+					if (int32 Num = MapHelper.Num())
+					{
+						if (FStructProperty* StructProperty = CastField<FStructProperty>(MapProperty->KeyProp))
+						{
+							FString StructName;
+							StructProperty->GetCPPMacroType(StructName);
+							ClassMap.FindOrAdd(StructName, 0) += Num;
+						}
+						if (FStructProperty* StructProperty = CastField<FStructProperty>(MapProperty->ValueProp))
+						{
+							FString StructName;
+							StructProperty->GetCPPMacroType(StructName);
+							ClassMap.FindOrAdd(StructName, 0) += Num;
+						}
+					}
+				}
+			}
+
+			while (Class && !Class->IsNative())
+			{
+				Class = Class->GetSuperClass();
+			}
+
+			FString ClassName = FString::Format(TEXT("{0}{1}"), { Class->GetPrefixCPP(), Class->GetName() });
+			ClassMap.FindOrAdd(ClassName,0)++;
+		}
+
+		if( ClassMap.Num() )
+		{
+			FOutputDeviceFile CSVFile{ *FPaths::Combine(*FPaths::ProfilingDir(), TEXT("ObjCount.csv")) };
+			CSVFile.SetSuppressEventTag(true);
+
+			for (const auto& pair : ClassMap)
+			{
+				CSVFile.Logf(ELogVerbosity::Log, TEXT("%s,%i"), *pair.Key, pair.Value);
+			}
+		}
+		return true;
+	}
+	
 	else if (FParse::Command(&Cmd, TEXT("PROPANALYSIS")))
 	{
 		HandlePropAnalysisCommand(Cmd, Ar);
