Optional thread safety locks for modular features. Ensure when doing unsafe calls into modular feature enumeration functions.

Use thread safety locks when finding latency marker modular features from rhi thread

diff --git a/Engine/Source/Runtime/ClothingSystemRuntimeInterface/Private/ClothingSimulationFactory.cpp b/Engine/Source/Runtime/ClothingSystemRuntimeInterface/Private/ClothingSimulationFactory.cpp
--- a/Engine/Source/Runtime/ClothingSystemRuntimeInterface/Private/ClothingSimulationFactory.cpp
+++ b/Engine/Source/Runtime/ClothingSystemRuntimeInterface/Private/ClothingSimulationFactory.cpp
@@ -33,8 +33,10 @@ TSubclassOf<class UClothingSimulationFactory> UClothingSimulationFactory::GetDef
 	TSubclassOf<UClothingSimulationFactory> DefaultClothingSimulationFactoryClass = nullptr;
 
 	const FString DefaultClothingSimulationFactoryClassName = ClothingSimulationFactoryConsoleVariables::CVarDefaultClothingSimulationFactoryClass.GetValueOnAnyThread();
-
+	
+	IModularFeatures::Get().LockModularFeatureList();
 	const TArray<IClothingSimulationFactoryClassProvider*> ClassProviders = IModularFeatures::Get().GetModularFeatureImplementations<IClothingSimulationFactoryClassProvider>(IClothingSimulationFactoryClassProvider::FeatureName);
+	IModularFeatures::Get().UnlockModularFeatureList();
 	for (const auto& ClassProvider : ClassProviders)
 	{
 		if (ClassProvider)
@@ -53,5 +55,6 @@ TSubclassOf<class UClothingSimulationFactory> UClothingSimulationFactory::GetDef
 			}
 		}
 	}
+
 	return DefaultClothingSimulationFactoryClass;
 }
diff --git a/Engine/Source/Runtime/Core/Private/Features/ModularFeatures.cpp b/Engine/Source/Runtime/Core/Private/Features/ModularFeatures.cpp
--- a/Engine/Source/Runtime/Core/Private/Features/ModularFeatures.cpp
+++ b/Engine/Source/Runtime/Core/Private/Features/ModularFeatures.cpp
@@ -10,15 +10,29 @@ IModularFeatures& IModularFeatures::Get()
 	return ModularFeatures;
 }
 
+void FModularFeatures::LockModularFeatureList()
+{
+	ModularFeaturesMapCriticalSection.Lock();
+	bModularFeatureListLocked = true;
+}
+
+void FModularFeatures::UnlockModularFeatureList()
+{
+	ModularFeaturesMapCriticalSection.Unlock();
+	bModularFeatureListLocked = false;
+}
 
 int32 FModularFeatures::GetModularFeatureImplementationCount( const FName Type )
 {
+	ensureMsgf(IsInGameThread() || bModularFeatureListLocked, TEXT("IModularFeature counting is not thread-safe unless wrapped with LockModularFeatureList/UnlockModularFeatureList"));
+
 	return ModularFeaturesMap.Num( Type );
 }
 
-
 IModularFeature* FModularFeatures::GetModularFeatureImplementation( const FName Type, const int32 Index )
 {
+	ensureMsgf(IsInGameThread() || bModularFeatureListLocked, TEXT("IModularFeature fetching is not thread-safe unless wrapped with LockModularFeatureList/UnlockModularFeatureList"));
+
 	IModularFeature* ModularFeature = nullptr;
 
 	int32 CurrentIndex = 0;
@@ -40,6 +54,8 @@ IModularFeature* FModularFeatures::GetModularFeatureImplementation( const FName
 
 void FModularFeatures::RegisterModularFeature( const FName Type, IModularFeature* ModularFeature )
 {
+	FScopeLock ScopeLock(&ModularFeaturesMapCriticalSection);
+
 	ModularFeaturesMap.AddUnique( Type, ModularFeature );
 	ModularFeatureRegisteredEvent.Broadcast( Type, ModularFeature );
 }
@@ -47,6 +63,8 @@ void FModularFeatures::RegisterModularFeature( const FName Type, IModularFeature
 
 void FModularFeatures::UnregisterModularFeature( const FName Type, IModularFeature* ModularFeature )
 {
+	FScopeLock ScopeLock(&ModularFeaturesMapCriticalSection);
+
 	ModularFeaturesMap.RemoveSingle( Type, ModularFeature );
 	ModularFeatureUnregisteredEvent.Broadcast( Type, ModularFeature );
 }
diff --git a/Engine/Source/Runtime/Core/Private/Features/ModularFeatures.h b/Engine/Source/Runtime/Core/Private/Features/ModularFeatures.h
--- a/Engine/Source/Runtime/Core/Private/Features/ModularFeatures.h
+++ b/Engine/Source/Runtime/Core/Private/Features/ModularFeatures.h
@@ -27,15 +27,22 @@ public:
 	virtual IModularFeatures::FOnModularFeatureRegistered& OnModularFeatureRegistered() override;
 	DECLARE_DERIVED_EVENT(FModularFeatures, IModularFeatures::FOnModularFeatureUnregistered, FOnModularFeatureUnregistered);
 	virtual IModularFeatures::FOnModularFeatureUnregistered& OnModularFeatureUnregistered() override;
+	virtual void LockModularFeatureList() override;
+	virtual void UnlockModularFeatureList() override;
 
 private:
 
 
 private:
-
 	/** Maps each feature type to a list of known providers of that feature */
 	TMultiMap< FName, class IModularFeature* > ModularFeaturesMap;
 
+	/** Used to fire ensure when ModularFeaturesMap is inspected cross-thread without being locked properly */
+	bool bModularFeatureListLocked = false;
+
+	/** Lock modular features map so it can be used across threads */
+	FCriticalSection ModularFeaturesMapCriticalSection;
+
 	/** Event used to inform clients that a modular feature has been registered */
 	FOnModularFeatureRegistered ModularFeatureRegisteredEvent;
 
diff --git a/Engine/Source/Runtime/Core/Public/Features/IModularFeatures.h b/Engine/Source/Runtime/Core/Public/Features/IModularFeatures.h
--- a/Engine/Source/Runtime/Core/Public/Features/IModularFeatures.h
+++ b/Engine/Source/Runtime/Core/Public/Features/IModularFeatures.h
@@ -79,6 +79,12 @@ public:
 
 public:
 
+	/** Lock modular features for access from another thread */
+	virtual void LockModularFeatureList() = 0;
+
+	/** Unlock modular features when finished accessing from another thread */
+	virtual void UnlockModularFeatureList() = 0;
+
 	/**
 	 * Returns the number of registered implementations of the specified feature type.
 	 *
diff --git a/Engine/Source/Runtime/Engine/Private/UnrealEngine.cpp b/Engine/Source/Runtime/Engine/Private/UnrealEngine.cpp
--- a/Engine/Source/Runtime/Engine/Private/UnrealEngine.cpp
+++ b/Engine/Source/Runtime/Engine/Private/UnrealEngine.cpp
@@ -2342,7 +2342,9 @@ void UEngine::SetSimulationLatencyMarkerEnd(uint64 FrameNumber)
 
 void UEngine::SetPresentLatencyMarkerStart(uint64 FrameNumber)
 {
+	IModularFeatures::Get().LockModularFeatureList();
 	TArray<ILatencyMarkerModule*> LatencyMarkerModules = IModularFeatures::Get().GetModularFeatureImplementations<ILatencyMarkerModule>(ILatencyMarkerModule::GetModularFeatureName());
+	IModularFeatures::Get().UnlockModularFeatureList();
 
 	for (ILatencyMarkerModule* LatencyMarkerModule : LatencyMarkerModules)
 	{
@@ -2352,7 +2354,9 @@ void UEngine::SetPresentLatencyMarkerStart(uint64 FrameNumber)
 
 void UEngine::SetPresentLatencyMarkerEnd(uint64 FrameNumber)
 {
+	IModularFeatures::Get().LockModularFeatureList();
 	TArray<ILatencyMarkerModule*> LatencyMarkerModules = IModularFeatures::Get().GetModularFeatureImplementations<ILatencyMarkerModule>(ILatencyMarkerModule::GetModularFeatureName());
+	IModularFeatures::Get().UnlockModularFeatureList();
 
 	for (ILatencyMarkerModule* LatencyMarkerModule : LatencyMarkerModules)
 	{
@@ -2363,7 +2367,9 @@ void UEngine::SetPresentLatencyMarkerEnd(uint64 FrameNumber)
 
 void UEngine::SetRenderSubmitLatencyMarkerStart(uint64 FrameNumber)
 {
+	IModularFeatures::Get().LockModularFeatureList();
 	TArray<ILatencyMarkerModule*> LatencyMarkerModules = IModularFeatures::Get().GetModularFeatureImplementations<ILatencyMarkerModule>(ILatencyMarkerModule::GetModularFeatureName());
+	IModularFeatures::Get().UnlockModularFeatureList();
 
 	for (ILatencyMarkerModule* LatencyMarkerModule : LatencyMarkerModules)
 	{
@@ -2373,7 +2379,9 @@ void UEngine::SetRenderSubmitLatencyMarkerStart(uint64 FrameNumber)
 
 void UEngine::SetRenderSubmitLatencyMarkerEnd(uint64 FrameNumber)
 {
+	IModularFeatures::Get().LockModularFeatureList();
 	TArray<ILatencyMarkerModule*> LatencyMarkerModules = IModularFeatures::Get().GetModularFeatureImplementations<ILatencyMarkerModule>(ILatencyMarkerModule::GetModularFeatureName());
+	IModularFeatures::Get().UnlockModularFeatureList();
 
 	for (ILatencyMarkerModule* LatencyMarkerModule : LatencyMarkerModules)
 	{
diff --git a/Engine/Source/Runtime/SignalProcessing/Private/FFTAlgorithm.cpp b/Engine/Source/Runtime/SignalProcessing/Private/FFTAlgorithm.cpp
--- a/Engine/Source/Runtime/SignalProcessing/Private/FFTAlgorithm.cpp
+++ b/Engine/Source/Runtime/SignalProcessing/Private/FFTAlgorithm.cpp
@@ -13,7 +13,9 @@ namespace Audio
 	{
 
 		// Get all IFFTAlgorithm factories. 
+		IModularFeatures::Get().LockModularFeatureList();
 		TArray<IFFTAlgorithmFactory*> Factories = IModularFeatures::Get().GetModularFeatureImplementations<IFFTAlgorithmFactory>(IFFTAlgorithmFactory::GetModularFeatureName());
+		IModularFeatures::Get().UnlockModularFeatureList();
 
 		// Remove null factories
 		Factories = Factories.FilterByPredicate([InAlgorithmFactoryName](const IFFTAlgorithmFactory* Factory) 
