Clear meta data of all index buffers used by staticmeshes, since we do the same for all vertex buffers. Also, dont clear the CachedNumIndices when we discard a LOD, for example wheh we cancel a FStaticMeshStreamIn.

diff --git a/Engine/Source/Runtime/Engine/Private/RawIndexBuffer.cpp b/Engine/Source/Runtime/Engine/Private/RawIndexBuffer.cpp
--- a/Engine/Source/Runtime/Engine/Private/RawIndexBuffer.cpp
+++ b/Engine/Source/Runtime/Engine/Private/RawIndexBuffer.cpp
@@ -426,11 +426,15 @@ void FRawStaticIndexBuffer::SerializeMetaData(FArchive& Ar)
 	Ar << CachedNumIndices << b32Bit;
 }
 
+void FRawStaticIndexBuffer::ClearMetaData()
+{
+	CachedNumIndices = -1;
+}
+
 void FRawStaticIndexBuffer::Discard()
 {
     IndexStorage.SetAllowCPUAccess(false);
     IndexStorage.Discard();
-	CachedNumIndices = IndexStorage.Num() / (b32Bit ? 4 : 2);
 }
 
 /*-----------------------------------------------------------------------------
diff --git a/Engine/Source/Runtime/Engine/Private/StaticMesh.cpp b/Engine/Source/Runtime/Engine/Private/StaticMesh.cpp
--- a/Engine/Source/Runtime/Engine/Private/StaticMesh.cpp
+++ b/Engine/Source/Runtime/Engine/Private/StaticMesh.cpp
@@ -545,6 +545,7 @@ void FStaticMeshLODResources::SerializeBuffers(FArchive& Ar, UStaticMesh* OwnerS
 		if (!bEnableReversedIndexBuffer)
 		{
 			SerializedAdditionalIndexBuffers->ReversedIndexBuffer.Discard();
+			SerializedAdditionalIndexBuffers->ReversedIndexBuffer.ClearMetaData();
 		}
 	}
 
@@ -554,6 +555,7 @@ void FStaticMeshLODResources::SerializeBuffers(FArchive& Ar, UStaticMesh* OwnerS
 	if (!bEnableDepthOnlyIndexBuffer)
 	{
 		DepthOnlyIndexBuffer.Discard();
+		DepthOnlyIndexBuffer.ClearMetaData();
 	}
 
 	if (bSerializeReversedIndexBuffer)
@@ -564,6 +566,7 @@ void FStaticMeshLODResources::SerializeBuffers(FArchive& Ar, UStaticMesh* OwnerS
 		if (!bEnableReversedIndexBuffer)
 		{
 			SerializedAdditionalIndexBuffers->ReversedDepthOnlyIndexBuffer.Discard();
+			SerializedAdditionalIndexBuffers->ReversedDepthOnlyIndexBuffer.ClearMetaData();
 		}
 	}
 
@@ -658,21 +661,25 @@ void FStaticMeshLODResources::SerializeAvailabilityInfo(FArchive& Ar)
 		// Reversed indices are either stripped during cook or will be stripped on load.
 		// In either case, clear CachedNumIndices to show that the buffer will be empty after actual loading
 		SerializedAdditionalIndexBuffers->ReversedIndexBuffer.Discard();
+		SerializedAdditionalIndexBuffers->ReversedIndexBuffer.ClearMetaData();
 	}
 	DepthOnlyIndexBuffer.SerializeMetaData(Ar);
 	if (!bHasDepthOnlyIndices)
 	{
 		DepthOnlyIndexBuffer.Discard();
+		DepthOnlyIndexBuffer.ClearMetaData();
 	}
 	SerializedAdditionalIndexBuffers->ReversedDepthOnlyIndexBuffer.SerializeMetaData(Ar);
 	if (!bHasReversedDepthOnlyIndices)
 	{
 		SerializedAdditionalIndexBuffers->ReversedDepthOnlyIndexBuffer.Discard();
+		SerializedAdditionalIndexBuffers->ReversedDepthOnlyIndexBuffer.ClearMetaData();
 	}
 	SerializedAdditionalIndexBuffers->WireframeIndexBuffer.SerializeMetaData(Ar);
 	if (!bHasWireframeIndices)
 	{
 		SerializedAdditionalIndexBuffers->WireframeIndexBuffer.Discard();
+		SerializedAdditionalIndexBuffers->WireframeIndexBuffer.ClearMetaData();
 	}
 	if (Ar.IsLoading() && Ar.CustomVer(FUE5ReleaseStreamObjectVersion::GUID) < FUE5ReleaseStreamObjectVersion::RemovingTessellation)
 	{
@@ -699,6 +706,17 @@ void FStaticMeshLODResources::ClearAvailabilityInfo()
 	VertexBuffers.PositionVertexBuffer.ClearMetaData();
 	VertexBuffers.ColorVertexBuffer.ClearMetaData();
 
+	IndexBuffer.ClearMetaData();
+	IndexBuffer.ClearMetaData();
+	DepthOnlyIndexBuffer.ClearMetaData();
+
+	if (AdditionalIndexBuffers)
+	{
+		AdditionalIndexBuffers->ReversedIndexBuffer.ClearMetaData();
+		AdditionalIndexBuffers->ReversedDepthOnlyIndexBuffer.ClearMetaData();
+		AdditionalIndexBuffers->WireframeIndexBuffer.ClearMetaData();
+	}
+
 	delete AdditionalIndexBuffers;
 	AdditionalIndexBuffers = nullptr;
 }
diff --git a/Engine/Source/Runtime/Engine/Public/RawIndexBuffer.h b/Engine/Source/Runtime/Engine/Public/RawIndexBuffer.h
--- a/Engine/Source/Runtime/Engine/Public/RawIndexBuffer.h
+++ b/Engine/Source/Runtime/Engine/Public/RawIndexBuffer.h
@@ -311,6 +311,8 @@ public:
 	/** Serialize only meta data (e.g. number of indices) but not the actual index data */
 	void SerializeMetaData(FArchive& Ar);
 
+	void ClearMetaData();
+
     /**
      * Discard
      * discards the serialized data when it is not needed
