IndexCurve::Scale/Shift-Curve crashes in case index mapping has not been populated

diff --git a/Engine/Source/Runtime/Engine/Private/Curves/IndexedCurve.cpp b/Engine/Source/Runtime/Engine/Private/Curves/IndexedCurve.cpp
--- a/Engine/Source/Runtime/Engine/Private/Curves/IndexedCurve.cpp
+++ b/Engine/Source/Runtime/Engine/Private/Curves/IndexedCurve.cpp
@@ -86,6 +86,7 @@ FKeyHandle FIndexedCurve::GetKeyHandle(int32 KeyIndex) const
 
 void FIndexedCurve::ShiftCurve(float DeltaTime)
 {
+	EnsureAllIndicesHaveHandles();
 	TSet<FKeyHandle> KeyHandles;
 	for (auto It = KeyHandlesToIndices.CreateConstIterator(); It; ++It)
 	{
@@ -112,6 +113,7 @@ void FIndexedCurve::ShiftCurve(float DeltaTime, const TSet<FKeyHandle>& KeyHandl
 
 void FIndexedCurve::ScaleCurve(float ScaleOrigin, float ScaleFactor)
 {
+	EnsureAllIndicesHaveHandles();
 	TSet<FKeyHandle> KeyHandles;
 	for (auto It = KeyHandlesToIndices.CreateConstIterator(); It; ++It)
 	{
