Strata does not crash the editor when the strata tree is incomplete.

diff --git a/Engine/Source/Runtime/Engine/Private/Materials/HLSLMaterialTranslator.cpp b/Engine/Source/Runtime/Engine/Private/Materials/HLSLMaterialTranslator.cpp
--- a/Engine/Source/Runtime/Engine/Private/Materials/HLSLMaterialTranslator.cpp
+++ b/Engine/Source/Runtime/Engine/Private/Materials/HLSLMaterialTranslator.cpp
@@ -146,6 +146,39 @@ static inline const TCHAR * GetFloatZeroVector(uint32 NumComponents)
 	}
 }
 
+static inline const TCHAR* GetStrataOperatorStr(int32 OperatorType)
+{
+	switch (OperatorType)
+	{
+	case STRATA_OPERATOR_WEIGHT:
+	{
+		return TEXT("WEIGHT    ");
+	}
+	case STRATA_OPERATOR_VERTICAL:
+	{
+		return TEXT("VERTICAL  ");
+	}
+	case STRATA_OPERATOR_HORIZONTAL:
+	{
+		return TEXT("HORIZONTAL");
+	}
+	case STRATA_OPERATOR_ADD:
+	{
+		return TEXT("ADD       ");
+	}
+	case STRATA_OPERATOR_BSDF:
+	{
+		return TEXT("BSDF      ");
+	}
+	case STRATA_OPERATOR_BSDF_LEGACY:
+	{
+		return TEXT("BSDFLEGACY");
+	}
+	}
+	return TEXT("UNKNOWN   ");
+};
+
+
 FHLSLMaterialTranslator::FHLSLMaterialTranslator(FMaterial* InMaterial,
 	FMaterialCompilationOutput& InMaterialCompilationOutput,
 	const FStaticParameterSet& InStaticParameters,
@@ -666,7 +699,10 @@ bool FHLSLMaterialTranslator::Translate()
 		if (bStrataEnabled && FrontMaterialExpr)
 		{
 			FrontMaterialExpr->StrataGenerateMaterialTopologyTree(this, nullptr, 0);
-			StrataGenerateDerivedMaterialOperatorData();
+			if (!StrataGenerateDerivedMaterialOperatorData())
+			{
+				Errorf(TEXT("Strata material errors encountered."));
+			}
 		}
 
 		// Generate code:
@@ -1989,45 +2025,6 @@ void FHLSLMaterialTranslator::GetMaterialEnvironment(EShaderPlatform InPlatform,
 			StrataMaterialDescription += FString::Printf(TEXT("Graph maximum distance to leaves %u\r\n"), RootMaximumDistanceToLeaves);
 			// Debug print operators according to depth from root.
 			{
-				auto OperatorTypeToStr = [&](int32 OperatorType)
-				{
-					switch (OperatorType)
-					{
-					case STRATA_OPERATOR_WEIGHT:
-					{
-						static const FString Name = TEXT("WEIGHT ");
-						return Name;
-					}
-					case STRATA_OPERATOR_VERTICAL:
-					{
-						static const FString Name = TEXT("VERTICA");
-						return Name;
-					}
-					case STRATA_OPERATOR_HORIZONTAL:
-					{
-						static const FString Name = TEXT("HORIZON");
-						return Name;
-					}
-					case STRATA_OPERATOR_ADD:
-					{
-						static const FString Name = TEXT("ADD    ");
-						return Name;
-					}
-					case STRATA_OPERATOR_BSDF:
-					{
-						static const FString Name = TEXT("BSDF   ");
-						return Name;
-					}
-					case STRATA_OPERATOR_BSDF_LEGACY:
-					{
-						static const FString Name = TEXT("BSDFLEG");
-						return Name;
-					}
-					}
-					static const FString Name = TEXT("UNKNOWN");
-					return Name;
-				};
-
 				check(StrataMaterialRootOperator);
 				for (int32 DistanceToLeaves = RootMaximumDistanceToLeaves; DistanceToLeaves >= 0; --DistanceToLeaves)
 				{
@@ -2037,7 +2034,7 @@ void FHLSLMaterialTranslator::GetMaterialEnvironment(EShaderPlatform InPlatform,
 						if (!It.IsDiscarded() && It.MaxDistanceFromLeaves == DistanceToLeaves)
 						{
 							StrataMaterialDescription += FString::Printf(TEXT("\tIdx=%d Op=%s ParentIdx=%d LeftIndex=%d RightIndex=%d BSDFIdx=%d LayerDepth=%d IsTop=%d IsBot=%d\r\n"),
-								It.Index, *OperatorTypeToStr(It.OperatorType), It.ParentIndex, It.LeftIndex, It.RightIndex, It.BSDFIndex, It.LayerDepth, It.bIsTop, It.bIsBottom);
+								It.Index, GetStrataOperatorStr(It.OperatorType), It.ParentIndex, It.LeftIndex, It.RightIndex, It.BSDFIndex, It.LayerDepth, It.bIsTop, It.bIsBottom);
 						}
 					}
 				}
@@ -9648,7 +9645,7 @@ FStrataOperator& FHLSLMaterialTranslator::StrataCompilationGetOperator(UMaterial
 	return StrataMaterialExpressionRegisteredOperators[*OperatorIndex];
 }
 
-void FHLSLMaterialTranslator::StrataGenerateDerivedMaterialOperatorData()
+bool FHLSLMaterialTranslator::StrataGenerateDerivedMaterialOperatorData()
 {
 	//
 	// Evaluate the one and only root node
@@ -9663,7 +9660,59 @@ void FHLSLMaterialTranslator::StrataGenerateDerivedMaterialOperatorData()
 		}
 	}
 	StrataMaterialRootOperator = &StrataMaterialExpressionRegisteredOperators[RootIndex];
-	check(StrataMaterialRootOperator);
+	if (!StrataMaterialRootOperator)
+	{
+		UE_LOG(LogMaterial, Error, TEXT("Cannot find the root of the Strata Tree for Material %s (asset: %s).\r\n"), *Material->GetDebugName(), *Material->GetAssetPath().ToString());
+		return false;
+	}
+
+	//
+	// Make sure each and every path of the strata tree ends up on a valid BSDF.
+	//
+	for (auto& It : StrataMaterialExpressionRegisteredOperators)
+	{
+		if (It.IsDiscarded())
+		{
+			continue; // ignore discarded operations in sub tree using parameter blending
+		}
+
+		bool bMustHaveLeftChild = false;
+		bool bMustHaveRightChild = false;
+		switch (It.OperatorType)
+		{
+		// Operators without any child
+		case STRATA_OPERATOR_BSDF:
+		break;
+
+		// Operators with two children
+		case STRATA_OPERATOR_HORIZONTAL:
+		case STRATA_OPERATOR_VERTICAL:
+		case STRATA_OPERATOR_ADD:
+		{
+			bMustHaveLeftChild = true;
+			bMustHaveRightChild = true;
+		}
+		break;
+
+		// Operators with a single child
+		case STRATA_OPERATOR_WEIGHT:
+		{
+			bMustHaveLeftChild = true;
+		}
+		break;
+		}
+
+		if (bMustHaveLeftChild && It.LeftIndex == INDEX_NONE)
+		{
+			UE_LOG(LogMaterial, Error, TEXT("A Strata Operator %s node is missing its first input from material %s (asset: %s).\r\n"), GetStrataOperatorStr(It.OperatorType), *Material->GetDebugName(), *Material->GetAssetPath().ToString());
+			return false;
+		}
+		if (bMustHaveRightChild && It.RightIndex == INDEX_NONE)
+		{
+			UE_LOG(LogMaterial, Error, TEXT("A Strata Operator %s node is missing its second input from material %s (asset: %s).\r\n"), GetStrataOperatorStr(It.OperatorType), *Material->GetDebugName(), *Material->GetAssetPath().ToString());
+			return false;
+		}
+	}
 
 	//
 	// Parse the tree and mark nodes that are the root of a subtree using parameter blending, while other nodes in that tree are forced to use parameter blending.
@@ -9865,6 +9914,8 @@ void FHLSLMaterialTranslator::StrataGenerateDerivedMaterialOperatorData()
 
 		WalkOperators(*StrataMaterialRootOperator);
 	}
+
+	return true; // Success
 }
 
 void FHLSLMaterialTranslator::StrataCompilationInfoRegisterCodeChunk(int32 CodeChunk, FStrataMaterialCompilationInfo& StrataMaterialCompilationInfo)
diff --git a/Engine/Source/Runtime/Engine/Private/Materials/HLSLMaterialTranslator.h b/Engine/Source/Runtime/Engine/Private/Materials/HLSLMaterialTranslator.h
--- a/Engine/Source/Runtime/Engine/Private/Materials/HLSLMaterialTranslator.h
+++ b/Engine/Source/Runtime/Engine/Private/Materials/HLSLMaterialTranslator.h
@@ -643,7 +643,7 @@ protected:
 
 	int32 GenericSwitch(const TCHAR* Function, int32 IfTrue, int32 IfFalse);
 
-	void StrataGenerateDerivedMaterialOperatorData();
+	bool StrataGenerateDerivedMaterialOperatorData();
 
 	// FMaterialCompiler interface.
 
