Ensure slow task dialog is shown when blocking on _all_ shader maps to be compiled.

- This way we can see the progress on how many shaders are being compiled.

diff --git a/Engine/Source/Runtime/Engine/Private/ShaderCompiler/ShaderCompiler.cpp b/Engine/Source/Runtime/Engine/Private/ShaderCompiler/ShaderCompiler.cpp
--- a/Engine/Source/Runtime/Engine/Private/ShaderCompiler/ShaderCompiler.cpp
+++ b/Engine/Source/Runtime/Engine/Private/ShaderCompiler/ShaderCompiler.cpp
@@ -4114,6 +4114,7 @@ void FShaderCompilingManager::BlockOnAllShaderMapCompletion(TMap<int32, FShaderM
 		}
 
 		FScopedSlowTask SlowTask(NumJobs, FText::Format(LOCTEXT("BlockOnAllShaderMapCompletion", "Compiling Shaders ({0})"), NumJobs), GIsEditor && !IsRunningCommandlet());
+		SlowTask.MakeDialog();
 
 		int32 NumPendingJobs = 0;
 		// Keep track of previous number of pending jobs so we can update the slow task with the amount of work done.
