Fix non-rewind case

diff --git a/Engine/Source/Runtime/Engine/Private/PhysicsEngine/PhysicsReplication.cpp b/Engine/Source/Runtime/Engine/Private/PhysicsEngine/PhysicsReplication.cpp
--- a/Engine/Source/Runtime/Engine/Private/PhysicsEngine/PhysicsReplication.cpp
+++ b/Engine/Source/Runtime/Engine/Private/PhysicsEngine/PhysicsReplication.cpp
@@ -161,7 +161,7 @@ bool FPhysicsReplication::ApplyRigidBodyState(float DeltaSeconds, FBodyInstance*
 	Chaos::FPhysicsSolver* Solver = PhysScene->GetSolver();
 	Chaos::FRewindData* RewindData = Solver->GetRewindData();
 
-	if (LocalFrame > RewindData->GetEarliestFrame_Internal() && LocalFrame < RewindData->CurrentFrame())
+	if (RewindData && LocalFrame > RewindData->GetEarliestFrame_Internal() && LocalFrame < RewindData->CurrentFrame())
 	{
 		auto Proxy = BI->GetPhysicsActorHandle();
 		const auto P = RewindData->GetPastStateAtFrame(*Proxy->GetHandle_LowLevel(), LocalFrame);
@@ -184,7 +184,7 @@ bool FPhysicsReplication::ApplyRigidBodyState(float DeltaSeconds, FBodyInstance*
 	// Call into the old ApplyRigidBodyState function for now,
 	// "new leash mode" should replace this.
 	// Note that old ApplyRigidBodyState is overridden in other projects, so consider backwards compat path
-	float PingSecondsOneWay = (Solver->GetLastDt() * NumPredictedFrames) * 0.5f;
+	float PingSecondsOneWay = RewindData == nullptr ? InPingSecondsOneWay : (Solver->GetLastDt() * NumPredictedFrames) * 0.5f;
 	return ApplyRigidBodyState(DeltaSeconds, BI, PhysicsTarget, ErrorCorrection, PingSecondsOneWay);
 }
 
