Make FAudioDeviceRenderInfo argument a const reference.

diff --git a/Engine/Source/Runtime/Engine/Public/AudioDevice.h b/Engine/Source/Runtime/Engine/Public/AudioDevice.h
--- a/Engine/Source/Runtime/Engine/Public/AudioDevice.h
+++ b/Engine/Source/Runtime/Engine/Public/AudioDevice.h
@@ -420,8 +420,8 @@ struct FAudioDeviceRenderInfo
 	int32 NumFrames = 0;
 };
 
-DECLARE_MULTICAST_DELEGATE_OneParam(FOnAudioDevicePreRender, FAudioDeviceRenderInfo);
-DECLARE_MULTICAST_DELEGATE_OneParam(FOnAudioDevicePostRender, FAudioDeviceRenderInfo);
+DECLARE_MULTICAST_DELEGATE_OneParam(FOnAudioDevicePreRender, const FAudioDeviceRenderInfo&);
+DECLARE_MULTICAST_DELEGATE_OneParam(FOnAudioDevicePostRender, const FAudioDeviceRenderInfo&);
 
 class ENGINE_API FAudioDevice : public FExec
 {
