Adding channel description to partial bunch max size errors

diff --git a/Engine/Source/Runtime/Engine/Private/DataChannel.cpp b/Engine/Source/Runtime/Engine/Private/DataChannel.cpp
--- a/Engine/Source/Runtime/Engine/Private/DataChannel.cpp
+++ b/Engine/Source/Runtime/Engine/Private/DataChannel.cpp
@@ -838,7 +838,7 @@ bool UChannel::ReceivedNextBunch( FInBunch & Bunch, bool & bOutSkipAck )
 
 		if (InPartialBunch && IsBunchTooLarge(Connection, InPartialBunch))
 		{
-			UE_LOG(LogNetPartialBunch, Error, TEXT("Received a partial bunch exceeding max allowed size. BunchSize=%d, MaximumSize=%d"), InPartialBunch->GetNumBytes(), NetMaxConstructedPartialBunchSizeBytes);
+			UE_LOG(LogNetPartialBunch, Error, TEXT("Received a partial bunch exceeding max allowed size. BunchSize=%d, MaximumSize=%d Channel: %s"), InPartialBunch->GetNumBytes(), NetMaxConstructedPartialBunchSizeBytes, *Describe());
 
 			Bunch.SetError();
 			AddToChainResultPtr(Bunch.ExtendedError, ENetCloseResult::PartialTooLarge);
@@ -1054,7 +1054,7 @@ FPacketIdRange UChannel::SendBunch( FOutBunch* Bunch, bool Merge )
 
 	if (IsBunchTooLarge(Connection, Bunch))
 	{
-		UE_LOG(LogNetPartialBunch, Error, TEXT("Attempted to send bunch exceeding max allowed size. BunchSize=%d, MaximumSize=%d"), Bunch->GetNumBytes(), NetMaxConstructedPartialBunchSizeBytes);
+		UE_LOG(LogNetPartialBunch, Error, TEXT("Attempted to send bunch exceeding max allowed size. BunchSize=%d, MaximumSize=%d Channel: %s"), Bunch->GetNumBytes(), NetMaxConstructedPartialBunchSizeBytes, *Describe());
 		Bunch->SetError();
 		return FPacketIdRange(INDEX_NONE);
 	}
