New HLSL translator supports custom outputs

diff --git a/Engine/Shaders/Private/GeometryCollectionVertexFactory.ush b/Engine/Shaders/Private/GeometryCollectionVertexFactory.ush
--- a/Engine/Shaders/Private/GeometryCollectionVertexFactory.ush
+++ b/Engine/Shaders/Private/GeometryCollectionVertexFactory.ush
@@ -269,7 +269,7 @@ FMaterialVertexParameters GetMaterialVertexParameters(FVertexFactoryInput Input,
 #endif
 
 #if ENABLE_NEW_HLSL_GENERATOR
-	Result.MaterialVertexAttributes = EvaluateVertexMaterialAttributes(Result);
+	Result.MaterialAttributes = EvaluateVertexMaterialAttributes(Result);
 #endif
 	return Result;
 }
diff --git a/Engine/Shaders/Private/LocalVertexFactory.ush b/Engine/Shaders/Private/LocalVertexFactory.ush
--- a/Engine/Shaders/Private/LocalVertexFactory.ush
+++ b/Engine/Shaders/Private/LocalVertexFactory.ush
@@ -661,7 +661,7 @@ FMaterialVertexParameters GetMaterialVertexParameters(FVertexFactoryInput Input,
 #endif
 
 #if ENABLE_NEW_HLSL_GENERATOR
-	Result.MaterialVertexAttributes = EvaluateVertexMaterialAttributes(Result);
+	Result.MaterialAttributes = EvaluateVertexMaterialAttributes(Result);
 #endif
 	return Result;
 }
diff --git a/Engine/Shaders/Private/MaterialTemplate.ush b/Engine/Shaders/Private/MaterialTemplate.ush
--- a/Engine/Shaders/Private/MaterialTemplate.ush
+++ b/Engine/Shaders/Private/MaterialTemplate.ush
@@ -450,6 +450,8 @@ struct FMaterialPixelParameters
 	FStrataTree StrataTree;
 	FStrataPixelFootprint StrataPixelFootprint;
 #endif
+
+	FMaterialAttributes MaterialAttributes;
 };
 
 /**
@@ -527,7 +529,7 @@ struct FMaterialVertexParameters
 	uint WaterWaveParamIndex;
 #endif
 
-	FMaterialAttributes MaterialVertexAttributes;
+	FMaterialAttributes MaterialAttributes;
 
 	/** Cached primitive and instance data */
 	FSceneDataIntermediates SceneData;
diff --git a/Engine/Source/Runtime/Engine/Classes/Materials/MaterialExpressionClearCoatNormalCustomOutput.h b/Engine/Source/Runtime/Engine/Classes/Materials/MaterialExpressionClearCoatNormalCustomOutput.h
--- a/Engine/Source/Runtime/Engine/Classes/Materials/MaterialExpressionClearCoatNormalCustomOutput.h
+++ b/Engine/Source/Runtime/Engine/Classes/Materials/MaterialExpressionClearCoatNormalCustomOutput.h
@@ -21,9 +21,9 @@ class UMaterialExpressionClearCoatNormalCustomOutput : public UMaterialExpressio
 	virtual int32 Compile(class FMaterialCompiler* Compiler, int32 OutputIndex) override;
 	virtual void GetCaption(TArray<FString>& OutCaptions) const override;
 	virtual FExpressionInput* GetInput(int32 InputIndex) override;
-
-	// Begin UObject Interface
 	virtual uint32 GetInputType(int32 InputIndex) override { return MCT_Float3; }
+	virtual UE::Shader::EValueType GetCustomOutputType(int32 OutputIndex) const override { return UE::Shader::EValueType::Float3; }
+	virtual bool GenerateHLSLExpression(FMaterialHLSLGenerator& Generator, UE::HLSLTree::FScope& Scope, int32 OutputIndex, UE::HLSLTree::FExpression*& OutExpression) override;
 #endif
 
 
diff --git a/Engine/Source/Runtime/Engine/Classes/Materials/MaterialExpressionCustomOutput.h b/Engine/Source/Runtime/Engine/Classes/Materials/MaterialExpressionCustomOutput.h
--- a/Engine/Source/Runtime/Engine/Classes/Materials/MaterialExpressionCustomOutput.h
+++ b/Engine/Source/Runtime/Engine/Classes/Materials/MaterialExpressionCustomOutput.h
@@ -25,6 +25,10 @@ class UMaterialExpressionCustomOutput : public UMaterialExpression
 	virtual bool NeedsCustomOutputDefines() { return true; }
 	virtual bool ShouldCompileBeforeAttributes() { return false; }
 	virtual EShaderFrequency GetShaderFrequency() { return SF_Pixel; }
+
+	// This is currently only needed by the new HLSL translator
+	// Should probably be merged with UMaterialExpression base interface, which already provides GetInputType/GetOutputType
+	virtual UE::Shader::EValueType GetCustomOutputType(int32 OutputIndex) const { return UE::Shader::EValueType::Void; }
 #endif
 };
 
diff --git a/Engine/Source/Runtime/Engine/Private/Materials/HLSLMaterialTranslator.cpp b/Engine/Source/Runtime/Engine/Private/Materials/HLSLMaterialTranslator.cpp
--- a/Engine/Source/Runtime/Engine/Private/Materials/HLSLMaterialTranslator.cpp
+++ b/Engine/Source/Runtime/Engine/Private/Materials/HLSLMaterialTranslator.cpp
@@ -2320,7 +2320,7 @@ FString FHLSLMaterialTranslator::GetMaterialShaderCode()
 
 	LazyPrintf.PushParam(*FString::Printf(TEXT("return %.5f"), Material->GetOpacityMaskClipValue()));
 
-	LazyPrintf.PushParam(!bEnableExecutionFlow ? *GenerateFunctionCode(MP_WorldPositionOffset, BaseDerivativeVariation) : TEXT("return Parameters.MaterialVertexAttributes.WorldPositionOffset"));
+	LazyPrintf.PushParam(!bEnableExecutionFlow ? *GenerateFunctionCode(MP_WorldPositionOffset, BaseDerivativeVariation) : TEXT("return Parameters.MaterialAttributes.WorldPositionOffset"));
 	LazyPrintf.PushParam(!bEnableExecutionFlow ? *GenerateFunctionCode(CompiledMP_PrevWorldPositionOffset, BaseDerivativeVariation) : TEXT("return 0.0f"));
 	LazyPrintf.PushParam(!bEnableExecutionFlow ? *GenerateFunctionCode(MP_CustomData0, BaseDerivativeVariation) : TEXT("return 0.0f"));
 	LazyPrintf.PushParam(!bEnableExecutionFlow ? *GenerateFunctionCode(MP_CustomData1, BaseDerivativeVariation) : TEXT("return 0.0f"));
@@ -2334,7 +2334,7 @@ FString FHLSLMaterialTranslator::GetMaterialShaderCode()
 		if (bEnableExecutionFlow)
 		{
 			const FString AttributeName = FMaterialAttributeDefinitionMap::GetAttributeName((EMaterialProperty)(MP_CustomizedUVs0 + CustomUVIndex));
-			CustomUVAssignments += FString::Printf(TEXT("\tOutTexCoords[%u] = Parameters.MaterialVertexAttributes.%s;") LINE_TERMINATOR, CustomUVIndex, *AttributeName);
+			CustomUVAssignments += FString::Printf(TEXT("\tOutTexCoords[%u] = Parameters.MaterialAttributes.%s;") LINE_TERMINATOR, CustomUVIndex, *AttributeName);
 		}
 		else
 		{
diff --git a/Engine/Source/Runtime/Engine/Private/Materials/MaterialExpressionHLSL.cpp b/Engine/Source/Runtime/Engine/Private/Materials/MaterialExpressionHLSL.cpp
--- a/Engine/Source/Runtime/Engine/Private/Materials/MaterialExpressionHLSL.cpp
+++ b/Engine/Source/Runtime/Engine/Private/Materials/MaterialExpressionHLSL.cpp
@@ -77,6 +77,7 @@
 #include "Materials/MaterialExpressionTransformPosition.h"
 #include "Materials/MaterialExpressionIf.h"
 #include "Materials/MaterialExpressionCustom.h"
+#include "Materials/MaterialExpressionClearCoatNormalCustomOutput.h"
 #include "Materials/MaterialExpressionSetLocal.h"
 #include "Materials/MaterialExpressionIfThenElse.h"
 #include "Materials/MaterialExpressionForLoop.h"
@@ -1251,6 +1252,12 @@ bool UMaterialExpressionCustom::GenerateHLSLExpression(FMaterialHLSLGenerator& G
 	return true;
 }
 
+bool UMaterialExpressionClearCoatNormalCustomOutput::GenerateHLSLExpression(FMaterialHLSLGenerator& Generator, UE::HLSLTree::FScope& Scope, int32 OutputIndex, UE::HLSLTree::FExpression*& OutExpression)
+{
+	OutExpression = Input.AcquireHLSLExpression(Generator, Scope);
+	return OutExpression != nullptr;
+}
+
 bool UMaterialExpressionExecBegin::GenerateHLSLStatements(FMaterialHLSLGenerator& Generator, UE::HLSLTree::FScope& Scope)
 {
 	Exec.GenerateHLSLStatements(Generator, Scope);
diff --git a/Engine/Source/Runtime/Engine/Private/Materials/MaterialHLSLEmitter.cpp b/Engine/Source/Runtime/Engine/Private/Materials/MaterialHLSLEmitter.cpp
--- a/Engine/Source/Runtime/Engine/Private/Materials/MaterialHLSLEmitter.cpp
+++ b/Engine/Source/Runtime/Engine/Private/Materials/MaterialHLSLEmitter.cpp
@@ -178,12 +178,7 @@ static FString GenerateMaterialTemplateHLSL(EShaderPlatform ShaderPlatform,
 
 	LazyPrintf.PushParam(*PixelMembersDeclaration);
 
-	{
-		FString DerivativeHelpers;// = DerivativeAutogen.GenerateUsedFunctions(*this);
-		FString DerivativeHelpersAndResources;// = DerivativeHelpers + ResourcesString;
-		//LazyPrintf.PushParam(*ResourcesString);
-		LazyPrintf.PushParam(*DerivativeHelpersAndResources);
-	}
+	LazyPrintf.PushParam(SharedShaderCode);
 
 	// Anything used bye the GenerationFunctionCode() like WorldPositionOffset shouldn't be using texures, right?
 	// Let those use the standard finite differences textures, since they should be the same. If we actually want
@@ -215,10 +210,11 @@ static FString GenerateMaterialTemplateHLSL(EShaderPlatform ShaderPlatform,
 
 	LazyPrintf.PushParam(*FString::Printf(TEXT("return %.5f"), Material.GetOpacityMaskClipValue()));
 
-	LazyPrintf.PushParam(TEXT("return Parameters.MaterialVertexAttributes.WorldPositionOffset"));
-	LazyPrintf.PushParam(TEXT("return Parameters.MaterialVertexAttributes.PrevWorldPositionOffset"));
-	LazyPrintf.PushParam(TEXT("return 0.0f"));
-	LazyPrintf.PushParam(TEXT("return 0.0f"));
+	LazyPrintf.PushParam(TEXT("return Parameters.MaterialAttributes.WorldPositionOffset"));
+	LazyPrintf.PushParam(TEXT("return Parameters.MaterialAttributes.PrevWorldPositionOffset"));
+	// CustomData0/1 are named ClearCoat/ClearCoatRoughness
+	LazyPrintf.PushParam(TEXT("return Parameters.MaterialAttributes.ClearCoat"));
+	LazyPrintf.PushParam(TEXT("return Parameters.MaterialAttributes.ClearCoatRoughness"));
 
 	// Print custom texture coordinate assignments, should be fine with regular derivatives
 	FString CustomUVAssignments;
@@ -229,7 +225,7 @@ static FString GenerateMaterialTemplateHLSL(EShaderPlatform ShaderPlatform,
 		//if (bEnableExecutionFlow)
 		{
 			const FString AttributeName = FMaterialAttributeDefinitionMap::GetAttributeName((EMaterialProperty)(MP_CustomizedUVs0 + CustomUVIndex));
-			CustomUVAssignments += FString::Printf(TEXT("\tOutTexCoords[%u] = Parameters.MaterialVertexAttributes.%s;") LINE_TERMINATOR, CustomUVIndex, *AttributeName);
+			CustomUVAssignments += FString::Printf(TEXT("\tOutTexCoords[%u] = Parameters.MaterialAttributes.%s;") LINE_TERMINATOR, CustomUVIndex, *AttributeName);
 		}
 		/*else
 		{
@@ -293,7 +289,7 @@ static FString GenerateMaterialTemplateHLSL(EShaderPlatform ShaderPlatform,
 
 	//if (bEnableExecutionFlow)
 	{
-		FString EvaluateMaterialDeclaration = SharedShaderCode;
+		FString EvaluateMaterialDeclaration;
 
 		EvaluateMaterialDeclaration += TEXT("FMaterialAttributes EvaluateVertexMaterialAttributesInternal(FMaterialVertexParameters Parameters)" LINE_TERMINATOR);
 		EvaluateMaterialDeclaration += TEXT("{" LINE_TERMINATOR);
@@ -337,7 +333,12 @@ static FString GenerateMaterialTemplateHLSL(EShaderPlatform ShaderPlatform,
 			EvaluateMaterialDeclaration += PixelShaderCodePhase1;
 			EvaluateMaterialDeclaration += TEXT("}" LINE_TERMINATOR);
 
-			EvaluateMaterialAttributesPhase1 = TEXT("    FMaterialAttributes MaterialAttributesPhase1 = EvaluatePixelMaterialAttributesPhase1(Parameters);" LINE_TERMINATOR);
+			EvaluateMaterialAttributesPhase1  = TEXT("    FMaterialAttributes MaterialAttributesPhase1 = EvaluatePixelMaterialAttributesPhase1(Parameters);" LINE_TERMINATOR);
+			EvaluateMaterialAttributesPhase1 += TEXT("    Parameters.MaterialAttributes = MaterialAttributesPhase1;" LINE_TERMINATOR);
+		}
+		else
+		{
+			EvaluateMaterialAttributesPhase0 += TEXT("    Parameters.MaterialAttributes = MaterialAttributesPhase0;" LINE_TERMINATOR);
 		}
 
 		LazyPrintf.PushParam(*EvaluateMaterialDeclaration);
@@ -817,6 +818,7 @@ bool MaterialEmitHLSL(const FMaterialCompileTargetParameters& InCompilerTarget,
 
 		// Prepare all fields *except* normal
 		FRequestedType RequestedPixelAttributesType;
+		Generator.SetRequestedFields(SF_Pixel, RequestedPixelAttributesType);
 		for (const FGuid& AttributeID : OrderedVisibleAttributes)
 		{
 			if (FMaterialAttributeDefinitionMap::GetShaderFrequency(AttributeID) == SF_Pixel)
@@ -878,6 +880,7 @@ bool MaterialEmitHLSL(const FMaterialCompileTargetParameters& InCompilerTarget,
 	FStringBuilderMemstack VertexCode(Allocator, 1024 * 1024);
 	{
 		FRequestedType RequestedVertexAttributesType;
+		Generator.SetRequestedFields(SF_Vertex, RequestedVertexAttributesType);
 		for (const FGuid& AttributeID : OrderedVisibleAttributes)
 		{
 			if (FMaterialAttributeDefinitionMap::GetShaderFrequency(AttributeID) == SF_Vertex)
@@ -921,8 +924,9 @@ bool MaterialEmitHLSL(const FMaterialCompileTargetParameters& InCompilerTarget,
 	TypeRegistry.EmitDeclarationsCode(Declarations);
 
 	FStringBuilderMemstack SharedCode(Allocator, 64 * 1024);
+	Generator.EmitSharedCode(SharedCode);
 	EmitContext.EmitDeclarationsCode(SharedCode);
-
+	
 	FString MaterialTemplateSource = GenerateMaterialTemplateHLSL(InCompilerTarget.ShaderPlatform,
 		InOutMaterial,
 		EmitContext,
diff --git a/Engine/Source/Runtime/Engine/Private/Materials/MaterialHLSLGenerator.cpp b/Engine/Source/Runtime/Engine/Private/Materials/MaterialHLSLGenerator.cpp
--- a/Engine/Source/Runtime/Engine/Private/Materials/MaterialHLSLGenerator.cpp
+++ b/Engine/Source/Runtime/Engine/Private/Materials/MaterialHLSLGenerator.cpp
@@ -8,6 +8,7 @@
 #include "Materials/MaterialExpressionMaterialFunctionCall.h"
 #include "Materials/MaterialExpressionFunctionInput.h"
 #include "Materials/MaterialExpressionFunctionOutput.h"
+#include "Materials/MaterialExpressionCustomOutput.h"
 #include "Materials/MaterialExpressionVolumetricAdvancedMaterialOutput.h"
 #include "Materials/MaterialExpressionExecBegin.h"
 #include "Materials/Material.h"
@@ -96,6 +97,8 @@ FMaterialHLSLGenerator::FMaterialHLSLGenerator(const FMaterialCompileTargetParam
 	FFunctionCallEntry* RootFunctionEntry = new(InOutTree.GetAllocator()) FFunctionCallEntry();
 	FunctionCallStack.Add(RootFunctionEntry);
 
+	InOutMaterial.GatherCustomOutputExpressions(MaterialCustomOutputs);
+
 	TArray<UE::Shader::FStructFieldInitializer, TInlineAllocator<MP_MAX + 16>> MaterialAttributeFields;
 
 	const TArray<FGuid>& OrderedVisibleAttributes = FMaterialAttributeDefinitionMap::GetOrderedVisibleAttributeList();
@@ -123,6 +126,27 @@ FMaterialHLSLGenerator::FMaterialHLSLGenerator(const FMaterialCompileTargetParam
 		}
 	}
 
+	TArray<TStringBuilder<256>, TInlineAllocator<4>> CustomOutputNames;
+	CustomOutputNames.Reserve(MaterialCustomOutputs.Num());
+	for (UMaterialExpressionCustomOutput* CustomOutput : MaterialCustomOutputs)
+	{
+		const int32 NumOutputs = CustomOutput->GetNumOutputs();
+		const FString OutputName = CustomOutput->GetFunctionName();
+
+		check(!CustomOutput->ShouldCompileBeforeAttributes()); // not supported yet, looks like this isn't currently being used
+
+		for (int32 OutputIndex = 0; OutputIndex < NumOutputs; ++OutputIndex)
+		{
+			const UE::Shader::EValueType ValueType = CustomOutput->GetCustomOutputType(OutputIndex);
+			FStringBuilderBase& FormattedName = CustomOutputNames.AddDefaulted_GetRef();
+			FormattedName.Appendf(TEXT("%s%d"), *OutputName, OutputIndex);
+			MaterialAttributeFields.Emplace(FormattedName.ToView(), ValueType);
+
+			const UE::Shader::FValue DefaultValue(ValueType);
+			MaterialAttributesDefaultValue.Component.Append(DefaultValue.Component);
+		}
+	}
+
 	MaterialAttributeFields.Emplace(TEXT("PrevWorldPositionOffset"), UE::Shader::EValueType::Float3);
 	MaterialAttributesDefaultValue.Component.Append({ 0.0f, 0.0f, 0.0f });
 
@@ -193,6 +217,60 @@ bool FMaterialHLSLGenerator::Generate()
 	return HLSLTree->Finalize();
 }
 
+void FMaterialHLSLGenerator::SetRequestedFields(EShaderFrequency ShaderFrequency, UE::HLSLTree::FRequestedType& OutRequestedType)
+{
+	for (UMaterialExpressionCustomOutput* CustomOutput : MaterialCustomOutputs)
+	{
+		if (CustomOutput->GetShaderFrequency() != ShaderFrequency)
+		{
+			continue;
+		}
+
+		const int32 NumOutputs = CustomOutput->GetNumOutputs();
+		const FString OutputName = CustomOutput->GetFunctionName();
+
+		for (int32 OutputIndex = 0; OutputIndex < NumOutputs; ++OutputIndex)
+		{
+			TStringBuilder<256> FieldName;
+			FieldName.Appendf(TEXT("%s%d"), *OutputName, OutputIndex);
+			const UE::Shader::FStructField* CustomOutputField = GetMaterialAttributesType()->FindFieldByName(FieldName.ToString());
+			check(CustomOutputField);
+			OutRequestedType.SetFieldRequested(CustomOutputField);
+		}
+	}
+}
+
+void FMaterialHLSLGenerator::EmitSharedCode(FStringBuilderBase& OutCode) const
+{
+	using namespace UE::Shader;
+	for (UMaterialExpressionCustomOutput* CustomOutput : MaterialCustomOutputs)
+	{
+		const int32 NumOutputs = CustomOutput->GetNumOutputs();
+		const FString OutputName = CustomOutput->GetFunctionName();
+		const EShaderFrequency ShaderFrequency = CustomOutput->GetShaderFrequency();
+
+		if (CustomOutput->NeedsCustomOutputDefines())
+		{
+			OutCode.Appendf(TEXT("#define NUM_MATERIAL_OUTPUTS_%s %d\n"), *OutputName.ToUpper(), NumOutputs);
+		}
+
+		for (int32 OutputIndex = 0; OutputIndex < NumOutputs; ++OutputIndex)
+		{
+			const EValueType ValueType = CustomOutput->GetCustomOutputType(OutputIndex);
+			const FValueTypeDescription ValueTypeDesc = GetValueTypeDescription(ValueType);
+
+			OutCode.Appendf(TEXT("#define HAVE_%s%d 1\n"), *OutputName, OutputIndex);
+
+			OutCode.Appendf(TEXT("%s %s%d(FMaterial%sParameters Parameters) { return Parameters.MaterialAttributes.%s%d; }\n"),
+				ValueTypeDesc.Name,
+				*OutputName, OutputIndex,
+				ShaderFrequency == SF_Pixel ? TEXT("Pixel") : TEXT("Vertex"),
+				*OutputName, OutputIndex);
+		}
+		OutCode.Append(TEXT("\n"));
+	}
+}
+
 static UE::HLSLTree::FExpression* CompileMaterialInput(FMaterialHLSLGenerator& Generator,
 	UE::HLSLTree::FScope& Scope,
 	EMaterialProperty InputProperty,
@@ -301,6 +379,22 @@ bool FMaterialHLSLGenerator::GenerateResult(UE::HLSLTree::FScope& Scope)
 				}
 			}
 
+			for (UMaterialExpressionCustomOutput* CustomOutput : MaterialCustomOutputs)
+			{
+				const int32 NumOutputs = CustomOutput->GetNumOutputs();
+				const FString OutputName = CustomOutput->GetFunctionName();
+				for (int32 OutputIndex = 0; OutputIndex < NumOutputs; ++OutputIndex)
+				{
+					TStringBuilder<256> FieldName;
+					FieldName.Appendf(TEXT("%s%d"), *OutputName, OutputIndex);
+					const FStructField* CustomOutputField = GetMaterialAttributesType()->FindFieldByName(FieldName.ToString());
+					check(CustomOutputField);
+
+					FExpression* CustomOutputExpression = AcquireExpression(Scope, CustomOutput, OutputIndex);
+					AttributesExpression = HLSLTree->NewExpression<FExpressionSetStructField>(GetMaterialAttributesType(), CustomOutputField, AttributesExpression, CustomOutputExpression);
+				}
+			}
+
 			if (AttributesExpression)
 			{
 				FStatementReturn* ReturnStatement = HLSLTree->NewStatement<FStatementReturn>(Scope);
diff --git a/Engine/Source/Runtime/Engine/Private/Shader/Preshader.cpp b/Engine/Source/Runtime/Engine/Private/Shader/Preshader.cpp
--- a/Engine/Source/Runtime/Engine/Private/Shader/Preshader.cpp
+++ b/Engine/Source/Runtime/Engine/Private/Shader/Preshader.cpp
@@ -323,10 +323,15 @@ static void EvaluateSetField(FPreshaderStack& Stack, FPreshaderDataContext& REST
 	}
 	else
 	{
-		for (int32 Index = 0; Index < ComponentNum; ++Index)
+		const int32 NumComponentsToCopy = FMath::Min(ComponentNum, Value.Component.Num());
+		for (int32 Index = 0; Index < NumComponentsToCopy; ++Index)
 		{
 			StructValue.Component[ComponentIndex + Index] = Value.Component[Index];
 		}
+		for (int32 Index = NumComponentsToCopy; Index < ComponentNum; ++Index)
+		{
+			StructValue.Component[ComponentIndex + Index] = FValueComponent();
+		}
 	}
 }
 
diff --git a/Engine/Source/Runtime/Engine/Public/MaterialHLSLGenerator.h b/Engine/Source/Runtime/Engine/Public/MaterialHLSLGenerator.h
--- a/Engine/Source/Runtime/Engine/Public/MaterialHLSLGenerator.h
+++ b/Engine/Source/Runtime/Engine/Public/MaterialHLSLGenerator.h
@@ -21,6 +21,7 @@ class UMaterialFunctionInterface;
 class UMaterialExpression;
 class UMaterialExpressionFunctionInput;
 class UMaterialExpressionFunctionOutput;
+class UMaterialExpressionCustomOutput;
 struct FFunctionExpressionInput;
 class ITargetPlatform;
 
@@ -82,6 +83,10 @@ public:
 	UE::HLSLTree::FExpression* GetResultExpression() { return ResultExpression; }
 	UE::HLSLTree::FStatement* GetResultStatement() { return ResultStatement; }
 
+	void SetRequestedFields(EShaderFrequency ShaderFrequency, UE::HLSLTree::FRequestedType& OutRequestedType);
+
+	void EmitSharedCode(FStringBuilderBase& OutCode) const;
+
 	bool GenerateResult(UE::HLSLTree::FScope& Scope);
 
 	UE::HLSLTree::FScope* NewScope(UE::HLSLTree::FScope& Scope, EMaterialNewScopeFlag Flags = EMaterialNewScopeFlag::None);
@@ -184,6 +189,7 @@ private:
 	UMaterial* TargetMaterial = nullptr;
 	FMaterialHLSLErrorHandler Errors;
 
+	TArray<UMaterialExpressionCustomOutput*> MaterialCustomOutputs;
 	const UE::Shader::FStructType* MaterialAttributesType = nullptr;
 	UE::Shader::FValue MaterialAttributesDefaultValue;
 
