Moved custom fast array delta state update flags into their own bit array
Only allocate custom fast array delta history items when required
Lower the number of history items from 64 to 32

diff --git a/Engine/Source/Runtime/Engine/Private/RepLayout.cpp b/Engine/Source/Runtime/Engine/Private/RepLayout.cpp
--- a/Engine/Source/Runtime/Engine/Private/RepLayout.cpp
+++ b/Engine/Source/Runtime/Engine/Private/RepLayout.cpp
@@ -450,12 +450,9 @@ struct FDeltaArrayHistoryItem
 	/** The set of changelists by element ID.*/
 	TMap<int32, TArray<uint16>> ChangelistByID;
 
-	bool bWasUpdated = false;
-
 	void Reset()
 	{
 		ChangelistByID.Empty();
-		bWasUpdated = false;
 	}
 
 	void CountBytes(FArchive& Ar) const
@@ -475,42 +472,49 @@ struct FDeltaArrayHistoryItem
 struct FDeltaArrayHistoryState
 {
 	/** The maximum number of individual changelists allowed.*/
-	static const uint16 MAX_CHANGE_HISTORY = FRepChangelistState::MAX_CHANGE_HISTORY;
-
-	//~ TODO: Investigate either making this a Dynamically sized container,
-	//~			or potentially changing ArrayStates to unique pointers / 
-	//~			an indirect array.
-	//~
-	//~			Right now, if a Delta Struct Serialization is disabled for
-	//~			a given Fast Array, this will still use up memory.
+	static const uint16 MAX_CHANGE_HISTORY = 32;
 
 	/** Circular buffer of changelists. */
-	FDeltaArrayHistoryItem ChangeHistory[MAX_CHANGE_HISTORY];
-
-	/** The latest ArrayReplicationKey sent to any connection. */
-	int32 ArrayReplicationKey = INDEX_NONE;
+	TArray<FDeltaArrayHistoryItem> ChangeHistory;
+	TBitArray<> ChangeHistoryUpdated;
 
 	/** Copy of the IDToIndexMap from the array when we last sent it. */
 	TMap<int32, int32> IDToIndexMap;
 
+	/** The latest ArrayReplicationKey sent to any connection. */
+	int32 ArrayReplicationKey = INDEX_NONE;
+
 	/** Index in the buffer where changelist history starts (i.e., the Oldest changelist). */
 	uint16 HistoryStart = 0;
 
 	/** Index in the buffer where changelist history ends (i.e., the Newest changelist). */
 	uint16 HistoryEnd = 0;
 
+	void InitHistory()
+	{
+		if (ChangeHistory.Num() == 0)
+		{
+			ChangeHistory.AddDefaulted(MAX_CHANGE_HISTORY);
+			ChangeHistoryUpdated.Init(false, MAX_CHANGE_HISTORY);
+		}
+	}
+
 	void CountBytes(FArchive& Ar) const
 	{
 		GRANULAR_NETWORK_MEMORY_TRACKING_INIT(Ar, "FDeltaArrayHistoryState::CountBytes");
 
 		GRANULAR_NETWORK_MEMORY_TRACKING_TRACK("IDToIndexMap", IDToIndexMap.CountBytes(Ar));
 
-		GRANULAR_NETWORK_MEMORY_TRACKING_TRACK("ChangeHistory",
+		GRANULAR_NETWORK_MEMORY_TRACKING_TRACK("ChangeHistory", ChangeHistory.CountBytes(Ar));
+
+		GRANULAR_NETWORK_MEMORY_TRACKING_TRACK("ChangeHistoryElements",
 			for (const FDeltaArrayHistoryItem& HistoryItem : ChangeHistory)
 			{
 				HistoryItem.CountBytes(Ar);
 			}
 		);
+
+		GRANULAR_NETWORK_MEMORY_TRACKING_TRACK("ChangeHistoryUpdated", ChangeHistoryUpdated.CountBytes(Ar));
 	}
 };
 
@@ -7323,15 +7327,18 @@ void FRepLayout::PreSendCustomDeltaProperties(
 							const int32 FastArrayReplicationKey = CustomDeltaProperty.GetFastArrayArrayReplicationKey(FastArraySerializer);
 							if (FastArrayHistoryState.ArrayReplicationKey != FastArrayReplicationKey)
 							{
+								FastArrayHistoryState.InitHistory();
+
 								const uint16 HistoryDelta = FastArrayHistoryState.HistoryEnd - FastArrayHistoryState.HistoryStart;
 								const uint16 CurrentHistoryIndex = FastArrayHistoryState.HistoryEnd % FDeltaArrayHistoryState::MAX_CHANGE_HISTORY;
 								const FDeltaArrayHistoryItem& CurrentHistory = FastArrayHistoryState.ChangeHistory[CurrentHistoryIndex];
+								const bool bCurrentHistoryUpdated = FastArrayHistoryState.ChangeHistoryUpdated[CurrentHistoryIndex];
 
 								// If we don't have any history items, go ahead and create one.
 								// Otherwise, check to see if our current history was actually updated.
 								// If it wasn't updated, that means that no one tried to replicate it last frame (which can be possible due
 								// to rep conditions), and there's no sense in creating a new one.
-								if (HistoryDelta == 0 || CurrentHistory.bWasUpdated)
+								if (HistoryDelta == 0 || bCurrentHistoryUpdated)
 								{
 									// If we've reached our buffer size, then move our start history marker up.
 									// In that case the old start history will become our new history.
@@ -7343,6 +7350,7 @@ void FRepLayout::PreSendCustomDeltaProperties(
 									++FastArrayHistoryState.HistoryEnd;
 									const uint16 NewHistory = FastArrayHistoryState.HistoryEnd % FDeltaArrayHistoryState::MAX_CHANGE_HISTORY;
 									FastArrayHistoryState.ChangeHistory[NewHistory].Reset();
+									FastArrayHistoryState.ChangeHistoryUpdated[NewHistory] = false;
 								}
 							}
 						}
@@ -7487,10 +7495,11 @@ ERepLayoutResult FRepLayout::DeltaSerializeFastArrayProperty(FFastArrayDeltaSeri
 					const uint16 RelativeNewHistory = NewChangelistHistory % FDeltaArrayHistoryState::MAX_CHANGE_HISTORY;
 					const uint16 CompareChangelistDelta = NewChangelistHistory - FastArrayState.HistoryStart;
 					FDeltaArrayHistoryItem& HistoryItem = FastArrayState.ChangeHistory[RelativeNewHistory];
+					const bool bHistoryItemUpdated = FastArrayState.ChangeHistoryUpdated[RelativeNewHistory];
 
-					if (!HistoryItem.bWasUpdated)
+					if (!bHistoryItemUpdated)
 					{
-						HistoryItem.bWasUpdated = true;
+						FastArrayState.ChangeHistoryUpdated[RelativeNewHistory] = true;
 
 						FastArrayState.ArrayReplicationKey = NewArrayDeltaState->ArrayReplicationKey;
 
