Locking the ModularFeatureList before accessing it in Get() and IsAvailable()

#rnx

diff --git a/Engine/Source/Runtime/Engine/Private/Animation/AnimRootMotionProvider.cpp b/Engine/Source/Runtime/Engine/Private/Animation/AnimRootMotionProvider.cpp
--- a/Engine/Source/Runtime/Engine/Private/Animation/AnimRootMotionProvider.cpp
+++ b/Engine/Source/Runtime/Engine/Private/Animation/AnimRootMotionProvider.cpp
@@ -14,14 +14,20 @@ const IAnimRootMotionProvider* IAnimRootMotionProvider::Get()
 {
 	if (IsAvailable())
 	{
-		return &IModularFeatures::Get().GetModularFeature<const IAnimRootMotionProvider>(ModularFeatureName);
+		IModularFeatures::Get().LockModularFeatureList();
+		const IAnimRootMotionProvider* AnimRootMotionProvider = &IModularFeatures::Get().GetModularFeature<const IAnimRootMotionProvider>(ModularFeatureName);
+		IModularFeatures::Get().UnlockModularFeatureList();
+		return AnimRootMotionProvider;
 	}
 	return nullptr;
 }
 
 bool IAnimRootMotionProvider::IsAvailable()
 {
-	return IModularFeatures::Get().IsModularFeatureAvailable(ModularFeatureName);
+	IModularFeatures::Get().LockModularFeatureList();
+	const bool bIsAvailable = IModularFeatures::Get().IsModularFeatureAvailable(ModularFeatureName);
+	IModularFeatures::Get().UnlockModularFeatureList();
+	return bIsAvailable;
 }
 
 }}; // namespace UE::Anim
\ No newline at end of file
