EditorBulkData: Add an Owner argument for UpdatePayload, and supply it from FTextureSource, so we can know ahead of the save which package an in-memory bulk data will be saved under.
#rnx

diff --git a/Engine/Source/Runtime/CoreUObject/Private/Serialization/EditorBulkData.cpp b/Engine/Source/Runtime/CoreUObject/Private/Serialization/EditorBulkData.cpp
--- a/Engine/Source/Runtime/CoreUObject/Private/Serialization/EditorBulkData.cpp
+++ b/Engine/Source/Runtime/CoreUObject/Private/Serialization/EditorBulkData.cpp
@@ -1446,7 +1446,7 @@ void FEditorBulkData::SerializeToPackageTrailer(FLinkerSave& LinkerSave, FCompre
 	LinkerSave.PackageTrailerBuilder->AddPayload(PayloadContentId, MoveTemp(PayloadToSerialize), MoveTemp(OnPayloadWritten));
 }
 
-void FEditorBulkData::UpdatePayloadImpl(FSharedBuffer&& InPayload, FIoHash&& InPayloadID)
+void FEditorBulkData::UpdatePayloadImpl(FSharedBuffer&& InPayload, FIoHash&& InPayloadID, UObject* Owner)
 {
 	TRACE_CPUPROFILER_EVENT_SCOPE(FEditorBulkData::UpdatePayloadImpl);
 
@@ -1487,7 +1487,7 @@ void FEditorBulkData::UpdatePayloadImpl(FSharedBuffer&& InPayload, FIoHash&& InP
 		{
 			BulkDataId = FGuid::NewGuid();
 		}
-		Register(nullptr);
+		Register(Owner);
 	}
 	else
 	{
@@ -1575,11 +1575,11 @@ TFuture<FCompressedBuffer>FEditorBulkData::GetCompressedPayload() const
 	return Promise.GetFuture();
 }
 
-void FEditorBulkData::UpdatePayload(FSharedBuffer InPayload)
+void FEditorBulkData::UpdatePayload(FSharedBuffer InPayload, UObject* Owner)
 {
 	TRACE_CPUPROFILER_EVENT_SCOPE(FEditorBulkData::UpdatePayload);
 	FIoHash NewPayloadId = HashPayload(InPayload);
-	UpdatePayloadImpl(MoveTemp(InPayload), MoveTemp(NewPayloadId));
+	UpdatePayloadImpl(MoveTemp(InPayload), MoveTemp(NewPayloadId), Owner);
 }
 
 FEditorBulkData::FSharedBufferWithID::FSharedBufferWithID(FSharedBuffer InPayload)
@@ -1588,9 +1588,9 @@ FEditorBulkData::FSharedBufferWithID::FSharedBufferWithID(FSharedBuffer InPayloa
 {
 }
 
-void FEditorBulkData::UpdatePayload(FSharedBufferWithID InPayload)
+void FEditorBulkData::UpdatePayload(FSharedBufferWithID InPayload, UObject* Owner)
 {
-	UpdatePayloadImpl(MoveTemp(InPayload.Payload), MoveTemp(InPayload.PayloadId));
+	UpdatePayloadImpl(MoveTemp(InPayload.Payload), MoveTemp(InPayload.PayloadId), Owner);
 }
 
 void FEditorBulkData::SetCompressionOptions(ECompressionOptions Option)
diff --git a/Engine/Source/Runtime/CoreUObject/Public/Serialization/EditorBulkData.h b/Engine/Source/Runtime/CoreUObject/Public/Serialization/EditorBulkData.h
--- a/Engine/Source/Runtime/CoreUObject/Public/Serialization/EditorBulkData.h
+++ b/Engine/Source/Runtime/CoreUObject/Public/Serialization/EditorBulkData.h
@@ -200,8 +200,9 @@ public:
 	 * to the bulkdata object.
 	 *
 	 * @param InPayload	The payload that this bulkdata object should reference. @see FSharedBuffer
+	 * @param Owner The object that owns the bulkdata, or null to not associate with a UObject.
 	 */
-	void UpdatePayload(FSharedBuffer InPayload);
+	void UpdatePayload(FSharedBuffer InPayload, UObject* Owner = nullptr);
 
 	/**
 	 * Utility struct used to compute the Payload ID before calling UpdatePayload
@@ -237,9 +238,10 @@ public:
 	 * 
 	 * Use this override if you want compute PayloadId before updating the bulkdata
 	 *
-	 * @param InPayload				The payload to update the bulkdata with
+	 * @param InPayload	The payload to update the bulkdata with
+	 * @param Owner The object that owns the bulkdata, or null to not associate with a UObject.
 	 */
-	void UpdatePayload(FSharedBufferWithID InPayload);
+	void UpdatePayload(FSharedBufferWithID InPayload, UObject* Owner = nullptr);
 
 	/** 
 	 * Sets the compression options to be applied to the payload during serialization.
@@ -364,7 +366,7 @@ private:
 	/** The new path that saves payloads to the FPackageTrailer which is then appended to the end of the package file */
 	void SerializeToPackageTrailer(FLinkerSave& LinkerSave, FCompressedBuffer PayloadToSerialize, EFlags UpdatedFlags, UObject* Owner);
 
-	void UpdatePayloadImpl(FSharedBuffer&& InPayload, FIoHash&& InPayloadID);
+	void UpdatePayloadImpl(FSharedBuffer&& InPayload, FIoHash&& InPayloadID, UObject* Owner);
 
 	FCompressedBuffer GetDataInternal() const;
 
diff --git a/Engine/Source/Runtime/Engine/Classes/Engine/Texture.h b/Engine/Source/Runtime/Engine/Classes/Engine/Texture.h
--- a/Engine/Source/Runtime/Engine/Classes/Engine/Texture.h
+++ b/Engine/Source/Runtime/Engine/Classes/Engine/Texture.h
@@ -361,6 +361,8 @@ struct FTextureSource
 	FORCEINLINE uint8* LockMip(int32 MipIndex) { return LockMip(0, 0, MipIndex); }
 	FORCEINLINE void UnlockMip(int32 MipIndex) { UnlockMip(0, 0, MipIndex); }
 
+	inline void SetOwner(UTexture* InOwner) { Owner = InOwner; }
+
 	struct FMipAllocation
 	{
 		/** Create an empty object */
@@ -461,6 +463,8 @@ private:
 #if WITH_EDITOR
 	/** Protects simultaneous access to BulkData */
 	TDontCopy<FRWLock> BulkDataLock;
+	/** Owner for associating BulkData with a package */
+	UTexture* Owner;
 #endif
 	/** The bulk source data. */
 	UE::Serialization::FEditorBulkData BulkData;
diff --git a/Engine/Source/Runtime/Engine/Private/Texture.cpp b/Engine/Source/Runtime/Engine/Private/Texture.cpp
--- a/Engine/Source/Runtime/Engine/Private/Texture.cpp
+++ b/Engine/Source/Runtime/Engine/Private/Texture.cpp
@@ -132,6 +132,7 @@ UTexture::UTexture(const FObjectInitializer& ObjectInitializer)
 	CompressionYCoCg = 0;
 	Downscale = 0.f;
 	DownscaleOptions = ETextureDownscaleOptions::Default;
+	Source.SetOwner(this);
 #endif // #if WITH_EDITORONLY_DATA
 
 	if (FApp::CanEverRender() && !IsTemplate())
@@ -1180,7 +1181,11 @@ FStreamableRenderResourceState UTexture::GetResourcePostInitState(const FTexture
 ------------------------------------------------------------------------------*/
 
 FTextureSource::FTextureSource()
-	: NumLockedMips(0u)
+	: 
+#if WITH_EDITOR
+	  Owner(nullptr),
+#endif
+	  NumLockedMips(0u)
 	, LockState(ELockState::None)
 #if WITH_EDITOR
 	, bHasHadBulkDataCleared(false)
@@ -1260,7 +1265,7 @@ void FTextureSource::InitBlocked(const ETextureSourceFormat* InLayerFormats,
 		}
 	}
 
-	BulkData.UpdatePayload(Buffer.MoveToShared());
+	BulkData.UpdatePayload(Buffer.MoveToShared(), Owner);
 	BulkData.SetCompressionOptions(UE::Serialization::ECompressionOptions::Default);
 }
 
@@ -1272,7 +1277,7 @@ void FTextureSource::InitBlocked(const ETextureSourceFormat* InLayerFormats,
 {
 	InitBlockedImpl(InLayerFormats, InBlocks, InNumLayers, InNumBlocks);
 
-	BulkData.UpdatePayload(MoveTemp(NewData));
+	BulkData.UpdatePayload(MoveTemp(NewData), Owner);
 	BulkData.SetCompressionOptions(UE::Serialization::ECompressionOptions::Default);
 }
 
@@ -1305,11 +1310,11 @@ void FTextureSource::InitLayered(
 	// it serves no purpose. Given a choice I'd assert on NewData == nullptr instead.
 	if (NewData != nullptr)
 	{
-		BulkData.UpdatePayload(FSharedBuffer::Clone(NewData, TotalBytes));
+		BulkData.UpdatePayload(FSharedBuffer::Clone(NewData, TotalBytes), Owner);
 	}
 	else
 	{
-		BulkData.UpdatePayload(FUniqueBuffer::Alloc(TotalBytes).MoveToShared());
+		BulkData.UpdatePayload(FUniqueBuffer::Alloc(TotalBytes).MoveToShared(), Owner);
 	}
 
 	BulkData.SetCompressionOptions(UE::Serialization::ECompressionOptions::Default);
@@ -1333,7 +1338,7 @@ void FTextureSource::InitLayered(
 		NewLayerFormat
 	);
 
-	BulkData.UpdatePayload(MoveTemp(NewData));
+	BulkData.UpdatePayload(MoveTemp(NewData), Owner);
 	BulkData.SetCompressionOptions(UE::Serialization::ECompressionOptions::Default);
 }
 
@@ -1419,7 +1424,7 @@ void FTextureSource::InitWithCompressedSourceData(
 
 	checkf(LockState == ELockState::None, TEXT("InitWithCompressedSourceData shouldn't be called in-between LockMip/UnlockMip"));
 
-	BulkData.UpdatePayload(FSharedBuffer::Clone(NewData.GetData(), NewData.Num()));
+	BulkData.UpdatePayload(FSharedBuffer::Clone(NewData.GetData(), NewData.Num()), Owner);
 
 	// Disable the internal bulkdata compression if the source data is already compressed
 	if (CompressionFormat == TSCF_None)
@@ -1439,6 +1444,7 @@ FTextureSource FTextureSource::CopyTornOff() const
 	Result.BulkData.TearOff();
 	// Use the default copy constructor to copy all the fields without having to write them manually
 	Result = *this;
+	Result.Owner = nullptr; // TornOffs don't count as belonging to the same owner
 	return Result;
 }
 
@@ -1475,7 +1481,7 @@ void FTextureSource::Compress()
 			TArray64<uint8> CompressedData = ImageWrapper->GetCompressed(PngQuality);
 			if ( CompressedData.Num() > 0 )
 			{
-				BulkData.UpdatePayload(MakeSharedBufferFromArray(MoveTemp(CompressedData)));
+				BulkData.UpdatePayload(MakeSharedBufferFromArray(MoveTemp(CompressedData)), Owner);
 
 				bPNGCompressed = true;
 				CompressionFormat = TSCF_PNG;
@@ -1590,7 +1596,7 @@ void FTextureSource::UnlockMip(int32 BlockIndex, int32 LayerIndex, int32 MipInde
 		{
 			UE_CLOG(CompressionFormat == TSCF_JPEG, LogTexture, Warning, TEXT("Call to FTextureSource::UnlockMip will cause texture source to lose it's jpeg storage format"));
 
-			BulkData.UpdatePayload(LockedMipData.Release());
+			BulkData.UpdatePayload(LockedMipData.Release(), Owner);
 			BulkData.SetCompressionOptions(UE::Serialization::ECompressionOptions::Default);
 
 			bPNGCompressed = false;
@@ -1928,7 +1934,7 @@ void FTextureSource::ImportCustomProperties(const TCHAR* SourceText, FFeedbackCo
 			
 			if (bSuccess)
 			{
-				BulkData.UpdatePayload(Buffer.MoveToShared());
+				BulkData.UpdatePayload(Buffer.MoveToShared(), Owner);
 			}
 		}
 
