Expose MeshDeformerInstance from SkinnedMeshComponent to Blueprint.
Will be used for setting variables on a mesh deformer instance from Blueprint.
#rnx

diff --git a/Engine/Source/Runtime/Engine/Classes/Components/SkinnedMeshComponent.h b/Engine/Source/Runtime/Engine/Classes/Components/SkinnedMeshComponent.h
--- a/Engine/Source/Runtime/Engine/Classes/Components/SkinnedMeshComponent.h
+++ b/Engine/Source/Runtime/Engine/Classes/Components/SkinnedMeshComponent.h
@@ -239,7 +239,7 @@ class ENGINE_API USkinnedMeshComponent : public UMeshComponent, public ILODSyncI
 	TObjectPtr<UMeshDeformer> MeshDeformer;
 
 	/** Object containing state for the bound MeshDeformer. */
-	UPROPERTY(Transient)
+	UPROPERTY(Transient, BlueprintReadOnly, Category = "Deformer")
 	TObjectPtr<UMeshDeformerInstance> MeshDeformerInstance;
 
 	/** const getters for previous transform idea */
