MRQ will submit remaining jobs for compilation for all the materials in the world and block on the results before rendering.

# MRQ
- When running with ODSC in the editor Materials can have incomplete or partial shader maps.  Shaders might be missing and we need to ensure all of them are compiled so we don't ever render a frame with the default material.
- This logic will run before every render frame to properly handle new materials being loaded.
- If there are no incomplete materials this logic is a no-op and very lightweight.  ~600us.

# UMaterialInterface
- Added function to UMaterialInterface to query if we have a complete set of shaders.
- Made CacheShaders and IsComplete part of the public interface of UMaterialInterface.

# Optimizations
- 31x speedup to `SubmitRemainingJobsForWorld`when there is no work to do.  This went from 20ms down to 0.645ms. It only performs the material update context (which flushes rendering commands) when there are materials that actually need to do work.
- This turns `SubmitRemainingJobsForWorld` into a no-op when there is no work.

# Misc
- When blocking on all shader map be sure to spawn a slow dialog so we can see the progress of shader compilation.
- Only spawn the dialog (a heavy operation) if there are jobs.
- Always show the dialog even in PIE.

diff --git a/Engine/Plugins/MovieScene/MovieRenderPipeline/Source/MovieRenderPipelineCore/Private/MoviePipelineRendering.cpp b/Engine/Plugins/MovieScene/MovieRenderPipeline/Source/MovieRenderPipelineCore/Private/MoviePipelineRendering.cpp
--- a/Engine/Plugins/MovieScene/MovieRenderPipeline/Source/MovieRenderPipelineCore/Private/MoviePipelineRendering.cpp
+++ b/Engine/Plugins/MovieScene/MovieRenderPipeline/Source/MovieRenderPipelineCore/Private/MoviePipelineRendering.cpp
@@ -36,6 +36,7 @@
 #include "AssetCompilingManager.h"
 #include "ShaderCompiler.h"
 #include "EngineUtils.h"
+#include "Materials/MaterialInterface.h"
 
 
 #define LOCTEXT_NAMESPACE "MoviePipeline"
@@ -496,6 +497,10 @@ void UMoviePipeline::FlushAsyncEngineSystems()
 		GetWorld()->BlockTillLevelStreamingCompleted();
 	}
 
+	// Ensure we have complete shader maps for all materials used by primitives in the world.
+	// This way we will never render with the default material.
+	UMaterialInterface::SubmitRemainingJobsForWorld(GetWorld());
+
 	// Flush all assets still being compiled asynchronously.
 	// A progressbar is already in place so the user can get feedback while waiting for everything to settle.
 	FAssetCompilingManager::Get().FinishAllCompilation();
diff --git a/Engine/Source/Runtime/Engine/Classes/Materials/Material.h b/Engine/Source/Runtime/Engine/Classes/Materials/Material.h
--- a/Engine/Source/Runtime/Engine/Classes/Materials/Material.h
+++ b/Engine/Source/Runtime/Engine/Classes/Materials/Material.h
@@ -1253,6 +1253,9 @@ public:
 	 */
 	ENGINE_API bool NeedsSetMaterialUsage_Concurrent(bool &bOutHasUsage, const EMaterialUsage Usage) const;
 
+	ENGINE_API virtual void CacheShaders(EMaterialShaderPrecompileMode CompileMode) override;
+	ENGINE_API virtual bool IsComplete() const override;
+
 #if WITH_EDITORONLY_DATA
 	/**
 	 * @param	OutParameterNames		Storage array for the parameter names.
@@ -1490,8 +1493,6 @@ private:
 	/** Caches shader maps for an array of material resources. */
 	void CacheShadersForResources(EShaderPlatform ShaderPlatform, const TArray<FMaterialResource*>& ResourcesToCache, EMaterialShaderPrecompileMode PrecompileMode = EMaterialShaderPrecompileMode::Default, const ITargetPlatform* TargetPlatform = nullptr);
 
-	ENGINE_API virtual void CacheShaders(EMaterialShaderPrecompileMode CompileMode) override;
-
 #if WITH_EDITOR
 	/**
 	 * If there is some texture reference used by a TextureProperty node in any expressions, this function
diff --git a/Engine/Source/Runtime/Engine/Classes/Materials/MaterialInstance.h b/Engine/Source/Runtime/Engine/Classes/Materials/MaterialInstance.h
--- a/Engine/Source/Runtime/Engine/Classes/Materials/MaterialInstance.h
+++ b/Engine/Source/Runtime/Engine/Classes/Materials/MaterialInstance.h
@@ -849,6 +849,9 @@ public:
 	void EndAllowCachingStaticParameterValues();
 #endif // WITH_EDITOR
 
+	ENGINE_API virtual void CacheShaders(EMaterialShaderPrecompileMode CompileMode) override;
+	ENGINE_API virtual bool IsComplete() const override;
+
 protected:
 
 	/**
@@ -913,8 +916,6 @@ protected:
 	/** Caches shader maps for an array of material resources. */
 	void CacheShadersForResources(EShaderPlatform ShaderPlatform, const TArray<FMaterialResource*>& ResourcesToCache, EMaterialShaderPrecompileMode PrecompileMode = EMaterialShaderPrecompileMode::Default, const ITargetPlatform* TargetPlatform = nullptr);
 
-	ENGINE_API virtual void CacheShaders(EMaterialShaderPrecompileMode CompileMode) override;
-
 	/** 
 	 * Copies over material instance parameters from the base material given a material interface.
 	 * This is a slow operation that is needed for the editor.
diff --git a/Engine/Source/Runtime/Engine/Classes/Materials/MaterialInterface.h b/Engine/Source/Runtime/Engine/Classes/Materials/MaterialInterface.h
--- a/Engine/Source/Runtime/Engine/Classes/Materials/MaterialInterface.h
+++ b/Engine/Source/Runtime/Engine/Classes/Materials/MaterialInterface.h
@@ -806,6 +806,18 @@ public:
 	*/
 	ENGINE_API virtual void CacheShaders(EMaterialShaderPrecompileMode CompileMode = EMaterialShaderPrecompileMode::Default) {}
 
+	/** @brief Checks to see if this material has all its shaders cached.
+	*
+	* Materials are not guaranteed to have all their shaders compiled after loading.  It can be useful to
+	* check for completeness in order to cache remaining shaders.
+	* 
+	* @return Whether or not all shaders for this material exist.
+	*
+	* @see CacheShaders
+	* @note This function will return true if the resources are not cache for this material yet.
+	*/
+	ENGINE_API virtual bool IsComplete() const { return true; }
+
 #if WITH_EDITOR
 	/** Clears the shader cache and recompiles the shader for rendering. */
 	ENGINE_API virtual void ForceRecompileForRendering() {}
diff --git a/Engine/Source/Runtime/Engine/Private/Materials/Material.cpp b/Engine/Source/Runtime/Engine/Private/Materials/Material.cpp
--- a/Engine/Source/Runtime/Engine/Private/Materials/Material.cpp
+++ b/Engine/Source/Runtime/Engine/Private/Materials/Material.cpp
@@ -2297,6 +2297,30 @@ void UMaterial::CacheShaders(EMaterialShaderPrecompileMode CompileMode)
 	CacheResourceShadersForRendering(false, CompileMode);
 }
 
+bool UMaterial::IsComplete() const
+{
+	bool bComplete = true;
+	if (FApp::CanEverRender())
+	{
+		const EMaterialQualityLevel::Type ActiveQualityLevel = GetCachedScalabilityCVars().MaterialQualityLevel;
+		uint32 FeatureLevelsToCompile = GetFeatureLevelsToCompileForRendering();
+
+		while (FeatureLevelsToCompile != 0)
+		{
+			const ERHIFeatureLevel::Type FeatureLevel = (ERHIFeatureLevel::Type)FBitSet::GetAndClearNextBit(FeatureLevelsToCompile);
+			const EShaderPlatform ShaderPlatform = GShaderPlatformForFeatureLevel[FeatureLevel];
+
+			FMaterialResource* CurrentResource = FindMaterialResource(MaterialResources, FeatureLevel, ActiveQualityLevel, true);
+			if (CurrentResource && !CurrentResource->IsGameThreadShaderMapComplete())
+			{
+				bComplete = false;
+				break;
+			}
+		}
+	}
+	return bComplete;
+}
+
 void UMaterial::ReleaseResourcesAndMutateDDCKey(const FGuid& TransformationId)
 {
 	if (TransformationId.IsValid())
diff --git a/Engine/Source/Runtime/Engine/Private/Materials/MaterialInstance.cpp b/Engine/Source/Runtime/Engine/Private/Materials/MaterialInstance.cpp
--- a/Engine/Source/Runtime/Engine/Private/Materials/MaterialInstance.cpp
+++ b/Engine/Source/Runtime/Engine/Private/Materials/MaterialInstance.cpp
@@ -2163,6 +2163,31 @@ bool UMaterialInstance::GetMaterialLayers(FMaterialLayersFunctions& OutLayers, T
 	return false;
 }
 
+bool UMaterialInstance::IsComplete() const
+{
+	bool bComplete = true;
+	if (bHasStaticPermutationResource && FApp::CanEverRender())
+	{
+		check(IsA(UMaterialInstanceConstant::StaticClass()));
+
+		uint32 FeatureLevelsToCompile = GetFeatureLevelsToCompileForRendering();
+		const EMaterialQualityLevel::Type ActiveQualityLevel = GetCachedScalabilityCVars().MaterialQualityLevel;
+
+		while (FeatureLevelsToCompile != 0)
+		{
+			const ERHIFeatureLevel::Type FeatureLevel = (ERHIFeatureLevel::Type)FBitSet::GetAndClearNextBit(FeatureLevelsToCompile);
+			const EShaderPlatform ShaderPlatform = GShaderPlatformForFeatureLevel[FeatureLevel];
+
+			FMaterialResource* CurrentResource = FindMaterialResource(StaticPermutationMaterialResources, FeatureLevel, ActiveQualityLevel, true);
+			if (CurrentResource && !CurrentResource->IsGameThreadShaderMapComplete())
+			{
+				bComplete = false;
+				break;
+			}
+		}
+	}
+	return bComplete;
+}
 
 #if WITH_EDITOR
 bool UMaterialInstance::SetMaterialLayers(const FMaterialLayersFunctions& LayersValue)
diff --git a/Engine/Source/Runtime/Engine/Private/Materials/MaterialInterface.cpp b/Engine/Source/Runtime/Engine/Private/Materials/MaterialInterface.cpp
--- a/Engine/Source/Runtime/Engine/Private/Materials/MaterialInterface.cpp
+++ b/Engine/Source/Runtime/Engine/Private/Materials/MaterialInterface.cpp
@@ -453,8 +453,6 @@ void UMaterialInterface::SubmitRemainingJobsForWorld(UWorld* World, EMaterialSha
 	TSet<UMaterialInterface*> MaterialsToCache;
 	FObjectCacheContextScope ObjectCacheScope;
 
-	FMaterialUpdateContext UpdateContext(FMaterialUpdateContext::EOptions::SyncWithRenderingThread);
-
 	for (UPrimitiveComponent* PrimitiveComponent : ObjectCacheScope.GetContext().GetPrimitiveComponents())
 	{
 		if (World && !World->ContainsActor(PrimitiveComponent->GetOwner()))
@@ -467,7 +465,7 @@ void UMaterialInterface::SubmitRemainingJobsForWorld(UWorld* World, EMaterialSha
 			TObjectCacheIterator<UMaterialInterface> UsedMaterials = ObjectCacheScope.GetContext().GetUsedMaterials(PrimitiveComponent);
 			for (UMaterialInterface* MaterialInterface : UsedMaterials)
 			{
-				if (MaterialInterface)
+				if (MaterialInterface && !MaterialInterface->IsComplete())
 				{
 					MaterialsToCache.Add(MaterialInterface);
 				}
@@ -480,18 +478,23 @@ void UMaterialInterface::SubmitRemainingJobsForWorld(UWorld* World, EMaterialSha
 	{
 		UMaterial* Material = It->GetMaterial();
 		if (Material &&
-			(Material->IsUIMaterial() || Material->IsPostProcessMaterial()))
+			(Material->IsUIMaterial() || Material->IsPostProcessMaterial()) &&
+			!Material->IsComplete())
 		{
 			MaterialsToCache.Add(*It);
 		}
 	}
 
-	for (UMaterialInterface* Material : MaterialsToCache)
+	if (MaterialsToCache.Num())
 	{
-		// This is needed because CacheShaders blindly recreates uniform buffers
-		// which can only be done if the draw command is going to be re-cached.
-		UpdateContext.AddMaterialInterface(Material);
-		Material->CacheShaders(CompileMode);
+		FMaterialUpdateContext UpdateContext(FMaterialUpdateContext::EOptions::SyncWithRenderingThread);
+		for (UMaterialInterface* Material : MaterialsToCache)
+		{
+			// This is needed because CacheShaders blindly recreates uniform buffers
+			// which can only be done if the draw command is going to be re-cached.
+			UpdateContext.AddMaterialInterface(Material);
+			Material->CacheShaders(CompileMode);
+		}
 	}
 }
 
diff --git a/Engine/Source/Runtime/Engine/Private/ShaderCompiler/ShaderCompiler.cpp b/Engine/Source/Runtime/Engine/Private/ShaderCompiler/ShaderCompiler.cpp
--- a/Engine/Source/Runtime/Engine/Private/ShaderCompiler/ShaderCompiler.cpp
+++ b/Engine/Source/Runtime/Engine/Private/ShaderCompiler/ShaderCompiler.cpp
@@ -4113,8 +4113,8 @@ void FShaderCompilingManager::BlockOnAllShaderMapCompletion(TMap<int32, FShaderM
 			}
 		}
 
-		FScopedSlowTask SlowTask(NumJobs, FText::Format(LOCTEXT("BlockOnAllShaderMapCompletion", "Compiling Shaders ({0})"), NumJobs), GIsEditor && !IsRunningCommandlet());
-		SlowTask.MakeDialog();
+		FScopedSlowTask SlowTask(NumJobs, FText::Format(LOCTEXT("BlockOnAllShaderMapCompletion", "Compiling Shaders ({0})"), NumJobs), NumJobs && GIsEditor && !IsRunningCommandlet());
+		SlowTask.MakeDialog(false, true);
 
 		int32 NumPendingJobs = 0;
 		// Keep track of previous number of pending jobs so we can update the slow task with the amount of work done.
