Fix for transition events not firing from conduits

diff --git a/Engine/Source/Runtime/Engine/Private/Animation/AnimNode_StateMachine.cpp b/Engine/Source/Runtime/Engine/Private/Animation/AnimNode_StateMachine.cpp
--- a/Engine/Source/Runtime/Engine/Private/Animation/AnimNode_StateMachine.cpp
+++ b/Engine/Source/Runtime/Engine/Private/Animation/AnimNode_StateMachine.cpp
@@ -1043,7 +1043,7 @@ void FAnimNode_StateMachine::TransitionToState(const FAnimationUpdateContext& Co
 		Context.AnimInstanceProxy->AddAnimNotifyFromGeneratedClass(ActiveTransitionArray[ActiveTransitionArray.Num() - 1].InterruptNotify);
 	}
 
-	const int32 PreviousState = TransitionInfo.PreviousState;
+	const int32 PreviousState = CurrentState;
 	const int32 NextState = TransitionInfo.NextState;
 
 	// Fire off Notifies for state transition
