Physics Replication:
Add logging for hard snaps.

diff --git a/Engine/Source/Runtime/Engine/Private/PhysicsEngine/PhysicsReplication.cpp b/Engine/Source/Runtime/Engine/Private/PhysicsEngine/PhysicsReplication.cpp
--- a/Engine/Source/Runtime/Engine/Private/PhysicsEngine/PhysicsReplication.cpp
+++ b/Engine/Source/Runtime/Engine/Private/PhysicsEngine/PhysicsReplication.cpp
@@ -18,7 +18,7 @@
 #include "Physics/PhysicsInterfaceCore.h"
 #include "Physics/PhysScene_PhysX.h"
 #include "Components/SkeletalMeshComponent.h"
-
+#include "Engine/DemoNetDriver.h"
 
 #if WITH_CHAOS
 #include "Chaos/ChaosMarshallingManager.h"
@@ -92,6 +92,10 @@ namespace PhysicsReplicationCVars
 {
 	static int32 SkipSkeletalRepOptimization = 1;
 	static FAutoConsoleVariableRef CVarSkipSkeletalRepOptimization(TEXT("p.SkipSkeletalRepOptimization"), SkipSkeletalRepOptimization, TEXT("If true, we don't move the skeletal mesh component during replication. This is ok because the skeletal mesh already polls physx after its results"));
+#if !UE_BUILD_SHIPPING
+	int32 LogPhysicsReplicationHardSnaps = 0;
+	static FAutoConsoleVariableRef CVarLogPhysicsReplicationHardSnaps(TEXT("p.LogPhysicsReplicationHardSnaps"), LogPhysicsReplicationHardSnaps, TEXT(""));
+#endif
 }
 
 #if WITH_CHAOS
@@ -188,7 +192,7 @@ bool FPhysicsReplication::ApplyRigidBodyState(float DeltaSeconds, FBodyInstance*
 	return ApplyRigidBodyState(DeltaSeconds, BI, PhysicsTarget, ErrorCorrection, PingSecondsOneWay);
 }
 
-bool FPhysicsReplication::ApplyRigidBodyState(float DeltaSeconds, FBodyInstance* BI, FReplicatedPhysicsTarget& PhysicsTarget, const FRigidBodyErrorCorrection& ErrorCorrection, const float PingSecondsOneWay)
+bool FPhysicsReplication::ApplyRigidBodyState(float DeltaSeconds, FBodyInstance* BI, FReplicatedPhysicsTarget& PhysicsTarget, const FRigidBodyErrorCorrection& ErrorCorrection, const float PingSecondsOneWay, bool* bDidHardSnap)
 {
 	if (CharacterMovementCVars::SkipPhysicsReplication)
 	{
@@ -363,6 +367,26 @@ bool FPhysicsReplication::ApplyRigidBodyState(float DeltaSeconds, FBodyInstance*
 
 		if (bHardSnap)
 		{
+#if !UE_BUILD_SHIPPING
+			if (PhysicsReplicationCVars::LogPhysicsReplicationHardSnaps && GetOwningWorld())
+			{
+				UE_LOG(LogTemp, Warning, TEXT("Simulated HARD SNAP - - %s, \nCurrent Pos - %s, Target Pos - %s\n CurrentState.LinVel - %s, New Lin Vel - %s\nTarget Extrapolation Delta - %s, Is Replay? - %d, Is Asleep - %d"),
+					*CurrentState.Position.ToString(), *TargetPos.ToString(), *CurrentState.LinVel.ToString(), *NewState.LinVel.ToString(),
+					*ExtrapolationDeltaPos.ToString(), (GetOwningWorld()->GetDemoNetDriver() && GetOwningWorld()->GetDemoNetDriver()->IsPlaying()), !BI->IsInstanceAwake());
+				if (bDidHardSnap)
+				{
+					*bDidHardSnap = true;
+				}
+				if (LinDiffSize > MaxLinearHardSnapDistance)
+				{
+					UE_LOG(LogTemp, Warning, TEXT("Hard snap due to linear difference error"));
+				}
+				else
+				{
+					UE_LOG(LogTemp, Warning, TEXT("Hard snap due to accumulated error"))
+				}
+			}
+#endif
 			// Too much error so just snap state here and be done with it
 			PhysicsTarget.AccumulatedErrorSeconds = 0.0f;
 			bRestoredState = true;
@@ -564,6 +588,14 @@ void FPhysicsReplication::OnTick(float DeltaSeconds, TMap<TWeakObjectPtr<UPrimit
 
 						if (UpdatedState.Flags & ERigidBodyFlags::NeedsUpdate)
 						{
+							bool bDidHardSnap = false;
+							const bool bRestoredState = ApplyRigidBodyState(DeltaSeconds, BI, PhysicsTarget, PhysicErrorCorrection, PingSecondsOneWay, &bDidHardSnap);
+#if !UE_BUILD_SHIPPING
+							if (PhysicsReplicationCVars::LogPhysicsReplicationHardSnaps && bDidHardSnap)
+							{
+								UE_LOG(LogTemp, Warning, TEXT("\nIn FPhysicsReplication::OnTick - Name - %s"), *PrimComp->GetOwner()->GetName());
+							}
+#endif
 							const int32 LocalFrame = PhysicsTarget.ServerFrame + LocalFrameOffset;
 							const bool bRestoredState = ApplyRigidBodyState(DeltaSeconds, BI, PhysicsTarget, PhysicErrorCorrection, PingSecondsOneWay, LocalFrame, NumPredictedFrames);
 
diff --git a/Engine/Source/Runtime/Engine/Public/PhysicsReplication.h b/Engine/Source/Runtime/Engine/Public/PhysicsReplication.h
--- a/Engine/Source/Runtime/Engine/Public/PhysicsReplication.h
+++ b/Engine/Source/Runtime/Engine/Public/PhysicsReplication.h
@@ -19,6 +19,13 @@ namespace Chaos
 	struct FSimCallbackInput;
 }
 
+#if !UE_BUILD_SHIPPING
+namespace PhysicsReplicationCVars
+{
+	extern ENGINE_API int32 LogPhysicsReplicationHardSnaps;
+}
+#endif
+
 class FPhysScene_PhysX;
 
 struct FReplicatedPhysicsTarget
@@ -98,7 +105,7 @@ protected:
 
 	/** Called when a dynamic rigid body receives a physics update */
 	virtual bool ApplyRigidBodyState(float DeltaSeconds, FBodyInstance* BI, FReplicatedPhysicsTarget& PhysicsTarget, const FRigidBodyErrorCorrection& ErrorCorrection, const float PingSecondsOneWay, int32 LocalFrame, int32 NumPredictedFrames);
-	virtual bool ApplyRigidBodyState(float DeltaSeconds, FBodyInstance* BI, FReplicatedPhysicsTarget& PhysicsTarget, const FRigidBodyErrorCorrection& ErrorCorrection, const float PingSecondsOneWay); // deprecated path with no localframe/numpredicted
+	virtual bool ApplyRigidBodyState(float DeltaSeconds, FBodyInstance* BI, FReplicatedPhysicsTarget& PhysicsTarget, const FRigidBodyErrorCorrection& ErrorCorrection, const float PingSecondsOneWay, bool* bDidHardSnap = nullptr); // deprecated path with no localframe/numpredicted
 
 	UWorld* GetOwningWorld();
 	const UWorld* GetOwningWorld() const;
