Fix compile error caused by merge.

diff --git a/Engine/Source/Runtime/Engine/Private/PhysicsEngine/PhysicsReplication.cpp b/Engine/Source/Runtime/Engine/Private/PhysicsEngine/PhysicsReplication.cpp
--- a/Engine/Source/Runtime/Engine/Private/PhysicsEngine/PhysicsReplication.cpp
+++ b/Engine/Source/Runtime/Engine/Private/PhysicsEngine/PhysicsReplication.cpp
@@ -589,16 +589,14 @@ void FPhysicsReplication::OnTick(float DeltaSeconds, TMap<TWeakObjectPtr<UPrimit
 						if (UpdatedState.Flags & ERigidBodyFlags::NeedsUpdate)
 						{
 							bool bDidHardSnap = false;
-							const bool bRestoredState = ApplyRigidBodyState(DeltaSeconds, BI, PhysicsTarget, PhysicErrorCorrection, PingSecondsOneWay, &bDidHardSnap);
+							const int32 LocalFrame = PhysicsTarget.ServerFrame + LocalFrameOffset;
+							const bool bRestoredState = ApplyRigidBodyState(DeltaSeconds, BI, PhysicsTarget, PhysicErrorCorrection, PingSecondsOneWay, LocalFrame, NumPredictedFrames, &bDidHardSnap);
 #if !UE_BUILD_SHIPPING
 							if (PhysicsReplicationCVars::LogPhysicsReplicationHardSnaps && bDidHardSnap)
 							{
 								UE_LOG(LogTemp, Warning, TEXT("\nIn FPhysicsReplication::OnTick - Name - %s"), *PrimComp->GetOwner()->GetName());
 							}
 #endif
-							const int32 LocalFrame = PhysicsTarget.ServerFrame + LocalFrameOffset;
-							const bool bRestoredState = ApplyRigidBodyState(DeltaSeconds, BI, PhysicsTarget, PhysicErrorCorrection, PingSecondsOneWay, LocalFrame, NumPredictedFrames);
-
 							// Need to update the component to match new position.
 							if (PhysicsReplicationCVars::SkipSkeletalRepOptimization == 0 || Cast<USkeletalMeshComponent>(PrimComp) == nullptr)	//simulated skeletal mesh does its own polling of physics results so we don't need to call this as it'll happen at the end of the physics sim
 							{
