Fixes WaterBody initialization crash when recreating a level using same CDO template as the previous level by preventing UActorComponent::PostInitProperties() to mark the CDO component for garbage collection in cases where the owner is cooked.

Issue was also present in regular UE but prevented because BP are recompiled after load, hence the second CDO instance was complete. Code responsible for the component deletion was introduced to fix UE-28374 (loading a previously saved actor component after removing the component in code), but does not seem relevant as the newly cooked actor would get rid of this component.

#changelist validated

diff --git a/Engine/Source/Runtime/Engine/Private/Components/ActorComponent.cpp b/Engine/Source/Runtime/Engine/Private/Components/ActorComponent.cpp
--- a/Engine/Source/Runtime/Engine/Private/Components/ActorComponent.cpp
+++ b/Engine/Source/Runtime/Engine/Private/Components/ActorComponent.cpp
@@ -341,7 +341,7 @@ void UActorComponent::PostInitProperties()
 	// Instance components will be added during the owner's initialization
 	if (OwnerPrivate && CreationMethod != EComponentCreationMethod::Instance)
 	{
-		if (!FPlatformProperties::RequiresCookedData() && CreationMethod == EComponentCreationMethod::Native && HasAllFlags(RF_NeedLoad|RF_DefaultSubObject))
+		if (!FPlatformProperties::RequiresCookedData() && !OwnerPrivate->GetClass()->bCooked && CreationMethod == EComponentCreationMethod::Native && HasAllFlags(RF_NeedLoad|RF_DefaultSubObject))
 		{
 			UObject* MyArchetype = GetArchetype();
 			if (IsValid(MyArchetype) && MyArchetype != GetClass()->ClassDefaultObject)
