Favor instance data over sparse data when resolving BP pins to properties

#rnx

diff --git a/Engine/Source/Editor/BlueprintGraph/Private/K2Node_Variable.cpp b/Engine/Source/Editor/BlueprintGraph/Private/K2Node_Variable.cpp
--- a/Engine/Source/Editor/BlueprintGraph/Private/K2Node_Variable.cpp
+++ b/Engine/Source/Editor/BlueprintGraph/Private/K2Node_Variable.cpp
@@ -394,19 +394,7 @@ FProperty* UK2Node_Variable::GetPropertyForVariable_Internal(UClass* OwningClass
 {
 	const FName VarName = GetVarName();
 
-	// Look in the sparse class data first
-	FProperty* VariableProperty = nullptr;
-	// TODO: move some of this into MemberReference::ResolveMember if possible
-	UClass* Scope = VariableReference.GetScope(OwningClass);
-	UScriptStruct* SparseClassDataStruct = Scope ? Scope->GetSparseClassDataStruct() : nullptr;
-	if (SparseClassDataStruct)
-	{
-		VariableProperty = FindFProperty<FProperty>(SparseClassDataStruct, VarName);
-	}
-	if (!VariableProperty)
-	{
-		VariableProperty = VariableReference.ResolveMember<FProperty>(OwningClass);
-	}
+	FProperty* VariableProperty = VariableReference.ResolveMember<FProperty>(OwningClass);
 
 	// if the variable has been deprecated, don't use it
 	if (VariableProperty != nullptr)
diff --git a/Engine/Source/Runtime/Engine/Classes/Engine/MemberReference.h b/Engine/Source/Runtime/Engine/Classes/Engine/MemberReference.h
--- a/Engine/Source/Runtime/Engine/Classes/Engine/MemberReference.h
+++ b/Engine/Source/Runtime/Engine/Classes/Engine/MemberReference.h
@@ -428,14 +428,13 @@ public:
 #endif
 				{
 					// Find in target scope or in the sparse class data
-					UScriptStruct* SparseClassDataStruct = TargetScope->GetSparseClassDataStruct();
-					if (SparseClassDataStruct)
+					ReturnField = FindUFieldOrFProperty(TargetScope, MemberName, EFieldIterationFlags::IncludeAll).Get<TFieldType>();
+					if (!ReturnField)
 					{
-						ReturnField = FindUFieldOrFProperty(SparseClassDataStruct, MemberName, EFieldIterationFlags::IncludeAll).Get<TFieldType>();
-					}
-					if (ReturnField == nullptr)
-					{
-						ReturnField = FindUFieldOrFProperty(TargetScope, MemberName, EFieldIterationFlags::IncludeAll).Get<TFieldType>();
+						if (UScriptStruct* SparseClassDataStruct = TargetScope->GetSparseClassDataStruct())
+						{
+							ReturnField = FindUFieldOrFProperty(SparseClassDataStruct, MemberName, EFieldIterationFlags::IncludeAll).Get<TFieldType>();
+						}
 					}
 				}
 
