Making streaming priority deprecated

diff --git a/Engine/Source/Runtime/Engine/Classes/Sound/SoundWave.h b/Engine/Source/Runtime/Engine/Classes/Sound/SoundWave.h
--- a/Engine/Source/Runtime/Engine/Classes/Sound/SoundWave.h
+++ b/Engine/Source/Runtime/Engine/Classes/Sound/SoundWave.h
@@ -366,8 +366,7 @@ private:
 
 public:
 
-	/** Priority of this sound when streaming (lower priority streams may not always play) */
-	UPROPERTY(EditAnywhere, Category = "Playback|Streaming", meta = (ClampMin = 0))
+	UPROPERTY(meta = (DeprecatedProperty, DeprecationMessage = "5.0 - Property is deprecated. Streaming priority has no effect with stream caching enabled."))
 	int32 StreamingPriority;
 
 	/** Quality of sample rate conversion for platforms that opt into resampling during cook. The sample rate for each enumeration is definable per platform in platform target settings. */
