Allowed BlendSpaceEvaluators to run with non-zero playrate so that they can generate notifies during playback.

diff --git a/Engine/Source/Runtime/AnimGraphRuntime/Private/AnimNodes/AnimNode_BlendSpaceEvaluator.cpp b/Engine/Source/Runtime/AnimGraphRuntime/Private/AnimNodes/AnimNode_BlendSpaceEvaluator.cpp
--- a/Engine/Source/Runtime/AnimGraphRuntime/Private/AnimNodes/AnimNode_BlendSpaceEvaluator.cpp
+++ b/Engine/Source/Runtime/AnimGraphRuntime/Private/AnimNodes/AnimNode_BlendSpaceEvaluator.cpp
@@ -36,5 +36,5 @@ void FAnimNode_BlendSpaceEvaluator::GatherDebugData(FNodeDebugData& DebugData)
 
 float FAnimNode_BlendSpaceEvaluator::GetPlayRate() const
 {
-	return 0.0f;
+	return bTeleportToNormalizedTime ? 0.0f : 1.0f;;
 }
diff --git a/Engine/Source/Runtime/AnimGraphRuntime/Public/AnimNodes/AnimNode_BlendSpaceEvaluator.h b/Engine/Source/Runtime/AnimGraphRuntime/Public/AnimNodes/AnimNode_BlendSpaceEvaluator.h
--- a/Engine/Source/Runtime/AnimGraphRuntime/Public/AnimNodes/AnimNode_BlendSpaceEvaluator.h
+++ b/Engine/Source/Runtime/AnimGraphRuntime/Public/AnimNodes/AnimNode_BlendSpaceEvaluator.h
@@ -9,7 +9,6 @@
 
 // Evaluates a point in a blendspace, using a specific time input rather than advancing time internally.
 // Typically the playback position of the animation for this node will represent something other than time, like jump height.
-// This node will not trigger any notifies present in the associated sequence.
 USTRUCT(BlueprintInternalUseOnly)
 struct ANIMGRAPHRUNTIME_API FAnimNode_BlendSpaceEvaluator : public FAnimNode_BlendSpacePlayer
 {
diff --git a/Engine/Source/Runtime/Engine/Private/Animation/BlendSpace.cpp b/Engine/Source/Runtime/Engine/Private/Animation/BlendSpace.cpp
--- a/Engine/Source/Runtime/Engine/Private/Animation/BlendSpace.cpp
+++ b/Engine/Source/Runtime/Engine/Private/Animation/BlendSpace.cpp
@@ -520,13 +520,17 @@ void UBlendSpace::TickAssetPlayer(FAnimTickRecord& Instance, struct FAnimNotifyQ
 			{
 				FAnimNotifyContext NotifyContext(Instance);
 				float ClampedNormalizedPreviousTime = FMath::Clamp<float>(NormalizedPreviousTime, 0.f, 1.f);
-				const float ClampedNormalizedCurrentTime = FMath::Clamp<float>(NormalizedCurrentTime, 0.f, 1.f);
+				float ClampedNormalizedCurrentTime = FMath::Clamp<float>(NormalizedCurrentTime, 0.f, 1.f);
 
 				if (Instance.BlendSpace.bIsEvaluator && !Instance.BlendSpace.bTeleportToTime)
 				{
 					// When running under an evaluator the time is being set explicitly and we want to add on the deltas.
-					// Note that this could take ClampedNormalizedPreviousTime out of the 0-1 range
 					ClampedNormalizedPreviousTime -= ExtraNormalizedDeltaTime;
+					ClampedNormalizedPreviousTime = ClampedNormalizedPreviousTime < 0.0f ? ClampedNormalizedPreviousTime + 1.0f : ClampedNormalizedPreviousTime;
+
+					// Also when under an evaluator, since the time is explicitly set before the update is called, the desired 
+					// current time is actually what we recorded before advancing time (effectively ignoring whatever was added).
+					ClampedNormalizedCurrentTime = FMath::Clamp<float>(NormalizedPreviousTime, 0.f, 1.f);
 				}
 
 				const bool bGenerateNotifies = (NormalizedCurrentTime != NormalizedPreviousTime) && NotifyTriggerMode != ENotifyTriggerMode::None;
