Using a separate connection flag (rather than internal ack) to indicate when the partial bunch size limit can be ignored

diff --git a/Engine/Source/Runtime/Engine/Classes/Engine/NetConnection.h b/Engine/Source/Runtime/Engine/Classes/Engine/NetConnection.h
--- a/Engine/Source/Runtime/Engine/Classes/Engine/NetConnection.h
+++ b/Engine/Source/Runtime/Engine/Classes/Engine/NetConnection.h
@@ -312,6 +312,18 @@ public:
 		bForceInitialDirty = bValue;
 	}
 
+	/**	Used to allow connections to ignore the bunch size limitation applied before splitting into partial bunch packets, 
+	*	or when receiving partial bunches to reassemble.
+	* 
+	*	!!!WARNING!!!  This is a security risk as the connection will accept much larger bunches before hitting the 
+	*	reliable buffer limit, or the array/string serialization limit.
+	*/
+	bool IsUnlimitedBunchSizeAllowed() const { return bUnlimitedBunchSizeAllowed; }
+	void SetUnlimitedBunchSizeAllowed(bool bValue)
+	{
+		bUnlimitedBunchSizeAllowed = bValue;
+	}
+
 	/** Destructor */
 	ENGINE_API virtual ~UNetConnection() {};
 
@@ -319,6 +331,7 @@ private:
 	uint32 bInternalAck : 1;	// Internally ack all packets, for 100% reliable connections.
 	uint32 bReplay : 1;			// Flag to indicate a replay connection, independent of reliability
 	uint32 bForceInitialDirty : 1;	// Force all properties dirty on initial replication
+	uint32 bUnlimitedBunchSizeAllowed : 1; // Ignore the value of net.MaxConstructedPartialBunchSizeBytes
 
 public:
 	struct FURL			URL;				// URL of the other side.
diff --git a/Engine/Source/Runtime/Engine/Private/DataChannel.cpp b/Engine/Source/Runtime/Engine/Private/DataChannel.cpp
--- a/Engine/Source/Runtime/Engine/Private/DataChannel.cpp
+++ b/Engine/Source/Runtime/Engine/Private/DataChannel.cpp
@@ -121,7 +121,7 @@ namespace UE
 template<typename T>
 static const bool IsBunchTooLarge(UNetConnection* Connection, T* Bunch)
 {
-	return !Connection->IsInternalAck() && Bunch != nullptr && Bunch->GetNumBytes() > NetMaxConstructedPartialBunchSizeBytes;
+	return !Connection->IsUnlimitedBunchSizeAllowed() && Bunch != nullptr && Bunch->GetNumBytes() > NetMaxConstructedPartialBunchSizeBytes;
 }
 
 /** Helper class to restrict access to the subobject list of actors and actor components */
diff --git a/Engine/Source/Runtime/Engine/Private/DemoNetDriver.cpp b/Engine/Source/Runtime/Engine/Private/DemoNetDriver.cpp
--- a/Engine/Source/Runtime/Engine/Private/DemoNetDriver.cpp
+++ b/Engine/Source/Runtime/Engine/Private/DemoNetDriver.cpp
@@ -4461,6 +4461,7 @@ UDemoNetConnection::UDemoNetConnection( const FObjectInitializer& ObjectInitiali
 	SetInternalAck(true);
 	SetReplay(true);
 	SetAutoFlush(true);
+	SetUnlimitedBunchSizeAllowed(true);
 }
 
 void UDemoNetConnection::InitConnection( UNetDriver* InDriver, EConnectionState InState, const FURL& InURL, int32 InConnectionSpeed, int32 InMaxPacket)
@@ -4472,6 +4473,7 @@ void UDemoNetConnection::InitConnection( UNetDriver* InDriver, EConnectionState
 	SetInternalAck(true);
 	SetReplay(true);
 	SetAutoFlush(true);
+	SetUnlimitedBunchSizeAllowed(true);
 
 	InitSendBuffer();
 
diff --git a/Engine/Source/Runtime/Engine/Private/NetConnection.cpp b/Engine/Source/Runtime/Engine/Private/NetConnection.cpp
--- a/Engine/Source/Runtime/Engine/Private/NetConnection.cpp
+++ b/Engine/Source/Runtime/Engine/Private/NetConnection.cpp
@@ -224,6 +224,7 @@ UNetConnection::UNetConnection(const FObjectInitializer& ObjectInitializer)
 ,	bInternalAck		( false )
 ,	bReplay				( false )
 ,	bForceInitialDirty	( false )
+,	bUnlimitedBunchSizeAllowed ( false )
 ,	RemoteAddr			( nullptr )
 ,	MaxPacketHandlerBits ( 0 )
 ,	State				( USOCK_Invalid )
@@ -4908,6 +4909,7 @@ void UNetConnection::NotifyActorChannelCleanedUp(UActorChannel* Channel, EChanne
 USimulatedClientNetConnection::USimulatedClientNetConnection( const FObjectInitializer& ObjectInitializer ) : Super( ObjectInitializer )
 {
 	SetInternalAck(true);
+	SetUnlimitedBunchSizeAllowed(true);	// here to avoid changing the previous behavior controlled by SetInternalAck, but probably not necessary
 }
 
 void USimulatedClientNetConnection::HandleClientPlayer( class APlayerController* PC, class UNetConnection* NetConnection )
diff --git a/Engine/Source/Runtime/Engine/Private/ReplayNetConnection.cpp b/Engine/Source/Runtime/Engine/Private/ReplayNetConnection.cpp
--- a/Engine/Source/Runtime/Engine/Private/ReplayNetConnection.cpp
+++ b/Engine/Source/Runtime/Engine/Private/ReplayNetConnection.cpp
@@ -18,6 +18,7 @@ UReplayNetConnection::UReplayNetConnection(const FObjectInitializer& ObjectIniti
 	SetInternalAck(true);
 	SetReplay(true);
 	SetAutoFlush(true);
+	SetUnlimitedBunchSizeAllowed(true);
 }
 
 void UReplayNetConnection::InitConnection(UNetDriver* InDriver, EConnectionState InState, const FURL& InURL, int32 InConnectionSpeed, int32 InMaxPacket)
@@ -30,6 +31,7 @@ void UReplayNetConnection::InitConnection(UNetDriver* InDriver, EConnectionState
 	SetInternalAck(true);
 	SetReplay(true);
 	SetAutoFlush(true);
+	SetUnlimitedBunchSizeAllowed(true);
 
 	InitSendBuffer();
 
