WorldPartition: Bring back InstanceTransform (ref CL: 18071431)

#rnx

diff --git a/Engine/Source/Runtime/Engine/Private/WorldPartition/WorldPartition.cpp b/Engine/Source/Runtime/Engine/Private/WorldPartition/WorldPartition.cpp
--- a/Engine/Source/Runtime/Engine/Private/WorldPartition/WorldPartition.cpp
+++ b/Engine/Source/Runtime/Engine/Private/WorldPartition/WorldPartition.cpp
@@ -260,6 +260,7 @@ UWorldPartition::UWorldPartition(const FObjectInitializer& ObjectInitializer)
 	, bIsPIE(false)
 #endif
 	, InitState(EWorldPartitionInitState::Uninitialized)
+	, InstanceTransform(FTransform::Identity)
 	, StreamingPolicy(nullptr)
 #if !UE_BUILD_SHIPPING
 	, Replay(nullptr)
@@ -390,6 +391,7 @@ void UWorldPartition::Initialize(UWorld* InWorld, const FTransform& InTransform)
 
 	check(InWorld);
 	World = InWorld;
+	InstanceTransform = InTransform;
 
 	check(InitState == EWorldPartitionInitState::Uninitialized);
 	InitState = EWorldPartitionInitState::Initializing;
@@ -466,7 +468,7 @@ void UWorldPartition::Initialize(UWorld* InWorld, const FTransform& InTransform)
 
 					InstancingContext.AddMapping(*LongActorPackageName, *InstancedName);
 
-					ActorDescIterator->TransformInstance(SourceWorldPath, RemappedWorldPath);
+					ActorDescIterator->TransformInstance(SourceWorldPath, RemappedWorldPath, InstanceTransform);
 				}
 
 				if (bIsEditor && !bIsCooking)
@@ -1113,6 +1115,7 @@ void UWorldPartition::OnActorDescRegistered(const FWorldPartitionActorDesc& Acto
 {
 	AActor* Actor = ActorDesc.GetActor();
 	check(Actor);
+	ApplyActorTransform(Actor, InstanceTransform);
 	Actor->GetLevel()->AddLoadedActor(Actor);
 }
 
@@ -1122,6 +1125,7 @@ void UWorldPartition::OnActorDescUnregistered(const FWorldPartitionActorDesc& Ac
 	if (IsValidChecked(Actor))
 	{
 		Actor->GetLevel()->RemoveLoadedActor(Actor);
+		ApplyActorTransform(Actor, InstanceTransform.Inverse());
 	}
 }
 
diff --git a/Engine/Source/Runtime/Engine/Private/WorldPartition/WorldPartitionActorDesc.cpp b/Engine/Source/Runtime/Engine/Private/WorldPartition/WorldPartitionActorDesc.cpp
--- a/Engine/Source/Runtime/Engine/Private/WorldPartition/WorldPartitionActorDesc.cpp
+++ b/Engine/Source/Runtime/Engine/Private/WorldPartition/WorldPartitionActorDesc.cpp
@@ -153,11 +153,21 @@ void FWorldPartitionActorDesc::SerializeTo(TArray<uint8>& OutData)
 	OutData.Append(PayloadData);
 }
 
-void FWorldPartitionActorDesc::TransformInstance(const FString& From, const FString& To)
+void FWorldPartitionActorDesc::TransformInstance(const FString& From, const FString& To, const FTransform& InstanceTransform)
 {
 	check(!HardRefCount);
 
 	ActorPath = *ActorPath.ToString().Replace(*From, *To);
+
+	// Transform BoundsLocation and BoundsExtent if necessary
+	if (!InstanceTransform.Equals(FTransform::Identity))
+	{
+		//@todo_ow: This will result in a new BoundsExtent that is larger than it should. To fix this, we would need the Object Oriented BoundingBox of the actor (the BV of the actor without rotation)
+		const FVector BoundsMin = BoundsLocation - BoundsExtent;
+		const FVector BoundsMax = BoundsLocation + BoundsExtent;
+		const FBox NewBounds = FBox(BoundsMin, BoundsMax).TransformBy(InstanceTransform);
+		NewBounds.GetCenterAndExtents(BoundsLocation, BoundsExtent);
+	}
 }
 
 FString FWorldPartitionActorDesc::ToString() const
diff --git a/Engine/Source/Runtime/Engine/Private/WorldPartition/WorldPartitionRuntimeLevelStreamingCell.cpp b/Engine/Source/Runtime/Engine/Private/WorldPartition/WorldPartitionRuntimeLevelStreamingCell.cpp
--- a/Engine/Source/Runtime/Engine/Private/WorldPartition/WorldPartitionRuntimeLevelStreamingCell.cpp
+++ b/Engine/Source/Runtime/Engine/Private/WorldPartition/WorldPartitionRuntimeLevelStreamingCell.cpp
@@ -125,7 +125,8 @@ UWorldPartitionLevelStreamingDynamic* UWorldPartitionRuntimeLevelStreamingCell::
 		FString PackageName = !InPackageName.IsEmpty() ? InPackageName : UWorldPartitionLevelStreamingPolicy::GetCellPackagePath(GetFName(), OuterWorld);
 		TSoftObjectPtr<UWorld> WorldAsset(FSoftObjectPath(FString::Printf(TEXT("%s.%s"), *PackageName, *OuterWorld->GetName())));
 		NewLevelStreaming->SetWorldAsset(WorldAsset);
-		NewLevelStreaming->LevelTransform = FTransform::Identity;
+		// Transfer WorldPartition's transform to Level
+		NewLevelStreaming->LevelTransform = WorldPartition->GetInstanceTransform();
 		NewLevelStreaming->bClientOnlyVisible = GetClientOnlyVisible();
 		NewLevelStreaming->Initialize(*this);
 
@@ -256,7 +257,8 @@ UWorldPartitionLevelStreamingDynamic* UWorldPartitionRuntimeLevelStreamingCell::
 			LevelStreaming->Rename(nullptr, OwningWorld);
 		}
 
-		LevelStreaming->LevelTransform = FTransform::Identity;
+		// Transfer WorldPartition's transform to LevelStreaming
+		LevelStreaming->LevelTransform = WorldPartition->GetInstanceTransform();
 
 		// When Partition outer level is an instance, make sure to also generate unique cell level instance name
 		ULevel* PartitionLevel = WorldPartition->GetTypedOuter<ULevel>();
diff --git a/Engine/Source/Runtime/Engine/Private/WorldPartition/WorldPartitionRuntimeSpatialHash.cpp b/Engine/Source/Runtime/Engine/Private/WorldPartition/WorldPartitionRuntimeSpatialHash.cpp
--- a/Engine/Source/Runtime/Engine/Private/WorldPartition/WorldPartitionRuntimeSpatialHash.cpp
+++ b/Engine/Source/Runtime/Engine/Private/WorldPartition/WorldPartitionRuntimeSpatialHash.cpp
@@ -293,9 +293,8 @@ EWorldPartitionRuntimeCellVisualizeMode FSpatialHashStreamingGrid::GetStreamingC
 	return VisualizeMode;
 }
 
-void FSpatialHashStreamingGrid::Draw3D(UWorld* World, const TArray<FWorldPartitionStreamingSource>& Sources) const
+void FSpatialHashStreamingGrid::Draw3D(UWorld* World, const TArray<FWorldPartitionStreamingSource>& Sources, const FTransform& Transform) const
 {
-	const FTransform& Transform = FTransform::Identity;
 	const EWorldPartitionRuntimeCellVisualizeMode VisualizeMode = GetStreamingCellVisualizeMode();
 	const UDataLayerSubsystem* DataLayerSubsystem = World->GetSubsystem<UDataLayerSubsystem>();
 	TMap<FName, FColor> DataLayerDebugColors;
@@ -1231,10 +1230,11 @@ void UWorldPartitionRuntimeSpatialHash::Draw3D(const TArray<FWorldPartitionStrea
 	UWorld* World = GetWorld();
 	UWorldPartition* WorldPartition = GetOuterUWorldPartition();
 
+	FTransform Transform = WorldPartition->GetInstanceTransform();
 	TArray<const FSpatialHashStreamingGrid*> FilteredStreamingGrids = GetFilteredStreamingGrids();
 	for (const FSpatialHashStreamingGrid* StreamingGrid : FilteredStreamingGrids)
 	{
-		StreamingGrid->Draw3D(World, Sources);
+		StreamingGrid->Draw3D(World, Sources, Transform);
 	}
 }
 
diff --git a/Engine/Source/Runtime/Engine/Private/WorldPartition/WorldPartitionStreamingPolicy.cpp b/Engine/Source/Runtime/Engine/Private/WorldPartition/WorldPartitionStreamingPolicy.cpp
--- a/Engine/Source/Runtime/Engine/Private/WorldPartition/WorldPartitionStreamingPolicy.cpp
+++ b/Engine/Source/Runtime/Engine/Private/WorldPartition/WorldPartitionStreamingPolicy.cpp
@@ -90,6 +90,8 @@ void UWorldPartitionStreamingPolicy::UpdateStreamingSources()
 		return;
 	}
 
+	const FTransform WorldToLocal = WorldPartition->GetInstanceTransform().Inverse();
+
 #if !UE_BUILD_SHIPPING
 	bool bUseReplaySources = false;
 	if (AWorldPartitionReplay* Replay = WorldPartition->Replay)
@@ -104,8 +106,8 @@ void UWorldPartitionStreamingPolicy::UpdateStreamingSources()
 		// We are in the SIE
 		if (GEnableSimulationStreamingSource && UWorldPartition::IsSimulating())
 		{
-			const FVector ViewLocation = GCurrentLevelEditingViewportClient->GetViewLocation();
-			const FRotator ViewRotation = GCurrentLevelEditingViewportClient->GetViewRotation();
+			const FVector ViewLocation = WorldToLocal.TransformPosition(GCurrentLevelEditingViewportClient->GetViewLocation());
+			const FRotator ViewRotation = WorldToLocal.TransformRotation(GCurrentLevelEditingViewportClient->GetViewRotation().Quaternion()).Rotator();
 			static const FName NAME_SIE(TEXT("SIE"));
 			StreamingSources.Add(FWorldPartitionStreamingSource(NAME_SIE, ViewLocation, ViewRotation, EStreamingSourceTargetState::Activated, /*bBlockOnSlowLoading=*/false, EStreamingSourcePriority::Default));
 		}
@@ -126,7 +128,9 @@ void UWorldPartitionStreamingPolicy::UpdateStreamingSources()
 						FRotator ViewRotation;
 						Player->PlayerController->GetPlayerViewPoint(ViewLocation, ViewRotation);
 
-				
+						// Transform to Local
+						ViewLocation = WorldToLocal.TransformPosition(ViewLocation);
+						ViewRotation = WorldToLocal.TransformRotation(ViewRotation.Quaternion()).Rotator();
 						const EStreamingSourceTargetState TargetState = Player->PlayerController->StreamingSourceShouldActivate() ? EStreamingSourceTargetState::Activated : EStreamingSourceTargetState::Loaded;
 						const bool bBlockOnSlowLoading = Player->PlayerController->StreamingSourceShouldBlockOnSlowStreaming();
 						StreamingSources.Add(FWorldPartitionStreamingSource(Player->PlayerController->GetFName(), ViewLocation, ViewRotation, TargetState, bBlockOnSlowLoading, EStreamingSourcePriority::Default));
@@ -142,6 +146,9 @@ void UWorldPartitionStreamingPolicy::UpdateStreamingSources()
 			StreamingSource.Priority = EStreamingSourcePriority::Low;
 			if (StreamingSourceProvider->GetStreamingSource(StreamingSource))
 			{
+				// Transform to Local
+				StreamingSource.Location = WorldToLocal.TransformPosition(StreamingSource.Location);
+				StreamingSource.Rotation = WorldToLocal.TransformRotation(StreamingSource.Rotation.Quaternion()).Rotator();
 				StreamingSources.Add(StreamingSource);
 			}
 		}
@@ -176,10 +183,11 @@ UE_SUPPRESS(LogWorldPartition, Verbosity, \
 	if (ToActivateCells.Num() > 0 || ToLoadCells.Num() > 0 || ToUnloadCells.Num() > 0) \
 	{ \
 		UE_LOG(LogWorldPartition, Verbosity, TEXT("UWorldPartitionStreamingPolicy: CellsToActivate(%d), CellsToLoad(%d), CellsToUnload(%d)"), ToActivateCells.Num(), ToLoadCells.Num(), ToUnloadCells.Num()); \
+		FTransform LocalToWorld = WorldPartition->GetInstanceTransform(); \
 		for (int i = 0; i < StreamingSources.Num(); ++i) \
 		{ \
-			FVector ViewLocation = StreamingSources[i].Location; \
-			FRotator ViewRotation = StreamingSources[i].Rotation; \
+			FVector ViewLocation = LocalToWorld.TransformPosition(StreamingSources[i].Location); \
+			FRotator ViewRotation = LocalToWorld.TransformRotation(StreamingSources[i].Rotation.Quaternion()).Rotator(); \
 			UE_LOG(LogWorldPartition, Verbosity, TEXT("UWorldPartitionStreamingPolicy: Sources[%d] = %s,%s"), i, *ViewLocation.ToString(), *ViewRotation.ToString()); \
 		} \
 	} \
diff --git a/Engine/Source/Runtime/Engine/Public/WorldPartition/WorldPartition.h b/Engine/Source/Runtime/Engine/Public/WorldPartition/WorldPartition.h
--- a/Engine/Source/Runtime/Engine/Public/WorldPartition/WorldPartition.h
+++ b/Engine/Source/Runtime/Engine/Public/WorldPartition/WorldPartition.h
@@ -169,6 +169,8 @@ public:
 	bool IsInitialized() const;
 	virtual void Uninitialize() override;
 
+	const FTransform& GetInstanceTransform() const { return InstanceTransform; }
+
 	void Tick(float DeltaSeconds);
 	void UpdateStreamingState();
 	bool CanAddLoadedLevelToWorld(class ULevel* InLevel) const;
@@ -221,6 +223,7 @@ public:
 
 private:
 	EWorldPartitionInitState InitState;
+	FTransform InstanceTransform;
 
 	UPROPERTY()
 	mutable TObjectPtr<UWorldPartitionStreamingPolicy> StreamingPolicy;
diff --git a/Engine/Source/Runtime/Engine/Public/WorldPartition/WorldPartitionActorDesc.h b/Engine/Source/Runtime/Engine/Public/WorldPartition/WorldPartitionActorDesc.h
--- a/Engine/Source/Runtime/Engine/Public/WorldPartition/WorldPartitionActorDesc.h
+++ b/Engine/Source/Runtime/Engine/Public/WorldPartition/WorldPartitionActorDesc.h
@@ -163,7 +163,7 @@ public:
 protected:
 	FWorldPartitionActorDesc();
 
-	void TransformInstance(const FString& From, const FString& To);
+	void TransformInstance(const FString& From, const FString& To, const FTransform& InstanceTransform);
 
 	virtual void TransferFrom(const FWorldPartitionActorDesc* From)
 	{
diff --git a/Engine/Source/Runtime/Engine/Public/WorldPartition/WorldPartitionRuntimeSpatialHash.h b/Engine/Source/Runtime/Engine/Public/WorldPartition/WorldPartitionRuntimeSpatialHash.h
--- a/Engine/Source/Runtime/Engine/Public/WorldPartition/WorldPartitionRuntimeSpatialHash.h
+++ b/Engine/Source/Runtime/Engine/Public/WorldPartition/WorldPartitionRuntimeSpatialHash.h
@@ -139,7 +139,7 @@ struct FSpatialHashStreamingGrid
 	void GetCells(const TArray<FWorldPartitionStreamingSource>& Sources, const class UDataLayerSubsystem* DataLayerSubsystem, UWorldPartitionRuntimeHash::FStreamingSourceCells& OutActivateCells, UWorldPartitionRuntimeHash::FStreamingSourceCells& OutLoadCells) const;
 	void GetAlwaysLoadedCells(const UDataLayerSubsystem* DataLayerSubsystem, TSet<const UWorldPartitionRuntimeCell*>& OutActivateCells, TSet<const UWorldPartitionRuntimeCell*>& OutLoadCells) const;
 	void Draw2D(UCanvas* Canvas, UWorld* World, const TArray<FWorldPartitionStreamingSource>& Sources, const FBox& Region, const FBox2D& GridScreenBounds, TFunctionRef<FVector2D(const FVector2D&)> WorldToScreen) const;
-	void Draw3D(UWorld* World, const TArray<FWorldPartitionStreamingSource>& Sources) const;
+	void Draw3D(UWorld* World, const TArray<FWorldPartitionStreamingSource>& Sources, const FTransform& Transform) const;
 
 private:
 	void DrawStreamingSource2D(UCanvas* Canvas, const FSphericalSector& Shape, TFunctionRef<FVector2D(const FVector2D&)> WorldToScreen, const FColor& Color) const;
