Completely hide StrataBlendMode when strata is not enabled.

diff --git a/Engine/Source/Runtime/Engine/Private/Materials/Material.cpp b/Engine/Source/Runtime/Engine/Private/Materials/Material.cpp
--- a/Engine/Source/Runtime/Engine/Private/Materials/Material.cpp
+++ b/Engine/Source/Runtime/Engine/Private/Materials/Material.cpp
@@ -982,6 +982,14 @@ UMaterial::UMaterial(const FObjectInitializer& ObjectInitializer)
 	PhysMaterialMask = nullptr;
 
 	FloatPrecisionMode = EMaterialFloatPrecisionMode::MFPM_Default;
+
+#if WITH_EDITOR
+	FProperty* StrataBlendModeProperty = FindFProperty< FProperty >(this->GetClass(), GET_MEMBER_NAME_CHECKED(UMaterial, StrataBlendMode));
+	if (StrataBlendModeProperty != nullptr && !Engine_IsStrataEnabled())
+	{
+		StrataBlendModeProperty->SetMetaData(TEXT("Category"), TEXT("NotVisible"));
+	}
+#endif
 }
 
 void UMaterial::PreSave(const class ITargetPlatform* TargetPlatform)
