Support many new epxression types with new HLSL translator
Remove 'ForwardValue' expressions, instead just forward values 'by hand' where needed

diff --git a/Engine/Shaders/Private/MaterialTexture.ush b/Engine/Shaders/Private/MaterialTexture.ush
--- a/Engine/Shaders/Private/MaterialTexture.ush
+++ b/Engine/Shaders/Private/MaterialTexture.ush
@@ -36,6 +36,41 @@ struct FTexture2D
 	int SamplerType; // EMaterialSamplerType
 };
 
+struct FTexture2DArray
+{
+	Texture2DArray Texture;
+	SamplerState Sampler;
+	int SamplerType; // EMaterialSamplerType
+};
+
+struct FTextureCube
+{
+	TextureCube Texture;
+	SamplerState Sampler;
+	int SamplerType; // EMaterialSamplerType
+};
+
+struct FTextureCubeArray
+{
+	TextureCubeArray Texture;
+	SamplerState Sampler;
+	int SamplerType; // EMaterialSamplerType
+};
+
+struct FTexture3D
+{
+	Texture3D Texture;
+	SamplerState Sampler;
+	int SamplerType; // EMaterialSamplerType
+};
+
+struct FTextureExternal
+{
+	TextureExternal Texture;
+	SamplerState Sampler;
+	int SamplerType; // EMaterialSamplerType
+};
+
 FTexture2D MakeTexture2D(Texture2D Texture, SamplerState Sampler, int SamplerType)
 {
 	FTexture2D Result;
@@ -44,6 +79,51 @@ FTexture2D MakeTexture2D(Texture2D Texture, SamplerState Sampler, int SamplerTyp
 	Result.SamplerType = SamplerType;
 	return Result;
 }
+
+FTexture2DArray MakeTexture2DArray(Texture2DArray Texture, SamplerState Sampler, int SamplerType)
+{
+	FTexture2DArray Result;
+	Result.Texture = Texture;
+	Result.Sampler = Sampler;
+	Result.SamplerType = SamplerType;
+	return Result;
+}
+
+FTextureCube MakeTexture2DArray(TextureCube Texture, SamplerState Sampler, int SamplerType)
+{
+	FTextureCube Result;
+	Result.Texture = Texture;
+	Result.Sampler = Sampler;
+	Result.SamplerType = SamplerType;
+	return Result;
+}
+
+FTextureCubeArray MakeTexture2DArray(TextureCubeArray Texture, SamplerState Sampler, int SamplerType)
+{
+	FTextureCubeArray Result;
+	Result.Texture = Texture;
+	Result.Sampler = Sampler;
+	Result.SamplerType = SamplerType;
+	return Result;
+}
+
+FTexture3D MakeTexture3D(Texture3D Texture, SamplerState Sampler, int SamplerType)
+{
+	FTexture3D Result;
+	Result.Texture = Texture;
+	Result.Sampler = Sampler;
+	Result.SamplerType = SamplerType;
+	return Result;
+}
+
+FTextureExternal MakeTextureExternal(TextureExternal Texture, SamplerState Sampler, int SamplerType)
+{
+	FTextureExternal Result;
+	Result.Texture = Texture;
+	Result.Sampler = Sampler;
+	Result.SamplerType = SamplerType;
+	return Result;
+}
 #endif // ENABLE_NEW_HLSL_GENERATOR
 
 MaterialFloat4 ProcessMaterialColorTextureLookup(MaterialFloat4 TextureValue)
diff --git a/Engine/Source/Runtime/Core/Public/Misc/MemStackUtility.h b/Engine/Source/Runtime/Core/Public/Misc/MemStackUtility.h
--- a/Engine/Source/Runtime/Core/Public/Misc/MemStackUtility.h
+++ b/Engine/Source/Runtime/Core/Public/Misc/MemStackUtility.h
@@ -43,6 +43,13 @@ inline const TCHAR* AllocateStringf(FMemStackBase& Allocator, const FormatType&
 	return AllocateString(Allocator, String);
 }
 
+template <typename FormatType, typename... ArgTypes>
+inline FStringView AllocateStringViewf(FMemStackBase& Allocator, const FormatType& Format, ArgTypes... Args)
+{
+	TStringBuilder<1024> String;
+	String.Appendf(Format, Forward<ArgTypes>(Args)...);
+	return AllocateStringView(Allocator, String.ToView());
+}
 
 template<typename T>
 inline TArrayView<T> AllocateArrayView(FMemStackBase& Allocator, TArrayView<T> View)
diff --git a/Engine/Source/Runtime/Engine/Classes/Materials/MaterialExpressionActorPositionWS.h b/Engine/Source/Runtime/Engine/Classes/Materials/MaterialExpressionActorPositionWS.h
--- a/Engine/Source/Runtime/Engine/Classes/Materials/MaterialExpressionActorPositionWS.h
+++ b/Engine/Source/Runtime/Engine/Classes/Materials/MaterialExpressionActorPositionWS.h
@@ -18,6 +18,7 @@ class UMaterialExpressionActorPositionWS : public UMaterialExpression
 #if WITH_EDITOR
 	virtual int32 Compile(class FMaterialCompiler* Compiler, int32 OutputIndex) override;
 	virtual void GetCaption(TArray<FString>& OutCaptions) const override;
+	virtual bool GenerateHLSLExpression(FMaterialHLSLGenerator& Generator, UE::HLSLTree::FScope& Scope, int32 OutputIndex, UE::HLSLTree::FExpression*& OutExpression) override;
 #endif
 	//~ End UMaterialExpression Interface
 };
diff --git a/Engine/Source/Runtime/Engine/Classes/Materials/MaterialExpressionArccosine.h b/Engine/Source/Runtime/Engine/Classes/Materials/MaterialExpressionArccosine.h
--- a/Engine/Source/Runtime/Engine/Classes/Materials/MaterialExpressionArccosine.h
+++ b/Engine/Source/Runtime/Engine/Classes/Materials/MaterialExpressionArccosine.h
@@ -22,6 +22,7 @@ class UMaterialExpressionArccosine : public UMaterialExpression
 	virtual void GetCaption(TArray<FString>& OutCaptions) const override;
 	virtual void GetExpressionToolTip(TArray<FString>& OutToolTip) override;
 	virtual FText GetKeywords() const override {return FText::FromString(TEXT("acos"));}
+	virtual bool GenerateHLSLExpression(FMaterialHLSLGenerator& Generator, UE::HLSLTree::FScope& Scope, int32 OutputIndex, UE::HLSLTree::FExpression*& OutExpression) override;
 #endif
 	//~ End UMaterialExpression Interface
 };
diff --git a/Engine/Source/Runtime/Engine/Classes/Materials/MaterialExpressionArccosineFast.h b/Engine/Source/Runtime/Engine/Classes/Materials/MaterialExpressionArccosineFast.h
--- a/Engine/Source/Runtime/Engine/Classes/Materials/MaterialExpressionArccosineFast.h
+++ b/Engine/Source/Runtime/Engine/Classes/Materials/MaterialExpressionArccosineFast.h
@@ -22,6 +22,7 @@ class UMaterialExpressionArccosineFast : public UMaterialExpression
 	virtual void GetCaption(TArray<FString>& OutCaptions) const override;
 	virtual void GetExpressionToolTip(TArray<FString>& OutToolTip) override;
 	virtual FText GetKeywords() const override {return FText::FromString(TEXT("acosfast"));}
+	virtual bool GenerateHLSLExpression(FMaterialHLSLGenerator& Generator, UE::HLSLTree::FScope& Scope, int32 OutputIndex, UE::HLSLTree::FExpression*& OutExpression) override;
 #endif
 	//~ End UMaterialExpression Interface
 };
diff --git a/Engine/Source/Runtime/Engine/Classes/Materials/MaterialExpressionArcsine.h b/Engine/Source/Runtime/Engine/Classes/Materials/MaterialExpressionArcsine.h
--- a/Engine/Source/Runtime/Engine/Classes/Materials/MaterialExpressionArcsine.h
+++ b/Engine/Source/Runtime/Engine/Classes/Materials/MaterialExpressionArcsine.h
@@ -22,6 +22,7 @@ class UMaterialExpressionArcsine : public UMaterialExpression
 	virtual void GetCaption(TArray<FString>& OutCaptions) const override;
 	virtual void GetExpressionToolTip(TArray<FString>& OutToolTip) override;
 	virtual FText GetKeywords() const override {return FText::FromString(TEXT("asin"));}
+	virtual bool GenerateHLSLExpression(FMaterialHLSLGenerator& Generator, UE::HLSLTree::FScope& Scope, int32 OutputIndex, UE::HLSLTree::FExpression*& OutExpression) override;
 #endif
 	//~ End UMaterialExpression Interface
 };
diff --git a/Engine/Source/Runtime/Engine/Classes/Materials/MaterialExpressionArcsineFast.h b/Engine/Source/Runtime/Engine/Classes/Materials/MaterialExpressionArcsineFast.h
--- a/Engine/Source/Runtime/Engine/Classes/Materials/MaterialExpressionArcsineFast.h
+++ b/Engine/Source/Runtime/Engine/Classes/Materials/MaterialExpressionArcsineFast.h
@@ -22,6 +22,7 @@ class UMaterialExpressionArcsineFast : public UMaterialExpression
 	virtual void GetCaption(TArray<FString>& OutCaptions) const override;
 	virtual void GetExpressionToolTip(TArray<FString>& OutToolTip) override;
 	virtual FText GetKeywords() const override {return FText::FromString(TEXT("asinfast"));}
+	virtual bool GenerateHLSLExpression(FMaterialHLSLGenerator& Generator, UE::HLSLTree::FScope& Scope, int32 OutputIndex, UE::HLSLTree::FExpression*& OutExpression) override;
 #endif
 	//~ End UMaterialExpression Interface
 };
diff --git a/Engine/Source/Runtime/Engine/Classes/Materials/MaterialExpressionArctangent.h b/Engine/Source/Runtime/Engine/Classes/Materials/MaterialExpressionArctangent.h
--- a/Engine/Source/Runtime/Engine/Classes/Materials/MaterialExpressionArctangent.h
+++ b/Engine/Source/Runtime/Engine/Classes/Materials/MaterialExpressionArctangent.h
@@ -22,6 +22,7 @@ class UMaterialExpressionArctangent : public UMaterialExpression
 	virtual void GetCaption(TArray<FString>& OutCaptions) const override;
 	virtual void GetExpressionToolTip(TArray<FString>& OutToolTip) override;
 	virtual FText GetKeywords() const override {return FText::FromString(TEXT("atan"));}
+	virtual bool GenerateHLSLExpression(FMaterialHLSLGenerator& Generator, UE::HLSLTree::FScope& Scope, int32 OutputIndex, UE::HLSLTree::FExpression*& OutExpression) override;
 #endif
 	//~ End UMaterialExpression Interface
 };
diff --git a/Engine/Source/Runtime/Engine/Classes/Materials/MaterialExpressionArctangent2.h b/Engine/Source/Runtime/Engine/Classes/Materials/MaterialExpressionArctangent2.h
--- a/Engine/Source/Runtime/Engine/Classes/Materials/MaterialExpressionArctangent2.h
+++ b/Engine/Source/Runtime/Engine/Classes/Materials/MaterialExpressionArctangent2.h
@@ -25,6 +25,7 @@ class UMaterialExpressionArctangent2 : public UMaterialExpression
 	virtual void GetCaption(TArray<FString>& OutCaptions) const override;
 	virtual void GetExpressionToolTip(TArray<FString>& OutToolTip) override;
 	virtual FText GetKeywords() const override {return FText::FromString(TEXT("atan2"));}
+	virtual bool GenerateHLSLExpression(FMaterialHLSLGenerator& Generator, UE::HLSLTree::FScope& Scope, int32 OutputIndex, UE::HLSLTree::FExpression*& OutExpression) override;
 #endif
 	//~ End UMaterialExpression Interface
 };
diff --git a/Engine/Source/Runtime/Engine/Classes/Materials/MaterialExpressionArctangent2Fast.h b/Engine/Source/Runtime/Engine/Classes/Materials/MaterialExpressionArctangent2Fast.h
--- a/Engine/Source/Runtime/Engine/Classes/Materials/MaterialExpressionArctangent2Fast.h
+++ b/Engine/Source/Runtime/Engine/Classes/Materials/MaterialExpressionArctangent2Fast.h
@@ -25,6 +25,7 @@ class UMaterialExpressionArctangent2Fast : public UMaterialExpression
 	virtual void GetCaption(TArray<FString>& OutCaptions) const override;
 	virtual void GetExpressionToolTip(TArray<FString>& OutToolTip) override;
 	virtual FText GetKeywords() const override {return FText::FromString(TEXT("atan2fast"));}
+	virtual bool GenerateHLSLExpression(FMaterialHLSLGenerator& Generator, UE::HLSLTree::FScope& Scope, int32 OutputIndex, UE::HLSLTree::FExpression*& OutExpression) override;
 #endif
 	//~ End UMaterialExpression Interface
 };
diff --git a/Engine/Source/Runtime/Engine/Classes/Materials/MaterialExpressionArctangentFast.h b/Engine/Source/Runtime/Engine/Classes/Materials/MaterialExpressionArctangentFast.h
--- a/Engine/Source/Runtime/Engine/Classes/Materials/MaterialExpressionArctangentFast.h
+++ b/Engine/Source/Runtime/Engine/Classes/Materials/MaterialExpressionArctangentFast.h
@@ -22,6 +22,7 @@ class UMaterialExpressionArctangentFast : public UMaterialExpression
 	virtual void GetCaption(TArray<FString>& OutCaptions) const override;
 	virtual void GetExpressionToolTip(TArray<FString>& OutToolTip) override;
 	virtual FText GetKeywords() const override {return FText::FromString(TEXT("atanfast"));}
+	virtual bool GenerateHLSLExpression(FMaterialHLSLGenerator& Generator, UE::HLSLTree::FScope& Scope, int32 OutputIndex, UE::HLSLTree::FExpression*& OutExpression) override;
 #endif
 	//~ End UMaterialExpression Interface
 };
diff --git a/Engine/Source/Runtime/Engine/Classes/Materials/MaterialExpressionBreakMaterialAttributes.h b/Engine/Source/Runtime/Engine/Classes/Materials/MaterialExpressionBreakMaterialAttributes.h
--- a/Engine/Source/Runtime/Engine/Classes/Materials/MaterialExpressionBreakMaterialAttributes.h
+++ b/Engine/Source/Runtime/Engine/Classes/Materials/MaterialExpressionBreakMaterialAttributes.h
@@ -31,6 +31,11 @@ class UMaterialExpressionBreakMaterialAttributes : public UMaterialExpression
 	virtual bool IsInputConnectionRequired(int32 InputIndex) const override;
 	virtual uint32 GetInputType(int32 InputIndex) override {return MCT_MaterialAttributes;}
 	virtual uint32 GetOutputType(int32 OutputIndex) override;
+	virtual bool GenerateHLSLExpression(FMaterialHLSLGenerator& Generator, UE::HLSLTree::FScope& Scope, int32 OutputIndex, UE::HLSLTree::FExpression*& OutExpression) override;
+
+	static TMap<EMaterialProperty, int32> PropertyToIOIndexMap;
+	static void BuildPropertyToIOIndexMap();
+
 #endif // WITH_EDITOR
 	//~ End UMaterialExpression Interface
 };
diff --git a/Engine/Source/Runtime/Engine/Classes/Materials/MaterialExpressionCameraVectorWS.h b/Engine/Source/Runtime/Engine/Classes/Materials/MaterialExpressionCameraVectorWS.h
--- a/Engine/Source/Runtime/Engine/Classes/Materials/MaterialExpressionCameraVectorWS.h
+++ b/Engine/Source/Runtime/Engine/Classes/Materials/MaterialExpressionCameraVectorWS.h
@@ -18,6 +18,7 @@ class UMaterialExpressionCameraVectorWS : public UMaterialExpression
 #if WITH_EDITOR
 	virtual int32 Compile(class FMaterialCompiler* Compiler, int32 OutputIndex) override;
 	virtual void GetCaption(TArray<FString>& OutCaptions) const override;
+	virtual bool GenerateHLSLExpression(FMaterialHLSLGenerator& Generator, UE::HLSLTree::FScope& Scope, int32 OutputIndex, UE::HLSLTree::FExpression*& OutExpression) override;
 #endif
 	//~ End UMaterialExpression Interface
 };
diff --git a/Engine/Source/Runtime/Engine/Classes/Materials/MaterialExpressionCloudLayer.h b/Engine/Source/Runtime/Engine/Classes/Materials/MaterialExpressionCloudLayer.h
--- a/Engine/Source/Runtime/Engine/Classes/Materials/MaterialExpressionCloudLayer.h
+++ b/Engine/Source/Runtime/Engine/Classes/Materials/MaterialExpressionCloudLayer.h
@@ -17,6 +17,7 @@ class UMaterialExpressionCloudSampleAttribute : public UMaterialExpression
 	virtual int32 Compile(class FMaterialCompiler* Compiler, int32 OutputIndex) override;
 	virtual void GetCaption(TArray<FString>& OutCaptions) const override;
 	virtual void GetExpressionToolTip(TArray<FString>& OutToolTip) override;
+	virtual bool GenerateHLSLExpression(FMaterialHLSLGenerator& Generator, UE::HLSLTree::FScope& Scope, int32 OutputIndex, UE::HLSLTree::FExpression*& OutExpression) override;
 #endif
 	//~ End UMaterialExpression Interface
 };
diff --git a/Engine/Source/Runtime/Engine/Classes/Materials/MaterialExpressionConstantBiasScale.h b/Engine/Source/Runtime/Engine/Classes/Materials/MaterialExpressionConstantBiasScale.h
--- a/Engine/Source/Runtime/Engine/Classes/Materials/MaterialExpressionConstantBiasScale.h
+++ b/Engine/Source/Runtime/Engine/Classes/Materials/MaterialExpressionConstantBiasScale.h
@@ -28,6 +28,7 @@ class UMaterialExpressionConstantBiasScale : public UMaterialExpression
 #if WITH_EDITOR
 	virtual int32 Compile(class FMaterialCompiler* Compiler, int32 OutputIndex) override;
 	virtual void GetCaption(TArray<FString>& OutCaptions) const override;
+	virtual bool GenerateHLSLExpression(FMaterialHLSLGenerator& Generator, UE::HLSLTree::FScope& Scope, int32 OutputIndex, UE::HLSLTree::FExpression*& OutExpression) override;
 #endif
 	//~ End UMaterialExpression Interface
 };
diff --git a/Engine/Source/Runtime/Engine/Classes/Materials/MaterialExpressionCosine.h b/Engine/Source/Runtime/Engine/Classes/Materials/MaterialExpressionCosine.h
--- a/Engine/Source/Runtime/Engine/Classes/Materials/MaterialExpressionCosine.h
+++ b/Engine/Source/Runtime/Engine/Classes/Materials/MaterialExpressionCosine.h
@@ -26,6 +26,7 @@ class UMaterialExpressionCosine : public UMaterialExpression
 	virtual int32 Compile(class FMaterialCompiler* Compiler, int32 OutputIndex) override;
 	virtual void GetCaption(TArray<FString>& OutCaptions) const override;
 	virtual FText GetKeywords() const override {return FText::FromString(TEXT("cos"));}
+	virtual bool GenerateHLSLExpression(FMaterialHLSLGenerator& Generator, UE::HLSLTree::FScope& Scope, int32 OutputIndex, UE::HLSLTree::FExpression*& OutExpression) override;
 #endif
 	//~ End UMaterialExpression Interface
 };
diff --git a/Engine/Source/Runtime/Engine/Classes/Materials/MaterialExpressionCrossProduct.h b/Engine/Source/Runtime/Engine/Classes/Materials/MaterialExpressionCrossProduct.h
--- a/Engine/Source/Runtime/Engine/Classes/Materials/MaterialExpressionCrossProduct.h
+++ b/Engine/Source/Runtime/Engine/Classes/Materials/MaterialExpressionCrossProduct.h
@@ -25,6 +25,7 @@ class UMaterialExpressionCrossProduct : public UMaterialExpression
 #if WITH_EDITOR
 	virtual int32 Compile(class FMaterialCompiler* Compiler, int32 OutputIndex) override;
 	virtual void GetCaption(TArray<FString>& OutCaptions) const override;
+	virtual bool GenerateHLSLExpression(FMaterialHLSLGenerator& Generator, UE::HLSLTree::FScope& Scope, int32 OutputIndex, UE::HLSLTree::FExpression*& OutExpression) override;
 #endif
 	//~ End UMaterialExpression Interface
 };
diff --git a/Engine/Source/Runtime/Engine/Classes/Materials/MaterialExpressionDDX.h b/Engine/Source/Runtime/Engine/Classes/Materials/MaterialExpressionDDX.h
--- a/Engine/Source/Runtime/Engine/Classes/Materials/MaterialExpressionDDX.h
+++ b/Engine/Source/Runtime/Engine/Classes/Materials/MaterialExpressionDDX.h
@@ -23,6 +23,7 @@ class UMaterialExpressionDDX : public UMaterialExpression
 #if WITH_EDITOR
 	virtual int32 Compile(class FMaterialCompiler* Compiler, int32 OutputIndex) override;
 	virtual void GetCaption(TArray<FString>& OutCaptions) const override;
+	virtual bool GenerateHLSLExpression(FMaterialHLSLGenerator& Generator, UE::HLSLTree::FScope& Scope, int32 OutputIndex, UE::HLSLTree::FExpression*& OutExpression) override;
 #endif
 	//~ End UMaterialExpression Interface
 };
diff --git a/Engine/Source/Runtime/Engine/Classes/Materials/MaterialExpressionDDY.h b/Engine/Source/Runtime/Engine/Classes/Materials/MaterialExpressionDDY.h
--- a/Engine/Source/Runtime/Engine/Classes/Materials/MaterialExpressionDDY.h
+++ b/Engine/Source/Runtime/Engine/Classes/Materials/MaterialExpressionDDY.h
@@ -23,6 +23,7 @@ class UMaterialExpressionDDY : public UMaterialExpression
 #if WITH_EDITOR
 	virtual int32 Compile(class FMaterialCompiler* Compiler, int32 OutputIndex) override;
 	virtual void GetCaption(TArray<FString>& OutCaptions) const override;
+	virtual bool GenerateHLSLExpression(FMaterialHLSLGenerator& Generator, UE::HLSLTree::FScope& Scope, int32 OutputIndex, UE::HLSLTree::FExpression*& OutExpression) override;
 #endif
 	//~ End UMaterialExpression Interface
 };
diff --git a/Engine/Source/Runtime/Engine/Classes/Materials/MaterialExpressionDesaturation.h b/Engine/Source/Runtime/Engine/Classes/Materials/MaterialExpressionDesaturation.h
--- a/Engine/Source/Runtime/Engine/Classes/Materials/MaterialExpressionDesaturation.h
+++ b/Engine/Source/Runtime/Engine/Classes/Materials/MaterialExpressionDesaturation.h
@@ -32,6 +32,7 @@ class UMaterialExpressionDesaturation : public UMaterialExpression
 	{
 		OutCaptions.Add(TEXT("Desaturation"));
 	}
+	virtual bool GenerateHLSLExpression(FMaterialHLSLGenerator& Generator, UE::HLSLTree::FScope& Scope, int32 OutputIndex, UE::HLSLTree::FExpression*& OutExpression) override;
 #endif
 	//~ End UMaterialExpression Interface
 };
diff --git a/Engine/Source/Runtime/Engine/Classes/Materials/MaterialExpressionFmod.h b/Engine/Source/Runtime/Engine/Classes/Materials/MaterialExpressionFmod.h
--- a/Engine/Source/Runtime/Engine/Classes/Materials/MaterialExpressionFmod.h
+++ b/Engine/Source/Runtime/Engine/Classes/Materials/MaterialExpressionFmod.h
@@ -26,6 +26,7 @@ class UMaterialExpressionFmod : public UMaterialExpression
 	virtual int32 Compile(class FMaterialCompiler* Compiler, int32 OutputIndex) override;
 	virtual void GetCaption(TArray<FString>& OutCaptions) const override;
 	virtual FText GetKeywords() const override {return FText::FromString(TEXT("%"));}
+	virtual bool GenerateHLSLExpression(FMaterialHLSLGenerator& Generator, UE::HLSLTree::FScope& Scope, int32 OutputIndex, UE::HLSLTree::FExpression*& OutExpression) override;
 #endif // WITH_EDITOR
 	//~ End UMaterialExpression Interface
 };
diff --git a/Engine/Source/Runtime/Engine/Classes/Materials/MaterialExpressionFresnel.h b/Engine/Source/Runtime/Engine/Classes/Materials/MaterialExpressionFresnel.h
--- a/Engine/Source/Runtime/Engine/Classes/Materials/MaterialExpressionFresnel.h
+++ b/Engine/Source/Runtime/Engine/Classes/Materials/MaterialExpressionFresnel.h
@@ -49,6 +49,7 @@ class UMaterialExpressionFresnel : public UMaterialExpression
 	{
 		OutCaptions.Add(FString(TEXT("Fresnel")));
 	}
+	virtual bool GenerateHLSLExpression(FMaterialHLSLGenerator& Generator, UE::HLSLTree::FScope& Scope, int32 OutputIndex, UE::HLSLTree::FExpression*& OutExpression) override;
 #endif
 	//~ End UMaterialExpression Interface
 };
diff --git a/Engine/Source/Runtime/Engine/Classes/Materials/MaterialExpressionHairAttributes.h b/Engine/Source/Runtime/Engine/Classes/Materials/MaterialExpressionHairAttributes.h
--- a/Engine/Source/Runtime/Engine/Classes/Materials/MaterialExpressionHairAttributes.h
+++ b/Engine/Source/Runtime/Engine/Classes/Materials/MaterialExpressionHairAttributes.h
@@ -22,6 +22,7 @@ class UMaterialExpressionHairAttributes : public UMaterialExpression
 #if WITH_EDITOR
 	virtual int32 Compile(class FMaterialCompiler* Compiler, int32 OutputIndex) override;
 	virtual void GetCaption(TArray<FString>& OutCaptions) const override;
+	virtual bool GenerateHLSLExpression(FMaterialHLSLGenerator& Generator, UE::HLSLTree::FScope& Scope, int32 OutputIndex, UE::HLSLTree::FExpression*& OutExpression) override;
 #endif
 	//~ End UMaterialExpression Interface
 };
diff --git a/Engine/Source/Runtime/Engine/Classes/Materials/MaterialExpressionLightmapUVs.h b/Engine/Source/Runtime/Engine/Classes/Materials/MaterialExpressionLightmapUVs.h
--- a/Engine/Source/Runtime/Engine/Classes/Materials/MaterialExpressionLightmapUVs.h
+++ b/Engine/Source/Runtime/Engine/Classes/Materials/MaterialExpressionLightmapUVs.h
@@ -22,6 +22,7 @@ class UMaterialExpressionLightmapUVs : public UMaterialExpression
 #if WITH_EDITOR
 	virtual int32 Compile(class FMaterialCompiler* Compiler, int32 OutputIndex) override;
 	virtual void GetCaption(TArray<FString>& OutCaptions) const override;
+	virtual bool GenerateHLSLExpression(FMaterialHLSLGenerator& Generator, UE::HLSLTree::FScope& Scope, int32 OutputIndex, UE::HLSLTree::FExpression*& OutExpression) override;
 #endif
 	//~ End UMaterialExpression Interface
 };
diff --git a/Engine/Source/Runtime/Engine/Classes/Materials/MaterialExpressionLogarithm10.h b/Engine/Source/Runtime/Engine/Classes/Materials/MaterialExpressionLogarithm10.h
--- a/Engine/Source/Runtime/Engine/Classes/Materials/MaterialExpressionLogarithm10.h
+++ b/Engine/Source/Runtime/Engine/Classes/Materials/MaterialExpressionLogarithm10.h
@@ -23,6 +23,7 @@ class UMaterialExpressionLogarithm10 : public UMaterialExpression
 	virtual void GetCaption(TArray<FString>& OutCaptions) const override;
 	virtual void GetExpressionToolTip(TArray<FString>& OutToolTip) override;
 	virtual FText GetKeywords() const override {return FText::FromString(TEXT("log10"));}
+	virtual bool GenerateHLSLExpression(FMaterialHLSLGenerator& Generator, UE::HLSLTree::FScope& Scope, int32 OutputIndex, UE::HLSLTree::FExpression*& OutExpression) override;
 #endif
 	//~ End UMaterialExpression Interface
 
diff --git a/Engine/Source/Runtime/Engine/Classes/Materials/MaterialExpressionLogarithm2.h b/Engine/Source/Runtime/Engine/Classes/Materials/MaterialExpressionLogarithm2.h
--- a/Engine/Source/Runtime/Engine/Classes/Materials/MaterialExpressionLogarithm2.h
+++ b/Engine/Source/Runtime/Engine/Classes/Materials/MaterialExpressionLogarithm2.h
@@ -23,6 +23,7 @@ class UMaterialExpressionLogarithm2 : public UMaterialExpression
 	virtual void GetCaption(TArray<FString>& OutCaptions) const override;
 	virtual void GetExpressionToolTip(TArray<FString>& OutToolTip) override;
 	virtual FText GetKeywords() const override {return FText::FromString(TEXT("log2"));}
+	virtual bool GenerateHLSLExpression(FMaterialHLSLGenerator& Generator, UE::HLSLTree::FScope& Scope, int32 OutputIndex, UE::HLSLTree::FExpression*& OutExpression) override;
 #endif
 	//~ End UMaterialExpression Interface
 
diff --git a/Engine/Source/Runtime/Engine/Classes/Materials/MaterialExpressionMakeMaterialAttributes.h b/Engine/Source/Runtime/Engine/Classes/Materials/MaterialExpressionMakeMaterialAttributes.h
--- a/Engine/Source/Runtime/Engine/Classes/Materials/MaterialExpressionMakeMaterialAttributes.h
+++ b/Engine/Source/Runtime/Engine/Classes/Materials/MaterialExpressionMakeMaterialAttributes.h
@@ -81,6 +81,7 @@ class UMaterialExpressionMakeMaterialAttributes : public UMaterialExpression
 	virtual void GetCaption(TArray<FString>& OutCaptions) const override;
 	virtual bool IsResultMaterialAttributes(int32 OutputIndex) override {return true;}
 	virtual uint32 GetInputType(int32 InputIndex) override;
+	virtual bool GenerateHLSLExpression(FMaterialHLSLGenerator& Generator, UE::HLSLTree::FScope& Scope, int32 OutputIndex, UE::HLSLTree::FExpression*& OutExpression) override;
 #endif // WITH_EDITOR
 	//~ End UMaterialExpression Interface
 };
diff --git a/Engine/Source/Runtime/Engine/Classes/Materials/MaterialExpressionNoise.h b/Engine/Source/Runtime/Engine/Classes/Materials/MaterialExpressionNoise.h
--- a/Engine/Source/Runtime/Engine/Classes/Materials/MaterialExpressionNoise.h
+++ b/Engine/Source/Runtime/Engine/Classes/Materials/MaterialExpressionNoise.h
@@ -111,6 +111,7 @@ class UMaterialExpressionNoise : public UMaterialExpression
 	virtual bool CanEditChange(const FProperty* InProperty) const override;
 	virtual int32 Compile(class FMaterialCompiler* Compiler, int32 OutputIndex) override;
 	virtual void GetCaption(TArray<FString>& OutCaptions) const override;
+	virtual bool GenerateHLSLExpression(FMaterialHLSLGenerator& Generator, UE::HLSLTree::FScope& Scope, int32 OutputIndex, UE::HLSLTree::FExpression*& OutExpression) override;
 #endif
 	//~ End UMaterialExpression Interface
 };
diff --git a/Engine/Source/Runtime/Engine/Classes/Materials/MaterialExpressionObjectBounds.h b/Engine/Source/Runtime/Engine/Classes/Materials/MaterialExpressionObjectBounds.h
--- a/Engine/Source/Runtime/Engine/Classes/Materials/MaterialExpressionObjectBounds.h
+++ b/Engine/Source/Runtime/Engine/Classes/Materials/MaterialExpressionObjectBounds.h
@@ -18,6 +18,7 @@ class UMaterialExpressionObjectBounds : public UMaterialExpression
 #if WITH_EDITOR
 	virtual int32 Compile(class FMaterialCompiler* Compiler, int32 OutputIndex) override;
 	virtual void GetCaption(TArray<FString>& OutCaptions) const override;
+	virtual bool GenerateHLSLExpression(FMaterialHLSLGenerator& Generator, UE::HLSLTree::FScope& Scope, int32 OutputIndex, UE::HLSLTree::FExpression*& OutExpression) override;
 #endif
 	//~ End UMaterialExpression Interface
 };
diff --git a/Engine/Source/Runtime/Engine/Classes/Materials/MaterialExpressionObjectOrientation.h b/Engine/Source/Runtime/Engine/Classes/Materials/MaterialExpressionObjectOrientation.h
--- a/Engine/Source/Runtime/Engine/Classes/Materials/MaterialExpressionObjectOrientation.h
+++ b/Engine/Source/Runtime/Engine/Classes/Materials/MaterialExpressionObjectOrientation.h
@@ -18,6 +18,7 @@ class ENGINE_API UMaterialExpressionObjectOrientation : public UMaterialExpressi
 #if WITH_EDITOR
 	virtual int32 Compile(class FMaterialCompiler* Compiler, int32 OutputIndex) override;
 	virtual void GetCaption(TArray<FString>& OutCaptions) const override;
+	virtual bool GenerateHLSLExpression(FMaterialHLSLGenerator& Generator, UE::HLSLTree::FScope& Scope, int32 OutputIndex, UE::HLSLTree::FExpression*& OutExpression) override;
 #endif
 	//~ End UMaterialExpression Interface
 };
diff --git a/Engine/Source/Runtime/Engine/Classes/Materials/MaterialExpressionObjectPositionWS.h b/Engine/Source/Runtime/Engine/Classes/Materials/MaterialExpressionObjectPositionWS.h
--- a/Engine/Source/Runtime/Engine/Classes/Materials/MaterialExpressionObjectPositionWS.h
+++ b/Engine/Source/Runtime/Engine/Classes/Materials/MaterialExpressionObjectPositionWS.h
@@ -18,6 +18,7 @@ class UMaterialExpressionObjectPositionWS : public UMaterialExpression
 #if WITH_EDITOR
 	virtual int32 Compile(class FMaterialCompiler* Compiler, int32 OutputIndex) override;
 	virtual void GetCaption(TArray<FString>& OutCaptions) const override;
+	virtual bool GenerateHLSLExpression(FMaterialHLSLGenerator& Generator, UE::HLSLTree::FScope& Scope, int32 OutputIndex, UE::HLSLTree::FExpression*& OutExpression) override;
 #endif
 	//~ End UMaterialExpression Interface
 };
diff --git a/Engine/Source/Runtime/Engine/Classes/Materials/MaterialExpressionObjectRadius.h b/Engine/Source/Runtime/Engine/Classes/Materials/MaterialExpressionObjectRadius.h
--- a/Engine/Source/Runtime/Engine/Classes/Materials/MaterialExpressionObjectRadius.h
+++ b/Engine/Source/Runtime/Engine/Classes/Materials/MaterialExpressionObjectRadius.h
@@ -18,6 +18,7 @@ class UMaterialExpressionObjectRadius : public UMaterialExpression
 #if WITH_EDITOR
 	virtual int32 Compile(class FMaterialCompiler* Compiler, int32 OutputIndex) override;
 	virtual void GetCaption(TArray<FString>& OutCaptions) const override;
+	virtual bool GenerateHLSLExpression(FMaterialHLSLGenerator& Generator, UE::HLSLTree::FScope& Scope, int32 OutputIndex, UE::HLSLTree::FExpression*& OutExpression) override;
 #endif
 	//~ End UMaterialExpression Interface
 };
diff --git a/Engine/Source/Runtime/Engine/Classes/Materials/MaterialExpressionParticleColor.h b/Engine/Source/Runtime/Engine/Classes/Materials/MaterialExpressionParticleColor.h
--- a/Engine/Source/Runtime/Engine/Classes/Materials/MaterialExpressionParticleColor.h
+++ b/Engine/Source/Runtime/Engine/Classes/Materials/MaterialExpressionParticleColor.h
@@ -19,6 +19,7 @@ class UMaterialExpressionParticleColor : public UMaterialExpression
 	virtual int32 Compile(class FMaterialCompiler* Compiler, int32 OutputIndex) override;
 	virtual void GetCaption(TArray<FString>& OutCaptions) const override;
 	virtual FText GetKeywords() const override { return FText::FromString(TEXT("particle color")); }
+	virtual bool GenerateHLSLExpression(FMaterialHLSLGenerator& Generator, UE::HLSLTree::FScope& Scope, int32 OutputIndex, UE::HLSLTree::FExpression*& OutExpression) override;
 #endif
 	//~ End UMaterialExpression Interface
 };
diff --git a/Engine/Source/Runtime/Engine/Classes/Materials/MaterialExpressionParticlePositionWS.h b/Engine/Source/Runtime/Engine/Classes/Materials/MaterialExpressionParticlePositionWS.h
--- a/Engine/Source/Runtime/Engine/Classes/Materials/MaterialExpressionParticlePositionWS.h
+++ b/Engine/Source/Runtime/Engine/Classes/Materials/MaterialExpressionParticlePositionWS.h
@@ -18,6 +18,7 @@ class UMaterialExpressionParticlePositionWS : public UMaterialExpression
 #if WITH_EDITOR
 	virtual int32 Compile(class FMaterialCompiler* Compiler, int32 OutputIndex) override;
 	virtual void GetCaption(TArray<FString>& OutCaptions) const override;
+	virtual bool GenerateHLSLExpression(FMaterialHLSLGenerator& Generator, UE::HLSLTree::FScope& Scope, int32 OutputIndex, UE::HLSLTree::FExpression*& OutExpression) override;
 #endif
 	//~ End UMaterialExpression Interface
 };
diff --git a/Engine/Source/Runtime/Engine/Classes/Materials/MaterialExpressionParticleRadius.h b/Engine/Source/Runtime/Engine/Classes/Materials/MaterialExpressionParticleRadius.h
--- a/Engine/Source/Runtime/Engine/Classes/Materials/MaterialExpressionParticleRadius.h
+++ b/Engine/Source/Runtime/Engine/Classes/Materials/MaterialExpressionParticleRadius.h
@@ -18,6 +18,7 @@ class UMaterialExpressionParticleRadius : public UMaterialExpression
 #if WITH_EDITOR
 	virtual int32 Compile(class FMaterialCompiler* Compiler, int32 OutputIndex) override;
 	virtual void GetCaption(TArray<FString>& OutCaptions) const override;
+	virtual bool GenerateHLSLExpression(FMaterialHLSLGenerator& Generator, UE::HLSLTree::FScope& Scope, int32 OutputIndex, UE::HLSLTree::FExpression*& OutExpression) override;
 #endif
 	//~ End UMaterialExpression Interface
 };
diff --git a/Engine/Source/Runtime/Engine/Classes/Materials/MaterialExpressionPerInstanceFadeAmount.h b/Engine/Source/Runtime/Engine/Classes/Materials/MaterialExpressionPerInstanceFadeAmount.h
--- a/Engine/Source/Runtime/Engine/Classes/Materials/MaterialExpressionPerInstanceFadeAmount.h
+++ b/Engine/Source/Runtime/Engine/Classes/Materials/MaterialExpressionPerInstanceFadeAmount.h
@@ -18,6 +18,7 @@ class UMaterialExpressionPerInstanceFadeAmount : public UMaterialExpression
 #if WITH_EDITOR
 	virtual int32 Compile(class FMaterialCompiler* Compiler, int32 OutputIndex) override;
 	virtual void GetCaption(TArray<FString>& OutCaptions) const override;
+	virtual bool GenerateHLSLExpression(FMaterialHLSLGenerator& Generator, UE::HLSLTree::FScope& Scope, int32 OutputIndex, UE::HLSLTree::FExpression*& OutExpression) override;
 #endif
 	//~ End UMaterialExpression Interface
 };
diff --git a/Engine/Source/Runtime/Engine/Classes/Materials/MaterialExpressionPixelNormalWS.h b/Engine/Source/Runtime/Engine/Classes/Materials/MaterialExpressionPixelNormalWS.h
--- a/Engine/Source/Runtime/Engine/Classes/Materials/MaterialExpressionPixelNormalWS.h
+++ b/Engine/Source/Runtime/Engine/Classes/Materials/MaterialExpressionPixelNormalWS.h
@@ -18,6 +18,7 @@ class UMaterialExpressionPixelNormalWS : public UMaterialExpression
 #if WITH_EDITOR
 	virtual int32 Compile(class FMaterialCompiler* Compiler, int32 OutputIndex) override;
 	virtual void GetCaption(TArray<FString>& OutCaptions) const override;
+	virtual bool GenerateHLSLExpression(FMaterialHLSLGenerator& Generator, UE::HLSLTree::FScope& Scope, int32 OutputIndex, UE::HLSLTree::FExpression*& OutExpression) override;
 #endif
 	//~ End UMaterialExpression Interface
 };
diff --git a/Engine/Source/Runtime/Engine/Classes/Materials/MaterialExpressionRotateAboutAxis.h b/Engine/Source/Runtime/Engine/Classes/Materials/MaterialExpressionRotateAboutAxis.h
--- a/Engine/Source/Runtime/Engine/Classes/Materials/MaterialExpressionRotateAboutAxis.h
+++ b/Engine/Source/Runtime/Engine/Classes/Materials/MaterialExpressionRotateAboutAxis.h
@@ -34,6 +34,7 @@ class UMaterialExpressionRotateAboutAxis : public UMaterialExpression
 #if WITH_EDITOR
 	virtual int32 Compile(class FMaterialCompiler* Compiler, int32 OutputIndex) override;
 	virtual void GetCaption(TArray<FString>& OutCaptions) const override;
+	virtual bool GenerateHLSLExpression(FMaterialHLSLGenerator& Generator, UE::HLSLTree::FScope& Scope, int32 OutputIndex, UE::HLSLTree::FExpression*& OutExpression) override;
 #endif
 	//~ End UMaterialExpression Interface
 
diff --git a/Engine/Source/Runtime/Engine/Classes/Materials/MaterialExpressionRotator.h b/Engine/Source/Runtime/Engine/Classes/Materials/MaterialExpressionRotator.h
--- a/Engine/Source/Runtime/Engine/Classes/Materials/MaterialExpressionRotator.h
+++ b/Engine/Source/Runtime/Engine/Classes/Materials/MaterialExpressionRotator.h
@@ -38,6 +38,7 @@ class UMaterialExpressionRotator : public UMaterialExpression
 	virtual int32 Compile(class FMaterialCompiler* Compiler, int32 OutputIndex) override;
 	virtual void GetCaption(TArray<FString>& OutCaptions) const override;
 	virtual bool NeedsRealtimePreview() override { return Time.Expression==NULL && Speed != 0.f; }
+	virtual bool GenerateHLSLExpression(FMaterialHLSLGenerator& Generator, UE::HLSLTree::FScope& Scope, int32 OutputIndex, UE::HLSLTree::FExpression*& OutExpression) override;
 #endif
 	//~ End UMaterialExpression Interface
 
diff --git a/Engine/Source/Runtime/Engine/Classes/Materials/MaterialExpressionSceneDepth.h b/Engine/Source/Runtime/Engine/Classes/Materials/MaterialExpressionSceneDepth.h
--- a/Engine/Source/Runtime/Engine/Classes/Materials/MaterialExpressionSceneDepth.h
+++ b/Engine/Source/Runtime/Engine/Classes/Materials/MaterialExpressionSceneDepth.h
@@ -46,6 +46,7 @@ class UMaterialExpressionSceneDepth : public UMaterialExpression
 	virtual int32 Compile(class FMaterialCompiler* Compiler, int32 OutputIndex) override;
 	virtual void GetCaption(TArray<FString>& OutCaptions) const override;
 	virtual FName GetInputName(int32 InputIndex) const override;
+	virtual bool GenerateHLSLExpression(FMaterialHLSLGenerator& Generator, UE::HLSLTree::FScope& Scope, int32 OutputIndex, UE::HLSLTree::FExpression*& OutExpression) override;
 #endif
 	//~ End UMaterialExpression Interface
 };
diff --git a/Engine/Source/Runtime/Engine/Classes/Materials/MaterialExpressionShadingModel.h b/Engine/Source/Runtime/Engine/Classes/Materials/MaterialExpressionShadingModel.h
--- a/Engine/Source/Runtime/Engine/Classes/Materials/MaterialExpressionShadingModel.h
+++ b/Engine/Source/Runtime/Engine/Classes/Materials/MaterialExpressionShadingModel.h
@@ -28,6 +28,7 @@ class UMaterialExpressionShadingModel : public UMaterialExpression
 	virtual int32 Compile(class FMaterialCompiler* Compiler, int32 OutputIndex) override;
 	virtual uint32 GetOutputType(int32 OutputIndex) override;
 	virtual void GetCaption(TArray<FString>& OutCaptions) const override;
+	virtual bool GenerateHLSLExpression(FMaterialHLSLGenerator& Generator, UE::HLSLTree::FScope& Scope, int32 OutputIndex, UE::HLSLTree::FExpression*& OutExpression) override;
 #endif
 public:
 	UPROPERTY(EditAnywhere, Category=ShadingModel,  meta=(ValidEnumValues="MSM_DefaultLit, MSM_Subsurface, MSM_PreintegratedSkin, MSM_ClearCoat, MSM_SubsurfaceProfile, MSM_TwoSidedFoliage, MSM_Hair, MSM_Cloth, MSM_Eye"))
diff --git a/Engine/Source/Runtime/Engine/Classes/Materials/MaterialExpressionShadingPathSwitch.h b/Engine/Source/Runtime/Engine/Classes/Materials/MaterialExpressionShadingPathSwitch.h
--- a/Engine/Source/Runtime/Engine/Classes/Materials/MaterialExpressionShadingPathSwitch.h
+++ b/Engine/Source/Runtime/Engine/Classes/Materials/MaterialExpressionShadingPathSwitch.h
@@ -33,6 +33,7 @@ class UMaterialExpressionShadingPathSwitch : public UMaterialExpression
 	virtual bool IsResultMaterialAttributes(int32 OutputIndex) override;
 	virtual uint32 GetInputType(int32 InputIndex) override {return MCT_Unknown;}
 	virtual uint32 GetOutputType(int32 InputIndex) override {return MCT_Unknown;}
+	virtual bool GenerateHLSLExpression(FMaterialHLSLGenerator& Generator, UE::HLSLTree::FScope& Scope, int32 OutputIndex, UE::HLSLTree::FExpression*& OutExpression) override;
 #endif // WITH_EDITOR
 	//~ End UMaterialExpression Interface
 };
diff --git a/Engine/Source/Runtime/Engine/Classes/Materials/MaterialExpressionSine.h b/Engine/Source/Runtime/Engine/Classes/Materials/MaterialExpressionSine.h
--- a/Engine/Source/Runtime/Engine/Classes/Materials/MaterialExpressionSine.h
+++ b/Engine/Source/Runtime/Engine/Classes/Materials/MaterialExpressionSine.h
@@ -26,6 +26,7 @@ class UMaterialExpressionSine : public UMaterialExpression
 	virtual int32 Compile(class FMaterialCompiler* Compiler, int32 OutputIndex) override;
 	virtual void GetCaption(TArray<FString>& OutCaptions) const override;
 	virtual FText GetKeywords() const override {return FText::FromString(TEXT("sin"));}
+	virtual bool GenerateHLSLExpression(FMaterialHLSLGenerator& Generator, UE::HLSLTree::FScope& Scope, int32 OutputIndex, UE::HLSLTree::FExpression*& OutExpression) override;
 #endif
 	//~ End UMaterialExpression Interface
 
diff --git a/Engine/Source/Runtime/Engine/Classes/Materials/MaterialExpressionSphereMask.h b/Engine/Source/Runtime/Engine/Classes/Materials/MaterialExpressionSphereMask.h
--- a/Engine/Source/Runtime/Engine/Classes/Materials/MaterialExpressionSphereMask.h
+++ b/Engine/Source/Runtime/Engine/Classes/Materials/MaterialExpressionSphereMask.h
@@ -44,6 +44,7 @@ class UMaterialExpressionSphereMask : public UMaterialExpression
 	virtual int32 Compile(class FMaterialCompiler* Compiler, int32 OutputIndex) override;
 	virtual void GetCaption(TArray<FString>& OutCaptions) const override;
 	virtual FText GetKeywords() const override { return FText::FromString(TEXT("sphere")); }
+	virtual bool GenerateHLSLExpression(FMaterialHLSLGenerator& Generator, UE::HLSLTree::FScope& Scope, int32 OutputIndex, UE::HLSLTree::FExpression*& OutExpression) override;
 #endif
 	//~ End UMaterialExpression Interface
 
diff --git a/Engine/Source/Runtime/Engine/Classes/Materials/MaterialExpressionSquareRoot.h b/Engine/Source/Runtime/Engine/Classes/Materials/MaterialExpressionSquareRoot.h
--- a/Engine/Source/Runtime/Engine/Classes/Materials/MaterialExpressionSquareRoot.h
+++ b/Engine/Source/Runtime/Engine/Classes/Materials/MaterialExpressionSquareRoot.h
@@ -23,6 +23,7 @@ class UMaterialExpressionSquareRoot : public UMaterialExpression
 	virtual int32 Compile(class FMaterialCompiler* Compiler, int32 OutputIndex) override;
 	virtual void GetCaption(TArray<FString>& OutCaptions) const override;
 	virtual FText GetKeywords() const override {return FText::FromString(TEXT("sqrt"));}
+	virtual bool GenerateHLSLExpression(FMaterialHLSLGenerator& Generator, UE::HLSLTree::FScope& Scope, int32 OutputIndex, UE::HLSLTree::FExpression*& OutExpression) override;
 #endif // WITH_EDITOR
 	//~ End UMaterialExpression Interface
 };
diff --git a/Engine/Source/Runtime/Engine/Classes/Materials/MaterialExpressionTangent.h b/Engine/Source/Runtime/Engine/Classes/Materials/MaterialExpressionTangent.h
--- a/Engine/Source/Runtime/Engine/Classes/Materials/MaterialExpressionTangent.h
+++ b/Engine/Source/Runtime/Engine/Classes/Materials/MaterialExpressionTangent.h
@@ -24,6 +24,7 @@ class UMaterialExpressionTangent : public UMaterialExpression
 	virtual int32 Compile(class FMaterialCompiler* Compiler, int32 OutputIndex) override;
 	virtual void GetCaption(TArray<FString>& OutCaptions) const override;
 	virtual FText GetKeywords() const override {return FText::FromString(TEXT("tan"));}
+	virtual bool GenerateHLSLExpression(FMaterialHLSLGenerator& Generator, UE::HLSLTree::FScope& Scope, int32 OutputIndex, UE::HLSLTree::FExpression*& OutExpression) override;
 #endif
 	//~ End UMaterialExpression Interface
 };
diff --git a/Engine/Source/Runtime/Engine/Classes/Materials/MaterialExpressionTwoSidedSign.h b/Engine/Source/Runtime/Engine/Classes/Materials/MaterialExpressionTwoSidedSign.h
--- a/Engine/Source/Runtime/Engine/Classes/Materials/MaterialExpressionTwoSidedSign.h
+++ b/Engine/Source/Runtime/Engine/Classes/Materials/MaterialExpressionTwoSidedSign.h
@@ -18,6 +18,7 @@ class UMaterialExpressionTwoSidedSign : public UMaterialExpression
 #if WITH_EDITOR
 	virtual int32 Compile(class FMaterialCompiler* Compiler, int32 OutputIndex) override;
 	virtual void GetCaption(TArray<FString>& OutCaptions) const override;
+	virtual bool GenerateHLSLExpression(FMaterialHLSLGenerator& Generator, UE::HLSLTree::FScope& Scope, int32 OutputIndex, UE::HLSLTree::FExpression*& OutExpression) override;
 #endif
 	//~ End UMaterialExpression Interface
 };
diff --git a/Engine/Source/Runtime/Engine/Classes/Materials/MaterialExpressionVertexNormalWS.h b/Engine/Source/Runtime/Engine/Classes/Materials/MaterialExpressionVertexNormalWS.h
--- a/Engine/Source/Runtime/Engine/Classes/Materials/MaterialExpressionVertexNormalWS.h
+++ b/Engine/Source/Runtime/Engine/Classes/Materials/MaterialExpressionVertexNormalWS.h
@@ -18,6 +18,7 @@ class UMaterialExpressionVertexNormalWS : public UMaterialExpression
 #if WITH_EDITOR
 	virtual int32 Compile(class FMaterialCompiler* Compiler, int32 OutputIndex) override;
 	virtual void GetCaption(TArray<FString>& OutCaptions) const override;
+	virtual bool GenerateHLSLExpression(FMaterialHLSLGenerator& Generator, UE::HLSLTree::FScope& Scope, int32 OutputIndex, UE::HLSLTree::FExpression*& OutExpression) override;
 #endif
 	//~ End UMaterialExpression Interface
 };
diff --git a/Engine/Source/Runtime/Engine/Classes/Materials/MaterialExpressionVertexTangentWS.h b/Engine/Source/Runtime/Engine/Classes/Materials/MaterialExpressionVertexTangentWS.h
--- a/Engine/Source/Runtime/Engine/Classes/Materials/MaterialExpressionVertexTangentWS.h
+++ b/Engine/Source/Runtime/Engine/Classes/Materials/MaterialExpressionVertexTangentWS.h
@@ -18,6 +18,7 @@ class UMaterialExpressionVertexTangentWS : public UMaterialExpression
 #if WITH_EDITOR
 	virtual int32 Compile(class FMaterialCompiler* Compiler, int32 OutputIndex) override;
 	virtual void GetCaption(TArray<FString>& OutCaptions) const override;
+	virtual bool GenerateHLSLExpression(FMaterialHLSLGenerator& Generator, UE::HLSLTree::FScope& Scope, int32 OutputIndex, UE::HLSLTree::FExpression*& OutExpression) override;
 #endif
 	//~ End UMaterialExpression Interface
 };
diff --git a/Engine/Source/Runtime/Engine/Classes/Materials/MaterialExpressionViewProperty.h b/Engine/Source/Runtime/Engine/Classes/Materials/MaterialExpressionViewProperty.h
--- a/Engine/Source/Runtime/Engine/Classes/Materials/MaterialExpressionViewProperty.h
+++ b/Engine/Source/Runtime/Engine/Classes/Materials/MaterialExpressionViewProperty.h
@@ -56,6 +56,7 @@ class UMaterialExpressionViewProperty : public UMaterialExpression
 #if WITH_EDITOR
 	virtual int32 Compile(class FMaterialCompiler* Compiler, int32 OutputIndex) override;
 	virtual void GetCaption(TArray<FString>& OutCaptions) const override;
+	virtual bool GenerateHLSLExpression(FMaterialHLSLGenerator& Generator, UE::HLSLTree::FScope& Scope, int32 OutputIndex, UE::HLSLTree::FExpression*& OutExpression) override;
 #endif
 	//~ End UMaterialExpression Interface
 };
diff --git a/Engine/Source/Runtime/Engine/Classes/Materials/MaterialExpressionVolumetricAdvancedMaterialInput.h b/Engine/Source/Runtime/Engine/Classes/Materials/MaterialExpressionVolumetricAdvancedMaterialInput.h
--- a/Engine/Source/Runtime/Engine/Classes/Materials/MaterialExpressionVolumetricAdvancedMaterialInput.h
+++ b/Engine/Source/Runtime/Engine/Classes/Materials/MaterialExpressionVolumetricAdvancedMaterialInput.h
@@ -16,6 +16,7 @@ class UMaterialExpressionVolumetricAdvancedMaterialInput : public UMaterialExpre
 #if WITH_EDITOR
 	virtual int32 Compile(class FMaterialCompiler* Compiler, int32 OutputIndex) override;
 	virtual void GetCaption(TArray<FString>& OutCaptions) const override;
+	virtual bool GenerateHLSLExpression(FMaterialHLSLGenerator& Generator, UE::HLSLTree::FScope& Scope, int32 OutputIndex, UE::HLSLTree::FExpression*& OutExpression) override;
 #endif
 	//~ End UMaterialExpression Interface
 };
diff --git a/Engine/Source/Runtime/Engine/Classes/Materials/MaterialExpressionVolumetricAdvancedMaterialOutput.h b/Engine/Source/Runtime/Engine/Classes/Materials/MaterialExpressionVolumetricAdvancedMaterialOutput.h
--- a/Engine/Source/Runtime/Engine/Classes/Materials/MaterialExpressionVolumetricAdvancedMaterialOutput.h
+++ b/Engine/Source/Runtime/Engine/Classes/Materials/MaterialExpressionVolumetricAdvancedMaterialOutput.h
@@ -98,6 +98,9 @@ public:
 	bool GetEvaluatePhaseOncePerSample() const;
 	uint32 GetMultiScatteringApproximationOctaveCount() const;
 
+	virtual bool GenerateHLSLExpression(FMaterialHLSLGenerator& Generator, UE::HLSLTree::FScope& Scope, int32 OutputIndex, UE::HLSLTree::FExpression*& OutExpression) override;
+	virtual UE::Shader::EValueType GetCustomOutputType(int32 OutputIndex) const override;
+
 #endif
 
 	//~ Begin UMaterialExpressionCustomOutput Interface
diff --git a/Engine/Source/Runtime/Engine/Private/HLSLTree/HLSLTree.cpp b/Engine/Source/Runtime/Engine/Private/HLSLTree/HLSLTree.cpp
--- a/Engine/Source/Runtime/Engine/Private/HLSLTree/HLSLTree.cpp
+++ b/Engine/Source/Runtime/Engine/Private/HLSLTree/HLSLTree.cpp
@@ -98,6 +98,8 @@ FOperationDescription GetOperationDescription(EOperation Op)
 	case EOperation::Abs: return FOperationDescription(TEXT("Abs"), TEXT("abs"), 1, Shader::EPreshaderOpcode::Abs); break;
 	case EOperation::Neg: return FOperationDescription(TEXT("Neg"), TEXT("-"), 1, Shader::EPreshaderOpcode::Neg); break;
 	case EOperation::Rcp: return FOperationDescription(TEXT("Rcp"), TEXT("/"), 1, Shader::EPreshaderOpcode::Rcp); break;
+	case EOperation::Sqrt: return FOperationDescription(TEXT("Sqrt"), TEXT("sqrt"), 1, Shader::EPreshaderOpcode::Sqrt); break;
+	case EOperation::Log2: return FOperationDescription(TEXT("Log2"), TEXT("log2"), 1, Shader::EPreshaderOpcode::Log2); break;
 	case EOperation::Frac: return FOperationDescription(TEXT("Frac"), TEXT("frac"), 1, Shader::EPreshaderOpcode::Frac); break;
 	case EOperation::Floor: return FOperationDescription(TEXT("Floor"), TEXT("floor"), 1, Shader::EPreshaderOpcode::Floor); break;
 	case EOperation::Ceil: return FOperationDescription(TEXT("Ceil"), TEXT("ceil"), 1, Shader::EPreshaderOpcode::Ceil); break;
@@ -107,6 +109,16 @@ FOperationDescription GetOperationDescription(EOperation Op)
 	case EOperation::Sign: return FOperationDescription(TEXT("Sign"), TEXT("sign"), 1, Shader::EPreshaderOpcode::Sign); break;
 	case EOperation::Length: return FOperationDescription(TEXT("Length"), TEXT("length"), 1, Shader::EPreshaderOpcode::Length); break;
 	case EOperation::Normalize: return FOperationDescription(TEXT("Normalize"), TEXT("normalize"), 1, Shader::EPreshaderOpcode::Normalize); break;
+	case EOperation::Sin: return FOperationDescription(TEXT("Sin"), TEXT("sin"), 1, Shader::EPreshaderOpcode::Sin); break;
+	case EOperation::Cos: return FOperationDescription(TEXT("Cos"), TEXT("cos"), 1, Shader::EPreshaderOpcode::Cos); break;
+	case EOperation::Tan: return FOperationDescription(TEXT("Tan"), TEXT("tan"), 1, Shader::EPreshaderOpcode::Tan); break;
+	case EOperation::Asin: return FOperationDescription(TEXT("Asin"), TEXT("asin"), 1, Shader::EPreshaderOpcode::Asin); break;
+	case EOperation::AsinFast: return FOperationDescription(TEXT("AsinFast"), TEXT("asinFast"), 1, Shader::EPreshaderOpcode::Asin); break;
+	case EOperation::Acos: return FOperationDescription(TEXT("Acos"), TEXT("acos"), 1, Shader::EPreshaderOpcode::Acos); break;
+	case EOperation::AcosFast: return FOperationDescription(TEXT("AcosFast"), TEXT("acosFast"), 1, Shader::EPreshaderOpcode::Acos); break;
+	case EOperation::Atan: return FOperationDescription(TEXT("Atan"), TEXT("atan"), 1, Shader::EPreshaderOpcode::Atan); break;
+	case EOperation::AtanFast: return FOperationDescription(TEXT("AtanFast"), TEXT("atanFast"), 1, Shader::EPreshaderOpcode::Atan); break;
+
 	// Binary
 	case EOperation::Add: return FOperationDescription(TEXT("Add"), TEXT("+"), 2, Shader::EPreshaderOpcode::Add); break;
 	case EOperation::Sub: return FOperationDescription(TEXT("Subtract"), TEXT("-"), 2, Shader::EPreshaderOpcode::Sub); break;
@@ -114,6 +126,8 @@ FOperationDescription GetOperationDescription(EOperation Op)
 	case EOperation::Div: return FOperationDescription(TEXT("Divide"), TEXT("/"), 2, Shader::EPreshaderOpcode::Div); break;
 	case EOperation::Fmod: return FOperationDescription(TEXT("Fmod"), TEXT("%"), 2, Shader::EPreshaderOpcode::Fmod); break;
 	case EOperation::PowPositiveClamped: return FOperationDescription(TEXT("PowPositiveClamped"), TEXT("PowPositiveClamped"), 2, Shader::EPreshaderOpcode::Nop); break;
+	case EOperation::Atan2: return FOperationDescription(TEXT("Atan2"), TEXT("atan2"), 2, Shader::EPreshaderOpcode::Atan2); break;
+	case EOperation::Atan2Fast: return FOperationDescription(TEXT("Atan2Fast"), TEXT("atan2Fast"), 2, Shader::EPreshaderOpcode::Atan2); break;
 	case EOperation::Dot: return FOperationDescription(TEXT("Dot"), TEXT("dot"), 2, Shader::EPreshaderOpcode::Dot); break;
 	case EOperation::Min: return FOperationDescription(TEXT("Min"), TEXT("min"), 2, Shader::EPreshaderOpcode::Min); break;
 	case EOperation::Max: return FOperationDescription(TEXT("Max"), TEXT("max"), 2, Shader::EPreshaderOpcode::Max); break;
@@ -167,11 +181,15 @@ EExpressionEvaluation CombineEvaluations(EExpressionEvaluation Lhs, EExpressionE
 	{
 		return EExpressionEvaluation::ConstantLoop;
 	}
+	else if (Lhs == EExpressionEvaluation::Constant || Rhs == EExpressionEvaluation::Constant)
+	{
+		return EExpressionEvaluation::Constant;
+	}
 
 	// Otherwise must be constant
-	check(Lhs == EExpressionEvaluation::Constant);
-	check(Rhs == EExpressionEvaluation::Constant);
-	return EExpressionEvaluation::Constant;
+	check(Lhs == EExpressionEvaluation::ConstantZero);
+	check(Rhs == EExpressionEvaluation::ConstantZero);
+	return EExpressionEvaluation::ConstantZero;
 }
 
 EExpressionEvaluation MakeLoopEvaluation(EExpressionEvaluation Evaluation)
@@ -251,54 +269,27 @@ FExpression* FExpressionLocalPHI::ComputePreviousFrame(FTree& Tree, const FReque
 bool FExpressionLocalPHI::PrepareValue(FEmitContext& Context, FEmitScope& Scope, const FRequestedType& RequestedType, FPrepareValueResult& OutResult) const
 {
 	check(NumValues <= MaxNumPreviousScopes);
-	FExpression* ForwardExpression = Values[0];
-	bool bForwardExpressionValid = true;
 
-	// There are 2 cases we want to optimize here
-	// 1) If the PHI node has the same value in all the previous scopes, we can avoid generating code for the previous scopes, and just use the value directly
-	for (int32 i = 1; i < NumValues; ++i)
-	{
-		FExpression* ScopeExpression = Values[i];
-		if (ScopeExpression != ForwardExpression)
-		{
-			ForwardExpression = nullptr;
-			bForwardExpressionValid = false;
-			break;
-		}
-	}
+	FEmitScope* EmitValueScopes[MaxNumPreviousScopes];
+	FScope* LiveScopes[MaxNumPreviousScopes];
+	FExpression* LiveValues[MaxNumPreviousScopes];
+	int32 NumLiveScopes = 0;
 
-	if (bForwardExpressionValid)
-	{
-		check(ForwardExpression);
-		return OutResult.SetForwardValue(Context, Scope, RequestedType, ForwardExpression);
-	}
-
-	// 2) PHI has different values in previous scopes, but possible some previous scopes may become dead due to constant folding
-	// In this case, we check to see if the value is the same in all live scopes, and forward if possible
-	FEmitScope* EmitScopes[MaxNumPreviousScopes] = { nullptr };
 	for (int32 i = 0; i < NumValues; ++i)
 	{
-		// Ignore values in dead scopes
-		EmitScopes[i] = Context.PrepareScope(Scopes[i]);
-		if (EmitScopes[i])
+		FEmitScope* EmitValueScope = Context.PrepareScope(Scopes[i]);
+		if (EmitValueScope)
 		{
-			FExpression* ScopeExpression = Values[i];
-			if (!ForwardExpression)
-			{
-				ForwardExpression = ScopeExpression;
-				bForwardExpressionValid = true;
-			}
-			else if (ForwardExpression != ScopeExpression)
-			{
-				bForwardExpressionValid = false;
-			}
+			const int32 LiveScopeIndex = NumLiveScopes++;
+			EmitValueScopes[LiveScopeIndex] = EmitValueScope;
+			LiveScopes[LiveScopeIndex] = Scopes[i];
+			LiveValues[LiveScopeIndex] = Values[i];
 		}
 	}
 
-	if (bForwardExpressionValid)
+	if (NumLiveScopes == 0)
 	{
-		check(ForwardExpression);
-		return OutResult.SetForwardValue(Context, Scope, RequestedType, ForwardExpression);
+		return OutResult.SetType(Context, RequestedType, EExpressionEvaluation::ConstantZero, Shader::EValueType::Float1);
 	}
 
 	FPreparedType TypePerValue[MaxNumPreviousScopes];
@@ -307,15 +298,15 @@ bool FExpressionLocalPHI::PrepareValue(FEmitContext& Context, FEmitScope& Scope,
 
 	auto UpdateValueTypes = [&]()
 	{
-		for (int32 i = 0; i < NumValues; ++i)
+		for (int32 i = 0; i < NumLiveScopes; ++i)
 		{
-			if (!EmitScopes[i] || EmitScopes[i]->Evaluation == EExpressionEvaluation::Unknown)
+			if (EmitValueScopes[i]->Evaluation == EExpressionEvaluation::Unknown)
 			{
-				EmitScopes[i] = Context.PrepareScope(Scopes[i]);
+				EmitValueScopes[i] = Context.PrepareScope(LiveScopes[i]);
 			}
-			if (TypePerValue[i].IsVoid() && EmitScopes[i])
+			if (TypePerValue[i].IsVoid())
 			{
-				const FPreparedType& ValueType = Context.PrepareExpression(Values[i], *EmitScopes[i], RequestedType);
+				const FPreparedType& ValueType = Context.PrepareExpression(LiveValues[i], *EmitValueScopes[i], RequestedType);
 				if (!ValueType.IsVoid())
 				{
 					TypePerValue[i] = ValueType;
@@ -323,12 +314,12 @@ bool FExpressionLocalPHI::PrepareValue(FEmitContext& Context, FEmitScope& Scope,
 					if (MergedType.IsVoid())
 					{
 						return Context.Errors->AddErrorf(TEXT("Mismatched types for local variable %s and %s"),
-							CurrentType.GetType().GetName(),
-							ValueType.GetType().GetName());
+							CurrentType.GetName(),
+							ValueType.GetName());
 					}
 					CurrentType = MergedType;
-					CurrentType.MergeEvaluation(EmitScopes[i]->Evaluation);
-					check(NumValidTypes < NumValues);
+					CurrentType.MergeEvaluation(EmitValueScopes[i]->Evaluation);
+					check(NumValidTypes < NumLiveScopes);
 					NumValidTypes++;
 				}
 			}
@@ -350,14 +341,14 @@ bool FExpressionLocalPHI::PrepareValue(FEmitContext& Context, FEmitScope& Scope,
 		return false;
 	}
 
-	if (NumValidTypes < NumValues)
+	if (NumValidTypes < NumLiveScopes)
 	{
 		// Now try to assign remaining types that failed the first iteration 
 		if (!UpdateValueTypes())
 		{
 			return false;
 		}
-		if (NumValidTypes < NumValues)
+		if (NumValidTypes < NumLiveScopes)
 		{
 			return Context.Errors->AddError(TEXT("Failed to compute all types for LocalPHI"));
 		}
@@ -371,13 +362,15 @@ bool FExpressionLocalPHI::PrepareValue(FEmitContext& Context, FEmitScope& Scope,
 			}
 
 			// Since we changed our type, need to update any dependant values again
-			for (int32 i = 0; i < NumValues; ++i)
+			for (int32 i = 0; i < NumLiveScopes; ++i)
 			{
-				const FPreparedType& ValueType = Context.PrepareExpression(Values[i], *EmitScopes[i], RequestedType);
+				const FPreparedType& ValueType = Context.PrepareExpression(LiveValues[i], *EmitValueScopes[i], RequestedType);
 				// Don't expect types to change *again*
 				if (ValueType.IsVoid() || MergePreparedTypes(CurrentType, ValueType) != CurrentType)
 				{
-					return Context.Errors->AddError(TEXT("Mismatched types for local variable %s and %s"));
+					return Context.Errors->AddErrorf(TEXT("Mismatched types for local variable %s and %s"),
+						CurrentType.GetName(),
+						ValueType.GetName());
 				}
 			}
 		}
@@ -386,74 +379,118 @@ bool FExpressionLocalPHI::PrepareValue(FEmitContext& Context, FEmitScope& Scope,
 	return true;
 }
 
+namespace Private
+{
+struct FLocalPHILiveScopes
+{
+	FEmitScope* EmitDeclarationScope = nullptr;
+	FEmitScope* EmitValueScopes[MaxNumPreviousScopes];
+	FExpression* LiveValues[MaxNumPreviousScopes];
+	int32 NumScopes = 0;
+	bool bCanForwardValue = true;
+};
+
+bool GetLiveScopes(FEmitContext& Context, const FExpressionLocalPHI& Expression, FLocalPHILiveScopes& OutLiveScopes)
+{
+	for (int32 ScopeIndex = 0; ScopeIndex < Expression.NumValues; ++ScopeIndex)
+	{
+		FEmitScope* EmitValueScope = Context.AcquireEmitScope(Expression.Scopes[ScopeIndex]);
+		if (!IsScopeDead(EmitValueScope))
+		{
+			if (OutLiveScopes.bCanForwardValue)
+			{
+				for (int32 PrevScopeIndex = 0; PrevScopeIndex < OutLiveScopes.NumScopes; ++PrevScopeIndex)
+				{
+					if (Expression.Values[ScopeIndex] != OutLiveScopes.LiveValues[PrevScopeIndex])
+					{
+						OutLiveScopes.bCanForwardValue = false;
+					}
+				}
+			}
+
+			const int32 LiveScopeIndex = OutLiveScopes.NumScopes++;
+			OutLiveScopes.LiveValues[LiveScopeIndex] = Expression.Values[ScopeIndex];
+			OutLiveScopes.EmitValueScopes[LiveScopeIndex] = EmitValueScope;
+			OutLiveScopes.EmitDeclarationScope = FEmitScope::FindSharedParent(OutLiveScopes.EmitDeclarationScope, EmitValueScope);
+			if (!OutLiveScopes.EmitDeclarationScope)
+			{
+				return Context.Errors->AddError(TEXT("Invalid LocalPHI"));
+			}
+		}
+	}
+	return OutLiveScopes.NumScopes > 0;
+}
+}
+
 void FExpressionLocalPHI::EmitValueShader(FEmitContext& Context, FEmitScope& Scope, const FRequestedType& RequestedType, FEmitValueShaderResult& OutResult) const
 {
 	FEmitShaderExpression* const* PrevEmitExpression = Context.EmitLocalPHIMap.Find(this);
 	FEmitShaderExpression* EmitExpression = PrevEmitExpression ? *PrevEmitExpression : nullptr;
 	if (!EmitExpression)
 	{
-		const int32 LocalPHIIndex = Context.NumExpressionLocalPHIs++;
-		const Shader::FType LocalType = GetType();
-
-		// This is the first time we've emitted shader code for this PHI
-		// Create an expression and add it to the map first, so if this is called recursively this path will only be taken the first time
-		EmitExpression = OutResult.Code = Context.EmitInlineExpression(Scope,
-			LocalType,
-			TEXT("LocalPHI%"), LocalPHIIndex);
-		Context.EmitLocalPHIMap.Add(this, EmitExpression);
-
 		// Find the outermost scope to declare our local variable
-		FEmitScope* EmitDeclarationScope = &Scope;
-		FEmitScope* EmitValueScopes[MaxNumPreviousScopes] = { nullptr };
-		for (int32 i = 0; i < NumValues; ++i)
+		Private::FLocalPHILiveScopes LiveScopes;
+		if (!Private::GetLiveScopes(Context, *this, LiveScopes))
 		{
-			EmitValueScopes[i] = Context.AcquireEmitScope(Scopes[i]);
-			EmitDeclarationScope = FEmitScope::FindSharedParent(EmitDeclarationScope, EmitValueScopes[i]);
-			if (!EmitDeclarationScope)
-			{
-				Context.Errors->AddError(TEXT("Invalid LocalPHI"));
-				return;
-			}
+			return;
 		}
 
-		FEmitShaderStatement* EmitDeclaration = nullptr;
-		for (int32 i = 0; i < NumValues; ++i)
+		if (LiveScopes.bCanForwardValue)
+		{
+			EmitExpression = LiveScopes.LiveValues[0]->GetValueShader(Context, Scope);
+			Context.EmitLocalPHIMap.Add(this, EmitExpression);
+		}
+		else
 		{
-			FEmitScope* EmitValueScope = EmitValueScopes[i];
-			if (EmitValueScope == EmitDeclarationScope)
+			// This is the first time we've emitted shader code for this PHI
+			// Create an expression and add it to the map first, so if this is called recursively this path will only be taken the first time
+			const int32 LocalPHIIndex = Context.NumExpressionLocalPHIs++;
+			const Shader::FType LocalType = GetType();
+
+			EmitExpression = OutResult.Code = Context.EmitInlineExpression(Scope,
+				LocalType,
+				TEXT("LocalPHI%"), LocalPHIIndex);
+			Context.EmitLocalPHIMap.Add(this, EmitExpression);
+
+			FEmitShaderStatement* EmitDeclaration = nullptr;
+			for (int32 i = 0; i < LiveScopes.NumScopes; ++i)
 			{
-				FEmitShaderExpression* ShaderValue = Values[i]->GetValueShader(Context, *EmitValueScope, LocalType);
-				EmitDeclaration = Context.EmitStatement(*EmitValueScope, TEXT("% LocalPHI% = %;"),
+				FEmitScope* EmitValueScope = LiveScopes.EmitValueScopes[i];
+				if (EmitValueScope == LiveScopes.EmitDeclarationScope)
+				{
+					FEmitShaderExpression* ShaderValue = LiveScopes.LiveValues[i]->GetValueShader(Context, *EmitValueScope, LocalType);
+					EmitDeclaration = Context.EmitStatement(*EmitValueScope, TEXT("% LocalPHI% = %;"),
+						LocalType.GetName(),
+						LocalPHIIndex,
+						ShaderValue);
+					break;
+				}
+			}
+			if (!EmitDeclaration)
+			{
+				EmitDeclaration = Context.EmitStatement(*LiveScopes.EmitDeclarationScope, TEXT("% LocalPHI%;"),
 					LocalType.GetName(),
-					LocalPHIIndex,
-					ShaderValue);
-				break;
+					LocalPHIIndex);
 			}
-		}
-		if (!EmitDeclaration)
-		{
-			EmitDeclaration = Context.EmitStatement(*EmitDeclarationScope, TEXT("% LocalPHI%;"),
-				LocalType.GetName(),
-				LocalPHIIndex);
-		}
 
-		FEmitShaderNode* Dependencies[MaxNumPreviousScopes] = { nullptr };
-		int32 NumDependencies = 0;
-		for (int32 i = 0; i < NumValues; ++i)
-		{
-			FEmitScope* EmitValueScope = EmitValueScopes[i];
-			if (EmitValueScope != EmitDeclarationScope)
+			FEmitShaderNode* Dependencies[MaxNumPreviousScopes] = { nullptr };
+			int32 NumDependencies = 0;
+			for (int32 i = 0; i < LiveScopes.NumScopes; ++i)
 			{
-				FEmitShaderExpression* ShaderValue = Values[i]->GetValueShader(Context, *EmitValueScope, LocalType);
-				FEmitShaderStatement* EmitAssignment = Context.EmitStatementWithDependency(*EmitValueScope, EmitDeclaration, TEXT("LocalPHI% = %;"),
-					LocalPHIIndex,
-					ShaderValue);
-				Dependencies[NumDependencies++] = EmitAssignment;
+				FEmitScope* EmitValueScope = LiveScopes.EmitValueScopes[i];
+				if (EmitValueScope != LiveScopes.EmitDeclarationScope)
+				{
+					FEmitShaderExpression* ShaderValue = LiveScopes.LiveValues[i]->GetValueShader(Context, *EmitValueScope, LocalType);
+					FEmitShaderStatement* EmitAssignment = Context.EmitStatementWithDependency(*EmitValueScope, EmitDeclaration, TEXT("LocalPHI% = %;"),
+						LocalPHIIndex,
+						ShaderValue);
+					Dependencies[NumDependencies++] = EmitAssignment;
+				}
 			}
-		}
 
-		// Fill in the expression's dependencies
-		EmitExpression->Dependencies = MemStack::AllocateArrayView(*Context.Allocator, MakeArrayView(Dependencies, NumDependencies));
+			// Fill in the expression's dependencies
+			EmitExpression->Dependencies = MemStack::AllocateArrayView(*Context.Allocator, MakeArrayView(Dependencies, NumDependencies));
+		}
 	}
 
 	OutResult.Code = EmitExpression;
@@ -475,27 +512,38 @@ void FExpressionLocalPHI::EmitValuePreshader(FEmitContext& Context, FEmitScope&
 	OutResult.Type = GetType();
 	if (ValueStackPosition == INDEX_NONE)
 	{
-		// Assign the initial value
-		Context.PreshaderStackPosition++;
-		OutResult.Preshader.WriteOpcode(Shader::EPreshaderOpcode::ConstantZero).Write(OutResult.Type);
-
-		ValueStackPosition = Context.PreshaderStackPosition;
-		const FPreshaderLocalPHIScope LocalPHIScope(this, ValueStackPosition);
-		Context.PreshaderLocalPHIScopes.Add(&LocalPHIScope);
+		Private::FLocalPHILiveScopes LiveScopes;
+		if (!Private::GetLiveScopes(Context, *this, LiveScopes))
+		{
+			return;
+		}
 
-		FEmitScope* EmitRootScope = nullptr;// &Scope;
-		FEmitPreshaderScope PreshaderScopes[MaxNumPreviousScopes];
-		for (int32 i = 0; i < NumValues; ++i)
+		if (LiveScopes.bCanForwardValue)
 		{
-			FEmitPreshaderScope& PreshaderScope = PreshaderScopes[i];
-			PreshaderScope.Scope = Context.AcquireEmitScope(Scopes[i]);
-			PreshaderScope.Value = Values[i];
-			EmitRootScope = FEmitScope::FindSharedParent(PreshaderScopes[i].Scope, EmitRootScope);
+			OutResult.Type = LiveScopes.LiveValues[0]->GetValuePreshader(Context, Scope, RequestedType, OutResult.Preshader);
 		}
+		else
+		{
+			// Assign the initial value
+			Context.PreshaderStackPosition++;
+			OutResult.Preshader.WriteOpcode(Shader::EPreshaderOpcode::ConstantZero).Write(OutResult.Type);
+
+			ValueStackPosition = Context.PreshaderStackPosition;
+			const FPreshaderLocalPHIScope LocalPHIScope(this, ValueStackPosition);
+			Context.PreshaderLocalPHIScopes.Add(&LocalPHIScope);
+
+			FEmitPreshaderScope PreshaderScopes[MaxNumPreviousScopes];
+			for (int32 i = 0; i < LiveScopes.NumScopes; ++i)
+			{
+				FEmitPreshaderScope& PreshaderScope = PreshaderScopes[i];
+				PreshaderScope.Scope = LiveScopes.EmitValueScopes[i];
+				PreshaderScope.Value = LiveScopes.LiveValues[i];
+			}
 
-		Context.EmitPreshaderScope(*EmitRootScope, RequestedType, MakeArrayView(PreshaderScopes, NumValues), OutResult.Preshader);
-		verify(Context.PreshaderLocalPHIScopes.Pop(false) == &LocalPHIScope);
-		check(Context.PreshaderStackPosition == ValueStackPosition);
+			Context.EmitPreshaderScope(*LiveScopes.EmitDeclarationScope, RequestedType, MakeArrayView(PreshaderScopes, LiveScopes.NumScopes), OutResult.Preshader);
+			verify(Context.PreshaderLocalPHIScopes.Pop(false) == &LocalPHIScope);
+			check(Context.PreshaderStackPosition == ValueStackPosition);
+		}
 	}
 	else
 	{
@@ -651,11 +699,6 @@ FPreparedComponent CombineComponents(const FPreparedComponent& Lhs, const FPrepa
 		{
 			Result.LoopScope = FEmitScope::FindSharedParent(Lhs.LoopScope, Rhs.LoopScope);
 		}
-		if (Lhs.ForwardValue == Rhs.ForwardValue && Lhs.ForwardComponentIndex == Rhs.ForwardComponentIndex)
-		{
-			Result.ForwardValue = Lhs.ForwardValue;
-			Result.ForwardComponentIndex = Lhs.ForwardComponentIndex;
-		}
 		if (Lhs.Bounds.Min == Rhs.Bounds.Min)
 		{
 			Result.Bounds.Min = Lhs.Bounds.Min;
@@ -776,24 +819,41 @@ EExpressionEvaluation FPreparedType::GetEvaluation(const FEmitScope& Scope, cons
 		{
 			Result = CombineEvaluations(Result, PreparedComponents[Index].GetEvaluation(Scope));
 		}
+		else
+		{
+			Result = CombineEvaluations(Result, EExpressionEvaluation::ConstantZero);
+		}
 	}
 	return Result;
 }
 
-EExpressionEvaluation FPreparedType::GetFieldEvaluation(const FEmitScope& Scope, int32 ComponentIndex, int32 NumComponents) const
+EExpressionEvaluation FPreparedType::GetFieldEvaluation(const FEmitScope& Scope, const FRequestedType& RequestedType, int32 ComponentIndex, int32 NumComponents) const
 {
 	EExpressionEvaluation Result = EExpressionEvaluation::None;
 	for (int32 Index = 0; Index < NumComponents; ++Index)
 	{
 		const FPreparedComponent Component = GetComponent(ComponentIndex + Index);
-		Result = CombineEvaluations(Result, Component.GetEvaluation(Scope));
+		if (RequestedType.IsComponentRequested(ComponentIndex + Index))
+		{
+			Result = CombineEvaluations(Result, Component.GetEvaluation(Scope));
+		}
+		else if(!Component.IsNone())
+		{
+			Result = CombineEvaluations(Result, EExpressionEvaluation::ConstantZero);
+		}
 	}
 	return Result;
 }
 
 FPreparedComponent FPreparedType::GetComponent(int32 Index) const
 {
-	return PreparedComponents.IsValidIndex(Index) ? PreparedComponents[Index] : FPreparedComponent();
+	if (PreparedComponents.IsValidIndex(Index))
+	{
+		return PreparedComponents[Index];
+	}
+	// Return 'ConstantZero' for unprepared components
+	// TODO - do we ever want to return 'None' instead?
+	return EExpressionEvaluation::ConstantZero;
 }
 
 void FPreparedType::EnsureNumComponents(int32 NumComponents)
@@ -870,19 +930,6 @@ void FPreparedType::SetLoopEvaluation(FEmitScope& Scope, const FRequestedType& R
 	}
 }
 
-void FPreparedType::SetForwardValue(const FRequestedType& RequestedType, FExpression* ForwardValue)
-{
-	for (int32 Index = 0; Index < PreparedComponents.Num(); ++Index)
-	{
-		FPreparedComponent& Component = PreparedComponents[Index];
-		if (RequestedType.IsComponentRequested(Index) && !Component.IsNone())
-		{
-			Component.ForwardValue = ForwardValue;
-			Component.ForwardComponentIndex = Index;
-		}
-	}
-}
-
 void FPreparedType::UpdateBounds(const FRequestedType& RequestedType, Shader::FComponentBounds Bounds)
 {
 	for (int32 Index = 0; Index < PreparedComponents.Num(); ++Index)
@@ -985,7 +1032,6 @@ bool FPrepareValueResult::TryMergePreparedType(FEmitContext& Context, const Shad
 		return false;
 	}
 
-	// If we previously had a forwarded value set, reset that and start over
 	if (!PreparedType.IsInitialized())
 	{
 		PreparedType.PreparedComponents.Reset();
@@ -1031,6 +1077,11 @@ bool FPrepareValueResult::SetType(FEmitContext& Context, const FRequestedType& R
 				{
 					PreparedType.MergeComponent(Index, Evaluation);
 				}
+				else
+				{
+					// If component wasn't requested, it can be replaced with constant-0
+					PreparedType.MergeComponent(Index, EExpressionEvaluation::ConstantZero);
+				}
 			}
 		}
 		return true;
@@ -1052,26 +1103,24 @@ bool FPrepareValueResult::SetType(FEmitContext& Context, const FRequestedType& R
 
 	if (TryMergePreparedType(Context, Type.StructType, ComponentType))
 	{
-		const int32 NumComponents = RequestedType.GetNumComponents();
+		const int32 NumComponents = Type.GetNumComponents();
 		for (int32 Index = 0; Index < NumComponents; ++Index)
 		{
 			if (RequestedType.IsComponentRequested(Index))
 			{
 				PreparedType.MergeComponent(Index, Type.GetComponent(Index));
 			}
+			else
+			{
+				// If component wasn't requested, it can be replaced with constant-0
+				PreparedType.MergeComponent(Index, EExpressionEvaluation::ConstantZero);
+			}
 		}
 		return true;
 	}
 	return false;
 }
 
-bool FPrepareValueResult::SetForwardValue(FEmitContext& Context, FEmitScope& Scope, const FRequestedType& RequestedType, FExpression* InForwardValue)
-{
-	FPreparedType ForwardPreparedType = Context.PrepareExpression(InForwardValue, Scope, RequestedType);
-	ForwardPreparedType.SetForwardValue(RequestedType, InForwardValue);
-	return SetType(Context, RequestedType, ForwardPreparedType);
-}
-
 void FStatement::EmitPreshader(FEmitContext& Context, FEmitScope& Scope, const FRequestedType& RequestedType, TArrayView<const FEmitPreshaderScope> Scopes, Shader::FPreshaderData& OutPreshader) const
 {
 	check(false);
@@ -1097,95 +1146,24 @@ void FExpression::EmitValuePreshader(FEmitContext& Context, FEmitScope& Scope, c
 	check(false);
 }
 
-FExpression* FExpression::GetForwardValue(const FRequestedType& RequestedType)
-{
-	TArray<FExpression*, TInlineAllocator<16>> VisitedExpressions;
-	FExpression* Expression = this;
-	int32 BaseComponentIndex = 0;
-
-	while (true)
-	{
-		const FPreparedType& PreparedType = Expression->PrepareValueResult.PreparedType;
-
-		VisitedExpressions.Add(Expression);
-
-		FExpression* ForwardValue = nullptr;
-		int32 ForwardComponentIndex = INDEX_NONE;
-		int32 StartComponentIndex = INDEX_NONE;
-		// Scan the required components, and see if they all have the same ForwardValue, in the correct order
-		for (int32 Index = 0; Index < RequestedType.RequestedComponents.Num(); ++Index)
-		{
-			if (RequestedType.RequestedComponents[Index])
-			{
-				const FPreparedComponent Component = PreparedType.GetComponent(BaseComponentIndex + Index);
-				if (!Component.IsNone())
-				{
-					if (!Component.ForwardValue)
-					{
-						return nullptr;
-					}
-
-					if (StartComponentIndex == INDEX_NONE)
-					{
-						StartComponentIndex = Index;
-						ForwardValue = Component.ForwardValue;
-						ForwardComponentIndex = Component.ForwardComponentIndex;
-					}
-					else if (ForwardValue != Component.ForwardValue ||
-						ForwardComponentIndex + Index - StartComponentIndex != Component.ForwardComponentIndex)
-					{
-						return nullptr;
-					}
-				}
-			}
-		}
-
-		check(ForwardValue);
-		if (ForwardValue == this)
-		{
-			// Don't forward to ourselves
-			return nullptr;
-		}
-		if (ForwardComponentIndex == StartComponentIndex)
-		{
-			// Forward to a value with no component offset, can use that directly
-			return ForwardValue;
-		}
-
-		// If there's a component offset, we can't forward directly
-		// Instead, recursively check to see if the ForwardValue itself has a Forward value
-		// This can happen when forwarding across multiple set/get struct field expressions
-		BaseComponentIndex = ForwardComponentIndex - StartComponentIndex;
-		Expression = ForwardValue;
-
-		// Make sure we're not in a loop
-		check(!VisitedExpressions.Contains(Expression));
-	}
-
-	checkNoEntry();
-	return nullptr;
-}
-
 FEmitShaderExpression* FExpression::GetValueShader(FEmitContext& Context, FEmitScope& Scope, const FRequestedType& RequestedType)
 {
 	FOwnerScope OwnerScope(*Context.Errors, GetOwner());
-	FExpression* ForwardValue = GetForwardValue(RequestedType);
-	if (ForwardValue)
-	{
-		return ForwardValue->GetValueShader(Context, Scope, RequestedType);
-	}
-
+	
 	const EExpressionEvaluation Evaluation = PrepareValueResult.PreparedType.GetEvaluation(Scope, RequestedType);
-	check(Evaluation != EExpressionEvaluation::None);
+	check(Evaluation != EExpressionEvaluation::None && Evaluation != EExpressionEvaluation::Unknown);
 
 	FEmitShaderExpression* Value = nullptr;
-	if (Evaluation == EExpressionEvaluation::Constant || Evaluation == EExpressionEvaluation::Preshader)
+	if (Evaluation == EExpressionEvaluation::ConstantZero)
+	{
+		Value = Context.EmitConstantZero(Scope, GetType());
+	}
+	else if (Evaluation == EExpressionEvaluation::Constant || Evaluation == EExpressionEvaluation::Preshader)
 	{
 		Value = Context.EmitPreshaderOrConstant(Scope, RequestedType, this);
 	}
 	else
 	{
-		check(Evaluation != EExpressionEvaluation::None && Evaluation != EExpressionEvaluation::Unknown);
 		FEmitValueShaderResult Result;
 		EmitValueShader(Context, Scope, RequestedType, Result);
 		Value = Result.Code;
@@ -1228,17 +1206,19 @@ FEmitShaderExpression* FExpression::GetValueShader(FEmitContext& Context, FEmitS
 Shader::FType FExpression::GetValuePreshader(FEmitContext& Context, FEmitScope& Scope, const FRequestedType& RequestedType, Shader::FPreshaderData& OutPreshader)
 {
 	FOwnerScope OwnerScope(*Context.Errors, GetOwner());
-	FExpression* ForwardValue = GetForwardValue(RequestedType);
-	if (ForwardValue)
-	{
-		return ForwardValue->GetValuePreshader(Context, Scope, RequestedType, OutPreshader);
-	}
-
+	
 	const int32 PrevStackPosition = Context.PreshaderStackPosition;
 	const EExpressionEvaluation Evaluation = PrepareValueResult.PreparedType.GetEvaluation(Scope, RequestedType);
+	check(Evaluation != EExpressionEvaluation::None && Evaluation != EExpressionEvaluation::Unknown);
 
 	FEmitValuePreshaderResult Result(OutPreshader);
-	if (Evaluation == EExpressionEvaluation::Constant)
+	if (Evaluation == EExpressionEvaluation::ConstantZero)
+	{
+		Context.PreshaderStackPosition++;
+		Result.Type = GetType();
+		OutPreshader.WriteOpcode(Shader::EPreshaderOpcode::ConstantZero).Write(Result.Type);
+	}
+	else if (Evaluation == EExpressionEvaluation::Constant)
 	{
 		const Shader::FValue ConstantValue = GetValueConstant(Context, Scope, RequestedType);
 		Context.PreshaderStackPosition++;
@@ -1247,7 +1227,7 @@ Shader::FType FExpression::GetValuePreshader(FEmitContext& Context, FEmitScope&
 	}
 	else
 	{
-		check(Evaluation != EExpressionEvaluation::None && Evaluation != EExpressionEvaluation::Unknown && Evaluation != EExpressionEvaluation::Shader);
+		check(Evaluation != EExpressionEvaluation::Shader);
 		EmitValuePreshader(Context, Scope, RequestedType, Result);
 	}
 	check(Context.PreshaderStackPosition == PrevStackPosition + 1);
@@ -1257,32 +1237,32 @@ Shader::FType FExpression::GetValuePreshader(FEmitContext& Context, FEmitScope&
 Shader::FValue FExpression::GetValueConstant(FEmitContext& Context, FEmitScope& Scope, const FRequestedType& RequestedType)
 {
 	FOwnerScope OwnerScope(*Context.Errors, GetOwner());
-	FExpression* ForwardValue = GetForwardValue(RequestedType);
-	if (ForwardValue)
-	{
-		return ForwardValue->GetValueConstant(Context, Scope, RequestedType);
-	}
-
+	
 	check(!bReentryFlag);
-
 	const EExpressionEvaluation Evaluation = PrepareValueResult.PreparedType.GetEvaluation(Scope, RequestedType);
-	check(Evaluation == EExpressionEvaluation::Constant || Evaluation == EExpressionEvaluation::ConstantLoop);
-	
-	Shader::FPreshaderData ConstantPreshader;
+	if (Evaluation == EExpressionEvaluation::ConstantZero)
 	{
-		FExpressionReentryScope ReentryScope(this);
-		const int32 PrevPreshaderStackPosition = Context.PreshaderStackPosition;
-		FEmitValuePreshaderResult PreshaderResult(ConstantPreshader);
-		EmitValuePreshader(Context, Scope, RequestedType, PreshaderResult);
-		check(Context.PreshaderStackPosition == PrevPreshaderStackPosition + 1);
-		Context.PreshaderStackPosition--;
+		return Shader::FValue(GetType());
 	}
+	else
+	{
+		check(Evaluation == EExpressionEvaluation::Constant || Evaluation == EExpressionEvaluation::ConstantLoop);
+		Shader::FPreshaderData ConstantPreshader;
+		{
+			FExpressionReentryScope ReentryScope(this);
+			const int32 PrevPreshaderStackPosition = Context.PreshaderStackPosition;
+			FEmitValuePreshaderResult PreshaderResult(ConstantPreshader);
+			EmitValuePreshader(Context, Scope, RequestedType, PreshaderResult);
+			check(Context.PreshaderStackPosition == PrevPreshaderStackPosition + 1);
+			Context.PreshaderStackPosition--;
+		}
 
-	// Evaluate the constant preshader and store its value
-	Shader::FPreshaderStack Stack;
-	const Shader::FPreshaderValue PreshaderValue = ConstantPreshader.EvaluateConstant(*Context.Material, Stack);
-	Shader::FValue Result = PreshaderValue.AsShaderValue(Context.TypeRegistry);
-	return Result;
+		// Evaluate the constant preshader and store its value
+		Shader::FPreshaderStack Stack;
+		const Shader::FPreshaderValue PreshaderValue = ConstantPreshader.EvaluateConstant(*Context.Material, Stack);
+		Shader::FValue Result = PreshaderValue.AsShaderValue(Context.TypeRegistry);
+		return Result;
+	}
 }
 
 Shader::FValue FExpression::GetValueConstant(FEmitContext& Context, FEmitScope& Scope, const FRequestedType& RequestedType, const Shader::FType& ResultType)
diff --git a/Engine/Source/Runtime/Engine/Private/HLSLTree/HLSLTreeCommon.cpp b/Engine/Source/Runtime/Engine/Private/HLSLTree/HLSLTreeCommon.cpp
--- a/Engine/Source/Runtime/Engine/Private/HLSLTree/HLSLTreeCommon.cpp
+++ b/Engine/Source/Runtime/Engine/Private/HLSLTree/HLSLTreeCommon.cpp
@@ -51,6 +51,11 @@ FExternalInputDescription GetExternalInputDescription(EExternalInput Input)
 	case EExternalInput::TexCoord6_Ddy: return FExternalInputDescription(TEXT("TexCoord6_Ddy"), Shader::EValueType::Float2);
 	case EExternalInput::TexCoord7_Ddy: return FExternalInputDescription(TEXT("TexCoord7_Ddy"), Shader::EValueType::Float2);
 
+	case EExternalInput::LightmapTexCoord: return FExternalInputDescription(TEXT("LightmapTexCoord"), Shader::EValueType::Float2, EExternalInput::LightmapTexCoord_Ddx, EExternalInput::LightmapTexCoord_Ddy);
+	case EExternalInput::LightmapTexCoord_Ddx: return FExternalInputDescription(TEXT("LightmapTexCoord_Ddx"), Shader::EValueType::Float2);
+	case EExternalInput::LightmapTexCoord_Ddy: return FExternalInputDescription(TEXT("LightmapTexCoord_Ddy"), Shader::EValueType::Float2);
+
+	case EExternalInput::TwoSidedSign: return FExternalInputDescription(TEXT("TwoSidedSign"), Shader::EValueType::Float1);
 	case EExternalInput::VertexColor: return FExternalInputDescription(TEXT("VertexColor"), Shader::EValueType::Float4, EExternalInput::VertexColor_Ddx, EExternalInput::VertexColor_Ddy);
 	case EExternalInput::VertexColor_Ddx: return FExternalInputDescription(TEXT("VertexColor_Ddx"), Shader::EValueType::Float4);
 	case EExternalInput::VertexColor_Ddy: return FExternalInputDescription(TEXT("VertexColor_Ddy"), Shader::EValueType::Float4);
@@ -59,6 +64,7 @@ FExternalInputDescription GetExternalInputDescription(EExternalInput Input)
 	case EExternalInput::WorldPosition_NoOffsets: return FExternalInputDescription(TEXT("WorldPosition_NoOffsets"), Shader::EValueType::Double3, EExternalInput::WorldPosition_Ddx, EExternalInput::WorldPosition_Ddy, EExternalInput::PrevWorldPosition_NoOffsets);
 	case EExternalInput::TranslatedWorldPosition: return FExternalInputDescription(TEXT("TranslatedWorldPosition"), Shader::EValueType::Float3, EExternalInput::WorldPosition_Ddx, EExternalInput::WorldPosition_Ddy, EExternalInput::PrevTranslatedWorldPosition);
 	case EExternalInput::TranslatedWorldPosition_NoOffsets: return FExternalInputDescription(TEXT("TranslatedWorldPosition_NoOffsets"), Shader::EValueType::Float3, EExternalInput::WorldPosition_Ddx, EExternalInput::WorldPosition_Ddy, EExternalInput::PrevTranslatedWorldPosition_NoOffsets);
+	case EExternalInput::ActorWorldPosition: return FExternalInputDescription(TEXT("TranslatedWorldPosition_NoOffsets"), Shader::EValueType::Double3);
 
 	case EExternalInput::PrevWorldPosition: return FExternalInputDescription(TEXT("PrevWorldPosition"), Shader::EValueType::Double3, EExternalInput::WorldPosition_Ddx, EExternalInput::WorldPosition_Ddy);
 	case EExternalInput::PrevWorldPosition_NoOffsets: return FExternalInputDescription(TEXT("PrevWorldPosition_NoOffsets"), Shader::EValueType::Double3, EExternalInput::WorldPosition_Ddx, EExternalInput::WorldPosition_Ddy);
@@ -68,12 +74,30 @@ FExternalInputDescription GetExternalInputDescription(EExternalInput Input)
 	case EExternalInput::WorldPosition_Ddx: return FExternalInputDescription(TEXT("WorldPosition_Ddx"), Shader::EValueType::Float3);
 	case EExternalInput::WorldPosition_Ddy: return FExternalInputDescription(TEXT("WorldPosition_Ddx"), Shader::EValueType::Float3);
 
+	case EExternalInput::WorldNormal: return FExternalInputDescription(TEXT("WorldNormal"), Shader::EValueType::Float3);
+	case EExternalInput::WorldReflection: return FExternalInputDescription(TEXT("WorldReflection"), Shader::EValueType::Float3);
+	case EExternalInput::WorldVertexNormal: return FExternalInputDescription(TEXT("WorldVertexNormal"), Shader::EValueType::Float3);
+	case EExternalInput::WorldVertexTangent: return FExternalInputDescription(TEXT("WorldVertexTangent"), Shader::EValueType::Float3);
+
 	case EExternalInput::ViewportUV: return FExternalInputDescription(TEXT("ViewportUV"), Shader::EValueType::Float2);
 	case EExternalInput::PixelPosition: return FExternalInputDescription(TEXT("PixelPosition"), Shader::EValueType::Float2);
 	case EExternalInput::ViewSize: return FExternalInputDescription(TEXT("ViewSize"), Shader::EValueType::Float2);
 	case EExternalInput::RcpViewSize: return FExternalInputDescription(TEXT("RcpViewSize"), Shader::EValueType::Float2);
-
+	case EExternalInput::FieldOfView: return FExternalInputDescription(TEXT("FieldOfView"), Shader::EValueType::Float1, EExternalInput::None, EExternalInput::None, EExternalInput::PrevFieldOfView);
+	case EExternalInput::TanHalfFieldOfView: return FExternalInputDescription(TEXT("TanHalfFieldOfView"), Shader::EValueType::Float2, EExternalInput::None, EExternalInput::None, EExternalInput::PrevTanHalfFieldOfView);
+	case EExternalInput::CotanHalfFieldOfView: return FExternalInputDescription(TEXT("CotanHalfFieldOfView"), Shader::EValueType::Float2, EExternalInput::None, EExternalInput::None, EExternalInput::PrevCotanHalfFieldOfView);
+	case EExternalInput::TemporalSampleCount: return FExternalInputDescription(TEXT("TemporalSampleCount"), Shader::EValueType::Float1);
+	case EExternalInput::TemporalSampleIndex: return FExternalInputDescription(TEXT("TemporalSampleIndex"), Shader::EValueType::Float1);
+	case EExternalInput::TemporalSampleOffset: return FExternalInputDescription(TEXT("TemporalSampleOffset"), Shader::EValueType::Float2);
+	case EExternalInput::PreExposure: return FExternalInputDescription(TEXT("PreExposure"), Shader::EValueType::Float1);
+	case EExternalInput::RcpPreExposure: return FExternalInputDescription(TEXT("RcpPreExposure"), Shader::EValueType::Float1);
+	case EExternalInput::RuntimeVirtualTextureOutputLevel: return FExternalInputDescription(TEXT("RuntimeVirtualTextureOutputLevel"), Shader::EValueType::Float1);
+	case EExternalInput::RuntimeVirtualTextureOutputDerivative: return FExternalInputDescription(TEXT("RuntimeVirtualTextureOutputDerivative"), Shader::EValueType::Float2);
+	case EExternalInput::RuntimeVirtualTextureMaxLevel: return FExternalInputDescription(TEXT("RuntimeVirtualTextureMaxLevel"), Shader::EValueType::Float1);
+
+	case EExternalInput::CameraVector: return FExternalInputDescription(TEXT("CameraVector"), Shader::EValueType::Float3);
 	case EExternalInput::CameraWorldPosition: return FExternalInputDescription(TEXT("CameraWorldPosition"), Shader::EValueType::Double3, EExternalInput::None, EExternalInput::None, EExternalInput::PrevCameraWorldPosition);
+	case EExternalInput::ViewWorldPosition: return FExternalInputDescription(TEXT("ViewWorldPosition"), Shader::EValueType::Double3, EExternalInput::None, EExternalInput::None, EExternalInput::PrevViewWorldPosition);
 	case EExternalInput::PreViewTranslation: return FExternalInputDescription(TEXT("PreViewTranslation"), Shader::EValueType::Double3, EExternalInput::None, EExternalInput::None, EExternalInput::PrevPreViewTranslation);
 	case EExternalInput::TangentToWorld: return FExternalInputDescription(TEXT("TangentToWorld"), Shader::EValueType::Float4x4);
 	case EExternalInput::LocalToWorld: return FExternalInputDescription(TEXT("LocalToWorld"), Shader::EValueType::Double4x4, EExternalInput::None, EExternalInput::None, EExternalInput::PrevLocalToWorld);
@@ -87,7 +111,11 @@ FExternalInputDescription GetExternalInputDescription(EExternalInput Input)
 	case EExternalInput::ParticleToWorld: return FExternalInputDescription(TEXT("ParticleToWorld"), Shader::EValueType::Double4x4);
 	case EExternalInput::InstanceToWorld: return FExternalInputDescription(TEXT("InstanceToWorld"), Shader::EValueType::Double4x4);
 
+	case EExternalInput::PrevFieldOfView: return FExternalInputDescription(TEXT("PrevFieldOfView"), Shader::EValueType::Float2);
+	case EExternalInput::PrevTanHalfFieldOfView: return FExternalInputDescription(TEXT("PrevTanHalfFieldOfView"), Shader::EValueType::Float2);
+	case EExternalInput::PrevCotanHalfFieldOfView: return FExternalInputDescription(TEXT("PrevCotanHalfFieldOfView"), Shader::EValueType::Float2);
 	case EExternalInput::PrevCameraWorldPosition: return FExternalInputDescription(TEXT("PrevCameraWorldPosition"), Shader::EValueType::Double3);
+	case EExternalInput::PrevViewWorldPosition: return FExternalInputDescription(TEXT("PrevViewWorldPosition"), Shader::EValueType::Double3);
 	case EExternalInput::PrevPreViewTranslation: return FExternalInputDescription(TEXT("PrevPreViewTranslation"), Shader::EValueType::Double3);
 	case EExternalInput::PrevLocalToWorld: return FExternalInputDescription(TEXT("PrevLocalToWorld"), Shader::EValueType::Double4x4);
 	case EExternalInput::PrevWorldToLocal: return FExternalInputDescription(TEXT("PrevWorldToLocal"), Shader::EValueType::DoubleInverse4x4);
@@ -107,6 +135,10 @@ FExternalInputDescription GetExternalInputDescription(EExternalInput Input)
 	case EExternalInput::PrevGameTime: return FExternalInputDescription(TEXT("PrevGameTime"), Shader::EValueType::Float1);
 	case EExternalInput::PrevRealTime: return FExternalInputDescription(TEXT("PrevRealTime"), Shader::EValueType::Float1);
 
+	case EExternalInput::ParticleColor: return FExternalInputDescription(TEXT("ParticleColor"), Shader::EValueType::Float4);
+	case EExternalInput::ParticleTranslatedWorldPosition: return FExternalInputDescription(TEXT("ParticleTranslatedWorldPosition"), Shader::EValueType::Float3);
+	case EExternalInput::ParticleRadius: return FExternalInputDescription(TEXT("ParticleRadius"), Shader::EValueType::Float1);
+
 	default: checkNoEntry(); return FExternalInputDescription(TEXT("Invalid"), Shader::EValueType::Void);
 	}
 }
@@ -207,6 +239,27 @@ void FExpressionConstant::EmitValuePreshader(FEmitContext& Context, FEmitScope&
 	OutResult.Preshader.WriteOpcode(Shader::EPreshaderOpcode::Constant).Write(Value);
 }
 
+void FExpressionMaterialShadingModel::ComputeAnalyticDerivatives(FTree& Tree, FExpressionDerivatives& OutResult) const
+{
+	const Shader::FValue ZeroValue(Shader::EValueType::Float1);
+	OutResult.ExpressionDdx = Tree.NewConstant(ZeroValue);
+	OutResult.ExpressionDdy = OutResult.ExpressionDdx;
+}
+
+bool FExpressionMaterialShadingModel::PrepareValue(FEmitContext& Context, FEmitScope& Scope, const FRequestedType& RequestedType, FPrepareValueResult& OutResult) const
+{
+	return OutResult.SetType(Context, RequestedType, EExpressionEvaluation::Constant, Shader::EValueType::Int1);
+}
+
+void FExpressionMaterialShadingModel::EmitValuePreshader(FEmitContext& Context, FEmitScope& Scope, const FRequestedType& RequestedType, FEmitValuePreshaderResult& OutResult) const
+{
+	Context.ShadingModelsFromCompilation.AddShadingModel(ShadingModel);
+
+	Context.PreshaderStackPosition++;
+	OutResult.Type = Shader::EValueType::Int1;
+	OutResult.Preshader.WriteOpcode(Shader::EPreshaderOpcode::Constant).Write(Shader::FValue((int32)ShadingModel));
+}
+
 void FExpressionMaterialParameter::ComputeAnalyticDerivatives(FTree& Tree, FExpressionDerivatives& OutResult) const
 {
 	const Shader::FType Type = GetShaderValueType(ParameterType);
@@ -242,49 +295,69 @@ void FExpressionMaterialParameter::EmitValueShader(FEmitContext& Context, FEmitS
 		const Shader::FTextureValue* TextureValue = DefaultValue.AsTexture();
 		check(TextureValue);
 
-		FMaterialTextureParameterInfo TextureParameterInfo;
-		TextureParameterInfo.ParameterInfo = ParameterName;
-		TextureParameterInfo.TextureIndex = Context.Material->GetReferencedTextures().Find(TextureValue->Texture);
-		TextureParameterInfo.SamplerSource = SSM_FromTextureAsset; // TODO - Is this needed?
-		check(TextureParameterInfo.TextureIndex != INDEX_NONE);
-
 		const Shader::EValueType TextureType = TextureValue->GetType();
-		EMaterialTextureParameterType TextureParameterType = EMaterialTextureParameterType::Count;
 		const TCHAR* ConstructorName = nullptr;
 		const TCHAR* TextureTypeName = nullptr;
-		switch (TextureType)
+		int32 ParameterIndex = INDEX_NONE;
+
+		if (TextureType == Shader::EValueType::TextureExternal)
 		{
-		case Shader::EValueType::Texture2D:
-			TextureParameterType = EMaterialTextureParameterType::Standard2D;
-			ConstructorName = TEXT("MakeTexture2D");
-			TextureTypeName = TEXT("Texture2D");
-			break;
-		case Shader::EValueType::Texture2DArray:
-			TextureParameterType = EMaterialTextureParameterType::Array2D;
-			ConstructorName = TEXT("MakeTexture2DArray");
-			TextureTypeName = TEXT("Texture2DArray");
-			break;
-		case Shader::EValueType::TextureCube:
-			TextureParameterType = EMaterialTextureParameterType::Cube;
-			ConstructorName = TEXT("MakeTextureCube");
-			TextureTypeName = TEXT("TextureCube");
-			break;
-		case Shader::EValueType::TextureCubeArray:
-			TextureParameterType = EMaterialTextureParameterType::ArrayCube;
-			ConstructorName = TEXT("MakeTextureCubeArray");
-			TextureTypeName = TEXT("TextureCubeArray");
-			break;
-		case Shader::EValueType::Texture3D:
-			TextureParameterType = EMaterialTextureParameterType::Volume;
-			ConstructorName = TEXT("MakeTexture3D");
-			TextureTypeName = TEXT("VolumeTexture");
-			break;
-		default:
-			checkNoEntry();
-			break;
+			check(TextureValue->SamplerType == SAMPLERTYPE_External);
+
+			ConstructorName = TEXT("MakeTextureExternal");
+			TextureTypeName = TEXT("ExternalTexture");
+
+			FMaterialExternalTextureParameterInfo TextureParameterInfo;
+			TextureParameterInfo.ParameterName = NameToScriptName(ParameterName);
+			TextureParameterInfo.ExternalTextureGuid = TextureValue->ExternalTextureGuid;
+			if (TextureValue->Texture)
+			{
+				TextureParameterInfo.SourceTextureIndex = Context.Material->GetReferencedTextures().Find(TextureValue->Texture);
+			}
+			ParameterIndex = Context.MaterialCompilationOutput->UniformExpressionSet.FindOrAddExternalTextureParameter(TextureParameterInfo);
 		}
+		else
+		{
+			EMaterialTextureParameterType TextureParameterType = EMaterialTextureParameterType::Count;
+			switch (TextureType)
+			{
+			case Shader::EValueType::Texture2D:
+				TextureParameterType = EMaterialTextureParameterType::Standard2D;
+				ConstructorName = TEXT("MakeTexture2D");
+				TextureTypeName = TEXT("Texture2D");
+				break;
+			case Shader::EValueType::Texture2DArray:
+				TextureParameterType = EMaterialTextureParameterType::Array2D;
+				ConstructorName = TEXT("MakeTexture2DArray");
+				TextureTypeName = TEXT("Texture2DArray");
+				break;
+			case Shader::EValueType::TextureCube:
+				TextureParameterType = EMaterialTextureParameterType::Cube;
+				ConstructorName = TEXT("MakeTextureCube");
+				TextureTypeName = TEXT("TextureCube");
+				break;
+			case Shader::EValueType::TextureCubeArray:
+				TextureParameterType = EMaterialTextureParameterType::ArrayCube;
+				ConstructorName = TEXT("MakeTextureCubeArray");
+				TextureTypeName = TEXT("TextureCubeArray");
+				break;
+			case Shader::EValueType::Texture3D:
+				TextureParameterType = EMaterialTextureParameterType::Volume;
+				ConstructorName = TEXT("MakeTexture3D");
+				TextureTypeName = TEXT("VolumeTexture");
+				break;
+			default:
+				checkNoEntry();
+				break;
+			}
 
-		const int32 ParameterIndex = Context.MaterialCompilationOutput->UniformExpressionSet.FindOrAddTextureParameter(TextureParameterType, TextureParameterInfo);
+			FMaterialTextureParameterInfo TextureParameterInfo;
+			TextureParameterInfo.ParameterInfo = ParameterName;
+			TextureParameterInfo.TextureIndex = Context.Material->GetReferencedTextures().Find(TextureValue->Texture);
+			TextureParameterInfo.SamplerSource = SSM_FromTextureAsset; // TODO - Is this needed?
+			check(TextureParameterInfo.TextureIndex != INDEX_NONE);
+			ParameterIndex = Context.MaterialCompilationOutput->UniformExpressionSet.FindOrAddTextureParameter(TextureParameterType, TextureParameterInfo);
+		}
 
 		TStringBuilder<256> FormattedCode;
 		FormattedCode.Appendf(TEXT("%s(Material.%s_%d, Material.%s_%dSampler, %d)"), ConstructorName, TextureTypeName, ParameterIndex, TextureTypeName, ParameterIndex, (int32)TextureValue->SamplerType);
@@ -360,6 +433,14 @@ void FExpressionExternalInput::ComputeAnalyticDerivatives(FTree& Tree, FExpressi
 
 FExpression* FExpressionExternalInput::ComputePreviousFrame(FTree& Tree, const FRequestedType& RequestedType) const
 {
+	if (InputType == EExternalInput::ActorWorldPosition)
+	{
+		return Tree.NewBinaryOp(EOperation::VecMulMatrix3, Tree.NewBinaryOp(EOperation::VecMulMatrix3,
+			Tree.NewExpression<FExpressionExternalInput>(EExternalInput::ActorWorldPosition),
+			Tree.NewExpression< FExpressionExternalInput>(EExternalInput::WorldToLocal)),
+			Tree.NewExpression< FExpressionExternalInput>(EExternalInput::PrevLocalToWorld));
+	}
+
 	const FExternalInputDescription InputDesc = GetExternalInputDescription(InputType);
 	if (InputDesc.PreviousFrame != EExternalInput::None)
 	{
@@ -371,16 +452,27 @@ FExpression* FExpressionExternalInput::ComputePreviousFrame(FTree& Tree, const F
 bool FExpressionExternalInput::PrepareValue(FEmitContext& Context, FEmitScope& Scope, const FRequestedType& RequestedType, FPrepareValueResult& OutResult) const
 {
 	const FExternalInputDescription InputDesc = GetExternalInputDescription(InputType);
+
+	switch (InputType)
+	{
+	case EExternalInput::WorldNormal:
+	case EExternalInput::WorldReflection:
+		Context.bReadMaterialNormal = true;
+		break;
+	default:
+		break;
+	}
+
 	return OutResult.SetType(Context, RequestedType, EExpressionEvaluation::Shader, InputDesc.Type);
 }
 
 void FExpressionExternalInput::EmitValueShader(FEmitContext& Context, FEmitScope& Scope, const FRequestedType& RequestedType, FEmitValueShaderResult& OutResult) const
 {
 	const int32 TypeIndex = (int32)InputType;
+	Context.ExternalInputMask[Context.ShaderFrequency][TypeIndex] = true;
 	if (IsTexCoord(InputType))
 	{
 		const int32 TexCoordIndex = TypeIndex - (int32)EExternalInput::TexCoord0;
-		Context.TexCoordMask[Context.ShaderFrequency] |= (1u << TexCoordIndex);
 		OutResult.Code = Context.EmitInlineExpression(Scope, Shader::EValueType::Float2, TEXT("Parameters.TexCoords[%].xy"), TexCoordIndex);
 	}
 	else if (IsTexCoord_Ddx(InputType))
@@ -399,13 +491,18 @@ void FExpressionExternalInput::EmitValueShader(FEmitContext& Context, FEmitScope
 		const TCHAR* Code = nullptr;
 		switch (InputType)
 		{
-		case EExternalInput::VertexColor: Code = TEXT("Parameters.VertexColor"); Context.bUsesVertexColor = (Context.ShaderFrequency != SF_Vertex); break;
-		case EExternalInput::VertexColor_Ddx: Code = TEXT("Parameters.VertexColor_DDX"); Context.bUsesVertexColor = (Context.ShaderFrequency != SF_Vertex); break;
-		case EExternalInput::VertexColor_Ddy: Code = TEXT("Parameters.VertexColor_DDY"); Context.bUsesVertexColor = (Context.ShaderFrequency != SF_Vertex); break;
+		case EExternalInput::LightmapTexCoord: Code = TEXT("GetLightmapUVs(Parameters)"); break;
+		case EExternalInput::LightmapTexCoord_Ddx: Code = TEXT("GetLightmapUVs_DDX(Parameters)"); break;
+		case EExternalInput::LightmapTexCoord_Ddy: Code = TEXT("GetLightmapUVs_DDY(Parameters)"); break;
+		case EExternalInput::TwoSidedSign: Code = TEXT("Parameters.TwoSidedSign"); break;
+		case EExternalInput::VertexColor: Code = TEXT("Parameters.VertexColor"); Context.bUsesVertexColor |= (Context.ShaderFrequency != SF_Vertex); break;
+		case EExternalInput::VertexColor_Ddx: Code = TEXT("Parameters.VertexColor_DDX"); Context.bUsesVertexColor |= (Context.ShaderFrequency != SF_Vertex); break;
+		case EExternalInput::VertexColor_Ddy: Code = TEXT("Parameters.VertexColor_DDY"); Context.bUsesVertexColor |= (Context.ShaderFrequency != SF_Vertex); break;
 		case EExternalInput::WorldPosition: Code = TEXT("GetWorldPosition(Parameters)"); break;
 		case EExternalInput::WorldPosition_NoOffsets: Code = TEXT("GetWorldPosition_NoMaterialOffsets(Parameters)"); break;
 		case EExternalInput::TranslatedWorldPosition: Code = TEXT("GetTranslatedWorldPosition(Parameters)"); break;
 		case EExternalInput::TranslatedWorldPosition_NoOffsets: Code = TEXT("GetTranslatedWorldPosition_NoMaterialOffsets(Parameters)"); break;
+		case EExternalInput::ActorWorldPosition: Code = TEXT("GetActorWorldPosition(Parameters)"); break;
 		case EExternalInput::PrevWorldPosition: Code = TEXT("GetPrevWorldPosition(Parameters)"); break;
 		case EExternalInput::PrevWorldPosition_NoOffsets: Code = TEXT("GetPrevWorldPosition_NoMaterialOffsets(Parameters)"); break;
 		case EExternalInput::PrevTranslatedWorldPosition: Code = TEXT("GetPrevTranslatedWorldPosition(Parameters)"); break;
@@ -413,12 +510,31 @@ void FExpressionExternalInput::EmitValueShader(FEmitContext& Context, FEmitScope
 		case EExternalInput::WorldPosition_Ddx: Code = TEXT("Parameters.WorldPosition_DDX"); break;
 		case EExternalInput::WorldPosition_Ddy: Code = TEXT("Parameters.WorldPosition_DDY"); break;
 
+		case EExternalInput::WorldNormal: Code = TEXT("Parameters.WorldNormal"); break;
+		case EExternalInput::WorldReflection: Code = TEXT("Parameters.ReflectionVector"); break;
+		case EExternalInput::WorldVertexNormal: Code = TEXT("Parameters.TangentToWorld[2]"); break;
+		case EExternalInput::WorldVertexTangent: Code = TEXT("Parameters.TangentToWorld[0]"); break;
+
 		case EExternalInput::ViewportUV: Code = TEXT("GetViewportUV(Parameters)"); break;
 		case EExternalInput::PixelPosition: Code = TEXT("GetPixelPosition(Parameters)"); break;
 		case EExternalInput::ViewSize: Code = TEXT("View.ViewSizeAndInvSize.xy"); break;
 		case EExternalInput::RcpViewSize: Code = TEXT("View.ViewSizeAndInvSize.zw"); break;
 
+		case EExternalInput::FieldOfView: Code = TEXT("View.FieldOfViewWideAngles"); break;
+		case EExternalInput::TanHalfFieldOfView: Code = TEXT("GetTanHalfFieldOfView()"); break;
+		case EExternalInput::CotanHalfFieldOfView: Code = TEXT("GetCotanHalfFieldOfView()"); break;
+		case EExternalInput::TemporalSampleCount: Code = TEXT("View.TemporalAAParams.y"); break;
+		case EExternalInput::TemporalSampleIndex: Code = TEXT("View.TemporalAAParams.x"); break;
+		case EExternalInput::TemporalSampleOffset: Code = TEXT("View.TemporalAAParams.zw"); break;
+		case EExternalInput::PreExposure: Code = TEXT("View.PreExposure.x"); break;
+		case EExternalInput::RcpPreExposure: Code = TEXT("View.OneOverPreExposure.x"); break;
+		case EExternalInput::RuntimeVirtualTextureOutputLevel:  Code = TEXT("View.RuntimeVirtualTextureMipLevel.x"); break;
+		case EExternalInput::RuntimeVirtualTextureOutputDerivative: Code = TEXT("View.RuntimeVirtualTextureMipLevel.zw"); break;
+		case EExternalInput::RuntimeVirtualTextureMaxLevel:  Code = TEXT("View.RuntimeVirtualTextureMipLevel.y"); break;
+
+		case EExternalInput::CameraVector: Code = TEXT("Parameters.CameraVector"); break;
 		case EExternalInput::CameraWorldPosition: Code = TEXT("ResolvedView.WorldCameraOrigin"); break;
+		case EExternalInput::ViewWorldPosition: Code = TEXT("ResolvedView.WorldViewOrigin"); break;
 		case EExternalInput::PreViewTranslation: Code = TEXT("ResolvedView.PreViewTranslation"); break;
 		case EExternalInput::TangentToWorld: Code = TEXT("Parameters.TangentToWorld"); break;
 		case EExternalInput::LocalToWorld: Code = TEXT("GetLocalToWorld(Parameters)"); break;
@@ -432,7 +548,11 @@ void FExpressionExternalInput::EmitValueShader(FEmitContext& Context, FEmitScope
 		case EExternalInput::ParticleToWorld: Code = TEXT("Parameters.Particle.ParticleToWorld"); break;
 		case EExternalInput::InstanceToWorld: Code = TEXT("GetInstanceToWorld(Parameters)"); break;
 
+		case EExternalInput::PrevFieldOfView: Code = TEXT("View.PrevFieldOfViewWideAngles"); break;
+		case EExternalInput::PrevTanHalfFieldOfView: Code = TEXT("GetPrevTanHalfFieldOfView()"); break;
+		case EExternalInput::PrevCotanHalfFieldOfView: Code = TEXT("GetPrevCotanHalfFieldOfView()"); break;
 		case EExternalInput::PrevCameraWorldPosition: Code = TEXT("ResolvedView.PrevWorldCameraOrigin"); break;
+		case EExternalInput::PrevViewWorldPosition: Code = TEXT("ResolvedView.PrevWorldViewOrigin"); break;
 		case EExternalInput::PrevPreViewTranslation: Code = TEXT("ResolvedView.PrevPreViewTranslation"); break;
 		case EExternalInput::PrevLocalToWorld: Code = TEXT("GetPrevLocalToWorld(Parameters)"); break;
 		case EExternalInput::PrevWorldToLocal: Code = TEXT("GetPrimitiveData(Parameters).PreviousWorldToLocal"); break;
@@ -449,6 +569,11 @@ void FExpressionExternalInput::EmitValueShader(FEmitContext& Context, FEmitScope
 		case EExternalInput::DeltaTime: Code = TEXT("View.DeltaTime"); break;
 		case EExternalInput::PrevGameTime: Code = TEXT("View.PrevFrameGameTime"); break;
 		case EExternalInput::PrevRealTime: Code = TEXT("View.PrevFrameRealTime"); break;
+
+		case EExternalInput::ParticleColor: Code = TEXT("Parameters.Particle.Color"); Context.bUsesParticleColor |= (Context.ShaderFrequency != SF_Vertex); break;
+		case EExternalInput::ParticleTranslatedWorldPosition: Code = TEXT("Parameters.Particle.TranslatedWorldPositionAndSize.xyz"); Context.bNeedsParticlePosition = true; break;
+		case EExternalInput::ParticleRadius: Code = TEXT("Parameters.Particle.TranslatedWorldPositionAndSize.w"); Context.bNeedsParticlePosition = true; break;
+
 		default:
 			checkNoEntry();
 			break;
@@ -503,6 +628,63 @@ void FExpressionMaterialSceneTexture::EmitValueShader(FEmitContext& Context, FEm
 	OutResult.Code = EmitLookup;
 }
 
+bool FExpressionMaterialNoise::PrepareValue(FEmitContext& Context, FEmitScope& Scope, const FRequestedType& RequestedType, FPrepareValueResult& OutResult) const
+{
+	const FPreparedType& PositionType = Context.PrepareExpression(PositionExpression, Scope, ERequestedType::Vector3);
+	const FPreparedType& FilterWidthType = Context.PrepareExpression(FilterWidthExpression, Scope, ERequestedType::Scalar);
+	if (PositionType.IsVoid() || FilterWidthType.IsVoid())
+	{
+		return false;
+	}
+	return OutResult.SetType(Context, RequestedType, EExpressionEvaluation::Shader, Shader::EValueType::Float1);
+}
+
+void FExpressionMaterialNoise::EmitValueShader(FEmitContext& Context, FEmitScope& Scope, const FRequestedType& RequestedType, FEmitValueShaderResult& OutResult) const
+{
+	bool bIsLWC = Shader::IsLWCType(PositionExpression->GetPreparedType().ValueComponentType);
+	FEmitShaderExpression* EmitPosition = PositionExpression->GetValueShader(Context, Scope, bIsLWC ? Shader::EValueType::Double3 : Shader::EValueType::Float3);
+	FEmitShaderExpression* EmitFilterWidth = FilterWidthExpression->GetValueShader(Context, Scope, Shader::EValueType::Float1);
+
+	if (bIsLWC)
+	{
+		// If Noise is driven by a LWC position, just take the offset within the current tile
+		// Will generate discontinuity in noise at tile boudaries
+		// Could potentially add noise functions that operate directly on LWC values, but that would be very expensive
+		EmitPosition = Context.EmitExpression(Scope, Shader::EValueType::Float3, TEXT("LWCNormalizeTile(%).Offset"), EmitPosition);
+	}
+
+	OutResult.Code = Context.EmitExpression(Scope, Shader::EValueType::Float1, TEXT("MaterialExpressionNoise(%,%,%,%,%,%,%,%,%,%,%,%)"),
+		EmitPosition,
+		Parameters.Scale,
+		Parameters.Quality,
+		Parameters.NoiseFunction,
+		Parameters.bTurbulence,
+		Parameters.Levels,
+		Parameters.OutputMin,
+		Parameters.OutputMax,
+		Parameters.LevelScale,
+		EmitFilterWidth,
+		Parameters.bTiling,
+		Parameters.RepeatSize);
+}
+
+namespace Private
+{
+Shader::EValueType GetTexCoordType(Shader::EValueType TextureType)
+{
+	switch (TextureType)
+	{
+	case Shader::EValueType::Texture2D:
+	case Shader::EValueType::TextureExternal: return Shader::EValueType::Float2;
+	case Shader::EValueType::Texture2DArray:
+	case Shader::EValueType::TextureCube:
+	case Shader::EValueType::Texture3D: return Shader::EValueType::Float3;
+	case Shader::EValueType::TextureCubeArray: return Shader::EValueType::Float4;
+	default: checkNoEntry(); return Shader::EValueType::Void;
+	}
+}
+}
+
 bool FExpressionTextureSample::PrepareValue(FEmitContext& Context, FEmitScope& Scope, const FRequestedType& RequestedType, FPrepareValueResult& OutResult) const
 {
 	const FPreparedType& TextureType = Context.PrepareExpression(TextureExpression, Scope, ERequestedType::Texture);
@@ -511,14 +693,15 @@ bool FExpressionTextureSample::PrepareValue(FEmitContext& Context, FEmitScope& S
 		return Context.Errors->AddError(TEXT("Expected texture"));
 	}
 
-	const FPreparedType& TexCoordType = Context.PrepareExpression(TexCoordExpression, Scope, ERequestedType::Vector2);
+	const FRequestedType RequestedTexCoordType = Private::GetTexCoordType(TextureType.GetType());
+	const FPreparedType& TexCoordType = Context.PrepareExpression(TexCoordExpression, Scope, RequestedTexCoordType);
 	if (TexCoordType.IsVoid())
 	{
 		return false;
 	}
 
-	Context.PrepareExpression(TexCoordDerivatives.ExpressionDdx, Scope, ERequestedType::Vector2);
-	Context.PrepareExpression(TexCoordDerivatives.ExpressionDdy, Scope, ERequestedType::Vector2);
+	Context.PrepareExpression(TexCoordDerivatives.ExpressionDdx, Scope, RequestedTexCoordType);
+	Context.PrepareExpression(TexCoordDerivatives.ExpressionDdy, Scope, RequestedTexCoordType);
 
 	return OutResult.SetType(Context, RequestedType, EExpressionEvaluation::Shader, Shader::EValueType::Float4);
 }
@@ -545,11 +728,10 @@ void FExpressionTextureSample::EmitValueShader(FEmitContext& Context, FEmitScope
 	case Shader::EValueType::Texture3D:
 		SampleFunctionName = TEXT("Texture3DSample");
 		break;
-	/*case MCT_TextureExternal:
-		// TODO
+	case Shader::EValueType::TextureExternal:
 		SampleFunctionName = TEXT("TextureExternalSample");
 		break;
-	case MCT_TextureVirtual:
+	/*case MCT_TextureVirtual:
 		// TODO
 		SampleFunctionName = TEXT("TextureVirtualSample");
 		bVirtualTexture = true;
@@ -581,14 +763,15 @@ void FExpressionTextureSample::EmitValueShader(FEmitContext& Context, FEmitScope
 		break;
 	}
 
-	FEmitShaderExpression* TexCoordValue = TexCoordExpression->GetValueShader(Context, Scope, Shader::EValueType::Float2);
+	const Shader::EValueType TexCoordType = Private::GetTexCoordType(TextureType);
+	FEmitShaderExpression* TexCoordValue = TexCoordExpression->GetValueShader(Context, Scope, TexCoordType);
 	FEmitShaderExpression* TexCoordValueDdx = nullptr;
 	FEmitShaderExpression* TexCoordValueDdy = nullptr;
 	FEmitShaderExpression* TextureResult = nullptr;
 	if (TexCoordDerivatives.IsValid())
 	{
-		TexCoordValueDdx = TexCoordDerivatives.ExpressionDdx->GetValueShader(Context, Scope, Shader::EValueType::Float2);
-		TexCoordValueDdy = TexCoordDerivatives.ExpressionDdy->GetValueShader(Context, Scope, Shader::EValueType::Float2);
+		TexCoordValueDdx = TexCoordDerivatives.ExpressionDdx->GetValueShader(Context, Scope, TexCoordType);
+		TexCoordValueDdy = TexCoordDerivatives.ExpressionDdy->GetValueShader(Context, Scope, TexCoordType);
 		TextureResult = Context.EmitExpression(Scope, Shader::EValueType::Float4, TEXT("%Grad(%.Texture, %, %, %, %)"),
 			SampleFunctionName,
 			EmitTexture,
@@ -642,7 +825,6 @@ bool FExpressionGetStructField::PrepareValue(FEmitContext& Context, FEmitScope&
 		return Context.Errors->AddErrorf(TEXT("Expected type %s"), StructType->Name);
 	}
 
-	//StructPreparedType.SetForwardValue(RequestedStructType, StructExpression);
 	return OutResult.SetType(Context, RequestedType, StructPreparedType.GetFieldType(Field));
 }
 
@@ -711,9 +893,6 @@ bool FExpressionSetStructField::PrepareValue(FEmitContext& Context, FEmitScope&
 	const FRequestedType RequestedFieldType = RequestedType.GetField(Field);
 	FPreparedType FieldPreparedType = Context.PrepareExpression(FieldExpression, Scope, RequestedFieldType);
 
-	StructPreparedType.SetForwardValue(RequestedStructType, StructExpression);
-	FieldPreparedType.SetForwardValue(RequestedFieldType, FieldExpression);
-
 	FPreparedType ResultType(StructPreparedType);
 	if (ResultType.IsVoid())
 	{
@@ -731,12 +910,16 @@ void FExpressionSetStructField::EmitValueShader(FEmitContext& Context, FEmitScop
 
 	const FRequestedType RequestedFieldType = RequestedType.GetField(Field);
 	const EExpressionEvaluation FieldEvaluation = FieldExpression->GetEvaluation(Scope, RequestedFieldType);
-	FEmitShaderExpression* FieldValue = FieldExpression->GetValueShader(Context, Scope, RequestedFieldType, Field->Type);
 
-	if (StructEvaluation == EExpressionEvaluation::None)
-	{
-		check(FieldEvaluation != EExpressionEvaluation::None);
+	FEmitShaderExpression* StructValue = (StructEvaluation != EExpressionEvaluation::None) ? StructExpression->GetValueShader(Context, Scope, RequestedStructType) : nullptr;
+	FEmitShaderExpression* FieldValue = (FieldEvaluation != EExpressionEvaluation::None) ? FieldExpression->GetValueShader(Context, Scope, RequestedFieldType, Field->Type) : nullptr;
 
+	if (StructEvaluation == EExpressionEvaluation::None && FieldEvaluation == EExpressionEvaluation::None)
+	{
+		OutResult.Code = Context.EmitExpression(Scope, StructType, TEXT("((%)0)"), StructType->Name);
+	}
+	else if (StructEvaluation == EExpressionEvaluation::None)
+	{
 		// StructExpression is not used, so default to a zero-initialized struct
 		// This will happen if all the accessed struct fields are explicitly defined
 		OutResult.Code = Context.EmitExpression(Scope, StructType, TEXT("%_Set%((%)0, %)"),
@@ -745,9 +928,14 @@ void FExpressionSetStructField::EmitValueShader(FEmitContext& Context, FEmitScop
 			StructType->Name,
 			FieldValue);
 	}
+	else if (FieldEvaluation == EExpressionEvaluation::None)
+	{
+		// Don't need field, can just forward the struct value
+		OutResult.Code = StructValue;
+	}
 	else
 	{
-		FEmitShaderExpression* StructValue = StructExpression->GetValueShader(Context, Scope, RequestedStructType);
+		check(StructValue->Type.StructType == StructType);
 		OutResult.Code = Context.EmitExpression(Scope, StructType, TEXT("%_Set%(%, %)"),
 			StructType->Name,
 			Field->Name,
@@ -790,10 +978,13 @@ void FExpressionSetStructField::EmitValuePreshader(FEmitContext& Context, FEmitS
 bool FExpressionSelect::PrepareValue(FEmitContext& Context, FEmitScope& Scope, const FRequestedType& RequestedType, FPrepareValueResult& OutResult) const
 {
 	const FPreparedType& ConditionType = Context.PrepareExpression(ConditionExpression, Scope, ERequestedType::Scalar);
-	if (ConditionType.GetEvaluation(Scope, ERequestedType::Scalar) == EExpressionEvaluation::Constant)
+	const EExpressionEvaluation ConditionEvaluation = ConditionType.GetEvaluation(Scope, ERequestedType::Scalar);
+	if (ConditionEvaluation == EExpressionEvaluation::Constant)
 	{
 		const bool bCondition = ConditionExpression->GetValueConstant(Context, Scope, Shader::EValueType::Bool1).AsBoolScalar();
-		return OutResult.SetForwardValue(Context, Scope, RequestedType, bCondition ? TrueExpression : FalseExpression);
+		FPreparedType ResultType = Context.PrepareExpression(bCondition ? TrueExpression : FalseExpression, Scope, RequestedType);
+		ResultType.MergeEvaluation(EExpressionEvaluation::Shader); // TODO - support preshader
+		return OutResult.SetType(Context, RequestedType, ResultType);
 	}
 
 	const FPreparedType& LhsType = Context.PrepareExpression(FalseExpression, Scope, RequestedType);
@@ -805,24 +996,40 @@ bool FExpressionSelect::PrepareValue(FEmitContext& Context, FEmitScope& Scope, c
 		return Context.Errors->AddError(TEXT("Type mismatch"));
 	}
 
-	return OutResult.SetType(Context, RequestedType, MergePreparedTypes(LhsType, RhsType));
+	FPreparedType ResultType = MergePreparedTypes(LhsType, RhsType);
+	ResultType.MergeEvaluation(ConditionEvaluation);
+	ResultType.MergeEvaluation(EExpressionEvaluation::Shader); // TODO - support preshader
+	return OutResult.SetType(Context, RequestedType, ResultType);
 }
 
 void FExpressionSelect::EmitValueShader(FEmitContext& Context, FEmitScope& Scope, const FRequestedType& RequestedType, FEmitValueShaderResult& OutResult) const
 {
-	const Shader::FType LocalType = GetType();
-	FEmitShaderExpression* TrueValue = TrueExpression->GetValueShader(Context, Scope, RequestedType, LocalType);
-	FEmitShaderExpression* FalseValue = FalseExpression->GetValueShader(Context, Scope, RequestedType, LocalType);
+	const EExpressionEvaluation ConditionEvaluation = ConditionExpression->GetEvaluation(Scope, ERequestedType::Scalar);
+	if (ConditionEvaluation == EExpressionEvaluation::Constant)
+	{
+		const bool bCondition = ConditionExpression->GetValueConstant(Context, Scope, Shader::EValueType::Bool1).AsBoolScalar();
+		FExpression* InputExpression = bCondition ? TrueExpression : FalseExpression;
+		OutResult.Code = InputExpression->GetValueShader(Context, Scope, RequestedType);
+	}
+	else
+	{
+		const Shader::FType LocalType = GetType();
+		FEmitShaderExpression* TrueValue = TrueExpression->GetValueShader(Context, Scope, RequestedType, LocalType);
+		FEmitShaderExpression* FalseValue = FalseExpression->GetValueShader(Context, Scope, RequestedType, LocalType);
 
-	OutResult.Code = Context.EmitExpression(Scope, LocalType, TEXT("(% ? % : %)"),
-		ConditionExpression->GetValueShader(Context, Scope, Shader::EValueType::Bool1),
-		TrueValue,
-		FalseValue);
+		OutResult.Code = Context.EmitExpression(Scope, LocalType, TEXT("(% ? % : %)"),
+			ConditionExpression->GetValueShader(Context, Scope, Shader::EValueType::Bool1),
+			TrueValue,
+			FalseValue);
+	}
 }
 
 void FExpressionSelect::EmitValuePreshader(FEmitContext& Context, FEmitScope& Scope, const FRequestedType& RequestedType, FEmitValuePreshaderResult& OutResult) const
 {
-	check(false); // TODO
+	// TODO
+	Context.PreshaderStackPosition++;
+	OutResult.Type = GetType();
+	OutResult.Preshader.WriteOpcode(Shader::EPreshaderOpcode::ConstantZero).Write(OutResult.Type);
 }
 
 FExpression* FTree::NewUnaryOp(EOperation Op, FExpression* Input)
@@ -837,10 +1044,24 @@ FExpression* FTree::NewBinaryOp(EOperation Op, FExpression* Lhs, FExpression* Rh
 	return NewExpression<FExpressionOperation>(Op, Inputs);
 }
 
+FExpression* FTree::NewCross(FExpression* Lhs, FExpression* Rhs)
+{
+	//c_P[0] = v_A[1] * v_B[2] - v_A[2] * v_B[1];
+	//c_P[1] = -(v_A[0] * v_B[2] - v_A[2] * v_B[0]);
+	//c_P[2] = v_A[0] * v_B[1] - v_A[1] * v_B[0];
+	FExpression* Lhs0 = NewExpression<FExpressionSwizzle>(FSwizzleParameters(1, 0, 0), Lhs);
+	FExpression* Lhs1 = NewExpression<FExpressionSwizzle>(FSwizzleParameters(2, 2, 1), Lhs);
+	FExpression* Rhs0 = NewExpression<FExpressionSwizzle>(FSwizzleParameters(2, 2, 1), Rhs);
+	FExpression* Rhs1 = NewExpression<FExpressionSwizzle>(FSwizzleParameters(1, 0, 0), Rhs);
+	return NewSub(NewMul(NewMul(Lhs0, Rhs0), NewConstant(FVector3f(1.0f, -1.0f, 1.0f))), NewMul(Lhs1, Rhs1));
+}
+
 FExpressionOperation::FExpressionOperation(EOperation InOp, TConstArrayView<FExpression*> InInputs) : Op(InOp)
 {
 	const FOperationDescription OpDesc = GetOperationDescription(InOp);
 	check(OpDesc.NumInputs == InInputs.Num());
+	check(InInputs.Num() <= MaxInputs);
+
 	for (int32 i = 0; i < OpDesc.NumInputs; ++i)
 	{
 		Inputs[i] = InInputs[i];
@@ -848,6 +1069,153 @@ FExpressionOperation::FExpressionOperation(EOperation InOp, TConstArrayView<FExp
 	}
 }
 
+namespace Private
+{
+struct FOperationRequestedTypes
+{
+	FRequestedType InputType[FExpressionOperation::MaxInputs];
+	bool bIsMatrixOperation = false;
+};
+struct FOperationTypes
+{
+	
+	Shader::EValueType InputType[FExpressionOperation::MaxInputs];
+	Shader::EValueType ResultType;
+	bool bIsLWC = false;
+};
+
+FOperationRequestedTypes GetOperationRequestedTypes(EOperation Op, const FRequestedType& RequestedType)
+{
+	const FOperationDescription OpDesc = GetOperationDescription(Op);
+	FOperationRequestedTypes Types;
+	for (int32 Index = 0; Index < OpDesc.NumInputs; ++Index)
+	{
+		Types.InputType[Index] = RequestedType;
+	}
+	switch (Op)
+	{
+	case EOperation::Length:
+	case EOperation::Normalize:
+		Types.InputType[0] = ERequestedType::Vector4;
+		break;
+	case EOperation::Dot:
+		Types.InputType[0] = ERequestedType::Vector4;
+		Types.InputType[1] = ERequestedType::Vector4;
+		break;
+	case EOperation::VecMulMatrix3:
+		Types.bIsMatrixOperation = true;
+		Types.InputType[0] = ERequestedType::Vector3;
+		Types.InputType[1] = ERequestedType::Matrix4x4;
+		break;
+	case EOperation::VecMulMatrix4:
+		Types.bIsMatrixOperation = true;
+		Types.InputType[0] = ERequestedType::Vector3;
+		Types.InputType[1] = ERequestedType::Matrix4x4;
+		break;
+	case EOperation::Matrix3MulVec:
+	case EOperation::Matrix4MulVec:
+		// No LWC for transpose matrices
+		Types.bIsMatrixOperation = true;
+		Types.InputType[0] = ERequestedType::Matrix4x4;
+		Types.InputType[1] = ERequestedType::Vector3;
+		break;
+	default:
+		break;
+	}
+	return Types;
+}
+
+FOperationTypes GetOperationTypes(EOperation Op, TConstArrayView<Shader::EValueType> InputTypes)
+{
+	Shader::EValueComponentType ComponentType = Shader::EValueComponentType::Void;
+	int32 NumComponents = 0;
+	for (int32 Index = 0; Index < InputTypes.Num(); ++Index)
+	{
+		const Shader::FValueTypeDescription InputTypeDesc = Shader::GetValueTypeDescription(InputTypes[Index]);
+		ComponentType = Shader::CombineComponentTypes(ComponentType, InputTypeDesc.ComponentType);
+		NumComponents = FMath::Max<int32>(NumComponents, InputTypeDesc.NumComponents);
+	}
+
+	Shader::EValueType IntermediateType = Shader::MakeValueType(ComponentType, NumComponents);
+
+	FOperationTypes Types;
+	for (int32 Index = 0; Index < InputTypes.Num(); ++Index)
+	{
+		Types.InputType[Index] = IntermediateType;
+	}
+	Types.ResultType = IntermediateType;
+	Types.bIsLWC = (ComponentType == Shader::EValueComponentType::Double);
+	switch (Op)
+	{
+	case EOperation::Length:
+		Types.ResultType = Shader::MakeValueType(ComponentType, 1);
+		break;
+	case EOperation::Normalize:
+		Types.ResultType = Shader::MakeNonLWCType(IntermediateType);
+		break;
+	case EOperation::Frac:
+	case EOperation::Rcp:
+	case EOperation::Sqrt:
+	case EOperation::Sign:
+	case EOperation::Saturate:
+	case EOperation::Sin:
+	case EOperation::Cos:
+	case EOperation::Tan:
+	case EOperation::Asin:
+	case EOperation::AsinFast:
+	case EOperation::Acos:
+	case EOperation::AcosFast:
+	case EOperation::Atan:
+	case EOperation::AtanFast:
+		Types.ResultType = Shader::MakeNonLWCType(IntermediateType);
+		break;
+	case EOperation::Log2:
+		// No LWC support yet
+		Types.InputType[0] = Types.ResultType = Shader::MakeNonLWCType(IntermediateType);
+		break;
+	case EOperation::Less:
+	case EOperation::Greater:
+	case EOperation::LessEqual:
+	case EOperation::GreaterEqual:
+		Types.ResultType = Shader::MakeValueType(Shader::EValueComponentType::Bool, NumComponents);
+		break;
+	case EOperation::Fmod:
+		Types.InputType[1] = Types.ResultType = Shader::MakeNonLWCType(IntermediateType);
+		break;
+	case EOperation::PowPositiveClamped:
+	case EOperation::Atan2:
+	case EOperation::Atan2Fast:
+		// No LWC support yet
+		Types.InputType[0] = Types.InputType[1] = Types.ResultType = Shader::MakeNonLWCType(IntermediateType);
+		break;
+	case EOperation::Dot:
+		Types.ResultType = Shader::MakeValueType(ComponentType, 1);
+		break;
+	case EOperation::VecMulMatrix3:
+		// No LWC for matrix3
+		Types.InputType[0] = Shader::EValueType::Float3;
+		Types.InputType[1] = Shader::EValueType::Float4x4;
+		Types.ResultType = Shader::EValueType::Float3;
+		break;
+	case EOperation::VecMulMatrix4:
+		Types.InputType[0] = Shader::EValueType::Float3;
+		Types.InputType[1] = Shader::EValueType::Float4x4;
+		Types.ResultType = Shader::EValueType::Float3;
+		break;
+	case EOperation::Matrix3MulVec:
+	case EOperation::Matrix4MulVec:
+		// No LWC for transpose matrices
+		Types.InputType[0] = Shader::EValueType::Float4x4;
+		Types.InputType[1] = Shader::EValueType::Float3;
+		Types.ResultType = Shader::EValueType::Float3;
+		break;
+	default:
+		break;
+	}
+	return Types;
+}
+} // namespace Private
+
 void FExpressionOperation::ComputeAnalyticDerivatives(FTree& Tree, FExpressionDerivatives& OutResult) const
 {
 	// Operations with constant derivatives
@@ -907,6 +1275,19 @@ void FExpressionOperation::ComputeAnalyticDerivatives(FTree& Tree, FExpressionDe
 	case EOperation::Abs:
 	case EOperation::Saturate:
 	case EOperation::PowPositiveClamped:
+	case EOperation::Sin:
+	case EOperation::Cos:
+	case EOperation::Tan:
+	case EOperation::Asin:
+	case EOperation::AsinFast:
+	case EOperation::Acos:
+	case EOperation::AcosFast:
+	case EOperation::Atan:
+	case EOperation::AtanFast:
+	case EOperation::Atan2:
+	case EOperation::Atan2Fast:
+	case EOperation::Sqrt:
+	case EOperation::Log2:
 		// TODO
 		break;
 	case EOperation::Add:
@@ -974,11 +1355,12 @@ void FExpressionOperation::ComputeAnalyticDerivatives(FTree& Tree, FExpressionDe
 
 FExpression* FExpressionOperation::ComputePreviousFrame(FTree& Tree, const FRequestedType& RequestedType) const
 {
+	const Private::FOperationRequestedTypes RequestedTypes = Private::GetOperationRequestedTypes(Op, RequestedType);
 	const FOperationDescription OpDesc = GetOperationDescription(Op);
 	FExpression* PrevFrameInputs[MaxInputs];
 	for (int32 Index = 0; Index < OpDesc.NumInputs; ++Index)
 	{
-		PrevFrameInputs[Index] = Tree.GetPreviousFrame(Inputs[Index], RequestedType);
+		PrevFrameInputs[Index] = Tree.GetPreviousFrame(Inputs[Index], RequestedTypes.InputType[Index]);
 	}
 
 	return Tree.NewExpression<FExpressionOperation>(Op, MakeArrayView(PrevFrameInputs, OpDesc.NumInputs));
@@ -987,43 +1369,12 @@ FExpression* FExpressionOperation::ComputePreviousFrame(FTree& Tree, const FRequ
 bool FExpressionOperation::PrepareValue(FEmitContext& Context, FEmitScope& Scope, const FRequestedType& RequestedType, FPrepareValueResult& OutResult) const
 {
 	const FOperationDescription OpDesc = GetOperationDescription(Op);
-	FRequestedType InputRequestedType[MaxInputs];
-	bool bIsMatrixOperation = false;
-	switch (Op)
-	{
-	case EOperation::Length:
-	case EOperation::Normalize:
-		// Each component of result is influenced by all components of input
-		InputRequestedType[0] = ERequestedType::Vector4;
-		break;
-	case EOperation::VecMulMatrix3:
-	case EOperation::VecMulMatrix4:
-		InputRequestedType[0] = ERequestedType::Vector3;
-		InputRequestedType[1] = ERequestedType::Matrix4x4;
-		bIsMatrixOperation = true;
-		break;
-	case EOperation::Matrix3MulVec:
-	case EOperation::Matrix4MulVec:
-		InputRequestedType[0] = ERequestedType::Matrix4x4;
-		InputRequestedType[1] = ERequestedType::Vector3;
-		bIsMatrixOperation = true;
-		break;
-	case EOperation::Dot:
-		InputRequestedType[0] = ERequestedType::Vector4;
-		InputRequestedType[1] = ERequestedType::Vector4;
-		break;
-	default:
-		for (int32 Index = 0; Index < OpDesc.NumInputs; ++Index)
-		{
-			InputRequestedType[Index] = RequestedType;
-		}
-		break;
-	}
+	const Private::FOperationRequestedTypes RequestedTypes = Private::GetOperationRequestedTypes(Op, RequestedType);
 
 	FPreparedType InputType[MaxInputs];
 	for (int32 Index = 0; Index < OpDesc.NumInputs; ++Index)
 	{
-		InputType[Index] = Context.PrepareExpression(Inputs[Index], Scope, InputRequestedType[Index]);
+		InputType[Index] = Context.PrepareExpression(Inputs[Index], Scope, RequestedTypes.InputType[Index]);
 		if (InputType[Index].IsVoid())
 		{
 			return false;
@@ -1036,7 +1387,7 @@ bool FExpressionOperation::PrepareValue(FEmitContext& Context, FEmitScope& Scope
 	}
 
 	FPreparedType ResultType;
-	if (bIsMatrixOperation)
+	if (RequestedTypes.bIsMatrixOperation)
 	{
 		const EExpressionEvaluation Evaluation = CombineEvaluations(InputType[0].GetEvaluation(Scope), InputType[1].GetEvaluation(Scope));
 		const Shader::EValueComponentType ComponentType = Shader::CombineComponentTypes(InputType[0].ValueComponentType, InputType[1].ValueComponentType);
@@ -1060,10 +1411,10 @@ bool FExpressionOperation::PrepareValue(FEmitContext& Context, FEmitScope& Scope
 	}
 	else if(OpDesc.NumInputs == 2)
 	{
-		if (InputType[0].GetNumComponents() != 1 && InputType[1].GetNumComponents() != 1 && InputType[0].GetNumComponents() != InputType[1].GetNumComponents())
+		/*if (InputType[0].GetNumComponents() != 1 && InputType[1].GetNumComponents() != 1 && InputType[0].GetNumComponents() != InputType[1].GetNumComponents())
 		{
 			return Context.Errors->AddErrorf(TEXT("Invalid arithmetic between %s and %s"), InputType[0].GetType().GetName(), InputType[1].GetType().GetName());
-		}
+		}*/
 		ResultType = MergePreparedTypes(InputType[0], InputType[1]);
 		if (Op == EOperation::Less || Op == EOperation::Greater || Op == EOperation::LessEqual || Op == EOperation::GreaterEqual)
 		{
@@ -1087,87 +1438,6 @@ bool FExpressionOperation::PrepareValue(FEmitContext& Context, FEmitScope& Scope
 	return OutResult.SetType(Context, RequestedType, ResultType);
 }
 
-namespace Private
-{
-struct FOperationTypes
-{
-	Shader::EValueType InputType[FExpressionOperation::MaxInputs];
-	Shader::EValueType ResultType;
-	bool bIsLWC;
-};
-FOperationTypes GetOperationTypes(EOperation Op, TConstArrayView<Shader::EValueType> InputTypes)
-{
-	Shader::EValueComponentType ComponentType = Shader::EValueComponentType::Void;
-	int32 NumComponents = 0;
-	for (int32 Index = 0; Index < InputTypes.Num(); ++Index)
-	{
-		const Shader::FValueTypeDescription InputTypeDesc = Shader::GetValueTypeDescription(InputTypes[Index]);
-		ComponentType = Shader::CombineComponentTypes(ComponentType, InputTypeDesc.ComponentType);
-		NumComponents = FMath::Max<int32>(NumComponents, InputTypeDesc.NumComponents);
-	}
-
-	Shader::EValueType IntermediateType = Shader::MakeValueType(ComponentType, NumComponents);
-
-	FOperationTypes Types;
-	for (int32 Index = 0; Index < InputTypes.Num(); ++Index)
-	{
-		Types.InputType[Index] = IntermediateType;
-	}
-	Types.ResultType = IntermediateType;
-	Types.bIsLWC = (ComponentType == Shader::EValueComponentType::Double);
-	switch (Op)
-	{
-	case EOperation::Length:
-		Types.ResultType = Shader::MakeValueType(ComponentType, 1);
-		break;
-	case EOperation::Normalize:
-	case EOperation::Frac:
-	case EOperation::Rcp:
-	case EOperation::Sign:
-	case EOperation::Saturate:
-		Types.ResultType = Shader::MakeNonLWCType(IntermediateType);
-		break;
-	case EOperation::Less:
-	case EOperation::Greater:
-	case EOperation::LessEqual:
-	case EOperation::GreaterEqual:
-		Types.ResultType = Shader::MakeValueType(Shader::EValueComponentType::Bool, NumComponents);
-		break;
-	case EOperation::Fmod:
-		Types.InputType[1] = Types.ResultType = Shader::MakeNonLWCType(IntermediateType);
-		break;
-	case EOperation::PowPositiveClamped:
-		// No LWC support yet
-		Types.InputType[0] = Types.InputType[1] = Types.ResultType = Shader::MakeNonLWCType(IntermediateType);
-		break;
-	case EOperation::Dot:
-		Types.ResultType = Shader::MakeValueType(ComponentType, 1);
-		break;
-	case EOperation::VecMulMatrix3:
-		// No LWC for matrix3
-		Types.InputType[0] = Shader::EValueType::Float3;
-		Types.InputType[1] = Shader::EValueType::Float4x4;
-		Types.ResultType = Shader::EValueType::Float3;
-		break;
-	case EOperation::VecMulMatrix4:
-		Types.InputType[0] = Shader::EValueType::Float3;
-		Types.InputType[1] = Shader::EValueType::Float4x4;
-		Types.ResultType = Shader::EValueType::Float3;
-		break;
-	case EOperation::Matrix3MulVec:
-	case EOperation::Matrix4MulVec:
-		// No LWC for transpose matrices
-		Types.InputType[0] = Shader::EValueType::Float4x4;
-		Types.InputType[1] = Shader::EValueType::Float3;
-		Types.ResultType = Shader::EValueType::Float3;
-		break;
-	default:
-		break;
-	}
-	return Types;
-}
-} // namespace Private
-
 void FExpressionOperation::EmitValueShader(FEmitContext& Context, FEmitScope& Scope, const FRequestedType& RequestedType, FEmitValueShaderResult& OutResult) const
 {
 	const FOperationDescription OpDesc = GetOperationDescription(Op);
@@ -1176,11 +1446,12 @@ void FExpressionOperation::EmitValueShader(FEmitContext& Context, FEmitScope& Sc
 	{
 		InputTypes[Index] = Inputs[Index]->GetType();
 	}
+	const Private::FOperationRequestedTypes RequestedTypes = Private::GetOperationRequestedTypes(Op, RequestedType);
 	const Private::FOperationTypes Types = Private::GetOperationTypes(Op, MakeArrayView(InputTypes, OpDesc.NumInputs));
 	FEmitShaderExpression* InputValue[MaxInputs] = { nullptr };
 	for (int32 Index = 0; Index < OpDesc.NumInputs; ++Index)
 	{
-		InputValue[Index] = Inputs[Index]->GetValueShader(Context, Scope, Types.InputType[Index]);
+		InputValue[Index] = Inputs[Index]->GetValueShader(Context, Scope, RequestedTypes.InputType[Index], Types.InputType[Index]);
 	}
 
 	switch (Op)
@@ -1215,6 +1486,8 @@ void FExpressionOperation::EmitValueShader(FEmitContext& Context, FEmitScope& Sc
 			OutResult.Code = Context.EmitExpression(Scope, Types.ResultType, TEXT("rcp(%)"), InputValue[0]);
 		}
 		break;
+	case EOperation::Sqrt: OutResult.Code = Context.EmitExpression(Scope, Types.ResultType, Types.bIsLWC ? TEXT("LWCSqrt(%)") : TEXT("sqrt(%)"), InputValue[0]); break;
+	case EOperation::Log2: OutResult.Code = Context.EmitExpression(Scope, Types.ResultType, TEXT("log2(%)"), InputValue[0]); break;
 	case EOperation::Frac:
 		if (Types.bIsLWC)
 		{
@@ -1306,6 +1579,15 @@ void FExpressionOperation::EmitValueShader(FEmitContext& Context, FEmitScope& Sc
 			OutResult.Code = Context.EmitExpression(Scope, Types.ResultType, TEXT("normalize(%)"), InputValue[0]);
 		}
 		break;
+	case EOperation::Sin: OutResult.Code = Context.EmitExpression(Scope, Types.ResultType, Types.bIsLWC ? TEXT("LWCSin(%)") : TEXT("sin(%)"), InputValue[0]); break;
+	case EOperation::Cos: OutResult.Code = Context.EmitExpression(Scope, Types.ResultType, Types.bIsLWC ? TEXT("LWCCos(%)") : TEXT("cos(%)"), InputValue[0]); break;
+	case EOperation::Tan: OutResult.Code = Context.EmitExpression(Scope, Types.ResultType, Types.bIsLWC ? TEXT("LWCTan(%)") : TEXT("tan(%)"), InputValue[0]); break;
+	case EOperation::Asin: OutResult.Code = Context.EmitExpression(Scope, Types.ResultType, Types.bIsLWC ? TEXT("LWCAsin(%)") : TEXT("asin(%)"), InputValue[0]); break;
+	case EOperation::AsinFast: OutResult.Code = Context.EmitExpression(Scope, Types.ResultType, Types.bIsLWC ? TEXT("LWCAsin(%)") : TEXT("asinFast(%)"), InputValue[0]); break;
+	case EOperation::Acos: OutResult.Code = Context.EmitExpression(Scope, Types.ResultType, Types.bIsLWC ? TEXT("LWCAcos(%)") : TEXT("acos(%)"), InputValue[0]); break;
+	case EOperation::AcosFast: OutResult.Code = Context.EmitExpression(Scope, Types.ResultType, Types.bIsLWC ? TEXT("LWCAcos(%)") : TEXT("acosFast(%)"), InputValue[0]); break;
+	case EOperation::Atan: OutResult.Code = Context.EmitExpression(Scope, Types.ResultType, Types.bIsLWC ? TEXT("LWCAtan(%)") : TEXT("atan(%)"), InputValue[0]); break;
+	case EOperation::AtanFast: OutResult.Code = Context.EmitExpression(Scope, Types.ResultType, Types.bIsLWC ? TEXT("LWCAtan(%)") : TEXT("atanFast(%)"), InputValue[0]); break;
 	case EOperation::Add:
 		if (Types.bIsLWC)
 		{
@@ -1359,6 +1641,8 @@ void FExpressionOperation::EmitValueShader(FEmitContext& Context, FEmitScope& Sc
 	case EOperation::PowPositiveClamped:
 		OutResult.Code = Context.EmitExpression(Scope, Types.ResultType, TEXT("PositiveClampedPow(%, %)"), InputValue[0], InputValue[1]);
 		break;
+	case EOperation::Atan2: OutResult.Code = Context.EmitExpression(Scope, Types.ResultType, TEXT("atan2(%, %)"), InputValue[0], InputValue[1]); break;
+	case EOperation::Atan2Fast: OutResult.Code = Context.EmitExpression(Scope, Types.ResultType, TEXT("atan2Fast(%, %)"), InputValue[0], InputValue[1]); break;
 	case EOperation::Dot:
 		if (Types.bIsLWC)
 		{
@@ -1432,7 +1716,7 @@ void FExpressionOperation::EmitValueShader(FEmitContext& Context, FEmitScope& Sc
 	case EOperation::VecMulMatrix3:
 		if (Types.bIsLWC)
 		{
-			OutResult.Code = Context.EmitExpression(Scope, Types.ResultType, TEXT("LWCMlutiplyVector(%, %)"), InputValue[0], InputValue[1]);
+			OutResult.Code = Context.EmitExpression(Scope, Types.ResultType, TEXT("LWCMultiply(%, %)"), InputValue[0], InputValue[1]);
 		}
 		else
 		{
@@ -1442,7 +1726,7 @@ void FExpressionOperation::EmitValueShader(FEmitContext& Context, FEmitScope& Sc
 	case EOperation::VecMulMatrix4:
 		if (Types.bIsLWC)
 		{
-			OutResult.Code = Context.EmitExpression(Scope, Types.ResultType, TEXT("LWCMlutiply(%, %)"), InputValue[0], InputValue[1]);
+			OutResult.Code = Context.EmitExpression(Scope, Types.ResultType, TEXT("LWCMultiply(%, %)"), InputValue[0], InputValue[1]);
 		}
 		else
 		{
@@ -1469,21 +1753,76 @@ void FExpressionOperation::EmitValuePreshader(FEmitContext& Context, FEmitScope&
 	{
 		InputTypes[Index] = Inputs[Index]->GetType();
 	}
+	const Private::FOperationRequestedTypes RequestedTypes = Private::GetOperationRequestedTypes(Op, RequestedType);
 	const Private::FOperationTypes Types = Private::GetOperationTypes(Op, MakeArrayView(InputTypes, OpDesc.NumInputs));
 	check(OpDesc.PreshaderOpcode != Shader::EPreshaderOpcode::Nop);
 
 	for (int32 Index = 0; Index < OpDesc.NumInputs; ++Index)
 	{
-		Inputs[Index]->GetValuePreshader(Context, Scope, Types.InputType[Index], OutResult.Preshader);
+		Inputs[Index]->GetValuePreshader(Context, Scope, RequestedTypes.InputType[Index], OutResult.Preshader);
 	}
 
-	check(Context.PreshaderStackPosition > 0);
-	Context.PreshaderStackPosition--;
+	const int32 NumInputsToPop = OpDesc.NumInputs - 1;
+	if (NumInputsToPop > 0)
+	{
+		check(Context.PreshaderStackPosition >= NumInputsToPop);
+		Context.PreshaderStackPosition -= NumInputsToPop;
+	}
 
 	OutResult.Preshader.WriteOpcode(OpDesc.PreshaderOpcode);
 	OutResult.Type = Types.ResultType;
 }
 
+void FExpressionDerivative::ComputeAnalyticDerivatives(FTree& Tree, FExpressionDerivatives& OutResult) const
+{
+	// TODO
+}
+
+FExpression* FExpressionDerivative::ComputePreviousFrame(FTree& Tree, const FRequestedType& RequestedType) const
+{
+	return Tree.NewExpression<FExpressionDerivative>(Coord, Tree.GetPreviousFrame(Input, RequestedType));
+}
+
+bool FExpressionDerivative::PrepareValue(FEmitContext& Context, FEmitScope& Scope, const FRequestedType& RequestedType, FPrepareValueResult& OutResult) const
+{
+	FPreparedType ResultType = Context.PrepareExpression(Input, Scope, RequestedType);
+	if (ResultType.IsVoid())
+	{
+		return false;
+	}
+
+	ResultType.ValueComponentType = Shader::MakeNonLWCType(ResultType.ValueComponentType);
+
+	const EExpressionEvaluation InputEvaluation = ResultType.GetEvaluation(Scope, RequestedType);
+	if (InputEvaluation != EExpressionEvaluation::Shader)
+	{
+		ResultType.SetEvaluation(EExpressionEvaluation::Constant);
+	}
+	return OutResult.SetType(Context, RequestedType, ResultType);
+}
+
+void FExpressionDerivative::EmitValueShader(FEmitContext& Context, FEmitScope& Scope, const FRequestedType& RequestedType, FEmitValueShaderResult& OutResult) const
+{
+	FEmitShaderExpression* EmitInput = Input->GetValueShader(Context, Scope, RequestedType);
+	const bool bIsLWC = Shader::IsLWCType(EmitInput->Type);
+	const TCHAR* FunctionName = nullptr;
+	switch (Coord)
+	{
+	case EDerivativeCoordinate::Ddx: FunctionName = bIsLWC ? TEXT("LWCDdx") : TEXT("ddx"); break;
+	case EDerivativeCoordinate::Ddy: FunctionName = bIsLWC ? TEXT("LWCDdy") : TEXT("ddy"); break;
+	default: checkNoEntry(); break;
+	}
+	OutResult.Code = Context.EmitExpression(Scope, Shader::MakeNonLWCType(EmitInput->Type), TEXT("%(%)"), FunctionName, EmitInput);
+}
+
+void FExpressionDerivative::EmitValuePreshader(FEmitContext& Context, FEmitScope& Scope, const FRequestedType& RequestedType, FEmitValuePreshaderResult& OutResult) const
+{
+	// Derivative of a constant is 0
+	Context.PreshaderStackPosition++;
+	OutResult.Type = GetType();
+	OutResult.Preshader.WriteOpcode(Shader::EPreshaderOpcode::ConstantZero).Write(OutResult.Type);
+}
+
 void FExpressionSwizzle::ComputeAnalyticDerivatives(FTree& Tree, FExpressionDerivatives& OutResult) const
 {
 	const FExpressionDerivatives InputDerivatives = Tree.GetAnalyticDerivatives(Input);
@@ -1503,13 +1842,34 @@ FExpression* FExpressionSwizzle::ComputePreviousFrame(FTree& Tree, const FReques
 bool FExpressionSwizzle::PrepareValue(FEmitContext& Context, FEmitScope& Scope, const FRequestedType& RequestedType, FPrepareValueResult& OutResult) const
 {
 	const FRequestedType RequestedInputType = Parameters.GetRequestedInputType(RequestedType);
-	const FPreparedType& InputType = Context.PrepareExpression(Input, Scope, RequestedInputType);
 
-	FPreparedType ResultType(InputType.ValueComponentType);
-	for (int32 ComponentIndex = 0; ComponentIndex < Parameters.NumComponents; ++ComponentIndex)
+	FPreparedType ResultType;
+	if (RequestedInputType.IsVoid())
 	{
-		const int32 SwizzledComponentIndex = Parameters.ComponentIndex[ComponentIndex];
-		ResultType.SetComponent(ComponentIndex, InputType.GetComponent(SwizzledComponentIndex));
+		// All the requested components are outside the swizzle, so just return 0
+		ResultType.ValueComponentType = Shader::EValueComponentType::Float;
+		for (int32 ComponentIndex = 0; ComponentIndex < Parameters.NumComponents; ++ComponentIndex)
+		{
+			ResultType.SetComponent(ComponentIndex, EExpressionEvaluation::ConstantZero);
+		}
+	}
+	else
+	{
+		const FPreparedType& InputType = Context.PrepareExpression(Input, Scope, RequestedInputType);
+
+		ResultType.ValueComponentType = InputType.ValueComponentType;
+		for (int32 ComponentIndex = 0; ComponentIndex < Parameters.NumComponents; ++ComponentIndex)
+		{
+			if (RequestedType.IsComponentRequested(ComponentIndex))
+			{
+				const int32 SwizzledComponentIndex = Parameters.ComponentIndex[ComponentIndex];
+				ResultType.SetComponent(ComponentIndex, InputType.GetComponent(SwizzledComponentIndex));
+			}
+			else
+			{
+				ResultType.SetComponent(ComponentIndex, EExpressionEvaluation::ConstantZero);
+			}
+		}
 	}
 
 	return OutResult.SetType(Context, RequestedType, ResultType);
@@ -1517,16 +1877,24 @@ bool FExpressionSwizzle::PrepareValue(FEmitContext& Context, FEmitScope& Scope,
 
 void FExpressionSwizzle::EmitValueShader(FEmitContext& Context, FEmitScope& Scope, const FRequestedType& RequestedType, FEmitValueShaderResult& OutResult) const
 {
-	const FRequestedType RequestedInputType = Parameters.GetRequestedInputType(RequestedType);
-
 	static const TCHAR ComponentName[] = { 'x', 'y', 'z', 'w' };
 	TCHAR Swizzle[5] = TEXT("");
 	TCHAR LWCSwizzle[10] = TEXT("");
 	bool bHasSwizzleReorder = false;
 
-	for (int32 ComponentIndex = 0; ComponentIndex < Parameters.NumComponents; ++ComponentIndex)
+	const int32 NumComponents = FMath::Min(RequestedType.GetNumComponents(), Parameters.NumComponents);
+	FRequestedType RequestedInputType;
+	for (int32 ComponentIndex = 0; ComponentIndex < NumComponents; ++ComponentIndex)
 	{
-		const int32 SwizzledComponentIndex = Parameters.ComponentIndex[ComponentIndex];
+		// If component wasn't requested, we just refer to 'x' component of input, since that should always be present
+		int32 SwizzledComponentIndex = 0;
+
+		if (RequestedType.IsComponentRequested(ComponentIndex))
+		{
+			SwizzledComponentIndex = Parameters.ComponentIndex[ComponentIndex];
+			RequestedInputType.SetComponentRequest(SwizzledComponentIndex);
+		}
+
 		Swizzle[ComponentIndex] = ComponentName[SwizzledComponentIndex];
 		LWCSwizzle[ComponentIndex * 2 + 0] = TEXT(',');
 		LWCSwizzle[ComponentIndex * 2 + 1] = TEXT('0') + SwizzledComponentIndex;
@@ -1538,18 +1906,30 @@ void FExpressionSwizzle::EmitValueShader(FEmitContext& Context, FEmitScope& Scop
 	}
 
 	FEmitShaderExpression* InputValue = Input->GetValueShader(Context, Scope, RequestedInputType);
-	if (bHasSwizzleReorder || Parameters.NumComponents != InputValue->Type.GetNumComponents())
+	const Shader::FValueTypeDescription InputTypeDesc = Shader::GetValueTypeDescription(InputValue->Type);
+
+	if (bHasSwizzleReorder || NumComponents != InputTypeDesc.NumComponents)
 	{
-		const Shader::FValueTypeDescription InputTypeDesc = Shader::GetValueTypeDescription(InputValue->Type);
-		const Shader::EValueType ResultType = Shader::MakeValueType(InputTypeDesc.ComponentType, Parameters.NumComponents);
+		const Shader::EValueType ResultType = Shader::MakeValueType(InputTypeDesc.ComponentType, NumComponents);
+		const int32 NumRequestedComponents = RequestedInputType.GetNumComponents();
+		check(NumRequestedComponents > 0);
+	
+		if (NumRequestedComponents > InputTypeDesc.NumComponents)
+		{
+			// Zero-extend our input if needed, so we can acccess all the given components
+			InputValue = Context.EmitCast(Scope, InputValue, Shader::MakeValueType(InputTypeDesc.ComponentType, NumRequestedComponents), EEmitCastFlags::ZeroExtendScalar);
+		}
+
 		if (InputTypeDesc.ComponentType == Shader::EValueComponentType::Double)
 		{
+			check(LWCSwizzle[0]);
 			OutResult.Code = Context.EmitInlineExpression(Scope, ResultType, TEXT("LWCSwizzle(%%)"),
 				InputValue,
 				LWCSwizzle);
 		}
 		else
 		{
+			check(Swizzle[0]);
 			OutResult.Code = Context.EmitInlineExpression(Scope, ResultType, TEXT("%.%"),
 				InputValue,
 				Swizzle);
@@ -1564,17 +1944,25 @@ void FExpressionSwizzle::EmitValueShader(FEmitContext& Context, FEmitScope& Scop
 void FExpressionSwizzle::EmitValuePreshader(FEmitContext& Context, FEmitScope& Scope, const FRequestedType& RequestedType, FEmitValuePreshaderResult& OutResult) const
 {
 	const FRequestedType RequestedInputType = Parameters.GetRequestedInputType(RequestedType);
+	if (RequestedInputType.IsVoid())
+	{
+		Context.PreshaderStackPosition++;
+		OutResult.Type = Shader::MakeValueType(Shader::EValueComponentType::Float, Parameters.NumComponents);
+		OutResult.Preshader.WriteOpcode(Shader::EPreshaderOpcode::ConstantZero).Write(OutResult.Type);
+	}
+	else
+	{
+		const Shader::FType InputType = Input->GetValuePreshader(Context, Scope, RequestedInputType, OutResult.Preshader);
+		const Shader::FValueTypeDescription InputTypeDesc = Shader::GetValueTypeDescription(InputType);
 
-	const Shader::FType InputType = Input->GetValuePreshader(Context, Scope, RequestedInputType, OutResult.Preshader);
-	const Shader::FValueTypeDescription InputTypeDesc = Shader::GetValueTypeDescription(InputType);
-
-	OutResult.Preshader.WriteOpcode(Shader::EPreshaderOpcode::ComponentSwizzle)
-		.Write((uint8)Parameters.NumComponents)
-		.Write((uint8)Parameters.ComponentIndex[0])
-		.Write((uint8)Parameters.ComponentIndex[1])
-		.Write((uint8)Parameters.ComponentIndex[2])
-		.Write((uint8)Parameters.ComponentIndex[3]);
-	OutResult.Type = Shader::MakeValueType(InputTypeDesc.ComponentType, Parameters.NumComponents);
+		OutResult.Preshader.WriteOpcode(Shader::EPreshaderOpcode::ComponentSwizzle)
+			.Write((uint8)Parameters.NumComponents)
+			.Write((uint8)Parameters.ComponentIndex[0])
+			.Write((uint8)Parameters.ComponentIndex[1])
+			.Write((uint8)Parameters.ComponentIndex[2])
+			.Write((uint8)Parameters.ComponentIndex[3]);
+		OutResult.Type = Shader::MakeValueType(InputTypeDesc.ComponentType, Parameters.NumComponents);
+	}
 }
 
 void FExpressionAppend::ComputeAnalyticDerivatives(FTree& Tree, FExpressionDerivatives& OutResult) const
@@ -1600,28 +1988,31 @@ bool FExpressionAppend::PrepareValue(FEmitContext& Context, FEmitScope& Scope, c
 	const int32 NumRequestedComponents = RequestedType.GetNumComponents();
 	const int32 NumLhsComponents = FMath::Min(LhsType.GetNumComponents(), NumRequestedComponents);
 
+	FPreparedType ResultType(LhsType.ValueComponentType);
+	for (int32 Index = 0; Index < NumLhsComponents; ++Index)
+	{
+		ResultType.SetComponent(Index, LhsType.GetComponent(Index));
+	}
+
 	FRequestedType RhsRequestedType;
 	for (int32 Index = NumLhsComponents; Index < NumRequestedComponents; ++Index)
 	{
 		RhsRequestedType.SetComponentRequest(Index - NumLhsComponents, RequestedType.IsComponentRequested(Index));
 	}
 
-	const FPreparedType& RhsType = Context.PrepareExpression(Rhs, Scope, RhsRequestedType);
-	const int32 NumRhsComponents = FMath::Min(RhsType.GetNumComponents(), NumRequestedComponents - NumLhsComponents);
-
-	if (LhsType.ValueComponentType != RhsType.ValueComponentType)
+	if (!RhsRequestedType.IsVoid())
 	{
-		return Context.Errors->AddError(TEXT("Type mismatch"));
-	}
+		const FPreparedType& RhsType = Context.PrepareExpression(Rhs, Scope, RhsRequestedType);
+		if (LhsType.ValueComponentType != RhsType.ValueComponentType)
+		{
+			return Context.Errors->AddError(TEXT("Type mismatch"));
+		}
 
-	FPreparedType ResultType(LhsType.ValueComponentType);
-	for (int32 Index = 0; Index < NumLhsComponents; ++Index)
-	{
-		ResultType.SetComponent(Index, LhsType.GetComponent(Index));
-	}
-	for (int32 Index = 0; Index < NumRhsComponents; ++Index)
-	{
-		ResultType.SetComponent(NumLhsComponents + Index, LhsType.GetComponent(Index));
+		const int32 NumRhsComponents = FMath::Min(RhsType.GetNumComponents(), NumRequestedComponents - NumLhsComponents);
+		for (int32 Index = 0; Index < NumRhsComponents; ++Index)
+		{
+			ResultType.SetComponent(NumLhsComponents + Index, RhsType.GetComponent(Index));
+		}
 	}
 
 	return OutResult.SetType(Context, RequestedType, ResultType);
@@ -1634,9 +2025,11 @@ struct FAppendTypes
 	Shader::EValueType ResultType;
 	Shader::EValueType LhsType;
 	Shader::EValueType RhsType;
+	FRequestedType LhsRequestedType;
+	FRequestedType RhsRequestedType;
 	bool bIsLWC;
 };
-FAppendTypes GetAppendTypes(Shader::EValueType LhsType, Shader::EValueType RhsType)
+FAppendTypes GetAppendTypes(const FRequestedType& RequestedType, Shader::EValueType LhsType, Shader::EValueType RhsType)
 {
 	const Shader::FValueTypeDescription LhsTypeDesc = Shader::GetValueTypeDescription(LhsType);
 	const Shader::FValueTypeDescription RhsTypeDesc = Shader::GetValueTypeDescription(RhsType);
@@ -1644,9 +2037,24 @@ FAppendTypes GetAppendTypes(Shader::EValueType LhsType, Shader::EValueType RhsTy
 	const int32 NumComponents = FMath::Min(LhsTypeDesc.NumComponents + RhsTypeDesc.NumComponents, 4);
 
 	FAppendTypes Types;
+	for (int32 Index = 0; Index < LhsTypeDesc.NumComponents; ++Index)
+	{
+		if (RequestedType.IsComponentRequested(Index))
+		{
+			Types.LhsRequestedType.SetComponentRequest(Index);
+		}
+	}
+	for (int32 Index = LhsTypeDesc.NumComponents; Index < NumComponents; ++Index)
+	{
+		if (RequestedType.IsComponentRequested(Index))
+		{
+			Types.RhsRequestedType.SetComponentRequest(Index - LhsTypeDesc.NumComponents);
+		}
+	}
+
 	Types.ResultType = Shader::MakeValueType(ComponentType, NumComponents);
 	Types.LhsType = Shader::MakeValueType(ComponentType, LhsTypeDesc.NumComponents);
-	Types.RhsType = Shader::MakeValueType(ComponentType, FMath::Min<int32>(RhsTypeDesc.NumComponents, 4 - LhsTypeDesc.NumComponents));
+	Types.RhsType = Shader::MakeValueType(ComponentType, NumComponents - LhsTypeDesc.NumComponents);
 	Types.bIsLWC = ComponentType == Shader::EValueComponentType::Double;
 	return Types;
 }
@@ -1654,8 +2062,8 @@ FAppendTypes GetAppendTypes(Shader::EValueType LhsType, Shader::EValueType RhsTy
 
 void FExpressionAppend::EmitValueShader(FEmitContext& Context, FEmitScope& Scope, const FRequestedType& RequestedType, FEmitValueShaderResult& OutResult) const
 {
-	const Private::FAppendTypes Types = Private::GetAppendTypes(Lhs->GetType(), Rhs->GetType());
-	FEmitShaderExpression* LhsValue = Lhs->GetValueShader(Context, Scope, Types.LhsType);
+	const Private::FAppendTypes Types = Private::GetAppendTypes(RequestedType, Lhs->GetType(), Rhs->GetType());
+	FEmitShaderExpression* LhsValue = Lhs->GetValueShader(Context, Scope, Types.LhsRequestedType, Types.LhsType);
 
 	if (Types.RhsType == Shader::EValueType::Void)
 	{
@@ -1663,7 +2071,7 @@ void FExpressionAppend::EmitValueShader(FEmitContext& Context, FEmitScope& Scope
 	}
 	else
 	{
-		FEmitShaderExpression* RhsValue = Rhs->GetValueShader(Context, Scope, Types.RhsType);
+		FEmitShaderExpression* RhsValue = Rhs->GetValueShader(Context, Scope, Types.RhsRequestedType, Types.RhsType);
 		if (Types.bIsLWC)
 		{
 			OutResult.Code = Context.EmitInlineExpression(Scope, Types.ResultType, TEXT("MakeLWCVector(%, %)"),
@@ -1682,11 +2090,11 @@ void FExpressionAppend::EmitValueShader(FEmitContext& Context, FEmitScope& Scope
 
 void FExpressionAppend::EmitValuePreshader(FEmitContext& Context, FEmitScope& Scope, const FRequestedType& RequestedType, FEmitValuePreshaderResult& OutResult) const
 {
-	const Private::FAppendTypes Types = Private::GetAppendTypes(Lhs->GetType(), Rhs->GetType());
-	Lhs->GetValuePreshader(Context, Scope, Types.LhsType, OutResult.Preshader);
+	const Private::FAppendTypes Types = Private::GetAppendTypes(RequestedType, Lhs->GetType(), Rhs->GetType());
+	Lhs->GetValuePreshader(Context, Scope, Types.LhsRequestedType, OutResult.Preshader);
 	if (Types.RhsType != Shader::EValueType::Void)
 	{
-		Rhs->GetValuePreshader(Context, Scope, Types.RhsType, OutResult.Preshader);
+		Rhs->GetValuePreshader(Context, Scope, Types.RhsRequestedType, OutResult.Preshader);
 
 		check(Context.PreshaderStackPosition > 0);
 		Context.PreshaderStackPosition--;
@@ -1696,15 +2104,14 @@ void FExpressionAppend::EmitValuePreshader(FEmitContext& Context, FEmitScope& Sc
 	OutResult.Type = Types.ResultType;
 }
 
-bool FExpressionReflectionVector::PrepareValue(FEmitContext& Context, FEmitScope& Scope, const FRequestedType& RequestedType, FPrepareValueResult& OutResult) const
+bool FExpressionInlineCustomHLSL::PrepareValue(FEmitContext& Context, FEmitScope& Scope, const FRequestedType& RequestedType, FPrepareValueResult& OutResult) const
 {
-	Context.bReadMaterialNormal = true;
-	return OutResult.SetType(Context, RequestedType, EExpressionEvaluation::Shader, Shader::EValueType::Float3);
+	return OutResult.SetType(Context, RequestedType, EExpressionEvaluation::Shader, Shader::FType(ResultType));
 }
 
-void FExpressionReflectionVector::EmitValueShader(FEmitContext& Context, FEmitScope& Scope, const FRequestedType& RequestedType, FEmitValueShaderResult& OutResult) const
+void FExpressionInlineCustomHLSL::EmitValueShader(FEmitContext& Context, FEmitScope& Scope, const FRequestedType& RequestedType, FEmitValueShaderResult& OutResult) const
 {
-	OutResult.Code = Context.EmitInlineExpression(Scope, Shader::EValueType::Float3, TEXT("Parameters.ReflectionVector"));
+	OutResult.Code = Context.EmitExpression(Scope, ResultType, Code);
 }
 
 bool FExpressionCustomHLSL::PrepareValue(FEmitContext& Context, FEmitScope& Scope, const FRequestedType& RequestedType, FPrepareValueResult& OutResult) const
diff --git a/Engine/Source/Runtime/Engine/Private/HLSLTree/HLSLTreeEmit.cpp b/Engine/Source/Runtime/Engine/Private/HLSLTree/HLSLTreeEmit.cpp
--- a/Engine/Source/Runtime/Engine/Private/HLSLTree/HLSLTreeEmit.cpp
+++ b/Engine/Source/Runtime/Engine/Private/HLSLTree/HLSLTreeEmit.cpp
@@ -199,6 +199,10 @@ FEmitContext::FEmitContext(FMemStackBase& InAllocator, FErrorHandlerInterface& I
 	, Errors(&InErrors)
 	, TypeRegistry(&InTypeRegistry)
 {
+	for (int32 Index = 0; Index < SF_NumFrequencies; ++Index)
+	{
+		ExternalInputMask[Index].Init(false, 256);
+	}
 }
 
 FEmitContext::~FEmitContext()
@@ -640,6 +644,8 @@ int32 InternalFormatString(FStringBuilderBase* OutString, FEmitShaderDependencie
 				case EFormatArgType::ShaderValue: FormatArg_ShaderValue(Arg.ShaderValue, OutDependencies, *OutString); break;
 				case EFormatArgType::String: OutString->Append(Arg.String); break;
 				case EFormatArgType::Int: OutString->Appendf(TEXT("%d"), Arg.Int); break;
+				case EFormatArgType::Uint: OutString->Appendf(TEXT("%u"), Arg.Uint); break;
+				case EFormatArgType::Float: OutString->Appendf(TEXT("%#.9gf"), Arg.Float); break;
 				case EFormatArgType::Bool: OutString->Append(Arg.Bool ? TEXT("true") : TEXT("false")); break;
 				default:
 					checkNoEntry();
@@ -669,6 +675,11 @@ FEmitShaderExpression* FEmitContext::InternalEmitExpression(FEmitScope& Scope, T
 {
 	FEmitShaderExpression* ShaderValue = nullptr;
 
+	if (Code.Contains(TEXT("(Local253.xy - float3(0.00000000f, 0.00000000f, 1.00000000f))")))
+	{
+		int a = 0;
+	}
+
 	FXxHash64Builder Hasher;
 	Hasher.Update(Code.GetData(), Code.Len() * sizeof(TCHAR));
 	if (bInline)
@@ -823,7 +834,7 @@ FEmitShaderExpression* FEmitContext::EmitPreshaderOrConstant(FEmitScope& Scope,
 		const Shader::EValueType FieldType = Type.GetFlatFieldType(FieldIndex);
 		const Shader::FValueTypeDescription TypeDesc = Shader::GetValueTypeDescription(FieldType);
 		const int32 NumFieldComponents = TypeDesc.NumComponents;
-		const EExpressionEvaluation FieldEvaluation = Expression->GetPreparedType().GetFieldEvaluation(Scope, ComponentIndex, NumFieldComponents);
+		const EExpressionEvaluation FieldEvaluation = Expression->GetPreparedType().GetFieldEvaluation(Scope, RequestedType, ComponentIndex, NumFieldComponents);
 
 		if (FieldEvaluation == EExpressionEvaluation::Preshader)
 		{
@@ -910,7 +921,8 @@ FEmitShaderExpression* FEmitContext::EmitPreshaderOrConstant(FEmitScope& Scope,
 		else
 		{
 			// We allow FieldEvaluation to be 'None', since in that case we still need to fill in a value for the HLSL initializer
-			check(FieldEvaluation == EExpressionEvaluation::Constant ||
+			check(FieldEvaluation == EExpressionEvaluation::ConstantZero ||
+				FieldEvaluation == EExpressionEvaluation::Constant ||
 				FieldEvaluation == EExpressionEvaluation::ConstantLoop ||
 				FieldEvaluation == EExpressionEvaluation::None);
 
@@ -919,8 +931,16 @@ FEmitShaderExpression* FEmitContext::EmitPreshaderOrConstant(FEmitScope& Scope,
 			Shader::FValue FieldConstantValue(ConstantValue.Type.GetComponentType(ComponentIndex), NumFieldComponents);
 			for (int32 i = 0; i < NumFieldComponents; ++i)
 			{
-				// Allow replicating scalar values
-				FieldConstantValue.Component[i] = ConstantValue.Component.Num() == 1 ? ConstantValue.Component[0] : ConstantValue.Component[ComponentIndex + i];
+				if (ConstantValue.Component.Num() == 1)
+				{
+					// Allow replicating scalar values
+					FieldConstantValue.Component[i] = ConstantValue.Component[0];
+				}
+				else if (ConstantValue.Component.IsValidIndex(ComponentIndex + i))
+				{
+					FieldConstantValue.Component[i] = ConstantValue.Component[ComponentIndex + i];
+				}
+				// else component defaults to 0
 			}
 
 			if (TypeDesc.ComponentType == Shader::EValueComponentType::Double)
@@ -992,7 +1012,7 @@ FEmitShaderExpression* FEmitContext::EmitConstantZero(FEmitScope& Scope, const S
 	return EmitInlineExpression(Scope, Type, TEXT("((%)0)"), Type.GetName());
 }
 
-FEmitShaderExpression* FEmitContext::EmitCast(FEmitScope& Scope, FEmitShaderExpression* ShaderValue, const Shader::FType& DestType)
+FEmitShaderExpression* FEmitContext::EmitCast(FEmitScope& Scope, FEmitShaderExpression* ShaderValue, const Shader::FType& DestType, EEmitCastFlags Flags)
 {
 	check(ShaderValue);
 	check(!DestType.IsVoid());
@@ -1032,7 +1052,7 @@ FEmitShaderExpression* FEmitContext::EmitCast(FEmitScope& Scope, FEmitShaderExpr
 		}
 		else
 		{
-			const bool bReplicateScalar = (SourceTypeDesc.NumComponents == 1);
+			const bool bReplicateScalar = (SourceTypeDesc.NumComponents == 1) && !EnumHasAnyFlags(Flags, EEmitCastFlags::ZeroExtendScalar);
 
 			int32 NumComponents = 0;
 			bool bNeedClosingParen = false;
@@ -1043,7 +1063,7 @@ FEmitShaderExpression* FEmitContext::EmitCast(FEmitScope& Scope, FEmitShaderExpr
 			}
 			else
 			{
-				if (SourceTypeDesc.NumComponents == 1 || SourceTypeDesc.NumComponents == DestTypeDesc.NumComponents)
+				if ((SourceTypeDesc.NumComponents == 1 && bReplicateScalar) || SourceTypeDesc.NumComponents == DestTypeDesc.NumComponents)
 				{
 					NumComponents = DestTypeDesc.NumComponents;
 					// Cast the scalar to the correct type, HLSL language will replicate the scalar if needed when performing this cast
diff --git a/Engine/Source/Runtime/Engine/Private/Materials/MaterialExpressionHLSL.cpp b/Engine/Source/Runtime/Engine/Private/Materials/MaterialExpressionHLSL.cpp
--- a/Engine/Source/Runtime/Engine/Private/Materials/MaterialExpressionHLSL.cpp
+++ b/Engine/Source/Runtime/Engine/Private/Materials/MaterialExpressionHLSL.cpp
@@ -18,6 +18,8 @@
 #include "Materials/MaterialExpressionConstant3Vector.h"
 #include "Materials/MaterialExpressionConstant4Vector.h"
 #include "Materials/MaterialExpressionStaticBool.h"
+#include "Materials/MaterialExpressionConstantBiasScale.h"
+#include "Materials/MaterialExpressionShadingModel.h"
 #include "Materials/MaterialExpressionVectorParameter.h"
 #include "Materials/MaterialExpressionDoubleVectorParameter.h"
 #include "Materials/MaterialExpressionScalarParameter.h"
@@ -25,14 +27,25 @@
 #include "Materials/MaterialExpressionPixelDepth.h"
 #include "Materials/MaterialExpressionWorldPosition.h"
 #include "Materials/MaterialExpressionCameraPositionWS.h"
+#include "Materials/MaterialExpressionCameraVectorWS.h"
+#include "Materials/MaterialExpressionViewProperty.h"
 #include "Materials/MaterialExpressionTime.h"
 #include "Materials/MaterialExpressionDeltaTime.h"
 #include "Materials/MaterialExpressionScreenPosition.h"
 #include "Materials/MaterialExpressionSceneTexelSize.h"
 #include "Materials/MaterialExpressionViewSize.h"
-#include "Materials/MaterialExpressionPanner.h"
+#include "Materials/MaterialExpressionTwoSidedSign.h"
+#include "Materials/MaterialExpressionPixelNormalWS.h"
+#include "Materials/MaterialExpressionReflectionVectorWS.h"
+#include "Materials/MaterialExpressionActorPositionWS.h"
 #include "Materials/MaterialExpressionVertexColor.h"
+#include "Materials/MaterialExpressionVertexNormalWS.h"
+#include "Materials/MaterialExpressionVertexTangentWS.h"
 #include "Materials/MaterialExpressionTextureCoordinate.h"
+#include "Materials/MaterialExpressionLightmapUVs.h"
+#include "Materials/MaterialExpressionParticleColor.h"
+#include "Materials/MaterialExpressionParticlePositionWS.h"
+#include "Materials/MaterialExpressionParticleRadius.h"
 #include "Materials/MaterialExpressionTextureSample.h"
 #include "Materials/MaterialExpressionTextureSampleParameter.h"
 #include "Materials/MaterialExpressionFontSample.h"
@@ -40,14 +53,20 @@
 #include "Materials/MaterialExpressionTextureObject.h"
 #include "Materials/MaterialExpressionTextureObjectParameter.h"
 #include "Materials/MaterialExpressionSceneTexture.h"
+#include "Materials/MaterialExpressionSceneDepth.h"
+#include "Materials/MaterialExpressionNoise.h"
 #include "Materials/MaterialExpressionFunctionInput.h"
 #include "Materials/MaterialExpressionFunctionOutput.h"
 #include "Materials/MaterialExpressionMaterialFunctionCall.h"
 #include "Materials/MaterialExpressionStaticSwitch.h"
 #include "Materials/MaterialExpressionFeatureLevelSwitch.h"
+#include "Materials/MaterialExpressionShadingPathSwitch.h"
 #include "Materials/MaterialExpressionGetLocal.h"
 #include "Materials/MaterialExpressionOneMinus.h"
 #include "Materials/MaterialExpressionAbs.h"
+#include "Materials/MaterialExpressionSquareRoot.h"
+#include "Materials/MaterialExpressionLogarithm2.h"
+#include "Materials/MaterialExpressionLogarithm10.h"
 #include "Materials/MaterialExpressionFrac.h"
 #include "Materials/MaterialExpressionFloor.h"
 #include "Materials/MaterialExpressionCeil.h"
@@ -55,12 +74,27 @@
 #include "Materials/MaterialExpressionTruncate.h"
 #include "Materials/MaterialExpressionSaturate.h"
 #include "Materials/MaterialExpressionSign.h"
+#include "Materials/MaterialExpressionSine.h"
+#include "Materials/MaterialExpressionCosine.h"
+#include "Materials/MaterialExpressionTangent.h"
+#include "Materials/MaterialExpressionArcsine.h"
+#include "Materials/MaterialExpressionArccosine.h"
+#include "Materials/MaterialExpressionArctangent.h"
+#include "Materials/MaterialExpressionArctangent2.h"
+#include "Materials/MaterialExpressionArcsineFast.h"
+#include "Materials/MaterialExpressionArccosineFast.h"
+#include "Materials/MaterialExpressionArctangentFast.h"
+#include "Materials/MaterialExpressionArctangent2Fast.h"
+#include "Materials/MaterialExpressionDDX.h"
+#include "Materials/MaterialExpressionDDY.h"
 #include "Materials/MaterialExpressionAdd.h"
 #include "Materials/MaterialExpressionSubtract.h"
 #include "Materials/MaterialExpressionMultiply.h"
 #include "Materials/MaterialExpressionDivide.h"
+#include "Materials/MaterialExpressionFmod.h"
 #include "Materials/MaterialExpressionPower.h"
 #include "Materials/MaterialExpressionDotProduct.h"
+#include "Materials/MaterialExpressionCrossProduct.h"
 #include "Materials/MaterialExpressionMin.h"
 #include "Materials/MaterialExpressionMax.h"
 #include "Materials/MaterialExpressionClamp.h"
@@ -72,10 +106,26 @@
 #include "Materials/MaterialExpressionComponentMask.h"
 #include "Materials/MaterialExpressionGetMaterialAttributes.h"
 #include "Materials/MaterialExpressionSetMaterialAttributes.h"
-#include "Materials/MaterialExpressionReflectionVectorWS.h"
+#include "Materials/MaterialExpressionMakeMaterialAttributes.h"
+#include "Materials/MaterialExpressionBreakMaterialAttributes.h"
 #include "Materials/MaterialExpressionTransform.h"
 #include "Materials/MaterialExpressionTransformPosition.h"
 #include "Materials/MaterialExpressionIf.h"
+#include "Materials/MaterialExpressionFresnel.h"
+#include "Materials/MaterialExpressionDesaturation.h"
+#include "Materials/MaterialExpressionSphereMask.h"
+#include "Materials/MaterialExpressionRotateAboutAxis.h"
+#include "Materials/MaterialExpressionPanner.h"
+#include "Materials/MaterialExpressionRotator.h"
+#include "Materials/MaterialExpressionVolumetricAdvancedMaterialInput.h"
+#include "Materials/MaterialExpressionVolumetricAdvancedMaterialOutput.h"
+#include "Materials/MaterialExpressionHairAttributes.h"
+#include "Materials/MaterialExpressionCloudLayer.h"
+#include "Materials/MaterialExpressionPerInstanceFadeAmount.h"
+#include "Materials/MaterialExpressionObjectBounds.h"
+#include "Materials/MaterialExpressionObjectOrientation.h"
+#include "Materials/MaterialExpressionObjectPositionWS.h"
+#include "Materials/MaterialExpressionObjectRadius.h"
 #include "Materials/MaterialExpressionCustom.h"
 #include "Materials/MaterialExpressionClearCoatNormalCustomOutput.h"
 #include "Materials/MaterialExpressionSetLocal.h"
@@ -86,6 +136,7 @@
 #include "Engine/Engine.h"
 #include "Engine/Texture2D.h"
 #include "Engine/Font.h"
+#include "RenderUtils.h"
 
 bool UMaterialExpression::GenerateHLSLExpression(FMaterialHLSLGenerator& Generator, UE::HLSLTree::FScope& Scope, int32 OutputIndex, UE::HLSLTree::FExpression*& OutExpression)
 {
@@ -139,6 +190,25 @@ bool UMaterialExpressionStaticBool::GenerateHLSLExpression(FMaterialHLSLGenerato
 	return true;
 }
 
+bool UMaterialExpressionConstantBiasScale::GenerateHLSLExpression(FMaterialHLSLGenerator& Generator, UE::HLSLTree::FScope& Scope, int32 OutputIndex, UE::HLSLTree::FExpression*& OutExpression)
+{
+	using namespace UE::HLSLTree;
+	FExpression* ExpressionInput = Input.AcquireHLSLExpression(Generator, Scope);
+	if (!ExpressionInput)
+	{
+		return false;
+	}
+	OutExpression = Generator.GetTree().NewAdd(Generator.GetTree().NewMul(Generator.NewConstant(Scale), ExpressionInput), Generator.NewConstant(Bias));
+	return true;
+}
+
+bool UMaterialExpressionShadingModel::GenerateHLSLExpression(FMaterialHLSLGenerator& Generator, UE::HLSLTree::FScope& Scope, int32 OutputIndex, UE::HLSLTree::FExpression*& OutExpression)
+{
+	using namespace UE::HLSLTree;
+	OutExpression = Generator.GetTree().NewExpression<FExpressionMaterialShadingModel>(ShadingModel);
+	return true;
+}
+
 bool UMaterialExpressionStaticSwitch::GenerateHLSLExpression(FMaterialHLSLGenerator& Generator, UE::HLSLTree::FScope& Scope, int32 OutputIndex, UE::HLSLTree::FExpression*& OutExpression)
 {
 	UE::HLSLTree::FExpression* ConditionExpression = Value.AcquireHLSLExpressionOrConstant(Generator, Scope, (bool)DefaultValue);
@@ -157,7 +227,7 @@ bool UMaterialExpressionFeatureLevelSwitch::GenerateHLSLExpression(FMaterialHLSL
 
 	if (!Default.GetTracedInput().Expression)
 	{
-		return Generator.GetErrors().AddError(TEXT("Feature Level switch missing default input"));
+		return Generator.GetErrors().AddError(TEXT("Missing default input"));
 	}
 
 	if (FeatureInput.GetTracedInput().Expression)
@@ -171,6 +241,39 @@ bool UMaterialExpressionFeatureLevelSwitch::GenerateHLSLExpression(FMaterialHLSL
 	return OutExpression != nullptr;
 }
 
+bool UMaterialExpressionShadingPathSwitch::GenerateHLSLExpression(FMaterialHLSLGenerator& Generator, UE::HLSLTree::FScope& Scope, int32 OutputIndex, UE::HLSLTree::FExpression*& OutExpression)
+{
+	const EShaderPlatform ShaderPlatform = Generator.GetCompileTarget().ShaderPlatform;
+	ERHIShadingPath::Type ShadingPathToCompile = ERHIShadingPath::Deferred;
+	if (IsForwardShadingEnabled(ShaderPlatform))
+	{
+		ShadingPathToCompile = ERHIShadingPath::Forward;
+	}
+	else if (Generator.GetCompileTarget().FeatureLevel < ERHIFeatureLevel::SM5)
+	{
+		ShadingPathToCompile = ERHIShadingPath::Mobile;
+	}
+
+	check(ShadingPathToCompile < UE_ARRAY_COUNT(Inputs));
+	FExpressionInput ShadingPathInput = Inputs[ShadingPathToCompile].GetTracedInput();
+	FExpressionInput DefaultTraced = Default.GetTracedInput();
+
+	if (!DefaultTraced.Expression)
+	{
+		return Generator.GetErrors().AddError(TEXT("Missing default input"));
+	}
+
+	if (ShadingPathInput.Expression)
+	{
+		OutExpression = ShadingPathInput.AcquireHLSLExpression(Generator, Scope);
+	}
+	else
+	{
+		OutExpression = Default.AcquireHLSLExpression(Generator, Scope);
+	}
+	return OutExpression != nullptr;
+}
+
 bool UMaterialExpressionGetLocal::GenerateHLSLExpression(FMaterialHLSLGenerator& Generator, UE::HLSLTree::FScope& Scope, int32 OutputIndex, UE::HLSLTree::FExpression*& OutExpression)
 {
 	OutExpression = Generator.GetTree().AcquireLocal(Scope, LocalName);
@@ -237,6 +340,58 @@ bool UMaterialExpressionCameraPositionWS::GenerateHLSLExpression(FMaterialHLSLGe
 	return true;
 }
 
+bool UMaterialExpressionCameraVectorWS::GenerateHLSLExpression(FMaterialHLSLGenerator& Generator, UE::HLSLTree::FScope& Scope, int32 OutputIndex, UE::HLSLTree::FExpression*& OutExpression)
+{
+	using namespace UE::HLSLTree;
+	OutExpression = Generator.GetTree().NewExpression<FExpressionExternalInput>(EExternalInput::CameraVector);
+	return true;
+}
+
+bool UMaterialExpressionViewProperty::GenerateHLSLExpression(FMaterialHLSLGenerator& Generator, UE::HLSLTree::FScope& Scope, int32 OutputIndex, UE::HLSLTree::FExpression*& OutExpression)
+{
+	using namespace UE::HLSLTree;
+
+	const bool bRequestedRcp = (OutputIndex == 1);
+	bool bApplyRcp = false;
+	EExternalInput Input = EExternalInput::None;
+	switch (Property)
+	{
+	// To make sure any material that were correctly handling BufferUV != ViewportUV, we just lie to material
+	// to make it believe ViewSize == BufferSize, so they are still compatible with SceneTextureLookup().
+	// TODO: Remove MEVP_BufferSize, MEVP_ViewportOffset and do this at material load time. 
+	case MEVP_BufferSize: Input = bRequestedRcp ? EExternalInput::RcpViewSize : EExternalInput::ViewSize; break;
+
+	// We don't care about OutputIndex == 1 because doesn't have any meaning and 
+	// was already returning NaN on unconstrained unique view rendering.
+	case MEVP_ViewportOffset: OutExpression = Generator.NewConstant(0.0f); break;
+
+	case MEVP_FieldOfView: Input = EExternalInput::FieldOfView; bApplyRcp = bRequestedRcp; break;
+	case MEVP_TanHalfFieldOfView: Input = bRequestedRcp ? EExternalInput::CotanHalfFieldOfView : EExternalInput::TanHalfFieldOfView;
+	case MEVP_ViewSize: Input = bRequestedRcp ? EExternalInput::RcpViewSize : EExternalInput::ViewSize; break;
+	case MEVP_WorldSpaceViewPosition: Input = EExternalInput::ViewWorldPosition; bApplyRcp = bRequestedRcp; break;
+	case MEVP_WorldSpaceCameraPosition: Input = EExternalInput::CameraWorldPosition; bApplyRcp = bRequestedRcp; break;
+	case MEVP_TemporalSampleCount: Input = EExternalInput::TemporalSampleCount; bApplyRcp = bRequestedRcp; break;
+	case MEVP_TemporalSampleIndex: Input = EExternalInput::TemporalSampleIndex; bApplyRcp = bRequestedRcp; break;
+	case MEVP_TemporalSampleOffset: Input = EExternalInput::TemporalSampleOffset; bApplyRcp = bRequestedRcp; break;
+	case MEVP_RuntimeVirtualTextureOutputLevel: Input = EExternalInput::RuntimeVirtualTextureOutputLevel; bApplyRcp = bRequestedRcp; break;
+	case MEVP_RuntimeVirtualTextureOutputDerivative: Input = EExternalInput::RuntimeVirtualTextureOutputDerivative; bApplyRcp = bRequestedRcp; break;
+	case MEVP_PreExposure: Input = bRequestedRcp ? EExternalInput::RcpPreExposure : EExternalInput::PreExposure; break;
+	case MEVP_RuntimeVirtualTextureMaxLevel: Input = EExternalInput::RuntimeVirtualTextureMaxLevel; bApplyRcp = bRequestedRcp; break;
+	default: checkNoEntry(); break;
+	}
+
+	if (!OutExpression)
+	{
+		check(Input != EExternalInput::None);
+		OutExpression = Generator.GetTree().NewExpression<FExpressionExternalInput>(Input);
+		if (bApplyRcp)
+		{
+			OutExpression = Generator.GetTree().NewRcp(OutExpression);
+		}
+	}
+	return true;
+}
+
 bool UMaterialExpressionTime::GenerateHLSLExpression(FMaterialHLSLGenerator& Generator, UE::HLSLTree::FScope& Scope, int32 OutputIndex, UE::HLSLTree::FExpression*& OutExpression)
 {
 	using namespace UE::HLSLTree;
@@ -287,28 +442,33 @@ bool UMaterialExpressionViewSize::GenerateHLSLExpression(FMaterialHLSLGenerator&
 	return true;
 }
 
-bool UMaterialExpressionPanner::GenerateHLSLExpression(FMaterialHLSLGenerator& Generator, UE::HLSLTree::FScope& Scope, int32 OutputIndex, UE::HLSLTree::FExpression*& OutExpression)
+bool UMaterialExpressionTwoSidedSign::GenerateHLSLExpression(FMaterialHLSLGenerator& Generator, UE::HLSLTree::FScope& Scope, int32 OutputIndex, UE::HLSLTree::FExpression*& OutExpression)
 {
 	using namespace UE::HLSLTree;
+	OutExpression = Generator.GetTree().NewExpression<FExpressionExternalInput>(EExternalInput::TwoSidedSign);
+	return true;
+}
 
-	FExpression* ExpressionTime = Time.TryAcquireHLSLExpression(Generator, Scope);
-	if (!ExpressionTime)
-	{
-		ExpressionTime = Generator.GetTree().NewExpression<FExpressionExternalInput>(EExternalInput::GameTime);
-	}
-	FExpression* ExpressionSpeed = Speed.AcquireHLSLExpressionOrConstant(Generator, Scope, FVector2f(SpeedX, SpeedY));
-	FExpression* ExpressionOffset = Generator.GetTree().NewMul(ExpressionSpeed, ExpressionTime);
-	if (bFractionalPart)
-	{
-		ExpressionOffset = Generator.GetTree().NewFrac(ExpressionOffset);
-	}
-	FExpression* ExpressionTexCoord = Coordinate.TryAcquireHLSLExpression(Generator, Scope);
-	if (!ExpressionTexCoord)
-	{
-		ExpressionTexCoord = Generator.NewTexCoord(ConstCoordinate);
-	}
+bool UMaterialExpressionReflectionVectorWS::GenerateHLSLExpression(FMaterialHLSLGenerator& Generator, UE::HLSLTree::FScope& Scope, int32 OutputIndex, UE::HLSLTree::FExpression*& OutExpression)
+{
+	using namespace UE::HLSLTree;
+	check(!CustomWorldNormal.GetTracedInput().Expression); // TODO
 
-	OutExpression = Generator.GetTree().NewAdd(ExpressionTexCoord, ExpressionOffset);
+	OutExpression = Generator.GetTree().NewExpression<FExpressionExternalInput>(EExternalInput::WorldReflection);
+	return true;
+}
+
+bool UMaterialExpressionActorPositionWS::GenerateHLSLExpression(FMaterialHLSLGenerator& Generator, UE::HLSLTree::FScope& Scope, int32 OutputIndex, UE::HLSLTree::FExpression*& OutExpression)
+{
+	using namespace UE::HLSLTree;
+	OutExpression = Generator.GetTree().NewExpression<FExpressionExternalInput>(EExternalInput::ActorWorldPosition);
+	return true;
+}
+
+bool UMaterialExpressionPixelNormalWS::GenerateHLSLExpression(FMaterialHLSLGenerator& Generator, UE::HLSLTree::FScope& Scope, int32 OutputIndex, UE::HLSLTree::FExpression*& OutExpression)
+{
+	using namespace UE::HLSLTree;
+	OutExpression = Generator.GetTree().NewExpression<FExpressionExternalInput>(EExternalInput::WorldNormal);
 	return true;
 }
 
@@ -319,6 +479,20 @@ bool UMaterialExpressionVertexColor::GenerateHLSLExpression(FMaterialHLSLGenerat
 	return true;
 }
 
+bool UMaterialExpressionVertexNormalWS::GenerateHLSLExpression(FMaterialHLSLGenerator& Generator, UE::HLSLTree::FScope& Scope, int32 OutputIndex, UE::HLSLTree::FExpression*& OutExpression)
+{
+	using namespace UE::HLSLTree;
+	OutExpression = Generator.GetTree().NewExpression<FExpressionExternalInput>(EExternalInput::WorldVertexNormal);
+	return true;
+}
+
+bool UMaterialExpressionVertexTangentWS::GenerateHLSLExpression(FMaterialHLSLGenerator& Generator, UE::HLSLTree::FScope& Scope, int32 OutputIndex, UE::HLSLTree::FExpression*& OutExpression)
+{
+	using namespace UE::HLSLTree;
+	OutExpression = Generator.GetTree().NewExpression<FExpressionExternalInput>(EExternalInput::WorldVertexTangent);
+	return true;
+}
+
 bool UMaterialExpressionTextureCoordinate::GenerateHLSLExpression(FMaterialHLSLGenerator& Generator, UE::HLSLTree::FScope& Scope, int32 OutputIndex, UE::HLSLTree::FExpression*& OutExpression)
 {
 	OutExpression = Generator.NewTexCoord(CoordinateIndex);
@@ -342,6 +516,36 @@ bool UMaterialExpressionTextureCoordinate::GenerateHLSLExpression(FMaterialHLSLG
 	return true;
 }
 
+bool UMaterialExpressionLightmapUVs::GenerateHLSLExpression(FMaterialHLSLGenerator& Generator, UE::HLSLTree::FScope& Scope, int32 OutputIndex, UE::HLSLTree::FExpression*& OutExpression)
+{
+	using namespace UE::HLSLTree;
+	OutExpression = Generator.GetTree().NewExpression<FExpressionExternalInput>(EExternalInput::LightmapTexCoord);
+	return true;
+}
+
+bool UMaterialExpressionParticleColor::GenerateHLSLExpression(FMaterialHLSLGenerator& Generator, UE::HLSLTree::FScope& Scope, int32 OutputIndex, UE::HLSLTree::FExpression*& OutExpression)
+{
+	using namespace UE::HLSLTree;
+	OutExpression = Generator.GetTree().NewExpression<FExpressionExternalInput>(EExternalInput::ParticleColor);
+	return true;
+}
+
+bool UMaterialExpressionParticlePositionWS::GenerateHLSLExpression(FMaterialHLSLGenerator& Generator, UE::HLSLTree::FScope& Scope, int32 OutputIndex, UE::HLSLTree::FExpression*& OutExpression)
+{
+	using namespace UE::HLSLTree;
+	OutExpression = Generator.GetTree().NewExpression<FExpressionExternalInput>(EExternalInput::ParticleTranslatedWorldPosition);
+	OutExpression = Generator.GetTree().NewSub(OutExpression, Generator.GetTree().NewExpression<FExpressionExternalInput>(EExternalInput::PreViewTranslation));
+	return true;
+}
+
+bool UMaterialExpressionParticleRadius::GenerateHLSLExpression(FMaterialHLSLGenerator& Generator, UE::HLSLTree::FScope& Scope, int32 OutputIndex, UE::HLSLTree::FExpression*& OutExpression)
+{
+	using namespace UE::HLSLTree;
+	OutExpression = Generator.GetTree().NewExpression<FExpressionExternalInput>(EExternalInput::ParticleRadius);
+	OutExpression = Generator.GetTree().NewMax(OutExpression, Generator.NewConstant(0.0f));
+	return true;
+}
+
 bool UMaterialExpressionTextureObject::GenerateHLSLExpression(FMaterialHLSLGenerator& Generator, UE::HLSLTree::FScope& Scope, int32 OutputIndex, UE::HLSLTree::FExpression*& OutExpression)
 {
 	using namespace UE::HLSLTree;
@@ -486,13 +690,52 @@ bool UMaterialExpressionSceneTexture::GenerateHLSLExpression(FMaterialHLSLGenera
 	}
 	else if (OutputIndex == 1 || OutputIndex == 2)
 	{
-		//return Compiler->GetSceneTextureViewSize(SceneTextureId, /* InvProperty = */ OutputIndex == 2);
-		return false; // TODO
+		const bool bRcp = (OutputIndex == 2);
+		const FStringView Code = UE::MemStack::AllocateStringViewf(Generator.GetTree().GetAllocator(),
+			TEXT("GetSceneTextureViewSize(%d).%s"),
+			(int32)SceneTextureId,
+			bRcp ? TEXT("zw") : TEXT("xy"));
+		OutExpression = Generator.GetTree().NewExpression<FExpressionInlineCustomHLSL>(UE::Shader::EValueType::Float2, Code);
+		return true;
 	}
 
 	return Generator.GetErrors().AddError(TEXT("Invalid input parameter"));
 }
 
+bool UMaterialExpressionNoise::GenerateHLSLExpression(FMaterialHLSLGenerator& Generator, UE::HLSLTree::FScope& Scope, int32 OutputIndex, UE::HLSLTree::FExpression*& OutExpression)
+{
+	using namespace UE::HLSLTree;
+
+	FExpression* ExpressionPosition = Position.TryAcquireHLSLExpression(Generator, Scope);
+	if (!ExpressionPosition)
+	{
+		ExpressionPosition = Generator.GetTree().NewExpression<FExpressionExternalInput>(EExternalInput::WorldPosition);
+	}
+	FExpression* ExpressionFilterWidth = FilterWidth.AcquireHLSLExpressionOrConstant(Generator, Scope, 0.0f);
+
+	FMaterialNoiseParameters NoiseParameters;
+	NoiseParameters.Quality = Quality;
+	// to limit performance problems due to values outside reasonable range
+	NoiseParameters.Levels = FMath::Clamp(Levels, 1, 10);
+	NoiseParameters.Scale = Scale;
+	NoiseParameters.RepeatSize = RepeatSize;
+	NoiseParameters.OutputMin = OutputMin;
+	NoiseParameters.OutputMax = OutputMax;
+	NoiseParameters.LevelScale = LevelScale;
+	NoiseParameters.NoiseFunction = NoiseFunction;
+	NoiseParameters.bTiling = bTiling;
+	NoiseParameters.bTurbulence = bTurbulence;
+
+	OutExpression = Generator.GetTree().NewExpression<FExpressionMaterialNoise>(NoiseParameters, ExpressionPosition, ExpressionFilterWidth);
+	return true;
+}
+
+bool UMaterialExpressionSceneDepth::GenerateHLSLExpression(FMaterialHLSLGenerator& Generator, UE::HLSLTree::FScope& Scope, int32 OutputIndex, UE::HLSLTree::FExpression*& OutExpression)
+{
+	OutExpression = Generator.NewConstant(0.0f); // TODO
+	return true;
+}
+
 bool UMaterialExpressionOneMinus::GenerateHLSLExpression(FMaterialHLSLGenerator& Generator, UE::HLSLTree::FScope& Scope, int32 OutputIndex, UE::HLSLTree::FExpression*& OutExpression)
 {
 	UE::HLSLTree::FExpression* InputExpression = Input.AcquireHLSLExpression(Generator, Scope);
@@ -516,6 +759,43 @@ bool UMaterialExpressionAbs::GenerateHLSLExpression(FMaterialHLSLGenerator& Gene
 	return true;
 }
 
+bool UMaterialExpressionSquareRoot::GenerateHLSLExpression(FMaterialHLSLGenerator& Generator, UE::HLSLTree::FScope& Scope, int32 OutputIndex, UE::HLSLTree::FExpression*& OutExpression)
+{
+	using namespace UE::HLSLTree;
+	FExpression* InputExpression = Input.AcquireHLSLExpression(Generator, Scope);
+	if (!InputExpression)
+	{
+		return false;
+	}
+	OutExpression = Generator.GetTree().NewUnaryOp(EOperation::Sqrt, InputExpression);
+	return true;
+}
+
+bool UMaterialExpressionLogarithm2::GenerateHLSLExpression(FMaterialHLSLGenerator& Generator, UE::HLSLTree::FScope& Scope, int32 OutputIndex, UE::HLSLTree::FExpression*& OutExpression)
+{
+	using namespace UE::HLSLTree;
+	FExpression* InputExpression = X.AcquireHLSLExpression(Generator, Scope);
+	if (!InputExpression)
+	{
+		return false;
+	}
+	OutExpression = Generator.GetTree().NewUnaryOp(EOperation::Log2, InputExpression);
+	return true;
+}
+
+bool UMaterialExpressionLogarithm10::GenerateHLSLExpression(FMaterialHLSLGenerator& Generator, UE::HLSLTree::FScope& Scope, int32 OutputIndex, UE::HLSLTree::FExpression*& OutExpression)
+{
+	using namespace UE::HLSLTree;
+	FExpression* InputExpression = X.AcquireHLSLExpression(Generator, Scope);
+	if (!InputExpression)
+	{
+		return false;
+	}
+	OutExpression = Generator.GetTree().NewUnaryOp(EOperation::Log2, InputExpression);
+	OutExpression = Generator.GetTree().NewMul(OutExpression, Generator.NewConstant(1.0f / FMath::Log2(10.0f))); // Convert Log2 -> Log10
+	return true;
+}
+
 bool UMaterialExpressionFrac::GenerateHLSLExpression(FMaterialHLSLGenerator& Generator, UE::HLSLTree::FScope& Scope, int32 OutputIndex, UE::HLSLTree::FExpression*& OutExpression)
 {
 	using namespace UE::HLSLTree;
@@ -600,6 +880,119 @@ bool UMaterialExpressionSign::GenerateHLSLExpression(FMaterialHLSLGenerator& Gen
 	return true;
 }
 
+static bool GenerateHLSLExpressionTrig(FMaterialHLSLGenerator& Generator,
+	UE::HLSLTree::EOperation Op,
+	UE::HLSLTree::FExpression* InputExpression,
+	float Period,
+	UE::HLSLTree::FExpression*& OutExpression)
+{
+	if (!InputExpression)
+	{
+		return false;
+	}
+	if (Period > 0.0f)
+	{
+		InputExpression = Generator.GetTree().NewMul(InputExpression, Generator.NewConstant(2.0f * PI / Period));
+	}
+	OutExpression = Generator.GetTree().NewUnaryOp(Op, InputExpression);
+	return true;
+}
+
+bool UMaterialExpressionSine::GenerateHLSLExpression(FMaterialHLSLGenerator& Generator, UE::HLSLTree::FScope& Scope, int32 OutputIndex, UE::HLSLTree::FExpression*& OutExpression)
+{
+	return GenerateHLSLExpressionTrig(Generator, UE::HLSLTree::EOperation::Sin, Input.AcquireHLSLExpression(Generator, Scope), Period, OutExpression);
+}
+
+bool UMaterialExpressionCosine::GenerateHLSLExpression(FMaterialHLSLGenerator& Generator, UE::HLSLTree::FScope& Scope, int32 OutputIndex, UE::HLSLTree::FExpression*& OutExpression)
+{
+	return GenerateHLSLExpressionTrig(Generator, UE::HLSLTree::EOperation::Cos, Input.AcquireHLSLExpression(Generator, Scope), Period, OutExpression);
+}
+
+bool UMaterialExpressionTangent::GenerateHLSLExpression(FMaterialHLSLGenerator& Generator, UE::HLSLTree::FScope& Scope, int32 OutputIndex, UE::HLSLTree::FExpression*& OutExpression)
+{
+	return GenerateHLSLExpressionTrig(Generator, UE::HLSLTree::EOperation::Tan, Input.AcquireHLSLExpression(Generator, Scope), Period, OutExpression);
+}
+
+bool UMaterialExpressionArcsine::GenerateHLSLExpression(FMaterialHLSLGenerator& Generator, UE::HLSLTree::FScope& Scope, int32 OutputIndex, UE::HLSLTree::FExpression*& OutExpression)
+{
+	return GenerateHLSLExpressionTrig(Generator, UE::HLSLTree::EOperation::Asin, Input.AcquireHLSLExpression(Generator, Scope), 0.0f, OutExpression);
+}
+
+bool UMaterialExpressionArccosine::GenerateHLSLExpression(FMaterialHLSLGenerator& Generator, UE::HLSLTree::FScope& Scope, int32 OutputIndex, UE::HLSLTree::FExpression*& OutExpression)
+{
+	return GenerateHLSLExpressionTrig(Generator, UE::HLSLTree::EOperation::Acos, Input.AcquireHLSLExpression(Generator, Scope), 0.0f, OutExpression);
+}
+
+bool UMaterialExpressionArctangent::GenerateHLSLExpression(FMaterialHLSLGenerator& Generator, UE::HLSLTree::FScope& Scope, int32 OutputIndex, UE::HLSLTree::FExpression*& OutExpression)
+{
+	return GenerateHLSLExpressionTrig(Generator, UE::HLSLTree::EOperation::Atan, Input.AcquireHLSLExpression(Generator, Scope), 0.0f, OutExpression);
+}
+
+bool UMaterialExpressionArctangent2::GenerateHLSLExpression(FMaterialHLSLGenerator& Generator, UE::HLSLTree::FScope& Scope, int32 OutputIndex, UE::HLSLTree::FExpression*& OutExpression)
+{
+	using namespace UE::HLSLTree;
+	FExpression* Lhs = Y.AcquireHLSLExpression(Generator, Scope);
+	FExpression* Rhs = X.AcquireHLSLExpression(Generator, Scope);
+	if (!Lhs || !Rhs)
+	{
+		return false;
+	}
+	OutExpression = Generator.GetTree().NewBinaryOp(EOperation::Atan2, Lhs, Rhs);
+	return true;
+}
+
+bool UMaterialExpressionArcsineFast::GenerateHLSLExpression(FMaterialHLSLGenerator& Generator, UE::HLSLTree::FScope& Scope, int32 OutputIndex, UE::HLSLTree::FExpression*& OutExpression)
+{
+	return GenerateHLSLExpressionTrig(Generator, UE::HLSLTree::EOperation::AsinFast, Input.AcquireHLSLExpression(Generator, Scope), 0.0f, OutExpression);
+}
+
+bool UMaterialExpressionArccosineFast::GenerateHLSLExpression(FMaterialHLSLGenerator& Generator, UE::HLSLTree::FScope& Scope, int32 OutputIndex, UE::HLSLTree::FExpression*& OutExpression)
+{
+	return GenerateHLSLExpressionTrig(Generator, UE::HLSLTree::EOperation::AcosFast, Input.AcquireHLSLExpression(Generator, Scope), 0.0f, OutExpression);
+}
+
+bool UMaterialExpressionArctangentFast::GenerateHLSLExpression(FMaterialHLSLGenerator& Generator, UE::HLSLTree::FScope& Scope, int32 OutputIndex, UE::HLSLTree::FExpression*& OutExpression)
+{
+	return GenerateHLSLExpressionTrig(Generator, UE::HLSLTree::EOperation::AtanFast, Input.AcquireHLSLExpression(Generator, Scope), 0.0f, OutExpression);
+}
+
+bool UMaterialExpressionArctangent2Fast::GenerateHLSLExpression(FMaterialHLSLGenerator& Generator, UE::HLSLTree::FScope& Scope, int32 OutputIndex, UE::HLSLTree::FExpression*& OutExpression)
+{
+	using namespace UE::HLSLTree;
+	FExpression* Lhs = Y.AcquireHLSLExpression(Generator, Scope);
+	FExpression* Rhs = X.AcquireHLSLExpression(Generator, Scope);
+	if (!Lhs || !Rhs)
+	{
+		return false;
+	}
+	OutExpression = Generator.GetTree().NewBinaryOp(EOperation::Atan2Fast, Lhs, Rhs);
+	return true;
+}
+
+bool UMaterialExpressionDDX::GenerateHLSLExpression(FMaterialHLSLGenerator& Generator, UE::HLSLTree::FScope& Scope, int32 OutputIndex, UE::HLSLTree::FExpression*& OutExpression)
+{
+	using namespace UE::HLSLTree;
+	FExpression* ExpressionInput = Value.AcquireHLSLExpression(Generator, Scope);
+	if (!ExpressionInput)
+	{
+		return false;
+	}
+	OutExpression = Generator.GetTree().NewExpression<FExpressionDerivative>(EDerivativeCoordinate::Ddx, ExpressionInput);
+	return true;
+}
+
+bool UMaterialExpressionDDY::GenerateHLSLExpression(FMaterialHLSLGenerator& Generator, UE::HLSLTree::FScope& Scope, int32 OutputIndex, UE::HLSLTree::FExpression*& OutExpression)
+{
+	using namespace UE::HLSLTree;
+	FExpression* ExpressionInput = Value.AcquireHLSLExpression(Generator, Scope);
+	if (!ExpressionInput)
+	{
+		return false;
+	}
+	OutExpression = Generator.GetTree().NewExpression<FExpressionDerivative>(EDerivativeCoordinate::Ddy, ExpressionInput);
+	return true;
+}
+
 bool UMaterialExpressionBinaryOp::GenerateHLSLExpression(FMaterialHLSLGenerator& Generator, UE::HLSLTree::FScope& Scope, int32 OutputIndex, UE::HLSLTree::FExpression*& OutExpression)
 {
 	UE::HLSLTree::FExpression* Lhs = A.AcquireHLSLExpressionOrConstant(Generator, Scope, ConstA);
@@ -661,6 +1054,18 @@ bool UMaterialExpressionDivide::GenerateHLSLExpression(FMaterialHLSLGenerator& G
 	return true;
 }
 
+bool UMaterialExpressionFmod::GenerateHLSLExpression(FMaterialHLSLGenerator& Generator, UE::HLSLTree::FScope& Scope, int32 OutputIndex, UE::HLSLTree::FExpression*& OutExpression)
+{
+	UE::HLSLTree::FExpression* Lhs = A.AcquireHLSLExpression(Generator, Scope);
+	UE::HLSLTree::FExpression* Rhs = B.AcquireHLSLExpression(Generator, Scope);
+	if (!Lhs || !Rhs)
+	{
+		return false;
+	}
+	OutExpression = Generator.GetTree().NewFmod(Lhs, Rhs);
+	return true;
+}
+
 bool UMaterialExpressionPower::GenerateHLSLExpression(FMaterialHLSLGenerator& Generator, UE::HLSLTree::FScope& Scope, int32 OutputIndex, UE::HLSLTree::FExpression*& OutExpression)
 {
 	using namespace UE::HLSLTree;
@@ -686,6 +1091,18 @@ bool UMaterialExpressionDotProduct::GenerateHLSLExpression(FMaterialHLSLGenerato
 	return true;
 }
 
+bool UMaterialExpressionCrossProduct::GenerateHLSLExpression(FMaterialHLSLGenerator& Generator, UE::HLSLTree::FScope& Scope, int32 OutputIndex, UE::HLSLTree::FExpression*& OutExpression)
+{
+	UE::HLSLTree::FExpression* Lhs = A.AcquireHLSLExpression(Generator, Scope);
+	UE::HLSLTree::FExpression* Rhs = B.AcquireHLSLExpression(Generator, Scope);
+	if (!Lhs || !Rhs)
+	{
+		return false;
+	}
+	OutExpression = Generator.GetTree().NewCross(Lhs, Rhs);
+	return true;
+}
+
 bool UMaterialExpressionMin::GenerateHLSLExpression(FMaterialHLSLGenerator& Generator, UE::HLSLTree::FScope& Scope, int32 OutputIndex, UE::HLSLTree::FExpression*& OutExpression)
 {
 	UE::HLSLTree::FExpression* Lhs = A.AcquireHLSLExpressionOrConstant(Generator, Scope, ConstA);
@@ -742,9 +1159,7 @@ bool UMaterialExpressionLinearInterpolate::GenerateHLSLExpression(FMaterialHLSLG
 	{
 		return false;
 	}
-
-	// A + (B - A) * Alpha
-	OutExpression = Generator.GetTree().NewAdd(ExpressionA, Generator.GetTree().NewMul(Generator.GetTree().NewSub(ExpressionB, ExpressionA), ExpressionAlpha));
+	OutExpression = Generator.GetTree().NewLerp(ExpressionA, ExpressionB, ExpressionAlpha);
 	return true;
 }
 
@@ -854,11 +1269,83 @@ bool UMaterialExpressionSetMaterialAttributes::GenerateHLSLExpression(FMaterialH
 	return true;
 }
 
-bool UMaterialExpressionReflectionVectorWS::GenerateHLSLExpression(FMaterialHLSLGenerator& Generator, UE::HLSLTree::FScope& Scope, int32 OutputIndex, UE::HLSLTree::FExpression*& OutExpression)
+static UE::HLSLTree::FExpression* SetAttribute(FMaterialHLSLGenerator& Generator,
+	UE::HLSLTree::FScope& Scope,
+	EMaterialProperty Property,
+	const FExpressionInput& Input,
+	UE::HLSLTree::FExpression* AttributesExpression)
 {
-	check(!CustomWorldNormal.GetTracedInput().Expression); // TODO
+	using namespace UE::HLSLTree;
+	using namespace UE::Shader;
+
+	FExpression* InputExpression = Input.TryAcquireHLSLExpression(Generator, Scope);
+	if (InputExpression)
+	{
+		const FString& AttributeName = FMaterialAttributeDefinitionMap::GetAttributeName(Property);
+		const FStructField* AttributeField = Generator.GetMaterialAttributesType()->FindFieldByName(*AttributeName);
+		AttributesExpression = Generator.GetTree().NewExpression<FExpressionSetStructField>(Generator.GetMaterialAttributesType(), AttributeField, AttributesExpression, InputExpression);
+	}
+	return AttributesExpression;
+}
+
+bool UMaterialExpressionMakeMaterialAttributes::GenerateHLSLExpression(FMaterialHLSLGenerator& Generator, UE::HLSLTree::FScope& Scope, int32 OutputIndex, UE::HLSLTree::FExpression*& OutExpression)
+{
+	using namespace UE::HLSLTree;
+	FExpression* AttributesExpression = Generator.NewConstant(Generator.GetMaterialAttributesDefaultValue());
+	AttributesExpression = SetAttribute(Generator, Scope, MP_BaseColor, BaseColor, AttributesExpression);
+	AttributesExpression = SetAttribute(Generator, Scope, MP_Metallic, Metallic, AttributesExpression);
+
+	AttributesExpression = SetAttribute(Generator, Scope, MP_Specular, Specular, AttributesExpression);
+	AttributesExpression = SetAttribute(Generator, Scope, MP_Roughness, Roughness, AttributesExpression);
+	AttributesExpression = SetAttribute(Generator, Scope, MP_Anisotropy, Anisotropy, AttributesExpression);
+	AttributesExpression = SetAttribute(Generator, Scope, MP_EmissiveColor, EmissiveColor, AttributesExpression);
+	AttributesExpression = SetAttribute(Generator, Scope, MP_Opacity, Opacity, AttributesExpression);
+	AttributesExpression = SetAttribute(Generator, Scope, MP_OpacityMask, OpacityMask, AttributesExpression);
+	AttributesExpression = SetAttribute(Generator, Scope, MP_Normal, Normal, AttributesExpression);
+
+	AttributesExpression = SetAttribute(Generator, Scope, MP_Tangent, Tangent, AttributesExpression);
+	AttributesExpression = SetAttribute(Generator, Scope, MP_WorldPositionOffset, WorldPositionOffset, AttributesExpression);
+	AttributesExpression = SetAttribute(Generator, Scope, MP_SubsurfaceColor, SubsurfaceColor, AttributesExpression);
+	AttributesExpression = SetAttribute(Generator, Scope, MP_CustomData0, ClearCoat, AttributesExpression);
+	AttributesExpression = SetAttribute(Generator, Scope, MP_CustomData1, ClearCoatRoughness, AttributesExpression);
+	AttributesExpression = SetAttribute(Generator, Scope, MP_AmbientOcclusion, AmbientOcclusion, AttributesExpression);
+	AttributesExpression = SetAttribute(Generator, Scope, MP_PixelDepthOffset, PixelDepthOffset, AttributesExpression);
+	AttributesExpression = SetAttribute(Generator, Scope, MP_ShadingModel, ShadingModel, AttributesExpression);
+
+	for (int32 Index = 0; Index < 8; ++Index)
+	{
+		const EMaterialProperty UVProperty = (EMaterialProperty)(MP_CustomizedUVs0 + Index);
+		AttributesExpression = SetAttribute(Generator, Scope, UVProperty, CustomizedUVs[Index], AttributesExpression);
+	}
+
+	OutExpression = AttributesExpression;
+	return true;
+}
+
+bool UMaterialExpressionBreakMaterialAttributes::GenerateHLSLExpression(FMaterialHLSLGenerator& Generator, UE::HLSLTree::FScope& Scope, int32 OutputIndex, UE::HLSLTree::FExpression*& OutExpression)
+{
+	using namespace UE::HLSLTree;
+	using namespace UE::Shader;
 
-	OutExpression = Generator.GetTree().NewExpression<UE::HLSLTree::FExpressionReflectionVector>();
+	BuildPropertyToIOIndexMap();
+
+	// Here we don't care about any multiplex index coming in.
+	// We pass through our output index as the multiplex index so the MakeMaterialAttriubtes node at the other end can send us the right data.
+	const EMaterialProperty* Property = PropertyToIOIndexMap.FindKey(OutputIndex);
+	if (!Property)
+	{
+		return Generator.GetErrors().AddError(TEXT("Invalid output"));
+	}
+
+	FExpression* AttributesExpression = MaterialAttributes.AcquireHLSLExpression(Generator, Scope);
+	if (!AttributesExpression)
+	{
+		return false;
+	}
+
+	const FString& AttributeName = FMaterialAttributeDefinitionMap::GetAttributeName(*Property);
+	const FStructField* AttributeField = Generator.GetMaterialAttributesType()->FindFieldByName(*AttributeName);
+	OutExpression = Generator.GetTree().NewExpression<FExpressionGetStructField>(Generator.GetMaterialAttributesType(), AttributeField, AttributesExpression);
 	return true;
 }
 
@@ -1149,6 +1636,173 @@ bool UMaterialExpressionIf::GenerateHLSLExpression(FMaterialHLSLGenerator& Gener
 	return true;
 }
 
+bool UMaterialExpressionFresnel::GenerateHLSLExpression(FMaterialHLSLGenerator& Generator, UE::HLSLTree::FScope& Scope, int32 OutputIndex, UE::HLSLTree::FExpression*& OutExpression)
+{
+	using namespace UE::HLSLTree;
+
+	// pow(1 - max(0,Normal dot Camera),Exponent) * (1 - BaseReflectFraction) + BaseReflectFraction
+
+	FExpression* ExpressionExponent = ExponentIn.AcquireHLSLExpressionOrConstant(Generator, Scope, Exponent);
+	FExpression* ExpressionReflectFraction = BaseReflectFractionIn.AcquireHLSLExpressionOrConstant(Generator, Scope, BaseReflectFraction);
+	if (!ExpressionExponent || !ExpressionReflectFraction)
+	{
+		return false;
+	}
+
+	FExpression* ExpressionNormal = Normal.TryAcquireHLSLExpression(Generator, Scope);
+	if (!ExpressionNormal)
+	{
+		ExpressionNormal = Generator.GetTree().NewExpression<FExpressionExternalInput>(EExternalInput::WorldNormal);
+	}
+
+	FExpression* ExpressionNdotV = Generator.GetTree().NewDot(ExpressionNormal, Generator.GetTree().NewExpression<FExpressionExternalInput>(EExternalInput::CameraVector));
+	FExpression* ExpressionMax = Generator.GetTree().NewMax(Generator.NewConstant(0.0f), ExpressionNdotV);
+	FExpression* ExpressionMinus = Generator.GetTree().NewSub(Generator.NewConstant(1.0f), ExpressionMax);
+
+	// Compiler->Power got changed to call PositiveClampedPow instead of ClampedPow
+	// Manually implement ClampedPow to maintain backwards compatibility in the case where the input normal is not normalized (length > 1)
+	FExpression* ExpressionAbs = Generator.GetTree().NewMax(Generator.GetTree().NewAbs(ExpressionMinus), Generator.NewConstant(KINDA_SMALL_NUMBER));
+
+	FExpression* ExpressionPow = Generator.GetTree().NewBinaryOp(EOperation::PowPositiveClamped, ExpressionAbs, ExpressionExponent);
+
+	OutExpression = Generator.GetTree().NewMul(ExpressionPow, Generator.GetTree().NewSub(Generator.NewConstant(1.0f), ExpressionReflectFraction));
+	OutExpression = Generator.GetTree().NewAdd(OutExpression, ExpressionReflectFraction);
+	return true;
+}
+
+bool UMaterialExpressionDesaturation::GenerateHLSLExpression(FMaterialHLSLGenerator& Generator, UE::HLSLTree::FScope& Scope, int32 OutputIndex, UE::HLSLTree::FExpression*& OutExpression)
+{
+	using namespace UE::HLSLTree;
+	FExpression* ExpressionInput = Input.AcquireHLSLExpression(Generator, Scope);
+	if (!ExpressionInput)
+	{
+		return false;
+	}
+
+	FExpression* ExpressionGrey = Generator.GetTree().NewDot(ExpressionInput, Generator.NewConstant(FVector3f(LuminanceFactors.R, LuminanceFactors.G, LuminanceFactors.B)));
+	FExpression* ExpressionFraction = Fraction.TryAcquireHLSLExpression(Generator, Scope);
+	if (ExpressionFraction)
+	{
+		OutExpression = Generator.GetTree().NewLerp(ExpressionInput, ExpressionGrey, ExpressionFraction);
+	}
+	else
+	{
+		OutExpression = ExpressionGrey;
+	}
+	return true;
+}
+
+bool UMaterialExpressionSphereMask::GenerateHLSLExpression(FMaterialHLSLGenerator& Generator, UE::HLSLTree::FScope& Scope, int32 OutputIndex, UE::HLSLTree::FExpression*& OutExpression)
+{
+	using namespace UE::HLSLTree;
+
+	FExpression* ExpressionA = A.AcquireHLSLExpression(Generator, Scope);
+	FExpression* ExpressionB = B.AcquireHLSLExpression(Generator, Scope);
+	if (!ExpressionA || !ExpressionB)
+	{
+		return false;
+	}
+
+	FExpression* ExpressionDistance = Generator.GetTree().NewLength(Generator.GetTree().NewSub(ExpressionA, ExpressionB));
+	FExpression* ExpressionRadius = Radius.AcquireHLSLExpressionOrConstant(Generator, Scope, AttenuationRadius);
+	FExpression* ExpressionRcpRadius = Generator.GetTree().NewRcp(Generator.GetTree().NewMax(ExpressionRadius, Generator.NewConstant(0.00001f)));
+	FExpression* ExpressionNormalizedDistance = Generator.GetTree().NewMul(ExpressionDistance, ExpressionRcpRadius);
+	FExpression* ExpressionHardness = Hardness.AcquireHLSLExpressionOrConstant(Generator, Scope, HardnessPercent * 0.01f);
+	FExpression* ExpressionSoftness = Generator.GetTree().NewSub(Generator.NewConstant(1.0f), ExpressionHardness);
+	FExpression* ExpressionRcpSoftness = Generator.GetTree().NewRcp(Generator.GetTree().NewMax(ExpressionSoftness, Generator.NewConstant(0.00001f)));
+	FExpression* ExpressionOneMinusDistance = Generator.GetTree().NewSub(Generator.NewConstant(1.0f), ExpressionNormalizedDistance);
+	
+	OutExpression = Generator.GetTree().NewSaturate(Generator.GetTree().NewMul(ExpressionOneMinusDistance, ExpressionRcpSoftness));
+	return true;
+}
+
+bool UMaterialExpressionRotateAboutAxis::GenerateHLSLExpression(FMaterialHLSLGenerator& Generator, UE::HLSLTree::FScope& Scope, int32 OutputIndex, UE::HLSLTree::FExpression*& OutExpression)
+{
+	using namespace UE::HLSLTree;
+
+	FExpression* ExpressionAxis = NormalizedRotationAxis.AcquireHLSLExpression(Generator, Scope);
+	FExpression* ExpressionRotationAngle = RotationAngle.AcquireHLSLExpression(Generator, Scope);
+	FExpression* ExpressionPivotPoint = PivotPoint.AcquireHLSLExpression(Generator, Scope);
+	FExpression* ExpressionPosition = Position.AcquireHLSLExpression(Generator, Scope);
+	if (!ExpressionAxis || !ExpressionRotationAngle || !ExpressionPivotPoint || !ExpressionPosition)
+	{
+		return false;
+	}
+
+	FTree& Tree = Generator.GetTree();
+	ExpressionRotationAngle = Tree.NewMul(ExpressionRotationAngle, Generator.NewConstant(2.0f * (float)PI / Period));
+
+	FExpression* ClosestPointOnAxis = Tree.NewAdd(ExpressionPivotPoint, Tree.NewMul(ExpressionAxis, Tree.NewDot(ExpressionAxis, Tree.NewSub(ExpressionPosition, ExpressionPivotPoint))));
+	FExpression* UAxis = Tree.NewSub(ExpressionPosition, ClosestPointOnAxis);
+	FExpression* VAxis = Tree.NewCross(ExpressionAxis, UAxis);
+	FExpression* CosAngle = Tree.NewCos(ExpressionRotationAngle);
+	FExpression* SinAngle = Tree.NewSin(ExpressionRotationAngle);
+	FExpression* R = Tree.NewAdd(Tree.NewMul(UAxis, CosAngle), Tree.NewMul(VAxis, SinAngle));
+	FExpression* RotatedPosition = Tree.NewAdd(ClosestPointOnAxis, R);
+
+	// TODO - see MaterialTemplate.ush LWC version of RotateAboutAxis to optimize
+	OutExpression = Tree.NewSub(RotatedPosition, ExpressionPosition);
+	return true;
+}
+
+bool UMaterialExpressionPanner::GenerateHLSLExpression(FMaterialHLSLGenerator& Generator, UE::HLSLTree::FScope& Scope, int32 OutputIndex, UE::HLSLTree::FExpression*& OutExpression)
+{
+	using namespace UE::HLSLTree;
+
+	FExpression* ExpressionTime = Time.TryAcquireHLSLExpression(Generator, Scope);
+	if (!ExpressionTime)
+	{
+		ExpressionTime = Generator.GetTree().NewExpression<FExpressionExternalInput>(EExternalInput::GameTime);
+	}
+	FExpression* ExpressionSpeed = Speed.AcquireHLSLExpressionOrConstant(Generator, Scope, FVector2f(SpeedX, SpeedY));
+	FExpression* ExpressionOffset = Generator.GetTree().NewMul(ExpressionSpeed, ExpressionTime);
+	if (bFractionalPart)
+	{
+		ExpressionOffset = Generator.GetTree().NewFrac(ExpressionOffset);
+	}
+	FExpression* ExpressionTexCoord = Coordinate.TryAcquireHLSLExpression(Generator, Scope);
+	if (!ExpressionTexCoord)
+	{
+		ExpressionTexCoord = Generator.NewTexCoord(ConstCoordinate);
+	}
+
+	OutExpression = Generator.GetTree().NewAdd(ExpressionTexCoord, ExpressionOffset);
+	return true;
+}
+
+bool UMaterialExpressionRotator::GenerateHLSLExpression(FMaterialHLSLGenerator& Generator, UE::HLSLTree::FScope& Scope, int32 OutputIndex, UE::HLSLTree::FExpression*& OutExpression)
+{
+	using namespace UE::HLSLTree;
+
+	FExpression* ExpressionTime = Time.TryAcquireHLSLExpression(Generator, Scope);
+	if (!ExpressionTime)
+	{
+		ExpressionTime = Generator.GetTree().NewExpression<FExpressionExternalInput>(EExternalInput::GameTime);
+	}
+	ExpressionTime = Generator.GetTree().NewMul(ExpressionTime, Generator.NewConstant(Speed));
+
+	FExpression* ExpressionCos = Generator.GetTree().NewUnaryOp(EOperation::Cos, ExpressionTime);
+	FExpression* ExpressionSin = Generator.GetTree().NewUnaryOp(EOperation::Sin, ExpressionTime);
+	FExpression* ExpressionRowX = Generator.GetTree().NewExpression<FExpressionAppend>(ExpressionCos, Generator.GetTree().NewMul(ExpressionSin, Generator.NewConstant(-1.0f)));
+	FExpression* ExpressionRowY = Generator.GetTree().NewExpression<FExpressionAppend>(ExpressionSin, ExpressionCos);
+	FExpression* ExpressionOrigin = Generator.GetTree().NewConstant(FVector2f(CenterX, CenterY));
+
+	FExpression* ExpressionCoord = Coordinate.TryAcquireHLSLExpression(Generator, Scope);
+	if (!ExpressionCoord)
+	{
+		ExpressionCoord = Generator.NewTexCoord(ConstCoordinate);
+	}
+	FExpression* ExpressionCoordXY = Generator.GetTree().NewSub(Generator.NewSwizzle(FSwizzleParameters(0, 1), ExpressionCoord), ExpressionOrigin);
+
+	FExpression* ExpressionResultX = Generator.GetTree().NewDot(ExpressionRowX, ExpressionCoordXY);
+	FExpression* ExpressionResultY = Generator.GetTree().NewDot(ExpressionRowY, ExpressionCoordXY);
+	FExpression* ExpressionResult = Generator.GetTree().NewAdd(Generator.GetTree().NewExpression<FExpressionAppend>(ExpressionResultX, ExpressionResultY), ExpressionOrigin);
+
+	// If given coordinate was a float3, we append the Z coordinate onto the result
+	OutExpression = Generator.GetTree().NewExpression<FExpressionAppend>(ExpressionResult, Generator.NewSwizzle(FSwizzleParameters(2), ExpressionCoord));
+	return true;
+}
+
 static UE::Shader::FType GetCustomOutputType(const FMaterialHLSLGenerator& Generator, ECustomMaterialOutputType Type)
 {
 	using namespace UE::Shader;
@@ -1163,6 +1817,126 @@ static UE::Shader::FType GetCustomOutputType(const FMaterialHLSLGenerator& Gener
 	}
 }
 
+bool UMaterialExpressionVolumetricAdvancedMaterialInput::GenerateHLSLExpression(FMaterialHLSLGenerator& Generator, UE::HLSLTree::FScope& Scope, int32 OutputIndex, UE::HLSLTree::FExpression*& OutExpression)
+{
+	if (OutputIndex != 0)
+	{
+		return Generator.GetErrors().AddError(TEXT("Invalid output"));
+	}
+	OutExpression = Generator.GetTree().NewExpression<UE::HLSLTree::FExpressionInlineCustomHLSL>(UE::Shader::EValueType::Float3, TEXT("MaterialExpressionVolumeSampleConservativeDensity(Parameters)"));
+	return true;
+}
+
+bool UMaterialExpressionVolumetricAdvancedMaterialOutput::GenerateHLSLExpression(FMaterialHLSLGenerator& Generator, UE::HLSLTree::FScope& Scope, int32 OutputIndex, UE::HLSLTree::FExpression*& OutExpression)
+{
+	switch (OutputIndex)
+	{
+	case 0: OutExpression = PhaseG.AcquireHLSLExpressionOrConstant(Generator, Scope, ConstPhaseG); break;
+	case 1: OutExpression = PhaseG2.AcquireHLSLExpressionOrConstant(Generator, Scope, ConstPhaseG2); break;
+	case 2: OutExpression = PhaseBlend.AcquireHLSLExpressionOrConstant(Generator, Scope, ConstPhaseBlend); break;
+	case 3: OutExpression = MultiScatteringContribution.AcquireHLSLExpressionOrConstant(Generator, Scope, ConstMultiScatteringContribution); break;
+	case 4: OutExpression = MultiScatteringOcclusion.AcquireHLSLExpressionOrConstant(Generator, Scope, ConstMultiScatteringOcclusion); break;
+	case 5: OutExpression = MultiScatteringEccentricity.AcquireHLSLExpressionOrConstant(Generator, Scope, ConstMultiScatteringEccentricity); break;
+	case 6: OutExpression = ConservativeDensity.AcquireHLSLExpressionOrConstant(Generator, Scope, FVector3f(1.0f, 1.0f, 1.0f)); break;
+	default: return Generator.GetErrors().AddError(TEXT("Invlid output"));
+	}
+	return OutExpression != nullptr;
+}
+
+UE::Shader::EValueType UMaterialExpressionVolumetricAdvancedMaterialOutput::GetCustomOutputType(int32 OutputIndex) const
+{
+	if (OutputIndex >= 0 && OutputIndex < 6) return UE::Shader::EValueType::Float1;
+	else if (OutputIndex == 6) return UE::Shader::EValueType::Float3;
+	else return UE::Shader::EValueType::Void;
+}
+
+bool UMaterialExpressionHairAttributes::GenerateHLSLExpression(FMaterialHLSLGenerator& Generator, UE::HLSLTree::FScope& Scope, int32 OutputIndex, UE::HLSLTree::FExpression*& OutExpression)
+{
+	using namespace UE::HLSLTree;
+	using namespace UE::Shader;
+
+	const FStringView FunctionHairTangent = bUseTangentSpace
+		? FStringView(TEXT("MaterialExpressionGetHairTangent(Parameters, true)"))
+		: FStringView(TEXT("MaterialExpressionGetHairTangent(Parameters, false)"));
+
+	switch (OutputIndex)
+	{
+	case 0:
+	case 1: OutExpression = Generator.GetTree().NewExpression<FExpressionInlineCustomHLSL>(EValueType::Float2, TEXT("MaterialExpressionGetHairUV(Parameters)")); break;
+	case 2:
+	case 3: OutExpression = Generator.GetTree().NewExpression<FExpressionInlineCustomHLSL>(EValueType::Float2, TEXT("MaterialExpressionGetHairDimensions(Parameters)")); break;
+	case 4: OutExpression = Generator.GetTree().NewExpression<FExpressionInlineCustomHLSL>(EValueType::Float1, TEXT("MaterialExpressionGetHairSeed(Parameters)")); break;
+	case 5: OutExpression = Generator.GetTree().NewExpression<FExpressionInlineCustomHLSL>(EValueType::Float3, FunctionHairTangent); break;
+	case 6: OutExpression = Generator.GetTree().NewExpression<FExpressionInlineCustomHLSL>(EValueType::Float2, TEXT("MaterialExpressionGetHairRootUV(Parameters)")); break;
+	case 7: OutExpression = Generator.GetTree().NewExpression<FExpressionInlineCustomHLSL>(EValueType::Float3, TEXT("MaterialExpressionGetHairBaseColor(Parameters)")); break;
+	case 8: OutExpression = Generator.GetTree().NewExpression<FExpressionInlineCustomHLSL>(EValueType::Float1, TEXT("MaterialExpressionGetHairRoughness(Parameters)")); break;
+	case 9: OutExpression = Generator.GetTree().NewExpression<FExpressionInlineCustomHLSL>(EValueType::Float1, TEXT("MaterialExpressionGetHairDepth(Parameters)")); break;
+	case 10: OutExpression = Generator.GetTree().NewExpression<FExpressionInlineCustomHLSL>(EValueType::Float1, TEXT("MaterialExpressionGetHairCoverage(Parameters)")); break;
+	case 11: OutExpression = Generator.GetTree().NewExpression<FExpressionInlineCustomHLSL>(EValueType::Float4, TEXT("MaterialExpressionGetHairAuxilaryData(Parameters)")); break;
+	case 12: OutExpression = Generator.GetTree().NewExpression<FExpressionInlineCustomHLSL>(EValueType::Float2, TEXT("MaterialExpressionGetAtlasUVs(Parameters)")); break;
+	case 13: OutExpression = Generator.GetTree().NewExpression<FExpressionInlineCustomHLSL>(EValueType::Float1, TEXT("MaterialExpressionGetHairGroupIndex(Parameters)")); break;
+	default: return Generator.GetErrors().AddError(TEXT("Invalid output"));
+	}
+	return OutExpression != nullptr;
+}
+
+bool UMaterialExpressionCloudSampleAttribute::GenerateHLSLExpression(FMaterialHLSLGenerator& Generator, UE::HLSLTree::FScope& Scope, int32 OutputIndex, UE::HLSLTree::FExpression*& OutExpression)
+{
+	using namespace UE::HLSLTree;
+	using namespace UE::Shader;
+
+	switch (OutputIndex)
+	{
+	case 0: OutExpression = Generator.GetTree().NewExpression<FExpressionInlineCustomHLSL>(EValueType::Float1, TEXT("MaterialExpressionCloudSampleAltitude(Parameters)")); break;
+	case 1: OutExpression = Generator.GetTree().NewExpression<FExpressionInlineCustomHLSL>(EValueType::Float1, TEXT("MaterialExpressionCloudSampleAltitudeInLayer(Parameters)")); break;
+	case 2: OutExpression = Generator.GetTree().NewExpression<FExpressionInlineCustomHLSL>(EValueType::Float1, TEXT("MaterialExpressionCloudSampleNormAltitudeInLayer(Parameters)")); break;
+	case 3: OutExpression = Generator.GetTree().NewExpression<FExpressionInlineCustomHLSL>(EValueType::Float1, TEXT("MaterialExpressionVolumeSampleShadowSampleDistance(Parameters)")); break;
+	default: return Generator.GetErrors().AddError(TEXT("Invalid output"));
+	}
+	return OutExpression != nullptr;
+}
+
+bool UMaterialExpressionPerInstanceFadeAmount::GenerateHLSLExpression(FMaterialHLSLGenerator& Generator, UE::HLSLTree::FScope& Scope, int32 OutputIndex, UE::HLSLTree::FExpression*& OutExpression)
+{
+	using namespace UE::HLSLTree;
+	using namespace UE::Shader;
+	OutExpression = Generator.GetTree().NewExpression<FExpressionInlineCustomHLSL>(EValueType::Float1, TEXT("GetPerInstanceFadeAmount(Parameters)"));
+	return true;
+}
+
+bool UMaterialExpressionObjectBounds::GenerateHLSLExpression(FMaterialHLSLGenerator& Generator, UE::HLSLTree::FScope& Scope, int32 OutputIndex, UE::HLSLTree::FExpression*& OutExpression)
+{
+	using namespace UE::HLSLTree;
+	using namespace UE::Shader;
+	const FStringView Code(TEXT("float3(GetPrimitiveData(Parameters).ObjectBoundsX, GetPrimitiveData(Parameters).ObjectBoundsY, GetPrimitiveData(Parameters).ObjectBoundsZ)"));
+	OutExpression = Generator.GetTree().NewExpression<FExpressionInlineCustomHLSL>(EValueType::Float3, Code);
+	return true;
+}
+
+bool UMaterialExpressionObjectOrientation::GenerateHLSLExpression(FMaterialHLSLGenerator& Generator, UE::HLSLTree::FScope& Scope, int32 OutputIndex, UE::HLSLTree::FExpression*& OutExpression)
+{
+	using namespace UE::HLSLTree;
+	using namespace UE::Shader;
+	OutExpression = Generator.GetTree().NewExpression<FExpressionInlineCustomHLSL>(EValueType::Float3, TEXT("GetObjectOrientation(Parameters)"));
+	return true;
+}
+
+bool UMaterialExpressionObjectPositionWS::GenerateHLSLExpression(FMaterialHLSLGenerator& Generator, UE::HLSLTree::FScope& Scope, int32 OutputIndex, UE::HLSLTree::FExpression*& OutExpression)
+{
+	using namespace UE::HLSLTree;
+	using namespace UE::Shader;
+	OutExpression = Generator.GetTree().NewExpression<FExpressionInlineCustomHLSL>(EValueType::Double3, TEXT("GetObjectWorldPosition(Parameters)"));
+	return true;
+}
+
+bool UMaterialExpressionObjectRadius::GenerateHLSLExpression(FMaterialHLSLGenerator& Generator, UE::HLSLTree::FScope& Scope, int32 OutputIndex, UE::HLSLTree::FExpression*& OutExpression)
+{
+	using namespace UE::HLSLTree;
+	using namespace UE::Shader;
+	OutExpression = Generator.GetTree().NewExpression<FExpressionInlineCustomHLSL>(EValueType::Float1, TEXT("GetPrimitiveData(Parameters).ObjectRadius"));
+	return true;
+}
+
 bool UMaterialExpressionCustom::GenerateHLSLExpression(FMaterialHLSLGenerator& Generator, UE::HLSLTree::FScope& Scope, int32 OutputIndex, UE::HLSLTree::FExpression*& OutExpression)
 {
 	using namespace UE::HLSLTree;
diff --git a/Engine/Source/Runtime/Engine/Private/Materials/MaterialExpressions.cpp b/Engine/Source/Runtime/Engine/Private/Materials/MaterialExpressions.cpp
--- a/Engine/Source/Runtime/Engine/Private/Materials/MaterialExpressions.cpp
+++ b/Engine/Source/Runtime/Engine/Private/Materials/MaterialExpressions.cpp
@@ -6110,9 +6110,9 @@ void UMaterialExpressionBreakMaterialAttributes::Serialize(FStructuredArchive::F
 }
 
 #if WITH_EDITOR
-static TMap<EMaterialProperty, int32> PropertyToIOIndexMap;
+TMap<EMaterialProperty, int32> UMaterialExpressionBreakMaterialAttributes::PropertyToIOIndexMap;
 
-static void BuildPropertyToIOIndexMap()
+void UMaterialExpressionBreakMaterialAttributes::BuildPropertyToIOIndexMap()
 {
 	if (PropertyToIOIndexMap.Num() == 0)
 	{
diff --git a/Engine/Source/Runtime/Engine/Private/Materials/MaterialHLSLEmitter.cpp b/Engine/Source/Runtime/Engine/Private/Materials/MaterialHLSLEmitter.cpp
--- a/Engine/Source/Runtime/Engine/Private/Materials/MaterialHLSLEmitter.cpp
+++ b/Engine/Source/Runtime/Engine/Private/Materials/MaterialHLSLEmitter.cpp
@@ -14,6 +14,7 @@
 #include "MaterialHLSLGenerator.h"
 #include "ShaderCore.h"
 #include "HLSLTree/HLSLTree.h"
+#include "HLSLTree/HLSLTreeCommon.h"
 #include "HLSLTree/HLSLTreeEmit.h"
 #include "Containers/LazyPrintf.h"
 
@@ -52,6 +53,8 @@ static FString GenerateMaterialTemplateHLSL(EShaderPlatform ShaderPlatform,
 	const TCHAR* PixelShaderCodePhase0,
 	const TCHAR* PixelShaderCodePhase1)
 {
+	using namespace UE::HLSLTree;
+
 	FString MaterialTemplateSource;
 	LoadShaderSourceFileChecked(TEXT("/Engine/Private/MaterialTemplate.ush"), ShaderPlatform, MaterialTemplateSource);
 
@@ -76,9 +79,15 @@ static FString GenerateMaterialTemplateHLSL(EShaderPlatform ShaderPlatform,
 
 	FLazyPrintf LazyPrintf(*MaterialTemplateSource);
 
-	const uint32 VertexTexCoordMask = EmitContext.TexCoordMask[SF_Vertex];
-	const uint32 PixelTexCoordMask = EmitContext.TexCoordMask[SF_Pixel];
-	const uint32 NumTexCoords = PixelTexCoordMask ? (FMath::FloorLog2(PixelTexCoordMask) + 1u) : 0u;
+	uint32 NumTexCoords = 0u;
+	for (int32 TexCoordIndex = UE::HLSLTree::MaxNumTexCoords - 1; TexCoordIndex >= 0; --TexCoordIndex)
+	{
+		if (EmitContext.IsExternalInputUsed(SF_Pixel, MakeInputTexCoord(TexCoordIndex)))
+		{
+			NumTexCoords = TexCoordIndex + 1;
+			break;
+		}
+	}
 
 	const uint32 NumUserVertexTexCoords = NumTexCoords;
 	const uint32 NumUserTexCoords = NumTexCoords;
@@ -301,7 +310,9 @@ static FString GenerateMaterialTemplateHLSL(EShaderPlatform ShaderPlatform,
 		EvaluateMaterialDeclaration += TEXT("    FMaterialAttributes Result = EvaluateVertexMaterialAttributesInternal(Parameters);" LINE_TERMINATOR);
 		for (uint32 TexCoordIndex = 0; TexCoordIndex < NumTexCoords; ++TexCoordIndex)
 		{
-			if ((PixelTexCoordMask & (1u << TexCoordIndex)) && !(VertexTexCoordMask & (1u << TexCoordIndex)))
+			const EExternalInput TexCoordInput = MakeInputTexCoord(TexCoordIndex);
+
+			if (EmitContext.IsExternalInputUsed(SF_Pixel, TexCoordInput) && !EmitContext.IsExternalInputUsed(SF_Vertex, TexCoordInput))
 			{
 				const FString AttributeName = FMaterialAttributeDefinitionMap::GetAttributeName((EMaterialProperty)(MP_CustomizedUVs0 + TexCoordIndex));
 				EvaluateMaterialDeclaration += FString::Printf(TEXT("    Result.%s = Parameters.TexCoords[%d];") LINE_TERMINATOR, *AttributeName, TexCoordIndex);
@@ -442,7 +453,7 @@ static void GetMaterialEnvironment(EShaderPlatform InPlatform,
 		OutEnvironment.SetDefine(TEXT("USE_PARTICLE_SUBUVS"), TEXT("1"));
 	}
 
-	if (false)//bUsesLightmapUVs)
+	if (EmitContext.ExternalInputMask[SF_Pixel][(int32)UE::HLSLTree::EExternalInput::LightmapTexCoord])
 	{
 		OutEnvironment.SetDefine(TEXT("LIGHTMAP_UV_ACCESS"), TEXT("1"));
 	}
@@ -545,12 +556,11 @@ static void GetMaterialEnvironment(EShaderPlatform InPlatform,
 	// If the material gets its shading model from the material expressions, then we use the result from the compilation (assuming it's valid).
 	// This result will potentially be tighter than what GetShadingModels() returns, because it only picks up the shading models from the expressions that get compiled for a specific feature level and quality level
 	// For example, the material might have shading models behind static switches. GetShadingModels() will return both the true and the false paths from that switch, whereas the shading model field from the compilation will only contain the actual shading model selected 
-	/*if (InMaterial.IsShadingModelFromMaterialExpression() && ShadingModelsFromCompilation.IsValid())
+	if (InMaterial.IsShadingModelFromMaterialExpression() && EmitContext.ShadingModelsFromCompilation.IsValid())
 	{
 		// Shading models fetched from the compilation of the expression graph
-		ShadingModels = ShadingModelsFromCompilation;
-	}*/
-
+		ShadingModels = EmitContext.ShadingModelsFromCompilation;
+	}
 	ensure(ShadingModels.IsValid());
 
 	if (ShadingModels.IsLit())
@@ -726,17 +736,17 @@ static bool IsAttributeUsed(FMaterialHLSLGenerator& Generator,
 		return false;
 	}
 
+	FRequestedType RequestedType;
+	RequestedType.SetFieldRequested(PropertyField);
+
 	const int32 NumComponents = PropertyField->GetNumComponents();
-	const EExpressionEvaluation Evaluation = ResultType.GetFieldEvaluation(Scope, PropertyField->ComponentIndex, NumComponents);
+	const EExpressionEvaluation Evaluation = ResultType.GetFieldEvaluation(Scope, RequestedType, PropertyField->ComponentIndex, NumComponents);
 	if (Evaluation == EExpressionEvaluation::None)
 	{
 		return false;
 	}
-	else if (Evaluation == EExpressionEvaluation::Constant)
+	else if (Evaluation == EExpressionEvaluation::Constant || Evaluation == EExpressionEvaluation::ConstantZero)
 	{
-		FRequestedType RequestedType;
-		RequestedType.SetFieldRequested(PropertyField);
-
 		const FValue& DefaultValue = Generator.GetMaterialAttributesDefaultValue();
 		const FValue ConstantValue = Expression->GetValueConstant(Context, Scope, RequestedType);
 		if (DefaultValue.GetType() != ConstantValue.GetType())
diff --git a/Engine/Source/Runtime/Engine/Private/Materials/MaterialHLSLGenerator.cpp b/Engine/Source/Runtime/Engine/Private/Materials/MaterialHLSLGenerator.cpp
--- a/Engine/Source/Runtime/Engine/Private/Materials/MaterialHLSLGenerator.cpp
+++ b/Engine/Source/Runtime/Engine/Private/Materials/MaterialHLSLGenerator.cpp
@@ -274,16 +274,12 @@ void FMaterialHLSLGenerator::EmitSharedCode(FStringBuilderBase& OutCode) const
 static UE::HLSLTree::FExpression* CompileMaterialInput(FMaterialHLSLGenerator& Generator,
 	UE::HLSLTree::FScope& Scope,
 	EMaterialProperty InputProperty,
-	UMaterial* Material,
-	UE::HLSLTree::FExpression* AttributesExpression)
+	UMaterial* Material)
 {
 	using namespace UE::HLSLTree;
 
-	// We're only interesting in attributes that map to valid fields
-	const UE::Shader::FStructField* AttributeField = Generator.GetMaterialAttributesType()->FindFieldByName(*FMaterialAttributeDefinitionMap::GetAttributeName(InputProperty));
-
 	FExpression* Expression = nullptr;
-	if (AttributeField && Material->IsPropertyActive(InputProperty))
+	if (Material->IsPropertyActive(InputProperty))
 	{
 		FMaterialInputDescription InputDescription;
 		if (Material->GetExpressionInputDescription(InputProperty, InputDescription))
@@ -304,17 +300,7 @@ static UE::HLSLTree::FExpression* CompileMaterialInput(FMaterialHLSLGenerator& G
 		}
 	}
 
-	if (Expression)
-	{
-		FExpression* SetAttributeExpression = Generator.GetTree().NewExpression<FExpressionSetStructField>(
-			Generator.GetMaterialAttributesType(),
-			AttributeField,
-			AttributesExpression,
-			Expression);
-		return SetAttributeExpression;
-	}
-
-	return AttributesExpression;
+	return Expression;
 }
 
 bool FMaterialHLSLGenerator::GenerateResult(UE::HLSLTree::FScope& Scope)
@@ -348,6 +334,8 @@ bool FMaterialHLSLGenerator::GenerateResult(UE::HLSLTree::FScope& Scope)
 
 		if (TargetMaterial)
 		{
+			const FStructField* PrevWPOField = GetMaterialAttributesType()->FindFieldByName(TEXT("PrevWorldPositionOffset"));
+
 			FExpression* AttributesExpression = nullptr;
 			if (TargetMaterial->bUseMaterialAttributes)
 			{
@@ -355,18 +343,21 @@ bool FMaterialHLSLGenerator::GenerateResult(UE::HLSLTree::FScope& Scope)
 				if (TargetMaterial->GetExpressionInputDescription(MP_MaterialAttributes, InputDescription))
 				{
 					check(InputDescription.Type == UE::Shader::EValueType::Struct);
-					AttributesExpression = InputDescription.Input->AcquireHLSLExpression(*this, Scope);
+					AttributesExpression = InputDescription.Input->TryAcquireHLSLExpression(*this, Scope);
 
-					const FString& WPOName = FMaterialAttributeDefinitionMap::GetAttributeName(MP_WorldPositionOffset);
-					const FStructField* WPOField = GetMaterialAttributesType()->FindFieldByName(*WPOName);
-					const FStructField* PrevWPOField = GetMaterialAttributesType()->FindFieldByName(TEXT("PrevWorldPositionOffset"));
+					if (AttributesExpression)
+					{
+						const FString& WPOName = FMaterialAttributeDefinitionMap::GetAttributeName(MP_WorldPositionOffset);
+						const FStructField* WPOField = GetMaterialAttributesType()->FindFieldByName(*WPOName);
 
-					FRequestedType PrevRequestedType;
-					PrevRequestedType.SetFieldRequested(WPOField);
+						FRequestedType PrevRequestedType;
+						PrevRequestedType.SetFieldRequested(WPOField);
 
-					FExpression* PrevAttributesExpression = HLSLTree->GetPreviousFrame(AttributesExpression, PrevRequestedType);
-					FExpression* PrevWPOExpression = HLSLTree->NewExpression<FExpressionGetStructField>(GetMaterialAttributesType(), WPOField, PrevAttributesExpression);
-					AttributesExpression = HLSLTree->NewExpression<FExpressionSetStructField>(GetMaterialAttributesType(), PrevWPOField, AttributesExpression, PrevWPOExpression);
+						FExpression* PrevAttributesExpression = HLSLTree->GetPreviousFrame(AttributesExpression, PrevRequestedType);
+						ensure(PrevAttributesExpression);
+						FExpression* PrevWPOExpression = HLSLTree->NewExpression<FExpressionGetStructField>(GetMaterialAttributesType(), WPOField, PrevAttributesExpression);
+						AttributesExpression = HLSLTree->NewExpression<FExpressionSetStructField>(GetMaterialAttributesType(), PrevWPOField, AttributesExpression, PrevWPOExpression);
+					}
 				}
 			}
 			else
@@ -375,7 +366,27 @@ bool FMaterialHLSLGenerator::GenerateResult(UE::HLSLTree::FScope& Scope)
 				for (int32 PropertyIndex = 0; PropertyIndex < MP_MAX; ++PropertyIndex)
 				{
 					const EMaterialProperty Property = (EMaterialProperty)PropertyIndex;
-					AttributesExpression = CompileMaterialInput(*this, Scope, Property, TargetMaterial, AttributesExpression);
+
+					// We're only interesting in attributes that map to valid fields
+					const UE::Shader::FStructField* AttributeField = GetMaterialAttributesType()->FindFieldByName(*FMaterialAttributeDefinitionMap::GetAttributeName(Property));
+					if (AttributeField)
+					{
+						FExpression* InputExpression = CompileMaterialInput(*this, Scope, Property, TargetMaterial);
+						if (InputExpression)
+						{
+							AttributesExpression = GetTree().NewExpression<FExpressionSetStructField>(
+								GetMaterialAttributesType(),
+								AttributeField,
+								AttributesExpression,
+								InputExpression);
+							if (Property == MP_WorldPositionOffset)
+							{
+								FExpression* PrevWPOExpression = HLSLTree->GetPreviousFrame(InputExpression, ERequestedType::Vector3);
+								ensure(PrevWPOExpression);
+								AttributesExpression = HLSLTree->NewExpression<FExpressionSetStructField>(GetMaterialAttributesType(), PrevWPOField, AttributesExpression, PrevWPOExpression);
+							}
+						}
+					}
 				}
 			}
 
@@ -466,6 +477,7 @@ const UE::Shader::FTextureValue* FMaterialHLSLGenerator::AcquireTextureValue(con
 	FXxHash64Builder Hasher;
 	Hasher.Update(&InValue.Texture, sizeof(InValue.Texture));
 	Hasher.Update(&InValue.SamplerType, sizeof(InValue.SamplerType));
+	Hasher.Update(&InValue.ExternalTextureGuid, sizeof(InValue.ExternalTextureGuid));
 	const FXxHash64 Hash = Hasher.Finalize();
 
 	FTextureValue const* const* PrevValue = TextureValueMap.Find(Hash);
diff --git a/Engine/Source/Runtime/Engine/Private/Materials/MaterialUniformExpressions.cpp b/Engine/Source/Runtime/Engine/Private/Materials/MaterialUniformExpressions.cpp
--- a/Engine/Source/Runtime/Engine/Private/Materials/MaterialUniformExpressions.cpp
+++ b/Engine/Source/Runtime/Engine/Private/Materials/MaterialUniformExpressions.cpp
@@ -563,6 +563,19 @@ int32 FUniformExpressionSet::FindOrAddTextureParameter(EMaterialTextureParameter
 	return UniformTextureParameters[(int32)Type].Add(Info);
 }
 
+int32 FUniformExpressionSet::FindOrAddExternalTextureParameter(const FMaterialExternalTextureParameterInfo& Info)
+{
+	for (int32 i = 0; i < UniformExternalTextureParameters.Num(); ++i)
+	{
+		if (UniformExternalTextureParameters[i] == Info)
+		{
+			return i;
+		}
+	}
+
+	return UniformExternalTextureParameters.Add(Info);
+}
+
 int32 FUniformExpressionSet::FindOrAddNumericParameter(EMaterialParameterType Type, const FMaterialParameterInfo& ParameterInfo, uint32 DefaultValueOffset)
 {
 	for (int32 i = 0; i < UniformNumericParameters.Num(); ++i)
diff --git a/Engine/Source/Runtime/Engine/Private/Shader/ShaderValue.cpp b/Engine/Source/Runtime/Engine/Private/Shader/ShaderValue.cpp
--- a/Engine/Source/Runtime/Engine/Private/Shader/ShaderValue.cpp
+++ b/Engine/Source/Runtime/Engine/Private/Shader/ShaderValue.cpp
@@ -294,9 +294,16 @@ EValueType FTextureValue::GetType() const
 		case MCT_TextureCube: return EValueType::TextureCube;
 		case MCT_TextureCubeArray: return EValueType::TextureCubeArray;
 		case MCT_VolumeTexture: return EValueType::Texture3D;
+		case MCT_TextureExternal: return EValueType::TextureExternal;
+		case MCT_TextureVirtual: return EValueType::Texture2D; // TODO
 		default: checkNoEntry(); break;
 		}
 	}
+	else if (ExternalTextureGuid.IsValid())
+	{
+		return EValueType::TextureExternal;
+	}
+
 	return EValueType::Void;
 }
 
@@ -439,6 +446,7 @@ FValueComponentTypeDescription GetValueComponentTypeDescription(EValueComponentT
 	case EValueComponentType::TextureCube: return FValueComponentTypeDescription(TEXT("TextureCube"), sizeof(void*), EComponentBound::Zero, EComponentBound::Zero);
 	case EValueComponentType::TextureCubeArray: return FValueComponentTypeDescription(TEXT("TextureCubeArray"), sizeof(void*), EComponentBound::Zero, EComponentBound::Zero);
 	case EValueComponentType::Texture3D: return FValueComponentTypeDescription(TEXT("Texture3D"), sizeof(void*), EComponentBound::Zero, EComponentBound::Zero);
+	case EValueComponentType::TextureExternal: return FValueComponentTypeDescription(TEXT("TextureExternal"), sizeof(void*), EComponentBound::Zero, EComponentBound::Zero);
 	default: checkNoEntry() return FValueComponentTypeDescription();
 	}
 }
@@ -641,6 +649,7 @@ FValueTypeDescription GetValueTypeDescription(EValueType Type)
 	case EValueType::TextureCube: return FValueTypeDescription(TEXT("FTextureCube"), EValueComponentType::TextureCube, 1);
 	case EValueType::TextureCubeArray: return FValueTypeDescription(TEXT("FTextureCubeArray"), EValueComponentType::TextureCubeArray, 1);
 	case EValueType::Texture3D: return FValueTypeDescription(TEXT("FTexture3D"), EValueComponentType::Texture3D, 1);
+	case EValueType::TextureExternal: return FValueTypeDescription(TEXT("FTextureExternal"), EValueComponentType::TextureExternal, 1);
 	default: checkNoEntry(); return FValueTypeDescription(TEXT("<INVALID>"), EValueComponentType::Void, 0);
 	}
 }
@@ -701,6 +710,7 @@ EValueType MakeValueType(EValueComponentType ComponentType, int32 NumComponents)
 	case EValueComponentType::TextureCube: check(NumComponents == 1); return EValueType::TextureCube; break;
 	case EValueComponentType::TextureCubeArray: check(NumComponents == 1); return EValueType::TextureCubeArray; break;
 	case EValueComponentType::Texture3D: check(NumComponents == 1); return EValueType::Texture3D; break;
+	case EValueComponentType::TextureExternal: check(NumComponents == 1); return EValueType::TextureExternal; break;
 	default:
 		break;
 	}
diff --git a/Engine/Source/Runtime/Engine/Public/HLSLTree/HLSLTree.h b/Engine/Source/Runtime/Engine/Public/HLSLTree/HLSLTree.h
--- a/Engine/Source/Runtime/Engine/Public/HLSLTree/HLSLTree.h
+++ b/Engine/Source/Runtime/Engine/Public/HLSLTree/HLSLTree.h
@@ -243,8 +243,6 @@ struct FPreparedComponent
 	EExpressionEvaluation GetEvaluation(const FEmitScope& Scope) const;
 
 	FEmitScope* LoopScope = nullptr;
-	FExpression* ForwardValue = nullptr;
-	int32 ForwardComponentIndex = INDEX_NONE;
 	Shader::FComponentBounds Bounds;
 	EExpressionEvaluation Evaluation = EExpressionEvaluation::None;
 };
@@ -252,9 +250,7 @@ inline bool operator==(const FPreparedComponent& Lhs, const FPreparedComponent&
 {
 	return Lhs.Evaluation == Rhs.Evaluation &&
 		Lhs.Bounds == Rhs.Bounds &&
-		Lhs.LoopScope == Rhs.LoopScope &&
-		Lhs.ForwardValue == Rhs.ForwardValue &&
-		Lhs.ForwardComponentIndex == Rhs.ForwardComponentIndex;
+		Lhs.LoopScope == Rhs.LoopScope;
 }
 inline bool operator!=(const FPreparedComponent& Lhs, const FPreparedComponent& Rhs)
 {
@@ -277,7 +273,6 @@ public:
 	void SetEvaluation(EExpressionEvaluation Evaluation);
 	void MergeEvaluation(EExpressionEvaluation Evaluation);
 	void SetLoopEvaluation(FEmitScope& Scope, const FRequestedType& RequestedType);
-	void SetForwardValue(const FRequestedType& RequestedType, FExpression* ForwardValue);
 	void UpdateBounds(const FRequestedType& RequestedType, Shader::FComponentBounds Bounds);
 
 	void SetField(const Shader::FStructField* Field, const FPreparedType& FieldType);
@@ -295,7 +290,7 @@ public:
 
 	EExpressionEvaluation GetEvaluation(const FEmitScope& Scope) const;
 	EExpressionEvaluation GetEvaluation(const FEmitScope& Scope, const FRequestedType& RequestedType) const;
-	EExpressionEvaluation GetFieldEvaluation(const FEmitScope& Scope, int32 ComponentIndex, int32 NumComponents) const;
+	EExpressionEvaluation GetFieldEvaluation(const FEmitScope& Scope, const FRequestedType& RequestedType, int32 ComponentIndex, int32 NumComponents) const;
 	Shader::FComponentBounds GetBounds(const FRequestedType& RequestedType) const;
 	FPreparedComponent GetComponent(int32 Index) const;
 
@@ -334,8 +329,6 @@ public:
 	bool SetType(FEmitContext& Context, const FRequestedType& RequestedType, EExpressionEvaluation Evaluation, const Shader::FType& Type);
 	bool SetType(FEmitContext& Context, const FRequestedType& RequestedType, const FPreparedType& Type);
 
-	bool SetForwardValue(FEmitContext& Context, FEmitScope& Scope, const FRequestedType& RequestedType, FExpression* InValue);
-
 private:
 	bool TryMergePreparedType(FEmitContext& Context, const Shader::FStructType* StructType, Shader::EValueComponentType ComponentType);
 
@@ -387,7 +380,6 @@ public:
 	FRequestedType GetRequestedType() const { return PrepareValueResult.PreparedType.GetRequestedType(); }
 	Shader::FType GetType() const { return PrepareValueResult.PreparedType.GetType(); }
 	EExpressionEvaluation GetEvaluation(const FEmitScope& Scope, const FRequestedType& RequestedType) const { return PrepareValueResult.PreparedType.GetEvaluation(Scope, RequestedType); }
-	FExpression* GetForwardValue(const FRequestedType& RequestedType);
 
 	virtual void Reset() override;
 
@@ -568,10 +560,13 @@ public:
 
 	FExpression* NewAbs(FExpression* Input) { return NewUnaryOp(EOperation::Abs, Input); }
 	FExpression* NewNeg(FExpression* Input) { return NewUnaryOp(EOperation::Neg, Input); }
+	FExpression* NewSaturate(FExpression* Input) { return NewUnaryOp(EOperation::Saturate, Input); }
 	FExpression* NewRcp(FExpression* Input) { return NewUnaryOp(EOperation::Rcp, Input); }
 	FExpression* NewFrac(FExpression* Input) { return NewUnaryOp(EOperation::Frac, Input); }
 	FExpression* NewLength(FExpression* Input) { return NewUnaryOp(EOperation::Length, Input); }
 	FExpression* NewNormalize(FExpression* Input) { return NewUnaryOp(EOperation::Normalize, Input); }
+	FExpression* NewSin(FExpression* Input) { return NewUnaryOp(EOperation::Sin, Input); }
+	FExpression* NewCos(FExpression* Input) { return NewUnaryOp(EOperation::Cos, Input); }
 
 	FExpression* NewAdd(FExpression* Lhs, FExpression* Rhs) { return NewBinaryOp(EOperation::Add, Lhs, Rhs); }
 	FExpression* NewSub(FExpression* Lhs, FExpression* Rhs) { return NewBinaryOp(EOperation::Sub, Lhs, Rhs); }
@@ -586,6 +581,9 @@ public:
 	FExpression* NewLessEqual(FExpression* Lhs, FExpression* Rhs) { return NewBinaryOp(EOperation::LessEqual, Lhs, Rhs); }
 	FExpression* NewGreaterEqual(FExpression* Lhs, FExpression* Rhs) { return NewBinaryOp(EOperation::GreaterEqual, Lhs, Rhs); }
 
+	FExpression* NewCross(FExpression* Lhs, FExpression* Rhs);
+	FExpression* NewLerp(FExpression* A, FExpression* B, FExpression* T) { return NewAdd(A, NewMul(NewSub(B, A), T)); }
+
 private:
 	template<typename T, typename... ArgTypes>
 	inline T* NewNode(ArgTypes&&... Args)
diff --git a/Engine/Source/Runtime/Engine/Public/HLSLTree/HLSLTreeCommon.h b/Engine/Source/Runtime/Engine/Public/HLSLTree/HLSLTreeCommon.h
--- a/Engine/Source/Runtime/Engine/Public/HLSLTree/HLSLTreeCommon.h
+++ b/Engine/Source/Runtime/Engine/Public/HLSLTree/HLSLTreeCommon.h
@@ -25,6 +25,20 @@ public:
 	virtual void EmitValuePreshader(FEmitContext& Context, FEmitScope& Scope, const FRequestedType& RequestedType, FEmitValuePreshaderResult& OutResult) const override;
 };
 
+class FExpressionMaterialShadingModel : public FExpression
+{
+public:
+	explicit FExpressionMaterialShadingModel(EMaterialShadingModel InShadingModel)
+		: ShadingModel(InShadingModel)
+	{}
+
+	TEnumAsByte<EMaterialShadingModel> ShadingModel;
+
+	virtual void ComputeAnalyticDerivatives(FTree& Tree, FExpressionDerivatives& OutResult) const override;
+	virtual bool PrepareValue(FEmitContext& Context, FEmitScope& Scope, const FRequestedType& RequestedType, FPrepareValueResult& OutResult) const override;
+	virtual void EmitValuePreshader(FEmitContext& Context, FEmitScope& Scope, const FRequestedType& RequestedType, FEmitValuePreshaderResult& OutResult) const override;
+};
+
 class FExpressionMaterialParameter : public FExpression
 {
 public:
@@ -74,6 +88,11 @@ enum class EExternalInput : uint8
 	TexCoord6_Ddy,
 	TexCoord7_Ddy,
 
+	LightmapTexCoord,
+	LightmapTexCoord_Ddx,
+	LightmapTexCoord_Ddy,
+
+	TwoSidedSign,
 	VertexColor,
 	VertexColor_Ddx,
 	VertexColor_Ddy,
@@ -82,6 +101,7 @@ enum class EExternalInput : uint8
 	WorldPosition_NoOffsets,
 	TranslatedWorldPosition,
 	TranslatedWorldPosition_NoOffsets,
+	ActorWorldPosition,
 
 	PrevWorldPosition,
 	PrevWorldPosition_NoOffsets,
@@ -91,12 +111,30 @@ enum class EExternalInput : uint8
 	WorldPosition_Ddx,
 	WorldPosition_Ddy,
 
+	WorldVertexNormal,
+	WorldVertexTangent,
+	WorldNormal,
+	WorldReflection,
+
 	ViewportUV,
 	PixelPosition,
 	ViewSize,
 	RcpViewSize,
-
+	FieldOfView,
+	TanHalfFieldOfView,
+	CotanHalfFieldOfView,
+	TemporalSampleCount,
+	TemporalSampleIndex,
+	TemporalSampleOffset,
+	PreExposure,
+	RcpPreExposure,
+	RuntimeVirtualTextureOutputLevel,
+	RuntimeVirtualTextureOutputDerivative,
+	RuntimeVirtualTextureMaxLevel,
+
+	CameraVector,
 	CameraWorldPosition,
+	ViewWorldPosition,
 	PreViewTranslation,
 	TangentToWorld,
 	LocalToWorld,
@@ -110,7 +148,12 @@ enum class EExternalInput : uint8
 	ParticleToWorld,
 	InstanceToWorld,
 
+	PrevFieldOfView,
+	PrevTanHalfFieldOfView,
+	PrevCotanHalfFieldOfView,
+
 	PrevCameraWorldPosition,
+	PrevViewWorldPosition,
 	PrevPreViewTranslation,
 	PrevLocalToWorld,
 	PrevWorldToLocal,
@@ -129,8 +172,12 @@ enum class EExternalInput : uint8
 
 	PrevGameTime,
 	PrevRealTime,
+
+	ParticleColor,
+	ParticleTranslatedWorldPosition,
+	ParticleRadius,
 };
-static constexpr int32 NumTexCoords = 8;
+static constexpr int32 MaxNumTexCoords = 8;
 
 struct FExternalInputDescription
 {
@@ -149,19 +196,19 @@ FExternalInputDescription GetExternalInputDescription(EExternalInput Input);
 
 inline bool IsTexCoord(EExternalInput Type)
 {
-	return FMath::IsWithin((int32)Type, (int32)EExternalInput::TexCoord0, (int32)EExternalInput::TexCoord0 + NumTexCoords);
+	return FMath::IsWithin((int32)Type, (int32)EExternalInput::TexCoord0, (int32)EExternalInput::TexCoord0 + MaxNumTexCoords);
 }
 inline bool IsTexCoord_Ddx(EExternalInput Type)
 {
-	return FMath::IsWithin((int32)Type, (int32)EExternalInput::TexCoord0_Ddx, (int32)EExternalInput::TexCoord0_Ddx + NumTexCoords);
+	return FMath::IsWithin((int32)Type, (int32)EExternalInput::TexCoord0_Ddx, (int32)EExternalInput::TexCoord0_Ddx + MaxNumTexCoords);
 }
 inline bool IsTexCoord_Ddy(EExternalInput Type)
 {
-	return FMath::IsWithin((int32)Type, (int32)EExternalInput::TexCoord0_Ddy, (int32)EExternalInput::TexCoord0_Ddy + NumTexCoords);
+	return FMath::IsWithin((int32)Type, (int32)EExternalInput::TexCoord0_Ddy, (int32)EExternalInput::TexCoord0_Ddy + MaxNumTexCoords);
 }
 inline EExternalInput MakeInputTexCoord(int32 Index)
 {
-	check(Index >= 0 && Index < NumTexCoords);
+	check(Index >= 0 && Index < MaxNumTexCoords);
 	return (EExternalInput)((int32)EExternalInput::TexCoord0 + Index);
 }
 
@@ -195,6 +242,39 @@ public:
 	virtual void EmitValueShader(FEmitContext& Context, FEmitScope& Scope, const FRequestedType& RequestedType, FEmitValueShaderResult& OutResult) const override;
 };
 
+struct FMaterialNoiseParameters
+{
+	FMaterialNoiseParameters() { FMemory::Memzero(*this); }
+
+	int32 Quality;
+	int32 Levels;
+	float Scale;
+	uint32 RepeatSize;
+	float OutputMin;
+	float OutputMax;
+	float LevelScale;
+	uint8 NoiseFunction;
+	bool bTiling;
+	bool bTurbulence;
+};
+
+class FExpressionMaterialNoise : public FExpression
+{
+public:
+	FExpressionMaterialNoise(const FMaterialNoiseParameters& InParams, FExpression* InPositionExpression, FExpression* InFilterWidthExpression)
+		: PositionExpression(InPositionExpression)
+		, FilterWidthExpression(InFilterWidthExpression)
+		, Parameters(InParams)
+	{}
+
+	FExpression* PositionExpression;
+	FExpression* FilterWidthExpression;
+	FMaterialNoiseParameters Parameters;
+
+	virtual bool PrepareValue(FEmitContext& Context, FEmitScope& Scope, const FRequestedType& RequestedType, FPrepareValueResult& OutResult) const override;
+	virtual void EmitValueShader(FEmitContext& Context, FEmitScope& Scope, const FRequestedType& RequestedType, FEmitValueShaderResult& OutResult) const override;
+};
+
 class FExpressionTextureSample : public FExpression
 {
 public:
@@ -301,10 +381,25 @@ public:
 	virtual void EmitValuePreshader(FEmitContext& Context, FEmitScope& Scope, const FRequestedType& RequestedType, FEmitValuePreshaderResult& OutResult) const override;
 };
 
+class FExpressionDerivative : public FExpression
+{
+public:
+	FExpressionDerivative(EDerivativeCoordinate InCoord, FExpression* InInput) : Input(InInput), Coord(InCoord) {}
+
+	FExpression* Input;
+	EDerivativeCoordinate Coord;
+
+	virtual void ComputeAnalyticDerivatives(FTree& Tree, FExpressionDerivatives& OutResult) const override;
+	virtual FExpression* ComputePreviousFrame(FTree& Tree, const FRequestedType& RequestedType) const override;
+	virtual bool PrepareValue(FEmitContext& Context, FEmitScope& Scope, const FRequestedType& RequestedType, FPrepareValueResult& OutResult) const override;
+	virtual void EmitValueShader(FEmitContext& Context, FEmitScope& Scope, const FRequestedType& RequestedType, FEmitValueShaderResult& OutResult) const override;
+	virtual void EmitValuePreshader(FEmitContext& Context, FEmitScope& Scope, const FRequestedType& RequestedType, FEmitValuePreshaderResult& OutResult) const override;
+};
+
 struct FSwizzleParameters
 {
 	FSwizzleParameters() : NumComponents(0) { ComponentIndex[0] = ComponentIndex[1] = ComponentIndex[2] = ComponentIndex[3] = INDEX_NONE; }
-	FSwizzleParameters(int8 IndexR, int8 IndexG, int8 IndexB, int8 IndexA);
+	explicit FSwizzleParameters(int8 IndexR, int8 IndexG = INDEX_NONE, int8 IndexB = INDEX_NONE, int8 IndexA = INDEX_NONE);
 
 	FRequestedType GetRequestedInputType(const FRequestedType& RequestedType) const;
 	bool HasSwizzle() const;
@@ -350,11 +445,17 @@ public:
 	virtual void EmitValuePreshader(FEmitContext& Context, FEmitScope& Scope, const FRequestedType& RequestedType, FEmitValuePreshaderResult& OutResult) const override;
 };
 
-class FExpressionReflectionVector : public FExpression
+/** Can be used to emit HLSL chunks with no inputs, where it's not worth the trouble of defining a new expression type */
+class FExpressionInlineCustomHLSL : public FExpression
 {
 public:
+	FExpressionInlineCustomHLSL(Shader::EValueType InType, FStringView InCode) : Code(InCode), ResultType(InType) {}
+
 	virtual bool PrepareValue(FEmitContext& Context, FEmitScope& Scope, const FRequestedType& RequestedType, FPrepareValueResult& OutResult) const override;
 	virtual void EmitValueShader(FEmitContext& Context, FEmitScope& Scope, const FRequestedType& RequestedType, FEmitValueShaderResult& OutResult) const override;
+
+	FStringView Code;
+	Shader::EValueType ResultType;
 };
 
 class FExpressionCustomHLSL : public FExpression
diff --git a/Engine/Source/Runtime/Engine/Public/HLSLTree/HLSLTreeEmit.h b/Engine/Source/Runtime/Engine/Public/HLSLTree/HLSLTreeEmit.h
--- a/Engine/Source/Runtime/Engine/Public/HLSLTree/HLSLTreeEmit.h
+++ b/Engine/Source/Runtime/Engine/Public/HLSLTree/HLSLTreeEmit.h
@@ -37,6 +37,9 @@ class FEmitScope;
 class FEmitShaderExpression;
 class FEmitShaderStatement;
 
+// TODO - remove dependency
+enum class EExternalInput : uint8;
+
 struct FEmitPreshaderScope;
 
 struct FEmitShaderScopeEntry
@@ -112,6 +115,8 @@ enum class EFormatArgType : uint8
 	ShaderValue,
 	String,
 	Int,
+	Uint,
+	Float,
 	Bool,
 };
 
@@ -121,6 +126,8 @@ struct FFormatArgVariant
 	FFormatArgVariant(FEmitShaderExpression* InValue) : Type(EFormatArgType::ShaderValue), ShaderValue(InValue) { check(InValue); }
 	FFormatArgVariant(const TCHAR* InValue) : Type(EFormatArgType::String), String(InValue) { check(InValue); }
 	FFormatArgVariant(int32 InValue) : Type(EFormatArgType::Int), Int(InValue) {}
+	FFormatArgVariant(uint32 InValue) : Type(EFormatArgType::Uint), Uint(InValue) {}
+	FFormatArgVariant(float InValue) : Type(EFormatArgType::Float), Float(InValue) {}
 	FFormatArgVariant(bool InValue) : Type(EFormatArgType::Bool), Bool(InValue) {}
 
 	EFormatArgType Type = EFormatArgType::Void;
@@ -129,6 +136,8 @@ struct FFormatArgVariant
 		FEmitShaderExpression* ShaderValue;
 		const TCHAR* String;
 		int32 Int;
+		uint32 Uint;
+		float Float;
 		bool Bool;
 	};
 };
@@ -196,6 +205,10 @@ public:
 	EEmitScopeState State = EEmitScopeState::Uninitialized;
 	EExpressionEvaluation Evaluation = EExpressionEvaluation::None;
 };
+inline bool IsScopeDead(const FEmitScope* Scope)
+{
+	return !Scope || Scope->State == EEmitScopeState::Dead;
+}
 
 struct FEmitCustomHLSLInput
 {
@@ -228,6 +241,14 @@ struct FPreshaderLocalPHIScope
 	int32 ValueStackPosition;
 };
 
+enum class EEmitCastFlags : uint32
+{
+	None = 0u,
+	// Scalars are replicated by default, this will zero-extend them instead
+	ZeroExtendScalar = (1u << 0),
+};
+ENUM_CLASS_FLAGS(EEmitCastFlags)
+
 /** Tracks shared state while emitting HLSL code */
 class FEmitContext
 {
@@ -400,7 +421,7 @@ public:
 
 	FEmitShaderExpression* EmitPreshaderOrConstant(FEmitScope& Scope, const FRequestedType& RequestedType, FExpression* Expression);
 	FEmitShaderExpression* EmitConstantZero(FEmitScope& Scope, const Shader::FType& Type);
-	FEmitShaderExpression* EmitCast(FEmitScope& Scope, FEmitShaderExpression* ShaderValue, const Shader::FType& DestType);
+	FEmitShaderExpression* EmitCast(FEmitScope& Scope, FEmitShaderExpression* ShaderValue, const Shader::FType& DestType, EEmitCastFlags Flags = EEmitCastFlags::None);
 
 	/** FStringViews passed to this method are expected to reference persistent memory */
 	FEmitShaderExpression* EmitCustomHLSL(FEmitScope& Scope, FStringView DeclarationCode, FStringView FunctionCode, TConstArrayView<FCustomHLSLInput> Inputs, const Shader::FStructType* OutputType);
@@ -432,7 +453,12 @@ public:
 	uint32 CurrentNumBoolComponents = 32u;
 	bool bReadMaterialNormal = false;
 	bool bUsesVertexColor = false;
-	uint32 TexCoordMask[SF_NumFrequencies] = { 0u };
+	bool bNeedsParticlePosition = false;
+	bool bUsesParticleColor = false;
+	TBitArray<> ExternalInputMask[SF_NumFrequencies];
+	FMaterialShadingModelField ShadingModelsFromCompilation;
+
+	bool IsExternalInputUsed(EShaderFrequency Frequency, EExternalInput Input) const { return ExternalInputMask[Frequency][(uint8)Input]; }
 
 	int32 NumExpressionLocals = 0;
 	int32 NumExpressionLocalPHIs = 0;
diff --git a/Engine/Source/Runtime/Engine/Public/HLSLTree/HLSLTreeTypes.h b/Engine/Source/Runtime/Engine/Public/HLSLTree/HLSLTreeTypes.h
--- a/Engine/Source/Runtime/Engine/Public/HLSLTree/HLSLTreeTypes.h
+++ b/Engine/Source/Runtime/Engine/Public/HLSLTree/HLSLTreeTypes.h
@@ -47,6 +47,9 @@ enum class EExpressionEvaluation : uint8
 
 	/** The expression outputs constant preshader code evaluated at compile time (via FExpressionEmitResult::Preshader) */
 	Constant,
+
+	/** The expression evaluates to 0 */
+	ConstantZero,
 };
 
 EExpressionEvaluation CombineEvaluations(EExpressionEvaluation Lhs, EExpressionEvaluation Rhs);
@@ -66,6 +69,8 @@ enum class EOperation : uint8
 	Abs,
 	Neg,
 	Rcp,
+	Sqrt,
+	Log2,
 	Frac,
 	Floor,
 	Ceil,
@@ -75,6 +80,15 @@ enum class EOperation : uint8
 	Sign,
 	Length,
 	Normalize,
+	Sin,
+	Cos,
+	Tan,
+	Asin,
+	AsinFast,
+	Acos,
+	AcosFast,
+	Atan,
+	AtanFast,
 
 	// Binary Ops
 	Add,
@@ -83,6 +97,8 @@ enum class EOperation : uint8
 	Div,
 	Fmod,
 	PowPositiveClamped,
+	Atan2,
+	Atan2Fast,
 	Dot,
 	Min,
 	Max,
diff --git a/Engine/Source/Runtime/Engine/Public/MaterialShared.h b/Engine/Source/Runtime/Engine/Public/MaterialShared.h
--- a/Engine/Source/Runtime/Engine/Public/MaterialShared.h
+++ b/Engine/Source/Runtime/Engine/Public/MaterialShared.h
@@ -570,6 +570,7 @@ public:
 	ENGINE_API void GetTextureValue(int32 Index, const FMaterialRenderContext& Context, const FMaterial& Material, const URuntimeVirtualTexture*& OutValue) const;
 
 	int32 FindOrAddTextureParameter(EMaterialTextureParameterType Type, const FMaterialTextureParameterInfo& Info);
+	int32 FindOrAddExternalTextureParameter(const FMaterialExternalTextureParameterInfo& Info);
 	int32 FindOrAddNumericParameter(EMaterialParameterType Type, const FMaterialParameterInfo& ParameterInfo, uint32 DefaultValueOffset);
 	uint32 AddDefaultParameterValue(const UE::Shader::FValue& Value);
 
diff --git a/Engine/Source/Runtime/Engine/Public/Shader/ShaderTypes.h b/Engine/Source/Runtime/Engine/Public/Shader/ShaderTypes.h
--- a/Engine/Source/Runtime/Engine/Public/Shader/ShaderTypes.h
+++ b/Engine/Source/Runtime/Engine/Public/Shader/ShaderTypes.h
@@ -68,6 +68,7 @@ enum class EValueComponentType : uint8
 	TextureCube,
 	TextureCubeArray,
 	Texture3D,
+	TextureExternal,
 };
 
 struct FValueComponentTypeDescription
@@ -131,6 +132,7 @@ enum class EValueType : uint8
 	TextureCube,
 	TextureCubeArray,
 	Texture3D,
+	TextureExternal,
 };
 
 struct FValueTypeDescription
@@ -310,15 +312,21 @@ struct FTextureValue
 		, SamplerType(InSamplerType)
 	{}
 
+	FTextureValue(const FGuid& InGuid)
+		: ExternalTextureGuid(InGuid)
+		, SamplerType(SAMPLERTYPE_External)
+	{}
+
 	EValueType GetType() const;
 
 	UTexture* Texture = nullptr;
+	FGuid ExternalTextureGuid;
 	EMaterialSamplerType SamplerType = SAMPLERTYPE_Color;
 };
 
 inline bool operator==(const FTextureValue& Lhs, const FTextureValue& Rhs)
 {
-	return Lhs.Texture == Rhs.Texture && Lhs.SamplerType == Rhs.SamplerType;
+	return Lhs.Texture == Rhs.Texture && Lhs.SamplerType == Rhs.SamplerType && Lhs.ExternalTextureGuid == Rhs.ExternalTextureGuid;
 }
 inline bool operator!=(const FTextureValue& Lhs, const FTextureValue& Rhs)
 {
