Fix deadlock when deleting a skeletalmesh during build
#rnx

diff --git a/Engine/Source/Runtime/Engine/Private/MorphTools.cpp b/Engine/Source/Runtime/Engine/Private/MorphTools.cpp
--- a/Engine/Source/Runtime/Engine/Private/MorphTools.cpp
+++ b/Engine/Source/Runtime/Engine/Private/MorphTools.cpp
@@ -227,6 +227,13 @@ void FFinishBuildMorphTargetData::ApplyEditorData(USkeletalMesh * SkeletalMesh,
 		return;
 	}
 	
+	if (SkeletalMesh->HasAnyFlags(RF_BeginDestroyed))
+	{
+		//No need to apply the morph targets if the asset is being deleted.
+		//Also acquiring the GCScopeGuard when we are in the destruction process will create a deadlock, if there is code
+		//in the BeginDestroy function of the skeletalmesh that touch any property that is lock by the async build.
+		return;
+	}
 	//Make sure we do not create new morph target during a gc
 	FGCScopeGuard GCScopeGuard;
 
