Remove Dot operation from new HLSL translator, replace with new 'Sum' unary op, combined with a Mul

diff --git a/Engine/Shaders/Private/Common.ush b/Engine/Shaders/Private/Common.ush
--- a/Engine/Shaders/Private/Common.ush
+++ b/Engine/Shaders/Private/Common.ush
@@ -193,6 +193,11 @@ float4 FmodFloor(float4 Lhs, float4 Rhs)
 	return Lhs - floor(Lhs / Rhs) * Rhs;
 }
 
+float VectorSum(float V) { return V; }
+float VectorSum(float2 V) { return V.x + V.y; }
+float VectorSum(float3 V) { return V.x + V.y + V.z; }
+float VectorSum(float4 V) { return V.x + V.y + V.z + V.w; }
+
 #include "LargeWorldCoordinates.ush"
 
 #include "InstancedStereo.ush"
diff --git a/Engine/Shaders/Private/LargeWorldCoordinates.ush b/Engine/Shaders/Private/LargeWorldCoordinates.ush
--- a/Engine/Shaders/Private/LargeWorldCoordinates.ush
+++ b/Engine/Shaders/Private/LargeWorldCoordinates.ush
@@ -399,59 +399,19 @@ FLWCMatrix LWCMakeRelativeToTile(FLWCMatrix M, float3 NewTile)
 	return M;
 }
 
-FLWCScalar LWCDot(FLWCScalar Lhs, FLWCScalar Rhs)
-{
-	return LWCMultiply(Lhs, Rhs);
-}
-FLWCScalar LWCDot(FLWCScalar Lhs, float Rhs)
-{
-	return LWCMultiply(Lhs, Rhs);
-}
-
-FLWCScalar LWCDot(FLWCVector2 Lhs, FLWCVector2 Rhs)
-{
-	FLWCScalar X2 = LWCMultiply(LWCGetX(Lhs), LWCGetX(Rhs));
-	FLWCScalar Y2 = LWCMultiply(LWCGetY(Lhs), LWCGetY(Rhs));
-	return LWCAdd(X2, Y2);
-}
-FLWCScalar LWCDot(FLWCVector2 Lhs, float2 Rhs)
-{
-	FLWCScalar X2 = LWCMultiply(LWCGetX(Lhs), Rhs.x);
-	FLWCScalar Y2 = LWCMultiply(LWCGetY(Lhs), Rhs.y);
-	return LWCAdd(X2, Y2);
-}
-
-FLWCScalar LWCDot(FLWCVector3 Lhs, FLWCVector3 Rhs)
-{
-	FLWCScalar X2 = LWCMultiply(LWCGetX(Lhs), LWCGetX(Rhs));
-	FLWCScalar Y2 = LWCMultiply(LWCGetY(Lhs), LWCGetY(Rhs));
-	FLWCScalar Z2 = LWCMultiply(LWCGetZ(Lhs), LWCGetZ(Rhs));
-	return LWCAdd(LWCAdd(X2, Y2), Z2);
-}
-FLWCScalar LWCDot(FLWCVector3 Lhs, float3 Rhs)
-{
-	FLWCScalar X2 = LWCMultiply(LWCGetX(Lhs), Rhs.x);
-	FLWCScalar Y2 = LWCMultiply(LWCGetY(Lhs), Rhs.y);
-	FLWCScalar Z2 = LWCMultiply(LWCGetZ(Lhs), Rhs.z);
-	return LWCAdd(LWCAdd(X2, Y2), Z2);
-}
-
-FLWCScalar LWCDot(FLWCVector4 Lhs, FLWCVector4 Rhs)
-{
-	FLWCScalar X2 = LWCMultiply(LWCGetX(Lhs), LWCGetX(Rhs));
-	FLWCScalar Y2 = LWCMultiply(LWCGetY(Lhs), LWCGetY(Rhs));
-	FLWCScalar Z2 = LWCMultiply(LWCGetZ(Lhs), LWCGetZ(Rhs));
-	FLWCScalar W2 = LWCMultiply(LWCGetW(Lhs), LWCGetW(Rhs));
-	return LWCAdd(LWCAdd(LWCAdd(X2, Y2), Z2), W2);
-}
-FLWCScalar LWCDot(FLWCVector4 Lhs, float4 Rhs)
-{
-	FLWCScalar X2 = LWCMultiply(LWCGetX(Lhs), Rhs.x);
-	FLWCScalar Y2 = LWCMultiply(LWCGetY(Lhs), Rhs.y);
-	FLWCScalar Z2 = LWCMultiply(LWCGetZ(Lhs), Rhs.z);
-	FLWCScalar W2 = LWCMultiply(LWCGetW(Lhs), Rhs.w);
-	return LWCAdd(LWCAdd(LWCAdd(X2, Y2), Z2), W2);
-}
+FLWCScalar LWCVectorSum(FLWCScalar V) { return V; }
+FLWCScalar LWCVectorSum(FLWCVector2 V) { return LWCAdd(LWCGetX(V), LWCGetY(V)); }
+FLWCScalar LWCVectorSum(FLWCVector3 V) { return LWCAdd(LWCAdd(LWCGetX(V), LWCGetY(V)), LWCGetZ(V)); }
+FLWCScalar LWCVectorSum(FLWCVector4 V) { return LWCAdd(LWCAdd(LWCAdd(LWCGetX(V), LWCGetY(V)), LWCGetZ(V)), LWCGetW(V)); }
+
+FLWCScalar LWCDot(FLWCScalar Lhs, FLWCScalar Rhs) { return LWCMultiply(Lhs, Rhs); }
+FLWCScalar LWCDot(FLWCScalar Lhs, float Rhs) { return LWCMultiply(Lhs, Rhs); }
+FLWCScalar LWCDot(FLWCVector2 Lhs, FLWCVector2 Rhs) { return LWCVectorSum(LWCMultiply(Lhs, Rhs)); }
+FLWCScalar LWCDot(FLWCVector2 Lhs, float2 Rhs) { return LWCVectorSum(LWCMultiply(Lhs, Rhs)); }
+FLWCScalar LWCDot(FLWCVector3 Lhs, FLWCVector3 Rhs) { return LWCVectorSum(LWCMultiply(Lhs, Rhs)); }
+FLWCScalar LWCDot(FLWCVector3 Lhs, float3 Rhs) { return LWCVectorSum(LWCMultiply(Lhs, Rhs)); }
+FLWCScalar LWCDot(FLWCVector4 Lhs, FLWCVector4 Rhs) { return LWCVectorSum(LWCMultiply(Lhs, Rhs)); }
+FLWCScalar LWCDot(FLWCVector4 Lhs, float4 Rhs) { return LWCVectorSum(LWCMultiply(Lhs, Rhs)); }
 
 // LWCLength2Scaled returns the length2 of the vector, scaled by the LWC tile size
 FLWCScalar LWCLength2Scaled(FLWCScalar V)
diff --git a/Engine/Source/Runtime/Engine/Private/HLSLTree/HLSLTree.cpp b/Engine/Source/Runtime/Engine/Private/HLSLTree/HLSLTree.cpp
--- a/Engine/Source/Runtime/Engine/Private/HLSLTree/HLSLTree.cpp
+++ b/Engine/Source/Runtime/Engine/Private/HLSLTree/HLSLTree.cpp
@@ -109,6 +109,7 @@ FOperationDescription GetOperationDescription(EOperation Op)
 	case EOperation::Sign: return FOperationDescription(TEXT("Sign"), TEXT("sign"), 1, Shader::EPreshaderOpcode::Sign); break;
 	case EOperation::Length: return FOperationDescription(TEXT("Length"), TEXT("length"), 1, Shader::EPreshaderOpcode::Length); break;
 	case EOperation::Normalize: return FOperationDescription(TEXT("Normalize"), TEXT("normalize"), 1, Shader::EPreshaderOpcode::Normalize); break;
+	case EOperation::Sum: return FOperationDescription(TEXT("Sum"), TEXT("sum"), 1, Shader::EPreshaderOpcode::Nop); break; // TODO
 	case EOperation::Sin: return FOperationDescription(TEXT("Sin"), TEXT("sin"), 1, Shader::EPreshaderOpcode::Sin); break;
 	case EOperation::Cos: return FOperationDescription(TEXT("Cos"), TEXT("cos"), 1, Shader::EPreshaderOpcode::Cos); break;
 	case EOperation::Tan: return FOperationDescription(TEXT("Tan"), TEXT("tan"), 1, Shader::EPreshaderOpcode::Tan); break;
@@ -128,7 +129,6 @@ FOperationDescription GetOperationDescription(EOperation Op)
 	case EOperation::PowPositiveClamped: return FOperationDescription(TEXT("PowPositiveClamped"), TEXT("PowPositiveClamped"), 2, Shader::EPreshaderOpcode::Nop); break;
 	case EOperation::Atan2: return FOperationDescription(TEXT("Atan2"), TEXT("atan2"), 2, Shader::EPreshaderOpcode::Atan2); break;
 	case EOperation::Atan2Fast: return FOperationDescription(TEXT("Atan2Fast"), TEXT("atan2Fast"), 2, Shader::EPreshaderOpcode::Atan2); break;
-	case EOperation::Dot: return FOperationDescription(TEXT("Dot"), TEXT("dot"), 2, Shader::EPreshaderOpcode::Dot); break;
 	case EOperation::Min: return FOperationDescription(TEXT("Min"), TEXT("min"), 2, Shader::EPreshaderOpcode::Min); break;
 	case EOperation::Max: return FOperationDescription(TEXT("Max"), TEXT("max"), 2, Shader::EPreshaderOpcode::Max); break;
 	case EOperation::Less: return FOperationDescription(TEXT("Less"), TEXT("<"), 2, Shader::EPreshaderOpcode::Less); break;
diff --git a/Engine/Source/Runtime/Engine/Private/HLSLTree/HLSLTreeCommon.cpp b/Engine/Source/Runtime/Engine/Private/HLSLTree/HLSLTreeCommon.cpp
--- a/Engine/Source/Runtime/Engine/Private/HLSLTree/HLSLTreeCommon.cpp
+++ b/Engine/Source/Runtime/Engine/Private/HLSLTree/HLSLTreeCommon.cpp
@@ -1096,12 +1096,9 @@ FOperationRequestedTypes GetOperationRequestedTypes(EOperation Op, const FReques
 	{
 	case EOperation::Length:
 	case EOperation::Normalize:
+	case EOperation::Sum:
 		Types.InputType[0] = ERequestedType::Vector4;
 		break;
-	case EOperation::Dot:
-		Types.InputType[0] = ERequestedType::Vector4;
-		Types.InputType[1] = ERequestedType::Vector4;
-		break;
 	case EOperation::VecMulMatrix3:
 		Types.bIsMatrixOperation = true;
 		Types.InputType[0] = ERequestedType::Vector3;
@@ -1148,6 +1145,7 @@ FOperationTypes GetOperationTypes(EOperation Op, TConstArrayView<Shader::EValueT
 	switch (Op)
 	{
 	case EOperation::Length:
+	case EOperation::Sum:
 		Types.ResultType = Shader::MakeValueType(ComponentType, 1);
 		break;
 	case EOperation::Normalize:
@@ -1188,9 +1186,6 @@ FOperationTypes GetOperationTypes(EOperation Op, TConstArrayView<Shader::EValueT
 		// No LWC support yet
 		Types.InputType[0] = Types.InputType[1] = Types.ResultType = Shader::MakeNonLWCType(IntermediateType);
 		break;
-	case EOperation::Dot:
-		Types.ResultType = Shader::MakeValueType(ComponentType, 1);
-		break;
 	case EOperation::VecMulMatrix3:
 		// No LWC for matrix3
 		Types.InputType[0] = Shader::EValueType::Float3;
@@ -1262,6 +1257,10 @@ void FExpressionOperation::ComputeAnalyticDerivatives(FTree& Tree, FExpressionDe
 		OutResult.ExpressionDdy = Tree.NewMul(dFdA, InputDerivatives[0].ExpressionDdy);
 		break;
 	}
+	case EOperation::Sum:
+		OutResult.ExpressionDdx = Tree.NewSum(InputDerivatives[0].ExpressionDdx);
+		OutResult.ExpressionDdy = Tree.NewSum(InputDerivatives[0].ExpressionDdy);
+		break;
 	case EOperation::Frac:
 		OutResult = InputDerivatives[0];
 		break;
@@ -1316,17 +1315,6 @@ void FExpressionOperation::ComputeAnalyticDerivatives(FTree& Tree, FExpressionDe
 		// We can't really do anything meaningful in the non-zero case.
 		OutResult = InputDerivatives[0];
 		break;
-	case EOperation::Dot:
-	{
-		// Dot means multiply the values, then sum the resulting components
-		FExpression* MulDdx = Tree.NewAdd(Tree.NewMul(InputDerivatives[0].ExpressionDdx, Inputs[1]), Tree.NewMul(InputDerivatives[1].ExpressionDdx, Inputs[0]));
-		FExpression* MulDdy = Tree.NewAdd(Tree.NewMul(InputDerivatives[0].ExpressionDdy, Inputs[1]), Tree.NewMul(InputDerivatives[1].ExpressionDdy, Inputs[0]));
-		// Dot the products with 1 to sum them
-		FExpression* Const1 = Tree.NewConstant(FVector4f(1.0f, 1.0f, 1.0f, 1.0f));
-		OutResult.ExpressionDdx = Tree.NewDot(MulDdx, Const1);
-		OutResult.ExpressionDdy = Tree.NewDot(MulDdy, Const1);
-		break;
-	}
 	case EOperation::Min:
 	{
 		FExpression* Cond = Tree.NewLess(Inputs[0], Inputs[1]);
@@ -1398,6 +1386,7 @@ bool FExpressionOperation::PrepareValue(FEmitContext& Context, FEmitScope& Scope
 		switch (Op)
 		{
 		case EOperation::Length:
+		case EOperation::Sum:
 			ResultType = FPreparedType(Shader::MakeValueType(InputType[0].ValueComponentType, 1), InputType[0].GetEvaluation(Scope));
 			break;
 		default:
@@ -1579,6 +1568,7 @@ void FExpressionOperation::EmitValueShader(FEmitContext& Context, FEmitScope& Sc
 			OutResult.Code = Context.EmitExpression(Scope, Types.ResultType, TEXT("normalize(%)"), InputValue[0]);
 		}
 		break;
+	case EOperation::Sum: OutResult.Code = Context.EmitExpression(Scope, Types.ResultType, Types.bIsLWC ? TEXT("LWCVectorSum(%)") : TEXT("VectorSum(%)"), InputValue[0]); break;
 	case EOperation::Sin: OutResult.Code = Context.EmitExpression(Scope, Types.ResultType, Types.bIsLWC ? TEXT("LWCSin(%)") : TEXT("sin(%)"), InputValue[0]); break;
 	case EOperation::Cos: OutResult.Code = Context.EmitExpression(Scope, Types.ResultType, Types.bIsLWC ? TEXT("LWCCos(%)") : TEXT("cos(%)"), InputValue[0]); break;
 	case EOperation::Tan: OutResult.Code = Context.EmitExpression(Scope, Types.ResultType, Types.bIsLWC ? TEXT("LWCTan(%)") : TEXT("tan(%)"), InputValue[0]); break;
@@ -1643,16 +1633,6 @@ void FExpressionOperation::EmitValueShader(FEmitContext& Context, FEmitScope& Sc
 		break;
 	case EOperation::Atan2: OutResult.Code = Context.EmitExpression(Scope, Types.ResultType, TEXT("atan2(%, %)"), InputValue[0], InputValue[1]); break;
 	case EOperation::Atan2Fast: OutResult.Code = Context.EmitExpression(Scope, Types.ResultType, TEXT("atan2Fast(%, %)"), InputValue[0], InputValue[1]); break;
-	case EOperation::Dot:
-		if (Types.bIsLWC)
-		{
-			OutResult.Code = Context.EmitExpression(Scope, Types.ResultType, TEXT("LWCDot(%, %)"), InputValue[0], InputValue[1]);
-		}
-		else
-		{
-			OutResult.Code = Context.EmitExpression(Scope, Types.ResultType, TEXT("dot(%, %)"), InputValue[0], InputValue[1]);
-		}
-		break;
 	case EOperation::Min:
 		if (Types.bIsLWC)
 		{
diff --git a/Engine/Source/Runtime/Engine/Public/HLSLTree/HLSLTree.h b/Engine/Source/Runtime/Engine/Public/HLSLTree/HLSLTree.h
--- a/Engine/Source/Runtime/Engine/Public/HLSLTree/HLSLTree.h
+++ b/Engine/Source/Runtime/Engine/Public/HLSLTree/HLSLTree.h
@@ -561,6 +561,7 @@ public:
 	FExpression* NewAbs(FExpression* Input) { return NewUnaryOp(EOperation::Abs, Input); }
 	FExpression* NewNeg(FExpression* Input) { return NewUnaryOp(EOperation::Neg, Input); }
 	FExpression* NewSaturate(FExpression* Input) { return NewUnaryOp(EOperation::Saturate, Input); }
+	FExpression* NewSum(FExpression* Input) { return NewUnaryOp(EOperation::Sum, Input); }
 	FExpression* NewRcp(FExpression* Input) { return NewUnaryOp(EOperation::Rcp, Input); }
 	FExpression* NewFrac(FExpression* Input) { return NewUnaryOp(EOperation::Frac, Input); }
 	FExpression* NewLength(FExpression* Input) { return NewUnaryOp(EOperation::Length, Input); }
@@ -573,7 +574,6 @@ public:
 	FExpression* NewMul(FExpression* Lhs, FExpression* Rhs) { return NewBinaryOp(EOperation::Mul, Lhs, Rhs); }
 	FExpression* NewDiv(FExpression* Lhs, FExpression* Rhs) { return NewBinaryOp(EOperation::Div, Lhs, Rhs); }
 	FExpression* NewFmod(FExpression* Lhs, FExpression* Rhs) { return NewBinaryOp(EOperation::Fmod, Lhs, Rhs); }
-	FExpression* NewDot(FExpression* Lhs, FExpression* Rhs) { return NewBinaryOp(EOperation::Dot, Lhs, Rhs); }
 	FExpression* NewMin(FExpression* Lhs, FExpression* Rhs) { return NewBinaryOp(EOperation::Min, Lhs, Rhs); }
 	FExpression* NewMax(FExpression* Lhs, FExpression* Rhs) { return NewBinaryOp(EOperation::Max, Lhs, Rhs); }
 	FExpression* NewLess(FExpression* Lhs, FExpression* Rhs) { return NewBinaryOp(EOperation::Less, Lhs, Rhs); }
@@ -582,6 +582,7 @@ public:
 	FExpression* NewGreaterEqual(FExpression* Lhs, FExpression* Rhs) { return NewBinaryOp(EOperation::GreaterEqual, Lhs, Rhs); }
 
 	FExpression* NewCross(FExpression* Lhs, FExpression* Rhs);
+	FExpression* NewDot(FExpression* Lhs, FExpression* Rhs) { return NewSum(NewMul(Lhs, Rhs)); }
 	FExpression* NewLerp(FExpression* A, FExpression* B, FExpression* T) { return NewAdd(A, NewMul(NewSub(B, A), T)); }
 
 private:
diff --git a/Engine/Source/Runtime/Engine/Public/HLSLTree/HLSLTreeTypes.h b/Engine/Source/Runtime/Engine/Public/HLSLTree/HLSLTreeTypes.h
--- a/Engine/Source/Runtime/Engine/Public/HLSLTree/HLSLTreeTypes.h
+++ b/Engine/Source/Runtime/Engine/Public/HLSLTree/HLSLTreeTypes.h
@@ -80,6 +80,7 @@ enum class EOperation : uint8
 	Sign,
 	Length,
 	Normalize,
+	Sum,
 	Sin,
 	Cos,
 	Tan,
@@ -99,7 +100,6 @@ enum class EOperation : uint8
 	PowPositiveClamped,
 	Atan2,
 	Atan2Fast,
-	Dot,
 	Min,
 	Max,
 	Less,
