Strata node name update and clean up

diff --git a/Engine/Source/Runtime/Engine/Classes/Materials/MaterialExpressionStrata.h b/Engine/Source/Runtime/Engine/Classes/Materials/MaterialExpressionStrata.h
--- a/Engine/Source/Runtime/Engine/Classes/Materials/MaterialExpressionStrata.h
+++ b/Engine/Source/Runtime/Engine/Classes/Materials/MaterialExpressionStrata.h
@@ -34,7 +34,7 @@ class UMaterialExpressionStrataBSDF : public UMaterialExpression
 };
 
 
-UCLASS(MinimalAPI, collapsecategories, hidecategories = Object, DisplayName = "Strata Slab")
+UCLASS(MinimalAPI, collapsecategories, hidecategories = Object, DisplayName = "Strata Legacy Conversion")
 class UMaterialExpressionStrataLegacyConversion : public UMaterialExpressionStrataBSDF
 {
 	GENERATED_UCLASS_BODY()
@@ -557,7 +557,7 @@ class UMaterialExpressionStrataSingleLayerWaterBSDF : public UMaterialExpression
 ///////////////////////////////////////////////////////////////////////////////
 // Operator nodes
 
-UCLASS(MinimalAPI, collapsecategories, hidecategories = Object, DisplayName = "Strata BSDF Horizontal Blend")
+UCLASS(MinimalAPI, collapsecategories, hidecategories = Object, DisplayName = "Strata Horizontal Blend")
 class UMaterialExpressionStrataHorizontalMixing : public UMaterialExpressionStrataBSDF
 {
 	GENERATED_UCLASS_BODY()
@@ -596,7 +596,7 @@ class UMaterialExpressionStrataHorizontalMixing : public UMaterialExpressionStra
 	//~ End UMaterialExpression Interface
 };
 
-UCLASS(MinimalAPI, collapsecategories, hidecategories = Object, DisplayName = "Strata BSDF Vertical Layer")
+UCLASS(MinimalAPI, collapsecategories, hidecategories = Object, DisplayName = "Strata Vertical Layer")
 class UMaterialExpressionStrataVerticalLayering : public UMaterialExpressionStrataBSDF
 {
 	GENERATED_UCLASS_BODY()
@@ -629,7 +629,7 @@ class UMaterialExpressionStrataVerticalLayering : public UMaterialExpressionStra
 	//~ End UMaterialExpression Interface
 };
 
-UCLASS(MinimalAPI, collapsecategories, hidecategories = Object, DisplayName = "Strata BSDF Add")
+UCLASS(MinimalAPI, collapsecategories, hidecategories = Object, DisplayName = "Strata Add")
 class UMaterialExpressionStrataAdd : public UMaterialExpressionStrataBSDF
 {
 	GENERATED_UCLASS_BODY()
@@ -662,7 +662,7 @@ class UMaterialExpressionStrataAdd : public UMaterialExpressionStrataBSDF
 	//~ End UMaterialExpression Interface
 };
 
-UCLASS(MinimalAPI, collapsecategories, hidecategories = Object, DisplayName = "Strata BSDF Weight")
+UCLASS(MinimalAPI, collapsecategories, hidecategories = Object, DisplayName = "Strata Coverate Weight")
 class UMaterialExpressionStrataWeight : public UMaterialExpressionStrataBSDF
 {
 	GENERATED_UCLASS_BODY()
diff --git a/Engine/Source/Runtime/Engine/Private/Materials/MaterialExpressions.cpp b/Engine/Source/Runtime/Engine/Private/Materials/MaterialExpressions.cpp
--- a/Engine/Source/Runtime/Engine/Private/Materials/MaterialExpressions.cpp
+++ b/Engine/Source/Runtime/Engine/Private/Materials/MaterialExpressions.cpp
@@ -20573,7 +20573,7 @@ void UMaterialExpressionStrataLegacyConversion::PostEditChangeProperty(FProperty
 
 void UMaterialExpressionStrataLegacyConversion::GetCaption(TArray<FString>& OutCaptions) const
 {
-	OutCaptions.Add(TEXT("Strata Legacy conversion"));
+	OutCaptions.Add(TEXT("Strata Legacy Conversion"));
 }
 
 uint32 UMaterialExpressionStrataLegacyConversion::GetOutputType(int32 OutputIndex)
@@ -21633,7 +21633,7 @@ UMaterialExpressionStrataHorizontalMixing::UMaterialExpressionStrataHorizontalMi
 	struct FConstructorStatics
 	{
 		FText NAME_Strata;
-		FConstructorStatics() : NAME_Strata(LOCTEXT("Strata Ops", "Strata Ops")) { }
+		FConstructorStatics() : NAME_Strata(LOCTEXT("Strata Ops", "Strata Operators")) { }
 	};
 	static FConstructorStatics ConstructorStatics;
 #if WITH_EDITORONLY_DATA
@@ -21739,11 +21739,11 @@ void UMaterialExpressionStrataHorizontalMixing::GetCaption(TArray<FString>& OutC
 {
 	if (bUseParameterBlending)
 	{
-		OutCaptions.Add(TEXT("Strata BSDF Horizontal Blend (Parameter Blend)"));
+		OutCaptions.Add(TEXT("Strata Horizontal Blend (Parameter Blend)"));
 	}
 	else
 	{
-		OutCaptions.Add(TEXT("Strata BSDF Horizontal Blend"));
+		OutCaptions.Add(TEXT("Strata Horizontal Blend"));
 	}
 }
 
@@ -21802,7 +21802,7 @@ UMaterialExpressionStrataVerticalLayering::UMaterialExpressionStrataVerticalLaye
 	struct FConstructorStatics
 	{
 		FText NAME_Strata;
-		FConstructorStatics() : NAME_Strata(LOCTEXT("Strata Ops", "Strata Ops")) { }
+		FConstructorStatics() : NAME_Strata(LOCTEXT("Strata Ops", "Strata Operators")) { }
 	};
 	static FConstructorStatics ConstructorStatics;
 #if WITH_EDITORONLY_DATA
@@ -21901,11 +21901,11 @@ void UMaterialExpressionStrataVerticalLayering::GetCaption(TArray<FString>& OutC
 {
 	if (bUseParameterBlending)
 	{
-		OutCaptions.Add(TEXT("Strata BSDF Vertical Layer (Parameter Blend)"));
+		OutCaptions.Add(TEXT("Strata Vertical Layer (Parameter Blend)"));
 	}
 	else
 	{
-		OutCaptions.Add(TEXT("Strata BSDF Vertical Layer"));
+		OutCaptions.Add(TEXT("Strata Vertical Layer"));
 	}
 }
 
@@ -21964,7 +21964,7 @@ UMaterialExpressionStrataAdd::UMaterialExpressionStrataAdd(const FObjectInitiali
 	struct FConstructorStatics
 	{
 		FText NAME_Strata;
-		FConstructorStatics() : NAME_Strata(LOCTEXT("Strata Ops", "Strata Ops")) { }
+		FConstructorStatics() : NAME_Strata(LOCTEXT("Strata Ops", "Strata Operators")) { }
 	};
 	static FConstructorStatics ConstructorStatics;
 #if WITH_EDITORONLY_DATA
@@ -22062,11 +22062,11 @@ void UMaterialExpressionStrataAdd::GetCaption(TArray<FString>& OutCaptions) cons
 {
 	if (bUseParameterBlending)
 	{
-		OutCaptions.Add(TEXT("Strata BSDF Add (Parameter Blend)"));
+		OutCaptions.Add(TEXT("Strata Add (Parameter Blend)"));
 	}
 	else
 	{
-		OutCaptions.Add(TEXT("Strata BSDF Add"));
+		OutCaptions.Add(TEXT("Strata Add"));
 	}
 }
 
@@ -22124,7 +22124,7 @@ UMaterialExpressionStrataWeight::UMaterialExpressionStrataWeight(const FObjectIn
 	struct FConstructorStatics
 	{
 		FText NAME_Strata;
-		FConstructorStatics() : NAME_Strata(LOCTEXT("Strata Ops", "Strata Ops")) { }
+		FConstructorStatics() : NAME_Strata(LOCTEXT("Strata Ops", "Strata Operators")) { }
 	};
 	static FConstructorStatics ConstructorStatics;
 #if WITH_EDITORONLY_DATA
@@ -22172,7 +22172,7 @@ int32 UMaterialExpressionStrataWeight::Compile(class FMaterialCompiler* Compiler
 
 void UMaterialExpressionStrataWeight::GetCaption(TArray<FString>& OutCaptions) const
 {
-	OutCaptions.Add(TEXT("Strata BSDF Weight"));
+	OutCaptions.Add(TEXT("Strata Coverage Weight"));
 }
 
 uint32 UMaterialExpressionStrataWeight::GetOutputType(int32 OutputIndex)
