Make "hidden shadow" not visible when the actor is hidden in the editor UI (before, the shadow would remain even if hidden in editor).

diff --git a/Engine/Source/Runtime/Engine/Private/Components/PrimitiveComponent.cpp b/Engine/Source/Runtime/Engine/Private/Components/PrimitiveComponent.cpp
--- a/Engine/Source/Runtime/Engine/Private/Components/PrimitiveComponent.cpp
+++ b/Engine/Source/Runtime/Engine/Private/Components/PrimitiveComponent.cpp
@@ -1491,7 +1491,15 @@ void UPrimitiveComponent::SetOnlyOwnerSee(bool bNewOnlyOwnerSee)
 bool UPrimitiveComponent::ShouldComponentAddToScene() const
 {
 	bool bSceneAdd = USceneComponent::ShouldComponentAddToScene();
-	return bSceneAdd && (ShouldRender() || bCastHiddenShadow || bRayTracingFarField);
+
+#if WITH_EDITOR
+	AActor* Owner = GetOwner();
+	const bool bIsHiddenInEditor = GIsEditor && Owner && Owner->IsHiddenEd();
+#else
+	const bool bIsHiddenInEditor = false;
+#endif
+
+	return bSceneAdd && (ShouldRender() || (bCastHiddenShadow && !bIsHiddenInEditor) || bRayTracingFarField);
 }
 
 bool UPrimitiveComponent::ShouldCreatePhysicsState() const
