Fix missing shadow casting from material instance when setting Cast Shadow As Masked flag on parent material.
- Superset of changes in PR #6863

diff --git a/Engine/Source/Runtime/Engine/Classes/Materials/MaterialInstance.h b/Engine/Source/Runtime/Engine/Classes/Materials/MaterialInstance.h
--- a/Engine/Source/Runtime/Engine/Classes/Materials/MaterialInstance.h
+++ b/Engine/Source/Runtime/Engine/Classes/Materials/MaterialInstance.h
@@ -684,6 +684,7 @@ public:
 #endif // WITH_EDITOR
 
 	ENGINE_API virtual float GetOpacityMaskClipValue() const override;
+	ENGINE_API virtual bool GetCastDynamicShadowAsMasked() const override;
 	ENGINE_API virtual EBlendMode GetBlendMode() const override;
 	ENGINE_API virtual FMaterialShadingModelField GetShadingModels() const override;
 	ENGINE_API virtual bool IsShadingModelFromMaterialExpression() const override;
diff --git a/Engine/Source/Runtime/Engine/Private/Materials/MaterialInstance.cpp b/Engine/Source/Runtime/Engine/Private/Materials/MaterialInstance.cpp
--- a/Engine/Source/Runtime/Engine/Private/Materials/MaterialInstance.cpp
+++ b/Engine/Source/Runtime/Engine/Private/Materials/MaterialInstance.cpp
@@ -3760,6 +3760,11 @@ float UMaterialInstance::GetOpacityMaskClipValue() const
 	return OpacityMaskClipValue;
 }
 
+bool UMaterialInstance::GetCastDynamicShadowAsMasked() const
+{
+	return bCastDynamicShadowAsMasked;
+}
+
 EBlendMode UMaterialInstance::GetBlendMode() const
 {
 	return BlendMode;
@@ -3792,7 +3797,7 @@ bool UMaterialInstance::IsDitheredLODTransition() const
 
 bool UMaterialInstance::IsMasked() const
 {
-	return GetBlendMode() == EBlendMode::BLEND_Masked;
+	return GetBlendMode() == EBlendMode::BLEND_Masked || (GetBlendMode() == EBlendMode::BLEND_Translucent && GetCastDynamicShadowAsMasked());
 }
 
 USubsurfaceProfile* UMaterialInstance::GetSubsurfaceProfile_Internal() const
