Restore the functionality of the command "DEBUG OOM"

- It looks like this was accidently broken in CL 2920312 which was refactoring code around the memory leak implementation.
- In this restored version instead of relying on a tick in UnrealEngine we throw a lambda onto FTSTickerto create the leak every frame.
- The original defaults (1MiB per frame allocated in 64KiB chunks) remain but they can now be overriden when entering the command.
- To override the defaults pass in two numbers after the command, the first number will state how much memory to leak every frame (in MiB) and the second number states how large the allocations that leak should be (in KiB)
-- If a value is missing or zero the defaults are used.
-- If the total size to leak per frame does not divide by the allocation size exactly then we round up the number of allocations that we will make each frame, resulting in slightly larger leaks.

diff --git a/Engine/Source/Runtime/Engine/Private/UnrealEngine.cpp b/Engine/Source/Runtime/Engine/Private/UnrealEngine.cpp
--- a/Engine/Source/Runtime/Engine/Private/UnrealEngine.cpp
+++ b/Engine/Source/Runtime/Engine/Private/UnrealEngine.cpp
@@ -9536,10 +9536,6 @@ bool UEngine::HandleGetIniCommand(const TCHAR* Cmd, FOutputDevice& Ar)
 }
 #endif // !UE_BUILD_SHIPPING
 
-
-// debug flag to allocate memory every frame, to trigger an OOM condition
-static bool GDebugAllocMemEveryFrame = false;
-
 /** Helper function to cause a stack overflow crash */
 FORCENOINLINE void StackOverflowFunction(int32* DummyArg)
 {
@@ -10048,8 +10044,43 @@ bool UEngine::PerformError(const TCHAR* Cmd, FOutputDevice& Ar)
 	}
 	else if (FParse::Command(&Cmd, TEXT("OOM")))
 	{
-		Ar.Log(TEXT("Will continuously allocate 1MB per frame until we hit OOM"));
-		GDebugAllocMemEveryFrame = true;
+		auto ParseFromCmd = [](const TCHAR*& Cmd, uint32 DefaultValue)
+		{
+			uint32 ParsedValue = DefaultValue;
+
+			const FString Token = FParse::Token(Cmd, false);
+			if (!Token.IsEmpty())
+			{
+				LexFromString(ParsedValue, *Token);
+				if (ParsedValue == 0)
+				{
+					ParsedValue = DefaultValue;
+				}
+			}
+			
+			return ParsedValue;
+		};
+	
+		// By default we leak 1mb per frame in chunks of 64kb
+		const uint32 LeakPerFrameMiB = ParseFromCmd(Cmd, 1);
+		const uint32 AllocSizeKiB = ParseFromCmd(Cmd, 64);
+		
+		const uint32 NumAllocs = FMath::DivideAndRoundUp(LeakPerFrameMiB * 1024, AllocSizeKiB);
+		const uint32 AllocSizeBytes = AllocSizeKiB * 1024;
+
+		Ar.Logf(TEXT("Will continuously allocate %u MiB per frame in %u KiB chunks until the process runs OOM"), LeakPerFrameMiB, AllocSizeKiB);
+
+		FTSTicker::GetCoreTicker().AddTicker(FTickerDelegate::CreateLambda([NumAllocs, AllocSizeBytes](float)
+		{
+			for (uint32 Index = 0; Index < NumAllocs; ++Index)
+			{
+				void* Eat = FMemory::Malloc(AllocSizeBytes);
+				FMemory::Memset(Eat, 0, AllocSizeBytes);
+			}
+
+			return true;
+		}));
+
 		return true;
 	}
 	else if (FParse::Command(&Cmd, TEXT("STACKOVERFLOW")))
