Typo fix from PR 8922. Thank you, gtnardy!
https://github.com/EpicGames/UnrealEngine/pull/8922/commits/cbb9304a818bdb099a869e700d50db0625908a2d
#github
#8922

diff --git a/Engine/Source/Runtime/Engine/Private/TextureCube.cpp b/Engine/Source/Runtime/Engine/Private/TextureCube.cpp
--- a/Engine/Source/Runtime/Engine/Private/TextureCube.cpp
+++ b/Engine/Source/Runtime/Engine/Private/TextureCube.cpp
@@ -143,7 +143,7 @@ void UTextureCube::Serialize(FArchive& Ar)
 		// Previously default maximum resolution for cubemaps generated from long-lat sources was set to 512 pixels in the texture building code.
 		// This value is now explicitly set for cubemaps loaded from earlier versions of the stream in order to avoid texture rebuild.
 		MaxTextureSize = 512;
-		UE_LOG(LogTexture, Log, TEXT("Default maximum texture size for cubemaps generated from long-lat sources has been changed from 512 to unlimited. In order to preserve old behaiour for '%s', its maximum texture size has been explicitly set to 512."), *GetPathName());
+		UE_LOG(LogTexture, Log, TEXT("Default maximum texture size for cubemaps generated from long-lat sources has been changed from 512 to unlimited. In order to preserve old behavior for '%s', its maximum texture size has been explicitly set to 512."), *GetPathName());
 	}
 #endif // #if WITH_EDITORONLY_DATA
 
