Split HLSL 'Operations' into a separate file
Fix up type analysis for various operations

diff --git a/Engine/Shaders/Private/MaterialTemplate.ush b/Engine/Shaders/Private/MaterialTemplate.ush
--- a/Engine/Shaders/Private/MaterialTemplate.ush
+++ b/Engine/Shaders/Private/MaterialTemplate.ush
@@ -27,7 +27,7 @@
 
 // Update this GUID to force all material shaders to recompile
 // Merge conflicts on this line should be resolved by generating a new GUID
-// GUID = 5F290441-5611-4232-B26D-613F5E421262
+// GUID = 5F290441-5611-4232-B26D-613F5E421262_A
 
 #if USES_SPEEDTREE
 	#include "/Engine/Private/SpeedTreeCommon.ush"
diff --git a/Engine/Source/Runtime/Engine/Private/HLSLTree/HLSLTree.cpp b/Engine/Source/Runtime/Engine/Private/HLSLTree/HLSLTree.cpp
--- a/Engine/Source/Runtime/Engine/Private/HLSLTree/HLSLTree.cpp
+++ b/Engine/Source/Runtime/Engine/Private/HLSLTree/HLSLTree.cpp
@@ -80,69 +80,6 @@ public:
 	int32 OutputIndex;
 };
 
-FOperationDescription::FOperationDescription(const TCHAR* InName, const TCHAR* InOperator, int8 InNumInputs, Shader::EPreshaderOpcode InOpcode)
-	: Name(InName), Operator(InOperator), NumInputs(InNumInputs), PreshaderOpcode(InOpcode)
-{}
-
-FOperationDescription::FOperationDescription()
-	: Name(nullptr), Operator(nullptr), PreshaderOpcode(Shader::EPreshaderOpcode::Nop)
-{}
-
-
-FOperationDescription GetOperationDescription(EOperation Op)
-{
-	switch (Op)
-	{
-	case EOperation::None: return FOperationDescription(TEXT("None"), TEXT(""), 0, Shader::EPreshaderOpcode::Nop); break;
-	// Unary
-	case EOperation::Abs: return FOperationDescription(TEXT("Abs"), TEXT("abs"), 1, Shader::EPreshaderOpcode::Abs); break;
-	case EOperation::Neg: return FOperationDescription(TEXT("Neg"), TEXT("-"), 1, Shader::EPreshaderOpcode::Neg); break;
-	case EOperation::Rcp: return FOperationDescription(TEXT("Rcp"), TEXT("/"), 1, Shader::EPreshaderOpcode::Rcp); break;
-	case EOperation::Sqrt: return FOperationDescription(TEXT("Sqrt"), TEXT("sqrt"), 1, Shader::EPreshaderOpcode::Sqrt); break;
-	case EOperation::Log2: return FOperationDescription(TEXT("Log2"), TEXT("log2"), 1, Shader::EPreshaderOpcode::Log2); break;
-	case EOperation::Frac: return FOperationDescription(TEXT("Frac"), TEXT("frac"), 1, Shader::EPreshaderOpcode::Frac); break;
-	case EOperation::Floor: return FOperationDescription(TEXT("Floor"), TEXT("floor"), 1, Shader::EPreshaderOpcode::Floor); break;
-	case EOperation::Ceil: return FOperationDescription(TEXT("Ceil"), TEXT("ceil"), 1, Shader::EPreshaderOpcode::Ceil); break;
-	case EOperation::Round: return FOperationDescription(TEXT("Round"), TEXT("round"), 1, Shader::EPreshaderOpcode::Round); break;
-	case EOperation::Trunc: return FOperationDescription(TEXT("Trunc"), TEXT("trunc"), 1, Shader::EPreshaderOpcode::Trunc); break;
-	case EOperation::Saturate: return FOperationDescription(TEXT("Saturate"), TEXT("saturate"), 1, Shader::EPreshaderOpcode::Saturate); break;
-	case EOperation::Sign: return FOperationDescription(TEXT("Sign"), TEXT("sign"), 1, Shader::EPreshaderOpcode::Sign); break;
-	case EOperation::Length: return FOperationDescription(TEXT("Length"), TEXT("length"), 1, Shader::EPreshaderOpcode::Length); break;
-	case EOperation::Normalize: return FOperationDescription(TEXT("Normalize"), TEXT("normalize"), 1, Shader::EPreshaderOpcode::Normalize); break;
-	case EOperation::Sum: return FOperationDescription(TEXT("Sum"), TEXT("sum"), 1, Shader::EPreshaderOpcode::Nop); break; // TODO
-	case EOperation::Sin: return FOperationDescription(TEXT("Sin"), TEXT("sin"), 1, Shader::EPreshaderOpcode::Sin); break;
-	case EOperation::Cos: return FOperationDescription(TEXT("Cos"), TEXT("cos"), 1, Shader::EPreshaderOpcode::Cos); break;
-	case EOperation::Tan: return FOperationDescription(TEXT("Tan"), TEXT("tan"), 1, Shader::EPreshaderOpcode::Tan); break;
-	case EOperation::Asin: return FOperationDescription(TEXT("Asin"), TEXT("asin"), 1, Shader::EPreshaderOpcode::Asin); break;
-	case EOperation::AsinFast: return FOperationDescription(TEXT("AsinFast"), TEXT("asinFast"), 1, Shader::EPreshaderOpcode::Asin); break;
-	case EOperation::Acos: return FOperationDescription(TEXT("Acos"), TEXT("acos"), 1, Shader::EPreshaderOpcode::Acos); break;
-	case EOperation::AcosFast: return FOperationDescription(TEXT("AcosFast"), TEXT("acosFast"), 1, Shader::EPreshaderOpcode::Acos); break;
-	case EOperation::Atan: return FOperationDescription(TEXT("Atan"), TEXT("atan"), 1, Shader::EPreshaderOpcode::Atan); break;
-	case EOperation::AtanFast: return FOperationDescription(TEXT("AtanFast"), TEXT("atanFast"), 1, Shader::EPreshaderOpcode::Atan); break;
-
-	// Binary
-	case EOperation::Add: return FOperationDescription(TEXT("Add"), TEXT("+"), 2, Shader::EPreshaderOpcode::Add); break;
-	case EOperation::Sub: return FOperationDescription(TEXT("Subtract"), TEXT("-"), 2, Shader::EPreshaderOpcode::Sub); break;
-	case EOperation::Mul: return FOperationDescription(TEXT("Multiply"), TEXT("*"), 2, Shader::EPreshaderOpcode::Mul); break;
-	case EOperation::Div: return FOperationDescription(TEXT("Divide"), TEXT("/"), 2, Shader::EPreshaderOpcode::Div); break;
-	case EOperation::Fmod: return FOperationDescription(TEXT("Fmod"), TEXT("%"), 2, Shader::EPreshaderOpcode::Fmod); break;
-	case EOperation::PowPositiveClamped: return FOperationDescription(TEXT("PowPositiveClamped"), TEXT("PowPositiveClamped"), 2, Shader::EPreshaderOpcode::Nop); break;
-	case EOperation::Atan2: return FOperationDescription(TEXT("Atan2"), TEXT("atan2"), 2, Shader::EPreshaderOpcode::Atan2); break;
-	case EOperation::Atan2Fast: return FOperationDescription(TEXT("Atan2Fast"), TEXT("atan2Fast"), 2, Shader::EPreshaderOpcode::Atan2); break;
-	case EOperation::Min: return FOperationDescription(TEXT("Min"), TEXT("min"), 2, Shader::EPreshaderOpcode::Min); break;
-	case EOperation::Max: return FOperationDescription(TEXT("Max"), TEXT("max"), 2, Shader::EPreshaderOpcode::Max); break;
-	case EOperation::Less: return FOperationDescription(TEXT("Less"), TEXT("<"), 2, Shader::EPreshaderOpcode::Less); break;
-	case EOperation::Greater: return FOperationDescription(TEXT("Greater"), TEXT(">"), 2, Shader::EPreshaderOpcode::Greater); break;
-	case EOperation::LessEqual: return FOperationDescription(TEXT("LessEqual"), TEXT("<="), 2, Shader::EPreshaderOpcode::Nop); break;
-	case EOperation::GreaterEqual: return FOperationDescription(TEXT("GreaterEqual"), TEXT(">="), 2, Shader::EPreshaderOpcode::Nop); break;
-	case EOperation::VecMulMatrix3: return FOperationDescription(TEXT("VecMulMatrix3"), TEXT("mul"), 2, Shader::EPreshaderOpcode::Nop); break;
-	case EOperation::VecMulMatrix4: return FOperationDescription(TEXT("VecMulMatrix4"), TEXT("mul"), 2, Shader::EPreshaderOpcode::Nop); break;
-	case EOperation::Matrix3MulVec: return FOperationDescription(TEXT("Matrix3MulVec"), TEXT("mul"), 2, Shader::EPreshaderOpcode::Nop); break;
-	case EOperation::Matrix4MulVec: return FOperationDescription(TEXT("Matrix4MulVec"), TEXT("mul"), 2, Shader::EPreshaderOpcode::Nop); break;
-	default: checkNoEntry(); return FOperationDescription();
-	}
-}
-
 EExpressionEvaluation CombineEvaluations(EExpressionEvaluation Lhs, EExpressionEvaluation Rhs)
 {
 	if (Lhs == EExpressionEvaluation::None)
@@ -711,7 +648,12 @@ FPreparedComponent CombineComponents(const FPreparedComponent& Lhs, const FPrepa
 	}
 }
 
-FPreparedType::FPreparedType(const Shader::FType& InType, EExpressionEvaluation InEvalution)
+FPreparedType::FPreparedType(Shader::EValueType InType, const FPreparedComponent& InComponent)
+	: FPreparedType(Shader::FType(InType), InComponent)
+{
+}
+
+FPreparedType::FPreparedType(const Shader::FType& InType, const FPreparedComponent& InComponent)
 {
 	if (InType.IsStruct())
 	{
@@ -722,13 +664,13 @@ FPreparedType::FPreparedType(const Shader::FType& InType, EExpressionEvaluation
 		ValueComponentType = Shader::GetValueTypeDescription(InType).ComponentType;
 	}
 
-	if (InEvalution != EExpressionEvaluation::None)
+	if (!InComponent.IsNone())
 	{
 		const int32 NumComponents = InType.GetNumComponents();
 		PreparedComponents.Reserve(NumComponents);
 		for (int32 Index = 0; Index < NumComponents; ++Index)
 		{
-			SetComponent(Index, InEvalution);
+			SetComponent(Index, InComponent);
 		}
 	}
 }
@@ -845,6 +787,16 @@ EExpressionEvaluation FPreparedType::GetFieldEvaluation(const FEmitScope& Scope,
 	return Result;
 }
 
+FPreparedComponent FPreparedType::GetMergedComponent() const
+{
+	FPreparedComponent Result;
+	for (const FPreparedComponent& Component : PreparedComponents)
+	{
+		Result = CombineComponents(Result, Component);
+	}
+	return Result;
+}
+
 FPreparedComponent FPreparedType::GetComponent(int32 Index) const
 {
 	if (PreparedComponents.IsValidIndex(Index))
@@ -877,6 +829,18 @@ void FPreparedType::SetComponent(int32 Index, const FPreparedComponent& InCompon
 	}
 }
 
+void FPreparedType::SetComponentBounds(int32 Index, const Shader::FComponentBounds Bounds)
+{
+	if (PreparedComponents.IsValidIndex(Index))
+	{
+		FPreparedComponent& Component = PreparedComponents[Index];
+		if (!Component.IsNone())
+		{
+			Component.Bounds = Bounds;
+		}
+	}
+}
+
 void FPreparedType::MergeComponent(int32 Index, const FPreparedComponent& InComponent)
 {
 	if (InComponent.Evaluation != EExpressionEvaluation::None)
@@ -930,37 +894,37 @@ void FPreparedType::SetLoopEvaluation(FEmitScope& Scope, const FRequestedType& R
 	}
 }
 
-void FPreparedType::UpdateBounds(const FRequestedType& RequestedType, Shader::FComponentBounds Bounds)
+Shader::FComponentBounds FPreparedType::GetComponentBounds(int32 Index) const
 {
-	for (int32 Index = 0; Index < PreparedComponents.Num(); ++Index)
+	const FPreparedComponent Component = GetComponent(Index);
+	if (Component.Evaluation == EExpressionEvaluation::None)
 	{
-		FPreparedComponent& Component = PreparedComponents[Index];
-		if (RequestedType.IsComponentRequested(Index) && !Component.IsNone())
-		{
-			Component.Bounds.Min = Shader::MaxBound(Component.Bounds.Min, Bounds.Min);
-			Component.Bounds.Max = Shader::MinBound(Component.Bounds.Max, Bounds.Max);
-		}
+		return Shader::FComponentBounds();
 	}
+	else if (Component.Evaluation == EExpressionEvaluation::ConstantZero)
+	{
+		return Shader::FComponentBounds(Shader::EComponentBound::Zero, Shader::EComponentBound::Zero);
+	}
+
+	const Shader::EValueComponentType ComponentType = StructType ? StructType->ComponentTypes[Index] : ValueComponentType;
+	check(ComponentType != Shader::EValueComponentType::Void);
+
+	const Shader::FValueComponentTypeDescription ComponentTypeDesc = Shader::GetValueComponentTypeDescription(ComponentType);
+	const Shader::EComponentBound MinBound = Shader::MaxBound(Component.Bounds.Min, ComponentTypeDesc.Bounds.Min);
+	const Shader::EComponentBound MaxBound = Shader::MinBound(Component.Bounds.Max, ComponentTypeDesc.Bounds.Max);
+	return Shader::FComponentBounds(MinBound, MaxBound);
 }
 
 Shader::FComponentBounds FPreparedType::GetBounds(const FRequestedType& RequestedType) const
 {
 	Shader::FComponentBounds Result(Shader::EComponentBound::DoubleMax, Shader::EComponentBound::NegDoubleMax);
-	for (int32 Index = 0; Index < PreparedComponents.Num(); ++Index)
+	for (int32 Index = 0; Index < RequestedType.RequestedComponents.Num(); ++Index)
 	{
-		const FPreparedComponent& Component = PreparedComponents[Index];
-		if (RequestedType.IsComponentRequested(Index) && !Component.IsNone())
+		if (RequestedType.IsComponentRequested(Index))
 		{
-			// Constrain the component bounds we've stored by the range of the component's type
-			const Shader::EValueComponentType ComponentType = StructType ? StructType->ComponentTypes[Index] : ValueComponentType;
-			check(ComponentType != Shader::EValueComponentType::Void);
-
-			const Shader::FValueComponentTypeDescription ComponentTypeDesc = Shader::GetValueComponentTypeDescription(ComponentType);
-			const Shader::EComponentBound MinBound = Shader::MaxBound(Component.Bounds.Min, ComponentTypeDesc.Bounds.Min);
-			const Shader::EComponentBound MaxBound = Shader::MinBound(Component.Bounds.Max, ComponentTypeDesc.Bounds.Max);
-	
-			Result.Min = Shader::MinBound(Result.Min, MinBound);
-			Result.Max = Shader::MaxBound(Result.Max, MaxBound);
+			const Shader::FComponentBounds ComponentBounds = GetComponentBounds(Index);
+			Result.Min = Shader::MinBound(Result.Min, ComponentBounds.Min);
+			Result.Max = Shader::MaxBound(Result.Max, ComponentBounds.Max);
 		}
 	}
 	return Result;
@@ -1025,6 +989,16 @@ FPreparedType MergePreparedTypes(const FPreparedType& Lhs, const FPreparedType&
 	return Result;
 }
 
+FPreparedType MakeNonLWCType(const FPreparedType& Type)
+{
+	FPreparedType Result(Type);
+	if (Result.IsNumeric())
+	{
+		Result.ValueComponentType = Shader::MakeNonLWCType(Result.ValueComponentType);
+	}
+	return Result;
+}
+
 bool FPrepareValueResult::TryMergePreparedType(FEmitContext& Context, const Shader::FStructType* StructType, Shader::EValueComponentType ComponentType)
 {
 	if (!StructType && ComponentType == Shader::EValueComponentType::Void)
diff --git a/Engine/Source/Runtime/Engine/Private/HLSLTree/HLSLTreeCommon.cpp b/Engine/Source/Runtime/Engine/Private/HLSLTree/HLSLTreeCommon.cpp
--- a/Engine/Source/Runtime/Engine/Private/HLSLTree/HLSLTreeCommon.cpp
+++ b/Engine/Source/Runtime/Engine/Private/HLSLTree/HLSLTreeCommon.cpp
@@ -1032,727 +1032,6 @@ void FExpressionSelect::EmitValuePreshader(FEmitContext& Context, FEmitScope& Sc
 	OutResult.Preshader.WriteOpcode(Shader::EPreshaderOpcode::ConstantZero).Write(OutResult.Type);
 }
 
-FExpression* FTree::NewUnaryOp(EOperation Op, FExpression* Input)
-{
-	FExpression* Inputs[1] = { Input };
-	return NewExpression<FExpressionOperation>(Op, Inputs);
-}
-
-FExpression* FTree::NewBinaryOp(EOperation Op, FExpression* Lhs, FExpression* Rhs)
-{
-	FExpression* Inputs[2] = { Lhs, Rhs };
-	return NewExpression<FExpressionOperation>(Op, Inputs);
-}
-
-FExpression* FTree::NewCross(FExpression* Lhs, FExpression* Rhs)
-{
-	//c_P[0] = v_A[1] * v_B[2] - v_A[2] * v_B[1];
-	//c_P[1] = -(v_A[0] * v_B[2] - v_A[2] * v_B[0]);
-	//c_P[2] = v_A[0] * v_B[1] - v_A[1] * v_B[0];
-	FExpression* Lhs0 = NewExpression<FExpressionSwizzle>(FSwizzleParameters(1, 0, 0), Lhs);
-	FExpression* Lhs1 = NewExpression<FExpressionSwizzle>(FSwizzleParameters(2, 2, 1), Lhs);
-	FExpression* Rhs0 = NewExpression<FExpressionSwizzle>(FSwizzleParameters(2, 2, 1), Rhs);
-	FExpression* Rhs1 = NewExpression<FExpressionSwizzle>(FSwizzleParameters(1, 0, 0), Rhs);
-	return NewSub(NewMul(NewMul(Lhs0, Rhs0), NewConstant(FVector3f(1.0f, -1.0f, 1.0f))), NewMul(Lhs1, Rhs1));
-}
-
-FExpressionOperation::FExpressionOperation(EOperation InOp, TConstArrayView<FExpression*> InInputs) : Op(InOp)
-{
-	const FOperationDescription OpDesc = GetOperationDescription(InOp);
-	check(OpDesc.NumInputs == InInputs.Num());
-	check(InInputs.Num() <= MaxInputs);
-
-	for (int32 i = 0; i < OpDesc.NumInputs; ++i)
-	{
-		Inputs[i] = InInputs[i];
-		check(Inputs[i]);
-	}
-}
-
-namespace Private
-{
-struct FOperationRequestedTypes
-{
-	FRequestedType InputType[FExpressionOperation::MaxInputs];
-	bool bIsMatrixOperation = false;
-};
-struct FOperationTypes
-{
-	
-	Shader::EValueType InputType[FExpressionOperation::MaxInputs];
-	Shader::EValueType ResultType;
-	bool bIsLWC = false;
-};
-
-FOperationRequestedTypes GetOperationRequestedTypes(EOperation Op, const FRequestedType& RequestedType)
-{
-	const FOperationDescription OpDesc = GetOperationDescription(Op);
-	FOperationRequestedTypes Types;
-	for (int32 Index = 0; Index < OpDesc.NumInputs; ++Index)
-	{
-		Types.InputType[Index] = RequestedType;
-	}
-	switch (Op)
-	{
-	case EOperation::Length:
-	case EOperation::Normalize:
-	case EOperation::Sum:
-		Types.InputType[0] = ERequestedType::Vector4;
-		break;
-	case EOperation::VecMulMatrix3:
-		Types.bIsMatrixOperation = true;
-		Types.InputType[0] = ERequestedType::Vector3;
-		Types.InputType[1] = ERequestedType::Matrix4x4;
-		break;
-	case EOperation::VecMulMatrix4:
-		Types.bIsMatrixOperation = true;
-		Types.InputType[0] = ERequestedType::Vector3;
-		Types.InputType[1] = ERequestedType::Matrix4x4;
-		break;
-	case EOperation::Matrix3MulVec:
-	case EOperation::Matrix4MulVec:
-		// No LWC for transpose matrices
-		Types.bIsMatrixOperation = true;
-		Types.InputType[0] = ERequestedType::Matrix4x4;
-		Types.InputType[1] = ERequestedType::Vector3;
-		break;
-	default:
-		break;
-	}
-	return Types;
-}
-
-FOperationTypes GetOperationTypes(EOperation Op, TConstArrayView<Shader::EValueType> InputTypes)
-{
-	Shader::EValueComponentType ComponentType = Shader::EValueComponentType::Void;
-	int32 NumComponents = 0;
-	for (int32 Index = 0; Index < InputTypes.Num(); ++Index)
-	{
-		const Shader::FValueTypeDescription InputTypeDesc = Shader::GetValueTypeDescription(InputTypes[Index]);
-		ComponentType = Shader::CombineComponentTypes(ComponentType, InputTypeDesc.ComponentType);
-		NumComponents = FMath::Max<int32>(NumComponents, InputTypeDesc.NumComponents);
-	}
-
-	Shader::EValueType IntermediateType = Shader::MakeValueType(ComponentType, NumComponents);
-
-	FOperationTypes Types;
-	for (int32 Index = 0; Index < InputTypes.Num(); ++Index)
-	{
-		Types.InputType[Index] = IntermediateType;
-	}
-	Types.ResultType = IntermediateType;
-	Types.bIsLWC = (ComponentType == Shader::EValueComponentType::Double);
-	switch (Op)
-	{
-	case EOperation::Length:
-	case EOperation::Sum:
-		Types.ResultType = Shader::MakeValueType(ComponentType, 1);
-		break;
-	case EOperation::Normalize:
-		Types.ResultType = Shader::MakeNonLWCType(IntermediateType);
-		break;
-	case EOperation::Frac:
-	case EOperation::Rcp:
-	case EOperation::Sqrt:
-	case EOperation::Sign:
-	case EOperation::Saturate:
-	case EOperation::Sin:
-	case EOperation::Cos:
-	case EOperation::Tan:
-	case EOperation::Asin:
-	case EOperation::AsinFast:
-	case EOperation::Acos:
-	case EOperation::AcosFast:
-	case EOperation::Atan:
-	case EOperation::AtanFast:
-		Types.ResultType = Shader::MakeNonLWCType(IntermediateType);
-		break;
-	case EOperation::Log2:
-		// No LWC support yet
-		Types.InputType[0] = Types.ResultType = Shader::MakeNonLWCType(IntermediateType);
-		break;
-	case EOperation::Less:
-	case EOperation::Greater:
-	case EOperation::LessEqual:
-	case EOperation::GreaterEqual:
-		Types.ResultType = Shader::MakeValueType(Shader::EValueComponentType::Bool, NumComponents);
-		break;
-	case EOperation::Fmod:
-		Types.InputType[1] = Types.ResultType = Shader::MakeNonLWCType(IntermediateType);
-		break;
-	case EOperation::PowPositiveClamped:
-	case EOperation::Atan2:
-	case EOperation::Atan2Fast:
-		// No LWC support yet
-		Types.InputType[0] = Types.InputType[1] = Types.ResultType = Shader::MakeNonLWCType(IntermediateType);
-		break;
-	case EOperation::VecMulMatrix3:
-		// No LWC for matrix3
-		Types.InputType[0] = Shader::EValueType::Float3;
-		Types.InputType[1] = Shader::EValueType::Float4x4;
-		Types.ResultType = Shader::EValueType::Float3;
-		break;
-	case EOperation::VecMulMatrix4:
-		Types.InputType[0] = Shader::EValueType::Float3;
-		Types.InputType[1] = Shader::EValueType::Float4x4;
-		Types.ResultType = Shader::EValueType::Float3;
-		break;
-	case EOperation::Matrix3MulVec:
-	case EOperation::Matrix4MulVec:
-		// No LWC for transpose matrices
-		Types.InputType[0] = Shader::EValueType::Float4x4;
-		Types.InputType[1] = Shader::EValueType::Float3;
-		Types.ResultType = Shader::EValueType::Float3;
-		break;
-	default:
-		break;
-	}
-	return Types;
-}
-} // namespace Private
-
-void FExpressionOperation::ComputeAnalyticDerivatives(FTree& Tree, FExpressionDerivatives& OutResult) const
-{
-	// Operations with constant derivatives
-	switch (Op)
-	{
-	case EOperation::Less:
-	case EOperation::Greater:
-	case EOperation::LessEqual:
-	case EOperation::GreaterEqual:
-		OutResult.ExpressionDdx = Tree.NewConstant(0.0f);
-		OutResult.ExpressionDdx = OutResult.ExpressionDdy;
-		break;
-	default:
-		break;
-	}
-
-	if (OutResult.IsValid())
-	{
-		return;
-	}
-
-	const FOperationDescription OpDesc = GetOperationDescription(Op);
-	FExpressionDerivatives InputDerivatives[MaxInputs];
-	for (int32 Index = 0; Index < OpDesc.NumInputs; ++Index)
-	{
-		InputDerivatives[Index] = Tree.GetAnalyticDerivatives(Inputs[Index]);
-		if (!InputDerivatives[Index].IsValid())
-		{
-			return;
-		}
-	}
-
-	switch (Op)
-	{
-	case EOperation::Neg:
-		OutResult.ExpressionDdx = Tree.NewNeg(InputDerivatives[0].ExpressionDdx);
-		OutResult.ExpressionDdy = Tree.NewNeg(InputDerivatives[0].ExpressionDdy);
-		break;
-	case EOperation::Rcp:
-	{
-		FExpression* Result = Tree.NewRcp(Inputs[0]);
-		FExpression* dFdA = Tree.NewNeg(Tree.NewMul(Result, Result));
-		OutResult.ExpressionDdx = Tree.NewMul(dFdA, InputDerivatives[0].ExpressionDdx);
-		OutResult.ExpressionDdy = Tree.NewMul(dFdA, InputDerivatives[0].ExpressionDdy);
-		break;
-	}
-	case EOperation::Sum:
-		OutResult.ExpressionDdx = Tree.NewSum(InputDerivatives[0].ExpressionDdx);
-		OutResult.ExpressionDdy = Tree.NewSum(InputDerivatives[0].ExpressionDdy);
-		break;
-	case EOperation::Frac:
-		OutResult = InputDerivatives[0];
-		break;
-	case EOperation::Length:
-	case EOperation::Normalize:
-	case EOperation::Floor:
-	case EOperation::Ceil:
-	case EOperation::Round:
-	case EOperation::Trunc:
-	case EOperation::Sign:
-	case EOperation::Abs:
-	case EOperation::Saturate:
-	case EOperation::PowPositiveClamped:
-	case EOperation::Sin:
-	case EOperation::Cos:
-	case EOperation::Tan:
-	case EOperation::Asin:
-	case EOperation::AsinFast:
-	case EOperation::Acos:
-	case EOperation::AcosFast:
-	case EOperation::Atan:
-	case EOperation::AtanFast:
-	case EOperation::Atan2:
-	case EOperation::Atan2Fast:
-	case EOperation::Sqrt:
-	case EOperation::Log2:
-		// TODO
-		break;
-	case EOperation::Add:
-		OutResult.ExpressionDdx = Tree.NewAdd(InputDerivatives[0].ExpressionDdx, InputDerivatives[1].ExpressionDdx);
-		OutResult.ExpressionDdy = Tree.NewAdd(InputDerivatives[0].ExpressionDdy, InputDerivatives[1].ExpressionDdy);
-		break;
-	case EOperation::Sub:
-		OutResult.ExpressionDdx = Tree.NewSub(InputDerivatives[0].ExpressionDdx, InputDerivatives[1].ExpressionDdx);
-		OutResult.ExpressionDdy = Tree.NewSub(InputDerivatives[0].ExpressionDdy, InputDerivatives[1].ExpressionDdy);
-		break;
-	case EOperation::Mul:
-		OutResult.ExpressionDdx = Tree.NewAdd(Tree.NewMul(InputDerivatives[0].ExpressionDdx, Inputs[1]), Tree.NewMul(InputDerivatives[1].ExpressionDdx, Inputs[0]));
-		OutResult.ExpressionDdy = Tree.NewAdd(Tree.NewMul(InputDerivatives[0].ExpressionDdy, Inputs[1]), Tree.NewMul(InputDerivatives[1].ExpressionDdy, Inputs[0]));
-		break;
-	case EOperation::Div:
-	{
-		FExpression* Denom = Tree.NewRcp(Tree.NewMul(Inputs[1], Inputs[1]));
-		FExpression* dFdA = Tree.NewMul(Inputs[1], Denom);
-		FExpression* dFdB = Tree.NewNeg(Tree.NewMul(Inputs[0], Denom));
-		OutResult.ExpressionDdx = Tree.NewAdd(Tree.NewMul(dFdA, InputDerivatives[0].ExpressionDdx), Tree.NewMul(dFdB, InputDerivatives[1].ExpressionDdx));
-		OutResult.ExpressionDdy = Tree.NewAdd(Tree.NewMul(dFdA, InputDerivatives[0].ExpressionDdy), Tree.NewMul(dFdB, InputDerivatives[1].ExpressionDdy));
-		break;
-	}
-	case EOperation::Fmod:
-		// Only valid when B derivatives are zero.
-		// We can't really do anything meaningful in the non-zero case.
-		OutResult = InputDerivatives[0];
-		break;
-	case EOperation::Min:
-	{
-		FExpression* Cond = Tree.NewLess(Inputs[0], Inputs[1]);
-		OutResult.ExpressionDdx = Tree.NewExpression<FExpressionSelect>(Cond, InputDerivatives[0].ExpressionDdx, InputDerivatives[1].ExpressionDdx);
-		OutResult.ExpressionDdy = Tree.NewExpression<FExpressionSelect>(Cond, InputDerivatives[0].ExpressionDdy, InputDerivatives[1].ExpressionDdy);
-		break;
-	}
-	case EOperation::Max:
-	{
-		FExpression* Cond = Tree.NewGreater(Inputs[0], Inputs[1]);
-		OutResult.ExpressionDdx = Tree.NewExpression<FExpressionSelect>(Cond, InputDerivatives[0].ExpressionDdx, InputDerivatives[1].ExpressionDdx);
-		OutResult.ExpressionDdy = Tree.NewExpression<FExpressionSelect>(Cond, InputDerivatives[0].ExpressionDdy, InputDerivatives[1].ExpressionDdy);
-		break;
-	}
-	case EOperation::VecMulMatrix3:
-	case EOperation::VecMulMatrix4:
-	case EOperation::Matrix3MulVec:
-	case EOperation::Matrix4MulVec:
-		// TODO
-		break;
-	default:
-		checkNoEntry();
-		break;
-	}
-}
-
-FExpression* FExpressionOperation::ComputePreviousFrame(FTree& Tree, const FRequestedType& RequestedType) const
-{
-	const Private::FOperationRequestedTypes RequestedTypes = Private::GetOperationRequestedTypes(Op, RequestedType);
-	const FOperationDescription OpDesc = GetOperationDescription(Op);
-	FExpression* PrevFrameInputs[MaxInputs];
-	for (int32 Index = 0; Index < OpDesc.NumInputs; ++Index)
-	{
-		PrevFrameInputs[Index] = Tree.GetPreviousFrame(Inputs[Index], RequestedTypes.InputType[Index]);
-	}
-
-	return Tree.NewExpression<FExpressionOperation>(Op, MakeArrayView(PrevFrameInputs, OpDesc.NumInputs));
-}
-
-bool FExpressionOperation::PrepareValue(FEmitContext& Context, FEmitScope& Scope, const FRequestedType& RequestedType, FPrepareValueResult& OutResult) const
-{
-	const FOperationDescription OpDesc = GetOperationDescription(Op);
-	const Private::FOperationRequestedTypes RequestedTypes = Private::GetOperationRequestedTypes(Op, RequestedType);
-
-	FPreparedType InputType[MaxInputs];
-	for (int32 Index = 0; Index < OpDesc.NumInputs; ++Index)
-	{
-		InputType[Index] = Context.PrepareExpression(Inputs[Index], Scope, RequestedTypes.InputType[Index]);
-		if (InputType[Index].IsVoid())
-		{
-			return false;
-		}
-
-		if (!InputType[Index].IsNumeric())
-		{
-			return Context.Errors->AddError(TEXT("Invalid arithmetic between non-numeric types"));
-		}
-	}
-
-	FPreparedType ResultType;
-	if (RequestedTypes.bIsMatrixOperation)
-	{
-		const EExpressionEvaluation Evaluation = CombineEvaluations(InputType[0].GetEvaluation(Scope), InputType[1].GetEvaluation(Scope));
-		const Shader::EValueComponentType ComponentType = Shader::CombineComponentTypes(InputType[0].ValueComponentType, InputType[1].ValueComponentType);
-		ResultType = FPreparedType(Shader::MakeValueType(ComponentType, 3), Evaluation);
-	}
-	else if (OpDesc.NumInputs == 1)
-	{
-		switch (Op)
-		{
-		case EOperation::Length:
-		case EOperation::Sum:
-			ResultType = FPreparedType(Shader::MakeValueType(InputType[0].ValueComponentType, 1), InputType[0].GetEvaluation(Scope));
-			break;
-		default:
-			ResultType = InputType[0];
-			if (Op == EOperation::Normalize)
-			{
-				ResultType.ValueComponentType = Shader::MakeNonLWCType(ResultType.ValueComponentType);
-			}
-			break;
-		}
-	}
-	else if(OpDesc.NumInputs == 2)
-	{
-		/*if (InputType[0].GetNumComponents() != 1 && InputType[1].GetNumComponents() != 1 && InputType[0].GetNumComponents() != InputType[1].GetNumComponents())
-		{
-			return Context.Errors->AddErrorf(TEXT("Invalid arithmetic between %s and %s"), InputType[0].GetType().GetName(), InputType[1].GetType().GetName());
-		}*/
-		ResultType = MergePreparedTypes(InputType[0], InputType[1]);
-		if (Op == EOperation::Less || Op == EOperation::Greater || Op == EOperation::LessEqual || Op == EOperation::GreaterEqual)
-		{
-			ResultType.ValueComponentType = Shader::EValueComponentType::Bool;
-		}
-		if (Op == EOperation::Min || Op == EOperation::Max)
-		{
-			const Shader::FComponentBounds LhsBounds = InputType[0].GetBounds(RequestedType);
-			const Shader::FComponentBounds RhsBounds = InputType[1].GetBounds(RequestedType);
-			const Shader::FComponentBounds Bounds = (Op == EOperation::Min) ? Shader::MinBound(LhsBounds, RhsBounds) : Shader::MaxBound(LhsBounds, RhsBounds);
-			ResultType.UpdateBounds(RequestedType, Bounds);
-		}
-	}
-
-	if (OpDesc.PreshaderOpcode == Shader::EPreshaderOpcode::Nop)
-	{
-		// No preshader support
-		ResultType.SetEvaluation(EExpressionEvaluation::Shader);
-	}
-
-	return OutResult.SetType(Context, RequestedType, ResultType);
-}
-
-void FExpressionOperation::EmitValueShader(FEmitContext& Context, FEmitScope& Scope, const FRequestedType& RequestedType, FEmitValueShaderResult& OutResult) const
-{
-	const FOperationDescription OpDesc = GetOperationDescription(Op);
-	Shader::EValueType InputTypes[MaxInputs];
-	for (int32 Index = 0; Index < OpDesc.NumInputs; ++Index)
-	{
-		InputTypes[Index] = Inputs[Index]->GetType();
-	}
-	const Private::FOperationRequestedTypes RequestedTypes = Private::GetOperationRequestedTypes(Op, RequestedType);
-	const Private::FOperationTypes Types = Private::GetOperationTypes(Op, MakeArrayView(InputTypes, OpDesc.NumInputs));
-	FEmitShaderExpression* InputValue[MaxInputs] = { nullptr };
-	for (int32 Index = 0; Index < OpDesc.NumInputs; ++Index)
-	{
-		InputValue[Index] = Inputs[Index]->GetValueShader(Context, Scope, RequestedTypes.InputType[Index], Types.InputType[Index]);
-	}
-
-	switch (Op)
-	{
-	case EOperation::Abs:
-		if (Types.bIsLWC)
-		{
-			OutResult.Code = Context.EmitExpression(Scope, Types.ResultType, TEXT("LWCAbs(%)"), InputValue[0]);
-		}
-		else
-		{
-			OutResult.Code = Context.EmitInlineExpression(Scope, Types.ResultType, TEXT("abs(%)"), InputValue[0]);
-		}
-		break;
-	case EOperation::Neg:
-		if (Types.bIsLWC)
-		{
-			OutResult.Code = Context.EmitExpression(Scope, Types.ResultType, TEXT("LWCNegate(%)"), InputValue[0]);
-		}
-		else
-		{
-			OutResult.Code = Context.EmitInlineExpression(Scope, Types.ResultType, TEXT("(-%)"), InputValue[0]);
-		}
-		break;
-	case EOperation::Rcp:
-		if (Types.bIsLWC)
-		{
-			OutResult.Code = Context.EmitExpression(Scope, Types.ResultType, TEXT("LWCRcp(%)"), InputValue[0]);
-		}
-		else
-		{
-			OutResult.Code = Context.EmitExpression(Scope, Types.ResultType, TEXT("rcp(%)"), InputValue[0]);
-		}
-		break;
-	case EOperation::Sqrt: OutResult.Code = Context.EmitExpression(Scope, Types.ResultType, Types.bIsLWC ? TEXT("LWCSqrt(%)") : TEXT("sqrt(%)"), InputValue[0]); break;
-	case EOperation::Log2: OutResult.Code = Context.EmitExpression(Scope, Types.ResultType, TEXT("log2(%)"), InputValue[0]); break;
-	case EOperation::Frac:
-		if (Types.bIsLWC)
-		{
-			OutResult.Code = Context.EmitExpression(Scope, Types.ResultType, TEXT("LWCFrac(%)"), InputValue[0]);
-		}
-		else
-		{
-			OutResult.Code = Context.EmitExpression(Scope, Types.ResultType, TEXT("frac(%)"), InputValue[0]);
-		}
-		break;
-
-	case EOperation::Floor:
-		if (Types.bIsLWC)
-		{
-			OutResult.Code = Context.EmitExpression(Scope, Types.ResultType, TEXT("LWCFloor(%)"), InputValue[0]);
-		}
-		else
-		{
-			OutResult.Code = Context.EmitExpression(Scope, Types.ResultType, TEXT("floor(%)"), InputValue[0]);
-		}
-		break;
-	case EOperation::Ceil:
-		if (Types.bIsLWC)
-		{
-			OutResult.Code = Context.EmitExpression(Scope, Types.ResultType, TEXT("LWCCeil(%)"), InputValue[0]);
-		}
-		else
-		{
-			OutResult.Code = Context.EmitExpression(Scope, Types.ResultType, TEXT("ceil(%)"), InputValue[0]);
-		}
-		break;
-	case EOperation::Round:
-		if (Types.bIsLWC)
-		{
-			OutResult.Code = Context.EmitExpression(Scope, Types.ResultType, TEXT("LWCRound(%)"), InputValue[0]);
-		}
-		else
-		{
-			OutResult.Code = Context.EmitExpression(Scope, Types.ResultType, TEXT("round(%)"), InputValue[0]);
-		}
-		break;
-	case EOperation::Trunc:
-		if (Types.bIsLWC)
-		{
-			OutResult.Code = Context.EmitExpression(Scope, Types.ResultType, TEXT("LWCTrunc(%)"), InputValue[0]);
-		}
-		else
-		{
-			OutResult.Code = Context.EmitExpression(Scope, Types.ResultType, TEXT("trunc(%)"), InputValue[0]);
-		}
-		break;
-	case EOperation::Saturate:
-		if (Types.bIsLWC)
-		{
-			OutResult.Code = Context.EmitExpression(Scope, Types.ResultType, TEXT("LWCSaturate(%)"), InputValue[0]);
-		}
-		else
-		{
-			OutResult.Code = Context.EmitExpression(Scope, Types.ResultType, TEXT("saturate(%)"), InputValue[0]);
-		}
-		break;
-	case EOperation::Sign:
-		if (Types.bIsLWC)
-		{
-			OutResult.Code = Context.EmitExpression(Scope, Types.ResultType, TEXT("LWCSign(%)"), InputValue[0]);
-		}
-		else
-		{
-			OutResult.Code = Context.EmitExpression(Scope, Types.ResultType, TEXT("sign(%)"), InputValue[0]);
-		}
-		break;
-	case EOperation::Length:
-		if (Types.bIsLWC)
-		{
-			OutResult.Code = Context.EmitExpression(Scope, Types.ResultType, TEXT("LWCLength(%)"), InputValue[0]);
-		}
-		else
-		{
-			OutResult.Code = Context.EmitExpression(Scope, Types.ResultType, TEXT("length(%)"), InputValue[0]);
-		}
-		break;
-	case EOperation::Normalize:
-		if (Types.bIsLWC)
-		{
-			OutResult.Code = Context.EmitExpression(Scope, Types.ResultType, TEXT("LWCNormalize(%)"), InputValue[0]);
-		}
-		else
-		{
-			OutResult.Code = Context.EmitExpression(Scope, Types.ResultType, TEXT("normalize(%)"), InputValue[0]);
-		}
-		break;
-	case EOperation::Sum: OutResult.Code = Context.EmitExpression(Scope, Types.ResultType, Types.bIsLWC ? TEXT("LWCVectorSum(%)") : TEXT("VectorSum(%)"), InputValue[0]); break;
-	case EOperation::Sin: OutResult.Code = Context.EmitExpression(Scope, Types.ResultType, Types.bIsLWC ? TEXT("LWCSin(%)") : TEXT("sin(%)"), InputValue[0]); break;
-	case EOperation::Cos: OutResult.Code = Context.EmitExpression(Scope, Types.ResultType, Types.bIsLWC ? TEXT("LWCCos(%)") : TEXT("cos(%)"), InputValue[0]); break;
-	case EOperation::Tan: OutResult.Code = Context.EmitExpression(Scope, Types.ResultType, Types.bIsLWC ? TEXT("LWCTan(%)") : TEXT("tan(%)"), InputValue[0]); break;
-	case EOperation::Asin: OutResult.Code = Context.EmitExpression(Scope, Types.ResultType, Types.bIsLWC ? TEXT("LWCAsin(%)") : TEXT("asin(%)"), InputValue[0]); break;
-	case EOperation::AsinFast: OutResult.Code = Context.EmitExpression(Scope, Types.ResultType, Types.bIsLWC ? TEXT("LWCAsin(%)") : TEXT("asinFast(%)"), InputValue[0]); break;
-	case EOperation::Acos: OutResult.Code = Context.EmitExpression(Scope, Types.ResultType, Types.bIsLWC ? TEXT("LWCAcos(%)") : TEXT("acos(%)"), InputValue[0]); break;
-	case EOperation::AcosFast: OutResult.Code = Context.EmitExpression(Scope, Types.ResultType, Types.bIsLWC ? TEXT("LWCAcos(%)") : TEXT("acosFast(%)"), InputValue[0]); break;
-	case EOperation::Atan: OutResult.Code = Context.EmitExpression(Scope, Types.ResultType, Types.bIsLWC ? TEXT("LWCAtan(%)") : TEXT("atan(%)"), InputValue[0]); break;
-	case EOperation::AtanFast: OutResult.Code = Context.EmitExpression(Scope, Types.ResultType, Types.bIsLWC ? TEXT("LWCAtan(%)") : TEXT("atanFast(%)"), InputValue[0]); break;
-	case EOperation::Add:
-		if (Types.bIsLWC)
-		{
-			OutResult.Code = Context.EmitExpression(Scope, Types.ResultType, TEXT("LWCAdd(%, %)"), InputValue[0], InputValue[1]);
-		}
-		else
-		{
-			OutResult.Code = Context.EmitExpression(Scope, Types.ResultType, TEXT("(% + %)"), InputValue[0], InputValue[1]);
-		}
-		break;
-	case EOperation::Sub:
-		if (Types.bIsLWC)
-		{
-			OutResult.Code = Context.EmitExpression(Scope, Types.ResultType, TEXT("LWCSubtract(%, %)"), InputValue[0], InputValue[1]);
-		}
-		else
-		{
-			OutResult.Code = Context.EmitExpression(Scope, Types.ResultType, TEXT("(% - %)"), InputValue[0], InputValue[1]);
-		}
-		break;
-	case EOperation::Mul:
-		if (Types.bIsLWC)
-		{
-			OutResult.Code = Context.EmitExpression(Scope, Types.ResultType, TEXT("LWCMultiply(%, %)"), InputValue[0], InputValue[1]);
-		}
-		else
-		{
-			OutResult.Code = Context.EmitExpression(Scope, Types.ResultType, TEXT("(% * %)"), InputValue[0], InputValue[1]);
-		}
-		break;
-	case EOperation::Div:
-		if (Types.bIsLWC)
-		{
-			OutResult.Code = Context.EmitExpression(Scope, Types.ResultType, TEXT("LWCDivide(%, %)"), InputValue[0], InputValue[1]);
-		}
-		else
-		{
-			OutResult.Code = Context.EmitExpression(Scope, Types.ResultType, TEXT("(% / %)"), InputValue[0], InputValue[1]);
-		}
-		break;
-	case EOperation::Fmod:
-		if (Types.bIsLWC)
-		{
-			OutResult.Code = Context.EmitExpression(Scope, Types.ResultType, TEXT("LWCFmod(%, %)"), InputValue[0], InputValue[1]);
-		}
-		else
-		{
-			OutResult.Code = Context.EmitExpression(Scope, Types.ResultType, TEXT("fmod(%, %)"), InputValue[0], InputValue[1]);
-		}
-		break;
-	case EOperation::PowPositiveClamped:
-		OutResult.Code = Context.EmitExpression(Scope, Types.ResultType, TEXT("PositiveClampedPow(%, %)"), InputValue[0], InputValue[1]);
-		break;
-	case EOperation::Atan2: OutResult.Code = Context.EmitExpression(Scope, Types.ResultType, TEXT("atan2(%, %)"), InputValue[0], InputValue[1]); break;
-	case EOperation::Atan2Fast: OutResult.Code = Context.EmitExpression(Scope, Types.ResultType, TEXT("atan2Fast(%, %)"), InputValue[0], InputValue[1]); break;
-	case EOperation::Min:
-		if (Types.bIsLWC)
-		{
-			OutResult.Code = Context.EmitExpression(Scope, Types.ResultType, TEXT("LWCMin(%, %)"), InputValue[0], InputValue[1]);
-		}
-		else
-		{
-			OutResult.Code = Context.EmitExpression(Scope, Types.ResultType, TEXT("min(%, %)"), InputValue[0], InputValue[1]);
-		}
-		break;
-	case EOperation::Max:
-		if (Types.bIsLWC)
-		{
-			OutResult.Code = Context.EmitExpression(Scope, Types.ResultType, TEXT("LWCMax(%, %)"), InputValue[0], InputValue[1]);
-		}
-		else
-		{
-			OutResult.Code = Context.EmitExpression(Scope, Types.ResultType, TEXT("max(%, %)"), InputValue[0], InputValue[1]);
-		}
-		break;
-	case EOperation::Less:
-		if (Types.bIsLWC)
-		{
-			OutResult.Code = Context.EmitExpression(Scope, Types.ResultType, TEXT("LWCLess(%, %)"), InputValue[0], InputValue[1]);
-		}
-		else
-		{
-			OutResult.Code = Context.EmitExpression(Scope, Types.ResultType, TEXT("(% < %)"), InputValue[0], InputValue[1]);
-		}
-		break;
-	case EOperation::Greater:
-		if (Types.bIsLWC)
-		{
-			OutResult.Code = Context.EmitExpression(Scope, Types.ResultType, TEXT("LWCGreater(%, %)"), InputValue[0], InputValue[1]);
-		}
-		else
-		{
-			OutResult.Code = Context.EmitExpression(Scope, Types.ResultType, TEXT("(% > %)"), InputValue[0], InputValue[1]);
-		}
-		break;
-	case EOperation::LessEqual:
-		if (Types.bIsLWC)
-		{
-			OutResult.Code = Context.EmitExpression(Scope, Types.ResultType, TEXT("LWCLessEqual(%, %)"), InputValue[0], InputValue[1]);
-		}
-		else
-		{
-			OutResult.Code = Context.EmitExpression(Scope, Types.ResultType, TEXT("(% <= %)"), InputValue[0], InputValue[1]);
-		}
-		break;
-	case EOperation::GreaterEqual:
-		if (Types.bIsLWC)
-		{
-			OutResult.Code = Context.EmitExpression(Scope, Types.ResultType, TEXT("LWCGreaterEqual(%, %)"), InputValue[0], InputValue[1]);
-		}
-		else
-		{
-			OutResult.Code = Context.EmitExpression(Scope, Types.ResultType, TEXT("(% >= %)"), InputValue[0], InputValue[1]);
-		}
-		break;
-	case EOperation::VecMulMatrix3:
-		if (Types.bIsLWC)
-		{
-			OutResult.Code = Context.EmitExpression(Scope, Types.ResultType, TEXT("LWCMultiply(%, %)"), InputValue[0], InputValue[1]);
-		}
-		else
-		{
-			OutResult.Code = Context.EmitExpression(Scope, Types.ResultType, TEXT("mul(%, (float3x3)%)"), InputValue[0], InputValue[1]);
-		}
-		break;
-	case EOperation::VecMulMatrix4:
-		if (Types.bIsLWC)
-		{
-			OutResult.Code = Context.EmitExpression(Scope, Types.ResultType, TEXT("LWCMultiply(%, %)"), InputValue[0], InputValue[1]);
-		}
-		else
-		{
-			OutResult.Code = Context.EmitExpression(Scope, Types.ResultType, TEXT("mul(%, %)"), InputValue[0], InputValue[1]);
-		}
-		break;
-	case EOperation::Matrix3MulVec:
-		OutResult.Code = Context.EmitExpression(Scope, Types.ResultType, TEXT("mul((float3x3)%, %)"), InputValue[0], InputValue[1]);
-		break;
-	case EOperation::Matrix4MulVec:
-		OutResult.Code = Context.EmitExpression(Scope, Types.ResultType, TEXT("mul(%, %)"), InputValue[0], InputValue[1]);
-		break;
-	default:
-		checkNoEntry();
-		break;
-	}
-}
-
-void FExpressionOperation::EmitValuePreshader(FEmitContext& Context, FEmitScope& Scope, const FRequestedType& RequestedType, FEmitValuePreshaderResult& OutResult) const
-{
-	const FOperationDescription OpDesc = GetOperationDescription(Op);
-	Shader::EValueType InputTypes[MaxInputs];
-	for (int32 Index = 0; Index < OpDesc.NumInputs; ++Index)
-	{
-		InputTypes[Index] = Inputs[Index]->GetType();
-	}
-	const Private::FOperationRequestedTypes RequestedTypes = Private::GetOperationRequestedTypes(Op, RequestedType);
-	const Private::FOperationTypes Types = Private::GetOperationTypes(Op, MakeArrayView(InputTypes, OpDesc.NumInputs));
-	check(OpDesc.PreshaderOpcode != Shader::EPreshaderOpcode::Nop);
-
-	for (int32 Index = 0; Index < OpDesc.NumInputs; ++Index)
-	{
-		Inputs[Index]->GetValuePreshader(Context, Scope, RequestedTypes.InputType[Index], OutResult.Preshader);
-	}
-
-	const int32 NumInputsToPop = OpDesc.NumInputs - 1;
-	if (NumInputsToPop > 0)
-	{
-		check(Context.PreshaderStackPosition >= NumInputsToPop);
-		Context.PreshaderStackPosition -= NumInputsToPop;
-	}
-
-	OutResult.Preshader.WriteOpcode(OpDesc.PreshaderOpcode);
-	OutResult.Type = Types.ResultType;
-}
-
 void FExpressionDerivative::ComputeAnalyticDerivatives(FTree& Tree, FExpressionDerivatives& OutResult) const
 {
 	// TODO
diff --git a/Engine/Source/Runtime/Engine/Private/HLSLTree/HLSLTreeOperation.cpp b/Engine/Source/Runtime/Engine/Private/HLSLTree/HLSLTreeOperation.cpp
new file mode 100644
--- /dev/null
+++ b/Engine/Source/Runtime/Engine/Private/HLSLTree/HLSLTreeOperation.cpp
@@ -0,0 +1,678 @@
+// Copyright Epic Games, Inc. All Rights Reserved.
+#include "HLSLTree/HLSLTreeCommon.h"
+#include "HLSLTree/HLSLTreeEmit.h"
+#include "Shader/Preshader.h"
+
+namespace UE
+{
+namespace HLSLTree
+{
+
+class FExpressionOperation : public FExpression
+{
+public:
+	FExpressionOperation(EOperation InOp, TConstArrayView<FExpression*> InInputs);
+
+	static constexpr int8 MaxInputs = 2;
+
+	EOperation Op;
+	FExpression* Inputs[MaxInputs];
+
+	virtual void ComputeAnalyticDerivatives(FTree& Tree, FExpressionDerivatives& OutResult) const override;
+	virtual FExpression* ComputePreviousFrame(FTree& Tree, const FRequestedType& RequestedType) const override;
+	virtual bool PrepareValue(FEmitContext& Context, FEmitScope& Scope, const FRequestedType& RequestedType, FPrepareValueResult& OutResult) const override;
+	virtual void EmitValueShader(FEmitContext& Context, FEmitScope& Scope, const FRequestedType& RequestedType, FEmitValueShaderResult& OutResult) const override;
+	virtual void EmitValuePreshader(FEmitContext& Context, FEmitScope& Scope, const FRequestedType& RequestedType, FEmitValuePreshaderResult& OutResult) const override;
+};
+
+FOperationDescription::FOperationDescription(const TCHAR* InName, const TCHAR* InOperator, int8 InNumInputs, Shader::EPreshaderOpcode InOpcode)
+	: Name(InName), Operator(InOperator), NumInputs(InNumInputs), PreshaderOpcode(InOpcode)
+{}
+
+FOperationDescription::FOperationDescription()
+	: Name(nullptr), Operator(nullptr), PreshaderOpcode(Shader::EPreshaderOpcode::Nop)
+{}
+
+
+FOperationDescription GetOperationDescription(EOperation Op)
+{
+	switch (Op)
+	{
+	case EOperation::None: return FOperationDescription(TEXT("None"), TEXT(""), 0, Shader::EPreshaderOpcode::Nop); break;
+
+	// Unary
+	case EOperation::Abs: return FOperationDescription(TEXT("Abs"), TEXT("abs"), 1, Shader::EPreshaderOpcode::Abs); break;
+	case EOperation::Neg: return FOperationDescription(TEXT("Neg"), TEXT("-"), 1, Shader::EPreshaderOpcode::Neg); break;
+	case EOperation::Rcp: return FOperationDescription(TEXT("Rcp"), TEXT("/"), 1, Shader::EPreshaderOpcode::Rcp); break;
+	case EOperation::Sqrt: return FOperationDescription(TEXT("Sqrt"), TEXT("sqrt"), 1, Shader::EPreshaderOpcode::Sqrt); break;
+	case EOperation::Rsqrt: return FOperationDescription(TEXT("Rsqrt"), TEXT("rsqrt"), 1, Shader::EPreshaderOpcode::Nop); break; // TODO
+	case EOperation::Log2: return FOperationDescription(TEXT("Log2"), TEXT("log2"), 1, Shader::EPreshaderOpcode::Log2); break;
+	case EOperation::Frac: return FOperationDescription(TEXT("Frac"), TEXT("frac"), 1, Shader::EPreshaderOpcode::Frac); break;
+	case EOperation::Floor: return FOperationDescription(TEXT("Floor"), TEXT("floor"), 1, Shader::EPreshaderOpcode::Floor); break;
+	case EOperation::Ceil: return FOperationDescription(TEXT("Ceil"), TEXT("ceil"), 1, Shader::EPreshaderOpcode::Ceil); break;
+	case EOperation::Round: return FOperationDescription(TEXT("Round"), TEXT("round"), 1, Shader::EPreshaderOpcode::Round); break;
+	case EOperation::Trunc: return FOperationDescription(TEXT("Trunc"), TEXT("trunc"), 1, Shader::EPreshaderOpcode::Trunc); break;
+	case EOperation::Saturate: return FOperationDescription(TEXT("Saturate"), TEXT("saturate"), 1, Shader::EPreshaderOpcode::Saturate); break;
+	case EOperation::Sign: return FOperationDescription(TEXT("Sign"), TEXT("sign"), 1, Shader::EPreshaderOpcode::Sign); break;
+	case EOperation::Length: return FOperationDescription(TEXT("Length"), TEXT("length"), 1, Shader::EPreshaderOpcode::Length); break;
+	case EOperation::Normalize: return FOperationDescription(TEXT("Normalize"), TEXT("normalize"), 1, Shader::EPreshaderOpcode::Normalize); break;
+	case EOperation::Sum: return FOperationDescription(TEXT("Sum"), TEXT("sum"), 1, Shader::EPreshaderOpcode::Nop); break; // TODO
+	case EOperation::Sin: return FOperationDescription(TEXT("Sin"), TEXT("sin"), 1, Shader::EPreshaderOpcode::Sin); break;
+	case EOperation::Cos: return FOperationDescription(TEXT("Cos"), TEXT("cos"), 1, Shader::EPreshaderOpcode::Cos); break;
+	case EOperation::Tan: return FOperationDescription(TEXT("Tan"), TEXT("tan"), 1, Shader::EPreshaderOpcode::Tan); break;
+	case EOperation::Asin: return FOperationDescription(TEXT("Asin"), TEXT("asin"), 1, Shader::EPreshaderOpcode::Asin); break;
+	case EOperation::AsinFast: return FOperationDescription(TEXT("AsinFast"), TEXT("asinFast"), 1, Shader::EPreshaderOpcode::Asin); break;
+	case EOperation::Acos: return FOperationDescription(TEXT("Acos"), TEXT("acos"), 1, Shader::EPreshaderOpcode::Acos); break;
+	case EOperation::AcosFast: return FOperationDescription(TEXT("AcosFast"), TEXT("acosFast"), 1, Shader::EPreshaderOpcode::Acos); break;
+	case EOperation::Atan: return FOperationDescription(TEXT("Atan"), TEXT("atan"), 1, Shader::EPreshaderOpcode::Atan); break;
+	case EOperation::AtanFast: return FOperationDescription(TEXT("AtanFast"), TEXT("atanFast"), 1, Shader::EPreshaderOpcode::Atan); break;
+
+	// Binary
+	case EOperation::Add: return FOperationDescription(TEXT("Add"), TEXT("+"), 2, Shader::EPreshaderOpcode::Add); break;
+	case EOperation::Sub: return FOperationDescription(TEXT("Subtract"), TEXT("-"), 2, Shader::EPreshaderOpcode::Sub); break;
+	case EOperation::Mul: return FOperationDescription(TEXT("Multiply"), TEXT("*"), 2, Shader::EPreshaderOpcode::Mul); break;
+	case EOperation::Div: return FOperationDescription(TEXT("Divide"), TEXT("/"), 2, Shader::EPreshaderOpcode::Div); break;
+	case EOperation::Fmod: return FOperationDescription(TEXT("Fmod"), TEXT("%"), 2, Shader::EPreshaderOpcode::Fmod); break;
+	case EOperation::PowPositiveClamped: return FOperationDescription(TEXT("PowPositiveClamped"), TEXT("PowPositiveClamped"), 2, Shader::EPreshaderOpcode::Nop); break;
+	case EOperation::Atan2: return FOperationDescription(TEXT("Atan2"), TEXT("atan2"), 2, Shader::EPreshaderOpcode::Atan2); break;
+	case EOperation::Atan2Fast: return FOperationDescription(TEXT("Atan2Fast"), TEXT("atan2Fast"), 2, Shader::EPreshaderOpcode::Atan2); break;
+	case EOperation::Min: return FOperationDescription(TEXT("Min"), TEXT("min"), 2, Shader::EPreshaderOpcode::Min); break;
+	case EOperation::Max: return FOperationDescription(TEXT("Max"), TEXT("max"), 2, Shader::EPreshaderOpcode::Max); break;
+	case EOperation::Less: return FOperationDescription(TEXT("Less"), TEXT("<"), 2, Shader::EPreshaderOpcode::Less); break;
+	case EOperation::Greater: return FOperationDescription(TEXT("Greater"), TEXT(">"), 2, Shader::EPreshaderOpcode::Greater); break;
+	case EOperation::LessEqual: return FOperationDescription(TEXT("LessEqual"), TEXT("<="), 2, Shader::EPreshaderOpcode::Nop); break;
+	case EOperation::GreaterEqual: return FOperationDescription(TEXT("GreaterEqual"), TEXT(">="), 2, Shader::EPreshaderOpcode::Nop); break;
+	case EOperation::VecMulMatrix3: return FOperationDescription(TEXT("VecMulMatrix3"), TEXT("mul"), 2, Shader::EPreshaderOpcode::Nop); break;
+	case EOperation::VecMulMatrix4: return FOperationDescription(TEXT("VecMulMatrix4"), TEXT("mul"), 2, Shader::EPreshaderOpcode::Nop); break;
+	case EOperation::Matrix3MulVec: return FOperationDescription(TEXT("Matrix3MulVec"), TEXT("mul"), 2, Shader::EPreshaderOpcode::Nop); break;
+	case EOperation::Matrix4MulVec: return FOperationDescription(TEXT("Matrix4MulVec"), TEXT("mul"), 2, Shader::EPreshaderOpcode::Nop); break;
+	default: checkNoEntry(); return FOperationDescription();
+	}
+}
+
+FExpression* FTree::NewUnaryOp(EOperation Op, FExpression* Input)
+{
+	FExpression* Inputs[1] = { Input };
+	return NewExpression<FExpressionOperation>(Op, Inputs);
+}
+
+FExpression* FTree::NewBinaryOp(EOperation Op, FExpression* Lhs, FExpression* Rhs)
+{
+	FExpression* Inputs[2] = { Lhs, Rhs };
+	return NewExpression<FExpressionOperation>(Op, Inputs);
+}
+
+FExpression* FTree::NewCross(FExpression* Lhs, FExpression* Rhs)
+{
+	//c_P[0] = v_A[1] * v_B[2] - v_A[2] * v_B[1];
+	//c_P[1] = -(v_A[0] * v_B[2] - v_A[2] * v_B[0]);
+	//c_P[2] = v_A[0] * v_B[1] - v_A[1] * v_B[0];
+	FExpression* Lhs0 = NewExpression<FExpressionSwizzle>(FSwizzleParameters(1, 0, 0), Lhs);
+	FExpression* Lhs1 = NewExpression<FExpressionSwizzle>(FSwizzleParameters(2, 2, 1), Lhs);
+	FExpression* Rhs0 = NewExpression<FExpressionSwizzle>(FSwizzleParameters(2, 2, 1), Rhs);
+	FExpression* Rhs1 = NewExpression<FExpressionSwizzle>(FSwizzleParameters(1, 0, 0), Rhs);
+	return NewSub(NewMul(NewMul(Lhs0, Rhs0), NewConstant(FVector3f(1.0f, -1.0f, 1.0f))), NewMul(Lhs1, Rhs1));
+}
+
+FExpressionOperation::FExpressionOperation(EOperation InOp, TConstArrayView<FExpression*> InInputs) : Op(InOp)
+{
+	const FOperationDescription OpDesc = GetOperationDescription(InOp);
+	check(OpDesc.NumInputs == InInputs.Num());
+	check(InInputs.Num() <= MaxInputs);
+
+	for (int32 i = 0; i < OpDesc.NumInputs; ++i)
+	{
+		Inputs[i] = InInputs[i];
+		check(Inputs[i]);
+	}
+}
+
+namespace Private
+{
+struct FOperationRequestedTypes
+{
+	FRequestedType InputType[FExpressionOperation::MaxInputs];
+	bool bIsMatrixOperation = false;
+};
+struct FOperationTypes
+{
+	Shader::EValueType InputType[FExpressionOperation::MaxInputs];
+	FPreparedType ResultType;
+	bool bIsLWC = false;
+};
+
+FOperationRequestedTypes GetOperationRequestedTypes(EOperation Op, const FRequestedType& RequestedType)
+{
+	const FOperationDescription OpDesc = GetOperationDescription(Op);
+	FOperationRequestedTypes Types;
+	for (int32 Index = 0; Index < OpDesc.NumInputs; ++Index)
+	{
+		Types.InputType[Index] = RequestedType;
+	}
+	switch (Op)
+	{
+	case EOperation::Length:
+	case EOperation::Normalize:
+	case EOperation::Sum:
+		Types.InputType[0] = ERequestedType::Vector4;
+		break;
+	case EOperation::VecMulMatrix3:
+		Types.bIsMatrixOperation = true;
+		Types.InputType[0] = ERequestedType::Vector3;
+		Types.InputType[1] = ERequestedType::Matrix4x4;
+		break;
+	case EOperation::VecMulMatrix4:
+		Types.bIsMatrixOperation = true;
+		Types.InputType[0] = ERequestedType::Vector3;
+		Types.InputType[1] = ERequestedType::Matrix4x4;
+		break;
+	case EOperation::Matrix3MulVec:
+	case EOperation::Matrix4MulVec:
+		// No LWC for transpose matrices
+		Types.bIsMatrixOperation = true;
+		Types.InputType[0] = ERequestedType::Matrix4x4;
+		Types.InputType[1] = ERequestedType::Vector3;
+		break;
+	default:
+		break;
+	}
+	return Types;
+}
+
+FOperationTypes GetOperationTypes(EOperation Op, TConstArrayView<FPreparedType> InputTypes)
+{
+	FOperationTypes Types;
+
+	if (Op == EOperation::VecMulMatrix3 ||
+		Op == EOperation::VecMulMatrix4 ||
+		Op == EOperation::Matrix3MulVec ||
+		Op == EOperation::Matrix4MulVec)
+	{
+		FPreparedComponent IntermediateComponent = InputTypes[0].GetMergedComponent();
+		Shader::EValueComponentType IntermediateComponentType = InputTypes[0].ValueComponentType;
+		for (int32 Index = 1; Index < InputTypes.Num(); ++Index)
+		{
+			IntermediateComponent = CombineComponents(IntermediateComponent, InputTypes[Index].GetMergedComponent());
+			IntermediateComponentType = Shader::CombineComponentTypes(IntermediateComponentType, InputTypes[Index].ValueComponentType);
+		}
+
+		switch (Op)
+		{
+		case EOperation::VecMulMatrix3:
+			// No LWC for matrix3
+			Types.InputType[0] = Shader::EValueType::Float3;
+			Types.InputType[1] = Shader::EValueType::Float4x4;
+			Types.ResultType = FPreparedType(Shader::EValueType::Float3, IntermediateComponent);
+			break;
+		case EOperation::VecMulMatrix4:
+			Types.InputType[0] = Shader::MakeValueType(IntermediateComponentType, 3);
+			Types.InputType[1] = Shader::MakeValueType(IntermediateComponentType, 16);
+			Types.ResultType = FPreparedType(Shader::MakeValueType(IntermediateComponentType, 3), IntermediateComponent);
+			break;
+		case EOperation::Matrix3MulVec:
+		case EOperation::Matrix4MulVec:
+			// No LWC for transpose matrices
+			Types.InputType[0] = Shader::EValueType::Float4x4;
+			Types.InputType[1] = Shader::EValueType::Float3;
+			Types.ResultType = FPreparedType(Shader::EValueType::Float3, IntermediateComponent);
+			break;
+		}
+	}
+	else
+	{
+		FPreparedType IntermediateType = InputTypes[0];
+		for (int32 Index = 1; Index < InputTypes.Num(); ++Index)
+		{
+			IntermediateType = MergePreparedTypes(IntermediateType, InputTypes[Index]);
+		}
+
+		const Shader::EValueType IntermediateValueType = IntermediateType.GetType();
+		for (int32 Index = 0; Index < InputTypes.Num(); ++Index)
+		{
+			Types.InputType[Index] = IntermediateValueType;
+		}
+		Types.ResultType = IntermediateType;
+		Types.bIsLWC = (IntermediateType.ValueComponentType == Shader::EValueComponentType::Double);
+
+		switch (Op)
+		{
+		case EOperation::Length:
+		case EOperation::Sum:
+			Types.ResultType = FPreparedType(Shader::MakeValueType(IntermediateType.ValueComponentType, 1), IntermediateType.GetMergedComponent());
+			break;
+		case EOperation::Normalize:
+		case EOperation::Rcp:
+		case EOperation::Sqrt:
+		case EOperation::Rsqrt:
+		case EOperation::Sign:
+		case EOperation::Tan:
+		case EOperation::Asin:
+		case EOperation::AsinFast:
+		case EOperation::Acos:
+		case EOperation::AcosFast:
+		case EOperation::Atan:
+		case EOperation::AtanFast:
+			Types.ResultType = MakeNonLWCType(IntermediateType);
+			break;
+		case EOperation::Saturate:
+		case EOperation::Frac:
+			for (int32 Index = 0; Index < Types.ResultType.PreparedComponents.Num(); ++Index)
+			{
+				Types.ResultType.SetComponentBounds(Index, Shader::FComponentBounds(Shader::EComponentBound::Zero, Shader::EComponentBound::One));
+			}
+			Types.ResultType = MakeNonLWCType(IntermediateType);
+			break;
+		case EOperation::Sin:
+		case EOperation::Cos:
+			for (int32 Index = 0; Index < Types.ResultType.PreparedComponents.Num(); ++Index)
+			{
+				Types.ResultType.SetComponentBounds(Index, Shader::FComponentBounds(Shader::EComponentBound::NegOne, Shader::EComponentBound::One));
+			}
+			Types.ResultType = MakeNonLWCType(IntermediateType);
+			break;
+		case EOperation::Log2:
+			// No LWC support yet
+			Types.InputType[0] = Shader::MakeNonLWCType(IntermediateValueType);
+			Types.ResultType = MakeNonLWCType(IntermediateType);
+			break;
+		case EOperation::Less:
+		case EOperation::Greater:
+		case EOperation::LessEqual:
+		case EOperation::GreaterEqual:
+			Types.ResultType.ValueComponentType = Shader::EValueComponentType::Bool;
+			break;
+		case EOperation::Fmod:
+			Types.InputType[1] = Shader::MakeNonLWCType(IntermediateValueType);
+			Types.ResultType = MakeNonLWCType(IntermediateType);
+			break;
+		case EOperation::PowPositiveClamped:
+		case EOperation::Atan2:
+		case EOperation::Atan2Fast:
+			// No LWC support yet
+			Types.InputType[0] = Types.InputType[1] = Shader::MakeNonLWCType(IntermediateValueType);
+			Types.ResultType = MakeNonLWCType(IntermediateType);
+			break;
+		case EOperation::Min:
+			for (int32 Index = 0; Index < Types.ResultType.PreparedComponents.Num(); ++Index)
+			{
+				Types.ResultType.SetComponentBounds(Index, Shader::MinBound(InputTypes[0].GetComponentBounds(Index), InputTypes[1].GetComponentBounds(Index)));
+			}
+			break;
+		case EOperation::Max:
+			for (int32 Index = 0; Index < Types.ResultType.PreparedComponents.Num(); ++Index)
+			{
+				Types.ResultType.SetComponentBounds(Index, Shader::MaxBound(InputTypes[0].GetComponentBounds(Index), InputTypes[1].GetComponentBounds(Index)));
+			}
+			break;
+		default:
+			break;
+		}
+	}
+	return Types;
+}
+} // namespace Private
+
+void FExpressionOperation::ComputeAnalyticDerivatives(FTree& Tree, FExpressionDerivatives& OutResult) const
+{
+	// Operations with constant derivatives
+	switch (Op)
+	{
+	case EOperation::Less:
+	case EOperation::Greater:
+	case EOperation::LessEqual:
+	case EOperation::GreaterEqual:
+	case EOperation::Floor:
+	case EOperation::Ceil:
+	case EOperation::Round:
+	case EOperation::Trunc:
+	case EOperation::Sign:
+		OutResult.ExpressionDdx = Tree.NewConstant(0.0f);
+		OutResult.ExpressionDdx = OutResult.ExpressionDdy;
+		break;
+	default:
+		break;
+	}
+
+	if (OutResult.IsValid())
+	{
+		return;
+	}
+
+	const FOperationDescription OpDesc = GetOperationDescription(Op);
+	FExpressionDerivatives InputDerivatives[MaxInputs];
+	for (int32 Index = 0; Index < OpDesc.NumInputs; ++Index)
+	{
+		InputDerivatives[Index] = Tree.GetAnalyticDerivatives(Inputs[Index]);
+		if (!InputDerivatives[Index].IsValid())
+		{
+			return;
+		}
+	}
+
+	switch (Op)
+	{
+	case EOperation::Neg:
+		OutResult.ExpressionDdx = Tree.NewNeg(InputDerivatives[0].ExpressionDdx);
+		OutResult.ExpressionDdy = Tree.NewNeg(InputDerivatives[0].ExpressionDdy);
+		break;
+	case EOperation::Rcp:
+	{
+		FExpression* Result = Tree.NewRcp(Inputs[0]);
+		FExpression* dFdA = Tree.NewNeg(Tree.NewMul(Result, Result));
+		OutResult.ExpressionDdx = Tree.NewMul(dFdA, InputDerivatives[0].ExpressionDdx);
+		OutResult.ExpressionDdy = Tree.NewMul(dFdA, InputDerivatives[0].ExpressionDdy);
+		break;
+	}
+	case EOperation::Sqrt:
+	{
+		FExpression* dFdA = Tree.NewMul(Tree.NewRsqrt(Tree.NewMax(Inputs[0], Tree.NewConstant(0.00001f))), Tree.NewConstant(0.5f));
+		OutResult.ExpressionDdx = Tree.NewMul(dFdA, InputDerivatives[0].ExpressionDdx);
+		OutResult.ExpressionDdy = Tree.NewMul(dFdA, InputDerivatives[0].ExpressionDdy);
+		break;
+	}
+	case EOperation::Rsqrt:
+	{
+		FExpression* dFdA = Tree.NewMul(Tree.NewMul(Tree.NewRsqrt(Inputs[0]), Tree.NewRcp(Inputs[0])), Tree.NewConstant(-0.5f));
+		OutResult.ExpressionDdx = Tree.NewMul(dFdA, InputDerivatives[0].ExpressionDdx);
+		OutResult.ExpressionDdy = Tree.NewMul(dFdA, InputDerivatives[0].ExpressionDdy);
+		break;
+	}
+	case EOperation::Sum:
+		OutResult.ExpressionDdx = Tree.NewSum(InputDerivatives[0].ExpressionDdx);
+		OutResult.ExpressionDdy = Tree.NewSum(InputDerivatives[0].ExpressionDdy);
+		break;
+	case EOperation::Frac:
+		OutResult = InputDerivatives[0];
+		break;
+	case EOperation::Sin:
+	{
+		FExpression* dFdA = Tree.NewCos(Inputs[0]);
+		OutResult.ExpressionDdx = Tree.NewMul(dFdA, InputDerivatives[0].ExpressionDdx);
+		OutResult.ExpressionDdy = Tree.NewMul(dFdA, InputDerivatives[0].ExpressionDdy);
+		break;
+	}
+	case EOperation::Cos:
+	{
+		FExpression* dFdA = Tree.NewNeg(Tree.NewSin(Inputs[0]));
+		OutResult.ExpressionDdx = Tree.NewMul(dFdA, InputDerivatives[0].ExpressionDdx);
+		OutResult.ExpressionDdy = Tree.NewMul(dFdA, InputDerivatives[0].ExpressionDdy);
+		break;
+	}
+	case EOperation::Tan:
+	{
+		FExpression* dFdA = Tree.NewRcp(Tree.NewPow2(Tree.NewCos(Inputs[0])));
+		OutResult.ExpressionDdx = Tree.NewMul(dFdA, InputDerivatives[0].ExpressionDdx);
+		OutResult.ExpressionDdy = Tree.NewMul(dFdA, InputDerivatives[0].ExpressionDdy);
+		break;
+	}
+	case EOperation::Asin:
+	case EOperation::AsinFast:
+	{
+		FExpression* dFdA = Tree.NewRsqrt(Tree.NewMax(Tree.NewSub(Tree.NewConstant(1.0f), Tree.NewPow2(Inputs[0])), Tree.NewConstant(0.00001f)));
+		OutResult.ExpressionDdx = Tree.NewMul(dFdA, InputDerivatives[0].ExpressionDdx);
+		OutResult.ExpressionDdy = Tree.NewMul(dFdA, InputDerivatives[0].ExpressionDdy);
+		break;
+	}
+	case EOperation::Acos:
+	case EOperation::AcosFast:
+	{
+		FExpression* dFdA = Tree.NewNeg(Tree.NewRsqrt(Tree.NewMax(Tree.NewSub(Tree.NewConstant(1.0f), Tree.NewPow2(Inputs[0])), Tree.NewConstant(0.00001f))));
+		OutResult.ExpressionDdx = Tree.NewMul(dFdA, InputDerivatives[0].ExpressionDdx);
+		OutResult.ExpressionDdy = Tree.NewMul(dFdA, InputDerivatives[0].ExpressionDdy);
+		break;
+	}
+	case EOperation::Atan:
+	case EOperation::AtanFast:
+	{
+		FExpression* dFdA = Tree.NewRcp(Tree.NewAdd(Tree.NewPow2(Inputs[0]), Tree.NewConstant(1.0f)));
+		OutResult.ExpressionDdx = Tree.NewMul(dFdA, InputDerivatives[0].ExpressionDdx);
+		OutResult.ExpressionDdy = Tree.NewMul(dFdA, InputDerivatives[0].ExpressionDdy);
+		break;
+	}
+	case EOperation::Atan2:
+	case EOperation::Atan2Fast:
+	{
+		FExpression* Denom = Tree.NewRcp(Tree.NewAdd(Tree.NewPow2(Inputs[0]), Tree.NewPow2(Inputs[1])));
+		FExpression* dFdA = Tree.NewMul(Inputs[1], Denom);
+		FExpression* dFdB = Tree.NewMul(Tree.NewNeg(Inputs[0]), Denom);
+		OutResult.ExpressionDdx = Tree.NewAdd(Tree.NewMul(dFdA, InputDerivatives[0].ExpressionDdx), Tree.NewMul(dFdB, InputDerivatives[1].ExpressionDdx));
+		OutResult.ExpressionDdy = Tree.NewAdd(Tree.NewMul(dFdA, InputDerivatives[0].ExpressionDdy), Tree.NewMul(dFdB, InputDerivatives[1].ExpressionDdy));
+		break;
+	}
+	case EOperation::Length:
+	case EOperation::Normalize:
+	case EOperation::Abs:
+	case EOperation::Saturate:
+	case EOperation::PowPositiveClamped:
+	case EOperation::Log2:
+		// TODO
+		break;
+	case EOperation::Add:
+		OutResult.ExpressionDdx = Tree.NewAdd(InputDerivatives[0].ExpressionDdx, InputDerivatives[1].ExpressionDdx);
+		OutResult.ExpressionDdy = Tree.NewAdd(InputDerivatives[0].ExpressionDdy, InputDerivatives[1].ExpressionDdy);
+		break;
+	case EOperation::Sub:
+		OutResult.ExpressionDdx = Tree.NewSub(InputDerivatives[0].ExpressionDdx, InputDerivatives[1].ExpressionDdx);
+		OutResult.ExpressionDdy = Tree.NewSub(InputDerivatives[0].ExpressionDdy, InputDerivatives[1].ExpressionDdy);
+		break;
+	case EOperation::Mul:
+		OutResult.ExpressionDdx = Tree.NewAdd(Tree.NewMul(InputDerivatives[0].ExpressionDdx, Inputs[1]), Tree.NewMul(InputDerivatives[1].ExpressionDdx, Inputs[0]));
+		OutResult.ExpressionDdy = Tree.NewAdd(Tree.NewMul(InputDerivatives[0].ExpressionDdy, Inputs[1]), Tree.NewMul(InputDerivatives[1].ExpressionDdy, Inputs[0]));
+		break;
+	case EOperation::Div:
+	{
+		FExpression* Denom = Tree.NewRcp(Tree.NewMul(Inputs[1], Inputs[1]));
+		FExpression* dFdA = Tree.NewMul(Inputs[1], Denom);
+		FExpression* dFdB = Tree.NewNeg(Tree.NewMul(Inputs[0], Denom));
+		OutResult.ExpressionDdx = Tree.NewAdd(Tree.NewMul(dFdA, InputDerivatives[0].ExpressionDdx), Tree.NewMul(dFdB, InputDerivatives[1].ExpressionDdx));
+		OutResult.ExpressionDdy = Tree.NewAdd(Tree.NewMul(dFdA, InputDerivatives[0].ExpressionDdy), Tree.NewMul(dFdB, InputDerivatives[1].ExpressionDdy));
+		break;
+	}
+	case EOperation::Fmod:
+		// Only valid when B derivatives are zero.
+		// We can't really do anything meaningful in the non-zero case.
+		OutResult = InputDerivatives[0];
+		break;
+	case EOperation::Min:
+	{
+		FExpression* Cond = Tree.NewLess(Inputs[0], Inputs[1]);
+		OutResult.ExpressionDdx = Tree.NewExpression<FExpressionSelect>(Cond, InputDerivatives[0].ExpressionDdx, InputDerivatives[1].ExpressionDdx);
+		OutResult.ExpressionDdy = Tree.NewExpression<FExpressionSelect>(Cond, InputDerivatives[0].ExpressionDdy, InputDerivatives[1].ExpressionDdy);
+		break;
+	}
+	case EOperation::Max:
+	{
+		FExpression* Cond = Tree.NewGreater(Inputs[0], Inputs[1]);
+		OutResult.ExpressionDdx = Tree.NewExpression<FExpressionSelect>(Cond, InputDerivatives[0].ExpressionDdx, InputDerivatives[1].ExpressionDdx);
+		OutResult.ExpressionDdy = Tree.NewExpression<FExpressionSelect>(Cond, InputDerivatives[0].ExpressionDdy, InputDerivatives[1].ExpressionDdy);
+		break;
+	}
+	case EOperation::VecMulMatrix3:
+	case EOperation::VecMulMatrix4:
+	case EOperation::Matrix3MulVec:
+	case EOperation::Matrix4MulVec:
+		// TODO
+		break;
+	default:
+		checkNoEntry();
+		break;
+	}
+}
+
+FExpression* FExpressionOperation::ComputePreviousFrame(FTree& Tree, const FRequestedType& RequestedType) const
+{
+	const Private::FOperationRequestedTypes RequestedTypes = Private::GetOperationRequestedTypes(Op, RequestedType);
+	const FOperationDescription OpDesc = GetOperationDescription(Op);
+	FExpression* PrevFrameInputs[MaxInputs];
+	for (int32 Index = 0; Index < OpDesc.NumInputs; ++Index)
+	{
+		PrevFrameInputs[Index] = Tree.GetPreviousFrame(Inputs[Index], RequestedTypes.InputType[Index]);
+	}
+
+	return Tree.NewExpression<FExpressionOperation>(Op, MakeArrayView(PrevFrameInputs, OpDesc.NumInputs));
+}
+
+bool FExpressionOperation::PrepareValue(FEmitContext& Context, FEmitScope& Scope, const FRequestedType& RequestedType, FPrepareValueResult& OutResult) const
+{
+	const FOperationDescription OpDesc = GetOperationDescription(Op);
+	const Private::FOperationRequestedTypes RequestedTypes = Private::GetOperationRequestedTypes(Op, RequestedType);
+
+	FPreparedType InputTypes[MaxInputs];
+	for (int32 Index = 0; Index < OpDesc.NumInputs; ++Index)
+	{
+		InputTypes[Index] = Context.PrepareExpression(Inputs[Index], Scope, RequestedTypes.InputType[Index]);
+		if (InputTypes[Index].IsVoid())
+		{
+			return false;
+		}
+
+		if (!InputTypes[Index].IsNumeric())
+		{
+			return Context.Errors->AddError(TEXT("Invalid arithmetic between non-numeric types"));
+		}
+	}
+
+	Private::FOperationTypes Types = Private::GetOperationTypes(Op, MakeArrayView(InputTypes, OpDesc.NumInputs));
+	if (OpDesc.PreshaderOpcode == Shader::EPreshaderOpcode::Nop)
+	{
+		// No preshader support
+		Types.ResultType.SetEvaluation(EExpressionEvaluation::Shader);
+	}
+
+	return OutResult.SetType(Context, RequestedType, Types.ResultType);
+}
+
+void FExpressionOperation::EmitValueShader(FEmitContext& Context, FEmitScope& Scope, const FRequestedType& RequestedType, FEmitValueShaderResult& OutResult) const
+{
+	const FOperationDescription OpDesc = GetOperationDescription(Op);
+	FPreparedType InputTypes[MaxInputs];
+	for (int32 Index = 0; Index < OpDesc.NumInputs; ++Index)
+	{
+		InputTypes[Index] = Inputs[Index]->GetPreparedType();
+	}
+	const Private::FOperationRequestedTypes RequestedTypes = Private::GetOperationRequestedTypes(Op, RequestedType);
+	const Private::FOperationTypes Types = Private::GetOperationTypes(Op, MakeArrayView(InputTypes, OpDesc.NumInputs));
+	FEmitShaderExpression* InputValue[MaxInputs] = { nullptr };
+	for (int32 Index = 0; Index < OpDesc.NumInputs; ++Index)
+	{
+		InputValue[Index] = Inputs[Index]->GetValueShader(Context, Scope, RequestedTypes.InputType[Index], Types.InputType[Index]);
+	}
+
+	const Shader::EValueType ResultType = Types.ResultType.GetType();
+	check(Shader::IsNumericType(ResultType));
+
+	switch (Op)
+	{
+	// Unary Ops
+	case EOperation::Abs: OutResult.Code = Context.EmitExpression(Scope, ResultType, Types.bIsLWC ? TEXT("LWCAbs(%)") : TEXT("abs(%)"), InputValue[0]); break;
+	case EOperation::Neg:
+		if (Types.bIsLWC)
+		{
+			OutResult.Code = Context.EmitExpression(Scope, ResultType, TEXT("LWCNegate(%)"), InputValue[0]);
+		}
+		else
+		{
+			OutResult.Code = Context.EmitInlineExpression(Scope, ResultType, TEXT("(-%)"), InputValue[0]);
+		}
+		break;
+	case EOperation::Rcp: OutResult.Code = Context.EmitExpression(Scope, ResultType, Types.bIsLWC ? TEXT("LWCRcp(%)") : TEXT("rcp(%)"), InputValue[0]); break;
+	case EOperation::Sqrt: OutResult.Code = Context.EmitExpression(Scope, ResultType, Types.bIsLWC ? TEXT("LWCSqrt(%)") : TEXT("sqrt(%)"), InputValue[0]); break;
+	case EOperation::Rsqrt: OutResult.Code = Context.EmitExpression(Scope, ResultType, Types.bIsLWC ? TEXT("LWCRsqrt(%)") : TEXT("rsqrt(%)"), InputValue[0]); break;
+	case EOperation::Log2: OutResult.Code = Context.EmitExpression(Scope, ResultType, TEXT("log2(%)"), InputValue[0]); break;
+	case EOperation::Frac: OutResult.Code = Context.EmitExpression(Scope, ResultType, Types.bIsLWC ? TEXT("LWCFrac(%)") : TEXT("frac(%)"), InputValue[0]); break;
+	case EOperation::Floor: OutResult.Code = Context.EmitExpression(Scope, ResultType, Types.bIsLWC ? TEXT("LWCFloor(%)") : TEXT("floor(%)"), InputValue[0]); break;
+	case EOperation::Ceil: OutResult.Code = Context.EmitExpression(Scope, ResultType, Types.bIsLWC ? TEXT("LWCCeil(%)") : TEXT("ceil(%)"), InputValue[0]); break;
+	case EOperation::Round: OutResult.Code = Context.EmitExpression(Scope, ResultType, Types.bIsLWC ? TEXT("LWCRound(%)") : TEXT("round(%)"), InputValue[0]); break;
+	case EOperation::Trunc: OutResult.Code = Context.EmitExpression(Scope, ResultType, Types.bIsLWC ? TEXT("LWCTrunc(%)") : TEXT("trunc(%)"), InputValue[0]); break;
+	case EOperation::Saturate: OutResult.Code = Context.EmitExpression(Scope, ResultType, Types.bIsLWC ? TEXT("LWCSaturate(%)") : TEXT("saturate(%)"), InputValue[0]); break;
+	case EOperation::Sign: OutResult.Code = Context.EmitExpression(Scope, ResultType, Types.bIsLWC ? TEXT("LWCSign(%)") : TEXT("sign(%)"), InputValue[0]); break;
+	case EOperation::Length: OutResult.Code = Context.EmitExpression(Scope, ResultType, Types.bIsLWC ? TEXT("LWCLength(%)") : TEXT("length(%)"), InputValue[0]); break;
+	case EOperation::Normalize: OutResult.Code = Context.EmitExpression(Scope, ResultType, Types.bIsLWC ? TEXT("LWCNormalize(%)") : TEXT("normalize(%)"), InputValue[0]); break;
+	case EOperation::Sum: OutResult.Code = Context.EmitExpression(Scope, ResultType, Types.bIsLWC ? TEXT("LWCVectorSum(%)") : TEXT("VectorSum(%)"), InputValue[0]); break;
+	case EOperation::Sin: OutResult.Code = Context.EmitExpression(Scope, ResultType, Types.bIsLWC ? TEXT("LWCSin(%)") : TEXT("sin(%)"), InputValue[0]); break;
+	case EOperation::Cos: OutResult.Code = Context.EmitExpression(Scope, ResultType, Types.bIsLWC ? TEXT("LWCCos(%)") : TEXT("cos(%)"), InputValue[0]); break;
+	case EOperation::Tan: OutResult.Code = Context.EmitExpression(Scope, ResultType, Types.bIsLWC ? TEXT("LWCTan(%)") : TEXT("tan(%)"), InputValue[0]); break;
+	case EOperation::Asin: OutResult.Code = Context.EmitExpression(Scope, ResultType, Types.bIsLWC ? TEXT("LWCAsin(%)") : TEXT("asin(%)"), InputValue[0]); break;
+	case EOperation::AsinFast: OutResult.Code = Context.EmitExpression(Scope, ResultType, Types.bIsLWC ? TEXT("LWCAsin(%)") : TEXT("asinFast(%)"), InputValue[0]); break;
+	case EOperation::Acos: OutResult.Code = Context.EmitExpression(Scope, ResultType, Types.bIsLWC ? TEXT("LWCAcos(%)") : TEXT("acos(%)"), InputValue[0]); break;
+	case EOperation::AcosFast: OutResult.Code = Context.EmitExpression(Scope, ResultType, Types.bIsLWC ? TEXT("LWCAcos(%)") : TEXT("acosFast(%)"), InputValue[0]); break;
+	case EOperation::Atan: OutResult.Code = Context.EmitExpression(Scope, ResultType, Types.bIsLWC ? TEXT("LWCAtan(%)") : TEXT("atan(%)"), InputValue[0]); break;
+	case EOperation::AtanFast: OutResult.Code = Context.EmitExpression(Scope, ResultType, Types.bIsLWC ? TEXT("LWCAtan(%)") : TEXT("atanFast(%)"), InputValue[0]); break;
+	
+	// Binary Ops
+	case EOperation::Add: OutResult.Code = Context.EmitExpression(Scope, ResultType, Types.bIsLWC ? TEXT("LWCAdd(%, %)") : TEXT("(% + %)"), InputValue[0], InputValue[1]); break;
+	case EOperation::Sub: OutResult.Code = Context.EmitExpression(Scope, ResultType, Types.bIsLWC ? TEXT("LWCSubtract(%, %)") : TEXT("(% - %)"), InputValue[0], InputValue[1]); break;
+	case EOperation::Mul: OutResult.Code = Context.EmitExpression(Scope, ResultType, Types.bIsLWC ? TEXT("LWCMultiply(%, %)") : TEXT("(% * %)"), InputValue[0], InputValue[1]); break;
+	case EOperation::Div: OutResult.Code = Context.EmitExpression(Scope, ResultType, Types.bIsLWC ? TEXT("LWCDivide(%, %)") : TEXT("(% / %)"), InputValue[0], InputValue[1]); break;
+	case EOperation::Fmod: OutResult.Code = Context.EmitExpression(Scope, ResultType, Types.bIsLWC ? TEXT("LWCFmod(%, %)") : TEXT("fmod(%, %)"), InputValue[0], InputValue[1]); break;
+	case EOperation::PowPositiveClamped: OutResult.Code = Context.EmitExpression(Scope, ResultType, TEXT("PositiveClampedPow(%, %)"), InputValue[0], InputValue[1]); break;
+	case EOperation::Atan2: OutResult.Code = Context.EmitExpression(Scope, ResultType, TEXT("atan2(%, %)"), InputValue[0], InputValue[1]); break;
+	case EOperation::Atan2Fast: OutResult.Code = Context.EmitExpression(Scope, ResultType, TEXT("atan2Fast(%, %)"), InputValue[0], InputValue[1]); break;
+	case EOperation::Min: OutResult.Code = Context.EmitExpression(Scope, ResultType, Types.bIsLWC ? TEXT("LWCMin(%, %)") : TEXT("min(%, %)"), InputValue[0], InputValue[1]); break;
+	case EOperation::Max: OutResult.Code = Context.EmitExpression(Scope, ResultType, Types.bIsLWC ? TEXT("LWCMax(%, %)") : TEXT("max(%, %)"), InputValue[0], InputValue[1]); break;
+	case EOperation::Less: OutResult.Code = Context.EmitExpression(Scope, ResultType, Types.bIsLWC ? TEXT("LWCLess(%, %)") : TEXT("(% < %)"), InputValue[0], InputValue[1]); break;
+	case EOperation::Greater: OutResult.Code = Context.EmitExpression(Scope, ResultType, Types.bIsLWC ? TEXT("LWCGreater(%, %)") : TEXT("(% > %)"), InputValue[0], InputValue[1]); break;
+	case EOperation::LessEqual: OutResult.Code = Context.EmitExpression(Scope, ResultType, Types.bIsLWC ? TEXT("LWCLessEqual(%, %)") : TEXT("(% <= %)"), InputValue[0], InputValue[1]); break;
+	case EOperation::GreaterEqual: OutResult.Code = Context.EmitExpression(Scope, ResultType, Types.bIsLWC ? TEXT("LWCGreaterEqual(%, %)") : TEXT("(% >= %)"), InputValue[0], InputValue[1]); break;
+	case EOperation::VecMulMatrix3:
+		if (Types.bIsLWC)
+		{
+			OutResult.Code = Context.EmitExpression(Scope, ResultType, TEXT("LWCMultiply(%, %)"), InputValue[0], InputValue[1]);
+		}
+		else
+		{
+			OutResult.Code = Context.EmitExpression(Scope, ResultType, TEXT("mul(%, (float3x3)%)"), InputValue[0], InputValue[1]);
+		}
+		break;
+	case EOperation::VecMulMatrix4:
+		if (Types.bIsLWC)
+		{
+			OutResult.Code = Context.EmitExpression(Scope, ResultType, TEXT("LWCMultiply(%, %)"), InputValue[0], InputValue[1]);
+		}
+		else
+		{
+			OutResult.Code = Context.EmitExpression(Scope, ResultType, TEXT("mul(%, %)"), InputValue[0], InputValue[1]);
+		}
+		break;
+	case EOperation::Matrix3MulVec:
+		OutResult.Code = Context.EmitExpression(Scope, ResultType, TEXT("mul((float3x3)%, %)"), InputValue[0], InputValue[1]);
+		break;
+	case EOperation::Matrix4MulVec:
+		OutResult.Code = Context.EmitExpression(Scope, ResultType, TEXT("mul(%, %)"), InputValue[0], InputValue[1]);
+		break;
+	default:
+		checkNoEntry();
+		break;
+	}
+}
+
+void FExpressionOperation::EmitValuePreshader(FEmitContext& Context, FEmitScope& Scope, const FRequestedType& RequestedType, FEmitValuePreshaderResult& OutResult) const
+{
+	const FOperationDescription OpDesc = GetOperationDescription(Op);
+	FPreparedType InputTypes[MaxInputs];
+	for (int32 Index = 0; Index < OpDesc.NumInputs; ++Index)
+	{
+		InputTypes[Index] = Inputs[Index]->GetPreparedType();
+	}
+	const Private::FOperationRequestedTypes RequestedTypes = Private::GetOperationRequestedTypes(Op, RequestedType);
+	const Private::FOperationTypes Types = Private::GetOperationTypes(Op, MakeArrayView(InputTypes, OpDesc.NumInputs));
+	check(OpDesc.PreshaderOpcode != Shader::EPreshaderOpcode::Nop);
+
+	for (int32 Index = 0; Index < OpDesc.NumInputs; ++Index)
+	{
+		Inputs[Index]->GetValuePreshader(Context, Scope, RequestedTypes.InputType[Index], OutResult.Preshader);
+	}
+
+	const int32 NumInputsToPop = OpDesc.NumInputs - 1;
+	if (NumInputsToPop > 0)
+	{
+		check(Context.PreshaderStackPosition >= NumInputsToPop);
+		Context.PreshaderStackPosition -= NumInputsToPop;
+	}
+
+	OutResult.Preshader.WriteOpcode(OpDesc.PreshaderOpcode);
+	OutResult.Type = Types.ResultType.GetType();
+}
+
+} // namespace HLSLTree
+} // namespace UE
diff --git a/Engine/Source/Runtime/Engine/Private/Materials/MaterialExpressionHLSL.cpp b/Engine/Source/Runtime/Engine/Private/Materials/MaterialExpressionHLSL.cpp
--- a/Engine/Source/Runtime/Engine/Private/Materials/MaterialExpressionHLSL.cpp
+++ b/Engine/Source/Runtime/Engine/Private/Materials/MaterialExpressionHLSL.cpp
@@ -1130,7 +1130,7 @@ bool UMaterialExpressionMax::GenerateHLSLExpression(FMaterialHLSLGenerator& Gene
 bool UMaterialExpressionClamp::GenerateHLSLExpression(FMaterialHLSLGenerator& Generator, UE::HLSLTree::FScope& Scope, int32 OutputIndex, UE::HLSLTree::FExpression*& OutExpression)
 {
 	UE::HLSLTree::FExpression* ExpressionMin = Min.AcquireHLSLExpressionOrConstant(Generator, Scope, MinDefault);
-	UE::HLSLTree::FExpression* ExpressionMax = Min.AcquireHLSLExpressionOrConstant(Generator, Scope, MaxDefault);
+	UE::HLSLTree::FExpression* ExpressionMax = Max.AcquireHLSLExpressionOrConstant(Generator, Scope, MaxDefault);
 	UE::HLSLTree::FExpression* ExpressionInput = Input.AcquireHLSLExpression(Generator, Scope);
 	if (!ExpressionMin || !ExpressionMax || !ExpressionInput)
 	{
diff --git a/Engine/Source/Runtime/Engine/Public/HLSLTree/HLSLTree.h b/Engine/Source/Runtime/Engine/Public/HLSLTree/HLSLTree.h
--- a/Engine/Source/Runtime/Engine/Public/HLSLTree/HLSLTree.h
+++ b/Engine/Source/Runtime/Engine/Public/HLSLTree/HLSLTree.h
@@ -268,12 +268,12 @@ public:
 	FPreparedType() = default;
 	FPreparedType(Shader::EValueComponentType InComponentType) : ValueComponentType(InComponentType) {}
 	FPreparedType(const Shader::FStructType* InStructType) : StructType(InStructType) {}
-	FPreparedType(const Shader::FType& InType, EExpressionEvaluation InEvalution = EExpressionEvaluation::None);
+	FPreparedType(Shader::EValueType InType, const FPreparedComponent& InComponent = FPreparedComponent());
+	FPreparedType(const Shader::FType& InType, const FPreparedComponent& InComponent = FPreparedComponent());
 
 	void SetEvaluation(EExpressionEvaluation Evaluation);
 	void MergeEvaluation(EExpressionEvaluation Evaluation);
 	void SetLoopEvaluation(FEmitScope& Scope, const FRequestedType& RequestedType);
-	void UpdateBounds(const FRequestedType& RequestedType, Shader::FComponentBounds Bounds);
 
 	void SetField(const Shader::FStructField* Field, const FPreparedType& FieldType);
 	FPreparedType GetFieldType(const Shader::FStructField* Field) const;
@@ -291,10 +291,13 @@ public:
 	EExpressionEvaluation GetEvaluation(const FEmitScope& Scope) const;
 	EExpressionEvaluation GetEvaluation(const FEmitScope& Scope, const FRequestedType& RequestedType) const;
 	EExpressionEvaluation GetFieldEvaluation(const FEmitScope& Scope, const FRequestedType& RequestedType, int32 ComponentIndex, int32 NumComponents) const;
+	Shader::FComponentBounds GetComponentBounds(int32 Index) const;
 	Shader::FComponentBounds GetBounds(const FRequestedType& RequestedType) const;
+	FPreparedComponent GetMergedComponent() const;
 	FPreparedComponent GetComponent(int32 Index) const;
 
 	void SetComponent(int32 Index, const FPreparedComponent& InComponent);
+	void SetComponentBounds(int32 Index, const Shader::FComponentBounds Bounds);
 	void MergeComponent(int32 Index, const FPreparedComponent& InComponent);
 
 	/** Unlike FRequestedType, one of these should be set */
@@ -318,6 +321,7 @@ inline bool operator!=(const FPreparedType& Lhs, const FPreparedType& Rhs)
 }
 
 FPreparedType MergePreparedTypes(const FPreparedType& Lhs, const FPreparedType& Rhs);
+FPreparedType MakeNonLWCType(const FPreparedType& Type);
 
 class FPrepareValueResult
 {
@@ -563,6 +567,8 @@ public:
 	FExpression* NewSaturate(FExpression* Input) { return NewUnaryOp(EOperation::Saturate, Input); }
 	FExpression* NewSum(FExpression* Input) { return NewUnaryOp(EOperation::Sum, Input); }
 	FExpression* NewRcp(FExpression* Input) { return NewUnaryOp(EOperation::Rcp, Input); }
+	FExpression* NewSqrt(FExpression* Input) { return NewUnaryOp(EOperation::Sqrt, Input); }
+	FExpression* NewRsqrt(FExpression* Input) { return NewUnaryOp(EOperation::Rsqrt, Input); }
 	FExpression* NewFrac(FExpression* Input) { return NewUnaryOp(EOperation::Frac, Input); }
 	FExpression* NewLength(FExpression* Input) { return NewUnaryOp(EOperation::Length, Input); }
 	FExpression* NewNormalize(FExpression* Input) { return NewUnaryOp(EOperation::Normalize, Input); }
@@ -581,6 +587,7 @@ public:
 	FExpression* NewLessEqual(FExpression* Lhs, FExpression* Rhs) { return NewBinaryOp(EOperation::LessEqual, Lhs, Rhs); }
 	FExpression* NewGreaterEqual(FExpression* Lhs, FExpression* Rhs) { return NewBinaryOp(EOperation::GreaterEqual, Lhs, Rhs); }
 
+	FExpression* NewPow2(FExpression* Input) { return NewMul(Input, Input); }
 	FExpression* NewCross(FExpression* Lhs, FExpression* Rhs);
 	FExpression* NewDot(FExpression* Lhs, FExpression* Rhs) { return NewSum(NewMul(Lhs, Rhs)); }
 	FExpression* NewLerp(FExpression* A, FExpression* B, FExpression* T) { return NewAdd(A, NewMul(NewSub(B, A), T)); }
diff --git a/Engine/Source/Runtime/Engine/Public/HLSLTree/HLSLTreeCommon.h b/Engine/Source/Runtime/Engine/Public/HLSLTree/HLSLTreeCommon.h
--- a/Engine/Source/Runtime/Engine/Public/HLSLTree/HLSLTreeCommon.h
+++ b/Engine/Source/Runtime/Engine/Public/HLSLTree/HLSLTreeCommon.h
@@ -364,23 +364,6 @@ public:
 	virtual void EmitValuePreshader(FEmitContext& Context, FEmitScope& Scope, const FRequestedType& RequestedType, FEmitValuePreshaderResult& OutResult) const override;
 };
 
-class FExpressionOperation : public FExpression
-{
-public:
-	FExpressionOperation(EOperation InOp, TConstArrayView<FExpression*> InInputs);
-
-	static constexpr int8 MaxInputs = 2;
-
-	EOperation Op;
-	FExpression* Inputs[MaxInputs];
-
-	virtual void ComputeAnalyticDerivatives(FTree& Tree, FExpressionDerivatives& OutResult) const override;
-	virtual FExpression* ComputePreviousFrame(FTree& Tree, const FRequestedType& RequestedType) const override;
-	virtual bool PrepareValue(FEmitContext& Context, FEmitScope& Scope, const FRequestedType& RequestedType, FPrepareValueResult& OutResult) const override;
-	virtual void EmitValueShader(FEmitContext& Context, FEmitScope& Scope, const FRequestedType& RequestedType, FEmitValueShaderResult& OutResult) const override;
-	virtual void EmitValuePreshader(FEmitContext& Context, FEmitScope& Scope, const FRequestedType& RequestedType, FEmitValuePreshaderResult& OutResult) const override;
-};
-
 class FExpressionDerivative : public FExpression
 {
 public:
diff --git a/Engine/Source/Runtime/Engine/Public/HLSLTree/HLSLTreeTypes.h b/Engine/Source/Runtime/Engine/Public/HLSLTree/HLSLTreeTypes.h
--- a/Engine/Source/Runtime/Engine/Public/HLSLTree/HLSLTreeTypes.h
+++ b/Engine/Source/Runtime/Engine/Public/HLSLTree/HLSLTreeTypes.h
@@ -70,6 +70,7 @@ enum class EOperation : uint8
 	Neg,
 	Rcp,
 	Sqrt,
+	Rsqrt,
 	Log2,
 	Frac,
 	Floor,
