Make sure old imported skeletalmesh (pre 4.24) have a "Regenerate" button verify the ddc key against the current state. This will ensure user apply there changes before quiting the persona editor.

#rnx

diff --git a/Engine/Source/Editor/Persona/Private/PersonaMeshDetails.cpp b/Engine/Source/Editor/Persona/Private/PersonaMeshDetails.cpp
--- a/Engine/Source/Editor/Persona/Private/PersonaMeshDetails.cpp
+++ b/Engine/Source/Editor/Persona/Private/PersonaMeshDetails.cpp
@@ -267,11 +267,6 @@ private:
 
 	bool IsNeedApplyLODChange() const
 	{
-		if (!BuildAvailable)
-		{
-			return true;
-		}
-
 		TSharedPtr<IPersonaToolkit> SharedToolkit = PersonaToolkit.Pin();
 		if (SharedToolkit.IsValid())
 		{
diff --git a/Engine/Source/Editor/SkeletalMeshEditor/Private/SkeletalMeshEditor.cpp b/Engine/Source/Editor/SkeletalMeshEditor/Private/SkeletalMeshEditor.cpp
--- a/Engine/Source/Editor/SkeletalMeshEditor/Private/SkeletalMeshEditor.cpp
+++ b/Engine/Source/Editor/SkeletalMeshEditor/Private/SkeletalMeshEditor.cpp
@@ -149,22 +149,6 @@ bool FSkeletalMeshEditor::OnRequestClose()
 				continue;
 			}
 			
-			//Do not prevent exiting if we are not using the skeletal mesh build workflow
-			if (!SkeletalMesh->IsLODImportedDataBuildAvailable(LODIndex))
-			{
-				if (!LODInfo->bHasBeenSimplified && !SkeletalMesh->IsReductionActive(LODIndex))
-				{
-					continue;
-				}
-				//Do not prevent exit if the generated LOD is base on a LODModel not using the skeletalmesh build workflow
-				int32 ReduceBaseLOD = LODInfo->ReductionSettings.BaseLOD;
-				if(!IsReductionParentBaseLODUseSkeletalMeshBuildWorkflow(SkeletalMesh, ReduceBaseLOD))
-				{
-					continue;
-				}
-				
-			}
-
 			bool bValidLODSettings = false;
 			if (SkeletalMesh->GetLODSettings() != nullptr)
 			{
diff --git a/Engine/Source/Runtime/Engine/Private/SkeletalMeshRenderData.cpp b/Engine/Source/Runtime/Engine/Private/SkeletalMeshRenderData.cpp
--- a/Engine/Source/Runtime/Engine/Private/SkeletalMeshRenderData.cpp
+++ b/Engine/Source/Runtime/Engine/Private/SkeletalMeshRenderData.cpp
@@ -205,24 +205,33 @@ FString BuildSkeletalMeshDerivedDataKey(const ITargetPlatform* TargetPlatform, U
 		bool bUseFullPrecisionUVs = BaseLODInfo ? BaseLODInfo->BuildSettings.bUseFullPrecisionUVs : false;
 		KeySuffix += SkelMesh->GetImportedModel()->GetIdString();
 		KeySuffix += (bUseFullPrecisionUVs || !GVertexElementTypeSupport.IsSupported(VET_Half2)) ? "1" : "0";
+		
+		//Dummy call to update the FSkeletalMeshLODModel::BuildStringID members.
+		SkelMesh->GetImportedModel()->GetLODModelIdString();
+		
+		//Dummy call to update the LODInfo::BuildGUID members.
+		{
+			FString TmpString;
+			SerializeLODInfoForDDC(SkelMesh, TmpString);
+		}
 	}
 	else
 	{
-		FString tmpDebugString;
+		FString TmpPartialKeySuffix;
 		//Synchronize the user data that are part of the key
 		SkelMesh->GetImportedModel()->SyncronizeLODUserSectionsData();
-		tmpDebugString = SkelMesh->GetImportedModel()->GetIdString();
-		KeySuffix += tmpDebugString;
-		tmpDebugString = SkelMesh->GetImportedModel()->GetLODModelIdString();
-		KeySuffix += tmpDebugString;
+		TmpPartialKeySuffix = SkelMesh->GetImportedModel()->GetIdString();
+		KeySuffix += TmpPartialKeySuffix;
+		TmpPartialKeySuffix = SkelMesh->GetImportedModel()->GetLODModelIdString();
+		KeySuffix += TmpPartialKeySuffix;
 		
 		//Add the max gpu bone per section
 		const int32 MaxGPUSkinBones = FGPUBaseSkinVertexFactory::GetMaxGPUSkinBones(TargetPlatform);
 		KeySuffix += FString::FromInt(MaxGPUSkinBones);
 
-		tmpDebugString = TEXT("");
-		SerializeLODInfoForDDC(SkelMesh, tmpDebugString);
-		KeySuffix += tmpDebugString;
+		TmpPartialKeySuffix = TEXT("");
+		SerializeLODInfoForDDC(SkelMesh, TmpPartialKeySuffix);
+		KeySuffix += TmpPartialKeySuffix;
 	}
 
 	KeySuffix += SkelMesh->GetHasVertexColors() ? "1" : "0";
