Fix Game Thread and Render Thread access issues when performing shader compilation using FMaterial.

- SubmitCompileJobs was assumed to be render thread only.  Make this explicit in the function name and add a render thread check.
- Add SubmitCompileJobs_GameThread to perform shader compilation submitted from the Game Thread safe.
- Keep a ref to a GameThreadPendingCompilerEnvironment because the one on the GameThreadShaderMap goes away after we clone.
- Add render thread check to GetRuntimeVirtualTextureOutputAttibuteMask_RenderThread
- Move simple enum types from ShaderCompiler.h to ShaderCompilerCore.h so we can access EShaderCompileJobPriority in MaterialShared.h
- Fixup incorrect usage SubmitCompileJobs that was submitted from the Game Thread using the Render Thread maps.

diff --git a/Engine/Source/Editor/MaterialEditor/Private/MaterialEditingLibrary.cpp b/Engine/Source/Editor/MaterialEditor/Private/MaterialEditingLibrary.cpp
--- a/Engine/Source/Editor/MaterialEditor/Private/MaterialEditingLibrary.cpp
+++ b/Engine/Source/Editor/MaterialEditor/Private/MaterialEditingLibrary.cpp
@@ -1218,7 +1218,7 @@ FMaterialStatistics UMaterialEditingLibrary::GetStatistics(class UMaterialInterf
 	{
 		if (!Resource->IsGameThreadShaderMapComplete())
 		{
-			Resource->SubmitCompileJobs(EShaderCompileJobPriority::High);
+			Resource->SubmitCompileJobs_GameThread(EShaderCompileJobPriority::High);
 		}
 		Resource->FinishCompilation();
 
diff --git a/Engine/Source/Editor/UnrealEd/Private/ObjectTools.cpp b/Engine/Source/Editor/UnrealEd/Private/ObjectTools.cpp
--- a/Engine/Source/Editor/UnrealEd/Private/ObjectTools.cpp
+++ b/Engine/Source/Editor/UnrealEd/Private/ObjectTools.cpp
@@ -4759,7 +4759,7 @@ namespace ThumbnailTools
 				{
 					if (!CurrentResource->IsGameThreadShaderMapComplete())
 					{
-						CurrentResource->SubmitCompileJobs(EShaderCompileJobPriority::High);
+						CurrentResource->SubmitCompileJobs_GameThread(EShaderCompileJobPriority::High);
 					}
 					CurrentResource->FinishCompilation();
 				}
diff --git a/Engine/Source/Runtime/Engine/Private/Materials/MaterialShared.cpp b/Engine/Source/Runtime/Engine/Private/Materials/MaterialShared.cpp
--- a/Engine/Source/Runtime/Engine/Private/Materials/MaterialShared.cpp
+++ b/Engine/Source/Runtime/Engine/Private/Materials/MaterialShared.cpp
@@ -943,6 +943,7 @@ bool FMaterial::UsesCustomDepthStencil_GameThread() const
 
 uint8 FMaterial::GetRuntimeVirtualTextureOutputAttibuteMask_RenderThread() const
 {
+	check(IsInParallelRenderingThread());
 	return RenderingThreadShaderMap ? RenderingThreadShaderMap->GetRuntimeVirtualTextureOutputAttributeMask() : 0;
 }
 
@@ -1001,6 +1002,7 @@ void FMaterial::SetInlineShaderMap(FMaterialShaderMap* InMaterialShaderMap)
 	// SetInlineShaderMap is called during PostLoad(), before given UMaterial(Instance) is fully initialized
 	// Can't check for completeness yet
 	bGameThreadShaderMapIsComplete = false;
+	GameThreadShaderMapSubmittedPriority = EShaderCompileJobPriority::None;
 
 	TRefCountPtr<FMaterial> Material = this;
 	TRefCountPtr<FMaterialShaderMap> ShaderMap = InMaterialShaderMap;
@@ -1024,6 +1026,9 @@ void FMaterial::SetCompilingShaderMap(FMaterialShaderMap* InMaterialShaderMap)
 		check(GameThreadCompilingShaderMapId != 0u);
 		InMaterialShaderMap->AddCompilingDependency(this);
 
+		GameThreadPendingCompilerEnvironment = InMaterialShaderMap->GetPendingCompilerEnvironment();
+		GameThreadShaderMapSubmittedPriority = EShaderCompileJobPriority::None;
+
 		TRefCountPtr<FMaterial> Material = this;
 		TRefCountPtr<FSharedShaderCompilerEnvironment> PendingCompilerEnvironment = InMaterialShaderMap->GetPendingCompilerEnvironment();
 		ENQUEUE_RENDER_COMMAND(SetCompilingShaderMap)([Material = MoveTemp(Material), CompilingShaderMapId, PendingCompilerEnvironment = MoveTemp(PendingCompilerEnvironment)](FRHICommandListImmediate& RHICmdList) mutable
@@ -2434,7 +2439,7 @@ bool FMaterial::CacheShaders(const FMaterialShaderMapId& ShaderMapId, EShaderPla
 			(PrecompileMode != EMaterialShaderPrecompileMode::None))
 		{
 			// Submit the remaining shaders in the map for compilation.
-			SubmitCompileJobs(EShaderCompileJobPriority::High);
+			SubmitCompileJobs_GameThread(EShaderCompileJobPriority::High);
 		}
 		else
 		{
@@ -3523,8 +3528,28 @@ void FMaterialRenderProxy::ReleaseResource()
 	}
 }
 
-void FMaterial::SubmitCompileJobs(EShaderCompileJobPriority Priority) const
+void FMaterial::SubmitCompileJobs_GameThread(EShaderCompileJobPriority Priority)
 {
+	check(IsInGameThread());
+
+	if (GameThreadCompilingShaderMapId != 0u && GameThreadShaderMap)
+	{
+		const EShaderCompileJobPriority SubmittedPriority = GameThreadShaderMapSubmittedPriority;
+
+		// To avoid as much useless work as possible, we make sure to submit our compile jobs only once per priority upgrade.
+		if (GameThreadShaderMapSubmittedPriority == EShaderCompileJobPriority::None || Priority > SubmittedPriority)
+		{
+			check(GameThreadPendingCompilerEnvironment.IsValid());
+
+			GameThreadShaderMapSubmittedPriority = Priority;
+			GameThreadShaderMap->SubmitCompileJobs(GameThreadCompilingShaderMapId, this, GameThreadPendingCompilerEnvironment, Priority);
+		}
+	}
+}
+
+void FMaterial::SubmitCompileJobs_RenderThread(EShaderCompileJobPriority Priority) const
+{
+	check(IsInParallelRenderingThread());
 	if (RenderingThreadCompilingShaderMapId != 0u && RenderingThreadShaderMap)
 	{
 		// std::atomic don't support enum class, so we have to make sure our cast assumptions are respected.
@@ -3558,7 +3583,7 @@ const FMaterial& FMaterialRenderProxy::GetMaterialWithFallback(ERHIFeatureLevel:
 
 		if (BaseMaterial)
 		{
-			BaseMaterial->SubmitCompileJobs(EShaderCompileJobPriority::Normal);
+			BaseMaterial->SubmitCompileJobs_RenderThread(EShaderCompileJobPriority::Normal);
 		}
 	}
 	return *Material;
diff --git a/Engine/Source/Runtime/Engine/Private/ShaderCompiler/ShaderCompiler.cpp b/Engine/Source/Runtime/Engine/Private/ShaderCompiler/ShaderCompiler.cpp
--- a/Engine/Source/Runtime/Engine/Private/ShaderCompiler/ShaderCompiler.cpp
+++ b/Engine/Source/Runtime/Engine/Private/ShaderCompiler/ShaderCompiler.cpp
@@ -4430,6 +4430,7 @@ void FShaderCompilingManager::ProcessCompiledShaderMaps(
 				{
 					check(Material->GameThreadCompilingShaderMapId != 0u);
 					Material->GameThreadCompilingShaderMapId = 0u;
+					Material->GameThreadPendingCompilerEnvironment.SafeRelease();
 					MaterialDependencies.RemoveAt(MaterialIndex);
 					MaterialsToReleaseCompilingId.Add(Material);
 				}
diff --git a/Engine/Source/Runtime/Engine/Public/MaterialShared.h b/Engine/Source/Runtime/Engine/Public/MaterialShared.h
--- a/Engine/Source/Runtime/Engine/Public/MaterialShared.h
+++ b/Engine/Source/Runtime/Engine/Public/MaterialShared.h
@@ -28,6 +28,7 @@
 #include "MaterialSceneTextureId.h"
 #include "VirtualTexturing.h"
 #include "Templates/UnrealTemplate.h"
+#include "ShaderCompilerCore.h"
 
 #include "Shader/Preshader.h"
 
@@ -2019,7 +2020,8 @@ public:
 
 	ENGINE_API bool ShouldCacheShaders(const FMaterialShaderTypes& InTypes, const FVertexFactoryType* InVertexFactoryType) const;
 
-	ENGINE_API void SubmitCompileJobs(EShaderCompileJobPriority Priority) const;
+	ENGINE_API void SubmitCompileJobs_GameThread(EShaderCompileJobPriority Priority);
+	ENGINE_API void SubmitCompileJobs_RenderThread(EShaderCompileJobPriority Priority) const;
 
 	/** Returns a string that describes the material's usage for debugging purposes. */
 	virtual FString GetMaterialUsageDescription() const = 0;
@@ -2190,6 +2192,7 @@ private:
 	uint32 GameThreadCompilingShaderMapId;
 	uint32 RenderingThreadCompilingShaderMapId;
 
+	TRefCountPtr<FSharedShaderCompilerEnvironment> GameThreadPendingCompilerEnvironment;
 	TRefCountPtr<FSharedShaderCompilerEnvironment> RenderingThreadPendingCompilerEnvironment;
 
 	/** 
@@ -2197,6 +2200,11 @@ private:
 	 */
 	mutable std::atomic<int8> RenderingThreadShaderMapSubmittedPriority { -1 };
 
+	/**
+	 * Used to prevent submitting the material more than once.  This is accessed by just the game thread so it doesn't need to be an atomic.
+	 */
+	EShaderCompileJobPriority GameThreadShaderMapSubmittedPriority = EShaderCompileJobPriority::None;
+
 	/** 
 	 * Game thread tracked shader map, which is ref counted and manages shader map lifetime. 
 	 * The shader map uses deferred deletion so that the rendering thread has a chance to process a release command when the shader map is no longer referenced.
diff --git a/Engine/Source/Runtime/Engine/Public/ShaderCompiler.h b/Engine/Source/Runtime/Engine/Public/ShaderCompiler.h
--- a/Engine/Source/Runtime/Engine/Public/ShaderCompiler.h
+++ b/Engine/Source/Runtime/Engine/Public/ShaderCompiler.h
@@ -40,46 +40,6 @@ class FShaderPipelineCompileJob;
 
 #define DEBUG_INFINITESHADERCOMPILE 0
 
-enum class EShaderCompilerWorkerType : uint8
-{
-	None,
-	LocalThread,
-	Distributed,
-};
-
-enum class EShaderCompileJobType : uint8
-{
-	Single,
-	Pipeline,
-	Num,
-};
-static const int32 NumShaderCompileJobTypes = (int32)EShaderCompileJobType::Num;
-
-enum class EShaderCompileJobPriority : uint8
-{
-	None = 0xff,
-
-	Low = 0u,
-	Normal,
-	High,
-	ForceLocal, // Force shader to skip XGE and compile on local machine
-	Num,
-};
-static const int32 NumShaderCompileJobPriorities = (int32)EShaderCompileJobPriority::Num;
-
-inline const TCHAR* ShaderCompileJobPriorityToString(EShaderCompileJobPriority v)
-{
-	switch (v)
-	{
-	case EShaderCompileJobPriority::None: return TEXT("None");
-	case EShaderCompileJobPriority::Low: return TEXT("Low");
-	case EShaderCompileJobPriority::Normal: return TEXT("Normal");
-	case EShaderCompileJobPriority::High: return TEXT("High");
-	case EShaderCompileJobPriority::ForceLocal: return TEXT("ForceLocal");
-	default: checkNoEntry(); return TEXT("");
-	}
-}
-
 extern ENGINE_API bool IsShaderJobCacheDDCEnabled();
 
 /** Results for a single compiled shader map. */
diff --git a/Engine/Source/Runtime/Landscape/Private/LandscapeEditLayers.cpp b/Engine/Source/Runtime/Landscape/Private/LandscapeEditLayers.cpp
--- a/Engine/Source/Runtime/Landscape/Private/LandscapeEditLayers.cpp
+++ b/Engine/Source/Runtime/Landscape/Private/LandscapeEditLayers.cpp
@@ -1722,7 +1722,7 @@ bool ALandscape::IsMaterialResourceCompiled(FMaterialResource* InMaterialResourc
 
 	{
 		TRACE_CPUPROFILER_EVENT_SCOPE(Landscape_WaitForMaterialCompilation);
-		InMaterialResource->SubmitCompileJobs(EShaderCompileJobPriority::ForceLocal);
+		InMaterialResource->SubmitCompileJobs_GameThread(EShaderCompileJobPriority::ForceLocal);
 		if (bInWaitForCompilation)
 		{
 			InMaterialResource->FinishCompilation();
diff --git a/Engine/Source/Runtime/RenderCore/Public/ShaderCompilerCore.h b/Engine/Source/Runtime/RenderCore/Public/ShaderCompilerCore.h
--- a/Engine/Source/Runtime/RenderCore/Public/ShaderCompilerCore.h
+++ b/Engine/Source/Runtime/RenderCore/Public/ShaderCompilerCore.h
@@ -526,3 +526,43 @@ extern RENDERCORE_API bool CheckVirtualShaderFilePath(FStringView VirtualPath, T
  * @return True if the file was successfully loaded.
  */
 extern RENDERCORE_API bool LoadShaderSourceFile(const TCHAR* VirtualFilePath, EShaderPlatform ShaderPlatform, FString* OutFileContents, TArray<FShaderCompilerError>* OutCompileErrors);
+
+enum class EShaderCompilerWorkerType : uint8
+{
+	None,
+	LocalThread,
+	Distributed,
+};
+
+enum class EShaderCompileJobType : uint8
+{
+	Single,
+	Pipeline,
+	Num,
+};
+static const int32 NumShaderCompileJobTypes = (int32)EShaderCompileJobType::Num;
+
+enum class EShaderCompileJobPriority : uint8
+{
+	None = 0xff,
+
+	Low = 0u,
+	Normal,
+	High,
+	ForceLocal, // Force shader to skip XGE and compile on local machine
+	Num,
+};
+static const int32 NumShaderCompileJobPriorities = (int32)EShaderCompileJobPriority::Num;
+
+inline const TCHAR* ShaderCompileJobPriorityToString(EShaderCompileJobPriority v)
+{
+	switch (v)
+	{
+	case EShaderCompileJobPriority::None: return TEXT("None");
+	case EShaderCompileJobPriority::Low: return TEXT("Low");
+	case EShaderCompileJobPriority::Normal: return TEXT("Normal");
+	case EShaderCompileJobPriority::High: return TEXT("High");
+	case EShaderCompileJobPriority::ForceLocal: return TEXT("ForceLocal");
+	default: checkNoEntry(); return TEXT("");
+	}
+}
