New HLSL translator generates permutation with and without analytic derivatives

diff --git a/Engine/Source/Runtime/Engine/Private/HLSLTree/HLSLTreeCommon.cpp b/Engine/Source/Runtime/Engine/Private/HLSLTree/HLSLTreeCommon.cpp
--- a/Engine/Source/Runtime/Engine/Private/HLSLTree/HLSLTreeCommon.cpp
+++ b/Engine/Source/Runtime/Engine/Private/HLSLTree/HLSLTreeCommon.cpp
@@ -413,6 +413,7 @@ void FExpressionExternalInput::ComputeAnalyticDerivatives(FTree& Tree, FExpressi
 	}
 	else
 	{
+		const Shader::EValueType DerivativeType = Shader::MakeDerivativeType(InputDesc.Type);
 		switch (InputType)
 		{
 		case EExternalInput::ViewportUV:
@@ -426,6 +427,10 @@ void FExpressionExternalInput::ComputeAnalyticDerivatives(FTree& Tree, FExpressi
 			break;
 		}
 		default:
+			if (DerivativeType != Shader::EValueType::Void)
+			{
+				OutResult.ExpressionDdx = OutResult.ExpressionDdy = Tree.NewConstant(Shader::FValue(DerivativeType));
+			}
 			break;
 		}
 	}
@@ -700,8 +705,16 @@ bool FExpressionTextureSample::PrepareValue(FEmitContext& Context, FEmitScope& S
 		return false;
 	}
 
-	Context.PrepareExpression(TexCoordDerivatives.ExpressionDdx, Scope, RequestedTexCoordType);
-	Context.PrepareExpression(TexCoordDerivatives.ExpressionDdy, Scope, RequestedTexCoordType);
+	const bool bUseAnalyticDerivatives = Context.bUseAnalyticDerivatives && (MipValueMode != TMVM_MipLevel) && TexCoordDerivatives.IsValid();
+	if (MipValueMode == TMVM_Derivative || bUseAnalyticDerivatives)
+	{
+		Context.PrepareExpression(TexCoordDerivatives.ExpressionDdx, Scope, RequestedTexCoordType);
+		Context.PrepareExpression(TexCoordDerivatives.ExpressionDdy, Scope, RequestedTexCoordType);
+	}
+	else if (MipValueMode == TMVM_MipLevel || MipValueMode == TMVM_MipBias)
+	{
+		Context.PrepareExpression(MipValueExpression, Scope, ERequestedType::Scalar);
+	}
 
 	return OutResult.SetType(Context, RequestedType, EExpressionEvaluation::Shader, Shader::EValueType::Float4);
 }
@@ -765,23 +778,39 @@ void FExpressionTextureSample::EmitValueShader(FEmitContext& Context, FEmitScope
 
 	const Shader::EValueType TexCoordType = Private::GetTexCoordType(TextureType);
 	FEmitShaderExpression* TexCoordValue = TexCoordExpression->GetValueShader(Context, Scope, TexCoordType);
-	FEmitShaderExpression* TexCoordValueDdx = nullptr;
-	FEmitShaderExpression* TexCoordValueDdy = nullptr;
+
 	FEmitShaderExpression* TextureResult = nullptr;
-	if (TexCoordDerivatives.IsValid())
+	if (MipValueMode == TMVM_MipLevel)
+	{
+		TextureResult = Context.EmitExpression(Scope, Shader::EValueType::Float4, TEXT("%Level(%.Texture, %, %, %)"),
+			SampleFunctionName,
+			EmitTexture,
+			FormattedSampler.ToString(),
+			TexCoordValue,
+			MipValueExpression->GetValueShader(Context, Scope, Shader::EValueType::Float1));
+	}
+	else if (MipValueMode == TMVM_Derivative || (Context.bUseAnalyticDerivatives && TexCoordDerivatives.IsValid()))
 	{
-		TexCoordValueDdx = TexCoordDerivatives.ExpressionDdx->GetValueShader(Context, Scope, TexCoordType);
-		TexCoordValueDdy = TexCoordDerivatives.ExpressionDdy->GetValueShader(Context, Scope, TexCoordType);
 		TextureResult = Context.EmitExpression(Scope, Shader::EValueType::Float4, TEXT("%Grad(%.Texture, %, %, %, %)"),
 			SampleFunctionName,
 			EmitTexture,
 			FormattedSampler.ToString(),
 			TexCoordValue,
-			TexCoordValueDdx,
-			TexCoordValueDdy);
+			TexCoordDerivatives.ExpressionDdx->GetValueShader(Context, Scope, TexCoordType),
+			TexCoordDerivatives.ExpressionDdy->GetValueShader(Context, Scope, TexCoordType));
+	}
+	else if (MipValueMode == TMVM_MipBias)
+	{
+		TextureResult = Context.EmitExpression(Scope, Shader::EValueType::Float4, TEXT("%Bias(%.Texture, %, %, %)"),
+			SampleFunctionName,
+			EmitTexture,
+			FormattedSampler.ToString(),
+			TexCoordValue,
+			MipValueExpression->GetValueShader(Context, Scope, Shader::EValueType::Float1));
 	}
 	else
 	{
+		check(MipValueMode == TMVM_None);
 		TextureResult = Context.EmitExpression(Scope, Shader::EValueType::Float4, TEXT("%(%.Texture, %, %)"),
 			SampleFunctionName,
 			EmitTexture,
@@ -975,6 +1004,22 @@ void FExpressionSetStructField::EmitValuePreshader(FEmitContext& Context, FEmitS
 	OutResult.Type = StructType;
 }
 
+void FExpressionSelect::ComputeAnalyticDerivatives(FTree& Tree, FExpressionDerivatives& OutResult) const
+{
+	const FExpressionDerivatives TrueDerivatives = Tree.GetAnalyticDerivatives(TrueExpression);
+	const FExpressionDerivatives FalseDerivatives = Tree.GetAnalyticDerivatives(FalseExpression);
+	OutResult.ExpressionDdx = Tree.NewExpression<FExpressionSelect>(ConditionExpression, TrueDerivatives.ExpressionDdx, FalseDerivatives.ExpressionDdx);
+	OutResult.ExpressionDdy = Tree.NewExpression<FExpressionSelect>(ConditionExpression, TrueDerivatives.ExpressionDdy, FalseDerivatives.ExpressionDdy);
+}
+
+FExpression* FExpressionSelect::ComputePreviousFrame(FTree& Tree, const FRequestedType& RequestedType) const
+{
+	return Tree.NewExpression<FExpressionSelect>(
+		Tree.GetPreviousFrame(ConditionExpression, ERequestedType::Scalar),
+		Tree.GetPreviousFrame(TrueExpression, RequestedType),
+		Tree.GetPreviousFrame(FalseExpression, RequestedType));
+}
+
 bool FExpressionSelect::PrepareValue(FEmitContext& Context, FEmitScope& Scope, const FRequestedType& RequestedType, FPrepareValueResult& OutResult) const
 {
 	const FPreparedType& ConditionType = Context.PrepareExpression(ConditionExpression, Scope, ERequestedType::Scalar);
diff --git a/Engine/Source/Runtime/Engine/Private/HLSLTree/HLSLTreeOperation.cpp b/Engine/Source/Runtime/Engine/Private/HLSLTree/HLSLTreeOperation.cpp
--- a/Engine/Source/Runtime/Engine/Private/HLSLTree/HLSLTreeOperation.cpp
+++ b/Engine/Source/Runtime/Engine/Private/HLSLTree/HLSLTreeOperation.cpp
@@ -47,6 +47,7 @@ FOperationDescription GetOperationDescription(EOperation Op)
 	case EOperation::Sqrt: return FOperationDescription(TEXT("Sqrt"), TEXT("sqrt"), 1, Shader::EPreshaderOpcode::Sqrt); break;
 	case EOperation::Rsqrt: return FOperationDescription(TEXT("Rsqrt"), TEXT("rsqrt"), 1, Shader::EPreshaderOpcode::Nop); break; // TODO
 	case EOperation::Log2: return FOperationDescription(TEXT("Log2"), TEXT("log2"), 1, Shader::EPreshaderOpcode::Log2); break;
+	case EOperation::Exp2: return FOperationDescription(TEXT("Exp2"), TEXT("exp2"), 1, Shader::EPreshaderOpcode::Nop); break; // TODO
 	case EOperation::Frac: return FOperationDescription(TEXT("Frac"), TEXT("frac"), 1, Shader::EPreshaderOpcode::Frac); break;
 	case EOperation::Floor: return FOperationDescription(TEXT("Floor"), TEXT("floor"), 1, Shader::EPreshaderOpcode::Floor); break;
 	case EOperation::Ceil: return FOperationDescription(TEXT("Ceil"), TEXT("ceil"), 1, Shader::EPreshaderOpcode::Ceil); break;
@@ -271,6 +272,7 @@ FOperationTypes GetOperationTypes(EOperation Op, TConstArrayView<FPreparedType>
 			Types.ResultType = MakeNonLWCType(IntermediateType);
 			break;
 		case EOperation::Log2:
+		case EOperation::Exp2:
 			// No LWC support yet
 			Types.InputType[0] = Shader::MakeNonLWCType(IntermediateValueType);
 			Types.ResultType = MakeNonLWCType(IntermediateType);
@@ -326,8 +328,7 @@ void FExpressionOperation::ComputeAnalyticDerivatives(FTree& Tree, FExpressionDe
 	case EOperation::Round:
 	case EOperation::Trunc:
 	case EOperation::Sign:
-		OutResult.ExpressionDdx = Tree.NewConstant(0.0f);
-		OutResult.ExpressionDdx = OutResult.ExpressionDdy;
+		OutResult.ExpressionDdx = OutResult.ExpressionDdy = Tree.NewConstant(0.0f);
 		break;
 	default:
 		break;
@@ -445,6 +446,7 @@ void FExpressionOperation::ComputeAnalyticDerivatives(FTree& Tree, FExpressionDe
 	case EOperation::Saturate:
 	case EOperation::PowPositiveClamped:
 	case EOperation::Log2:
+	case EOperation::Exp2:
 		// TODO
 		break;
 	case EOperation::Add:
@@ -492,6 +494,7 @@ void FExpressionOperation::ComputeAnalyticDerivatives(FTree& Tree, FExpressionDe
 	case EOperation::Matrix3MulVec:
 	case EOperation::Matrix4MulVec:
 		// TODO
+		OutResult.ExpressionDdx = OutResult.ExpressionDdy = Tree.NewConstant(FVector3f(0.0f, 0.0f, 0.0f));
 		break;
 	default:
 		checkNoEntry();
@@ -579,6 +582,7 @@ void FExpressionOperation::EmitValueShader(FEmitContext& Context, FEmitScope& Sc
 	case EOperation::Sqrt: OutResult.Code = Context.EmitExpression(Scope, ResultType, Types.bIsLWC ? TEXT("LWCSqrt(%)") : TEXT("sqrt(%)"), InputValue[0]); break;
 	case EOperation::Rsqrt: OutResult.Code = Context.EmitExpression(Scope, ResultType, Types.bIsLWC ? TEXT("LWCRsqrt(%)") : TEXT("rsqrt(%)"), InputValue[0]); break;
 	case EOperation::Log2: OutResult.Code = Context.EmitExpression(Scope, ResultType, TEXT("log2(%)"), InputValue[0]); break;
+	case EOperation::Exp2: OutResult.Code = Context.EmitExpression(Scope, ResultType, TEXT("exp2(%)"), InputValue[0]); break;
 	case EOperation::Frac: OutResult.Code = Context.EmitExpression(Scope, ResultType, Types.bIsLWC ? TEXT("LWCFrac(%)") : TEXT("frac(%)"), InputValue[0]); break;
 	case EOperation::Floor: OutResult.Code = Context.EmitExpression(Scope, ResultType, Types.bIsLWC ? TEXT("LWCFloor(%)") : TEXT("floor(%)"), InputValue[0]); break;
 	case EOperation::Ceil: OutResult.Code = Context.EmitExpression(Scope, ResultType, Types.bIsLWC ? TEXT("LWCCeil(%)") : TEXT("ceil(%)"), InputValue[0]); break;
diff --git a/Engine/Source/Runtime/Engine/Private/Materials/MaterialExpressionHLSL.cpp b/Engine/Source/Runtime/Engine/Private/Materials/MaterialExpressionHLSL.cpp
--- a/Engine/Source/Runtime/Engine/Private/Materials/MaterialExpressionHLSL.cpp
+++ b/Engine/Source/Runtime/Engine/Private/Materials/MaterialExpressionHLSL.cpp
@@ -573,8 +573,35 @@ bool UMaterialExpressionTextureSample::GenerateHLSLExpressionBase(FMaterialHLSLG
 	}
 
 	FExpression* TexCoordExpression = Coordinates.GetTracedInput().Expression ? Coordinates.TryAcquireHLSLExpression(Generator, Scope) : Generator.NewTexCoord(ConstCoordinate);
-	const FExpressionDerivatives TexCoordDerivatives;// = Generator.GetTree().GetAnalyticDerivatives(TexCoordExpression);
-	OutExpression = Generator.GetTree().NewExpression<FExpressionTextureSample>(TextureExpression, TexCoordExpression, TexCoordDerivatives, SamplerSource, MipValueMode);
+	FExpression* MipLevelExpression = nullptr;
+	FExpressionDerivatives TexCoordDerivatives;
+	switch (MipValueMode)
+	{
+	case TMVM_None:
+		TexCoordDerivatives = Generator.GetTree().GetAnalyticDerivatives(TexCoordExpression);
+		break;
+	case TMVM_MipBias:
+	{
+		MipLevelExpression = MipValue.AcquireHLSLExpressionOrConstant(Generator, Scope, ConstMipValue);
+		FExpression* DerivativeScale = Generator.GetTree().NewExp2(MipLevelExpression);
+		TexCoordDerivatives = Generator.GetTree().GetAnalyticDerivatives(TexCoordExpression);
+		TexCoordDerivatives.ExpressionDdx = Generator.GetTree().NewMul(TexCoordDerivatives.ExpressionDdx, DerivativeScale);
+		TexCoordDerivatives.ExpressionDdy = Generator.GetTree().NewMul(TexCoordDerivatives.ExpressionDdy, DerivativeScale);
+		break;
+	}
+	case TMVM_MipLevel:
+		MipLevelExpression = MipValue.AcquireHLSLExpressionOrConstant(Generator, Scope, ConstMipValue);
+		break;
+	case TMVM_Derivative:
+		TexCoordDerivatives.ExpressionDdx = CoordinatesDX.AcquireHLSLExpression(Generator, Scope);
+		TexCoordDerivatives.ExpressionDdy = CoordinatesDY.AcquireHLSLExpression(Generator, Scope);
+		break;
+	default:
+		checkNoEntry();
+		break;
+	}
+	
+	OutExpression = Generator.GetTree().NewExpression<FExpressionTextureSample>(TextureExpression, TexCoordExpression, MipLevelExpression, TexCoordDerivatives, SamplerSource, MipValueMode);
 	return true;
 }
 
@@ -669,9 +696,9 @@ bool UMaterialExpressionFontSample::GenerateHLSLExpression(FMaterialHLSLGenerato
 	const FTextureValue* TextureValue = Generator.AcquireTextureValue(FTextureValue(Texture, ExpectedSamplerType));
 	FExpression* TextureExpression = GenerateHLSLTextureExpression(Generator, TextureValue);
 	FExpression* TexCoordExpression = Generator.NewTexCoord(0);
-	FExpressionDerivatives TexCoordDerivatives;
+	const FExpressionDerivatives TexCoordDerivatives = Generator.GetTree().GetAnalyticDerivatives(TextureExpression);
 
-	OutExpression = Generator.GetTree().NewExpression<FExpressionTextureSample>(TextureExpression, TexCoordExpression, TexCoordDerivatives, SSM_FromTextureAsset, TMVM_None);
+	OutExpression = Generator.GetTree().NewExpression<FExpressionTextureSample>(TextureExpression, TexCoordExpression, nullptr, TexCoordDerivatives, SSM_FromTextureAsset, TMVM_None);
 	return true;
 }
 
diff --git a/Engine/Source/Runtime/Engine/Private/Materials/MaterialHLSLEmitter.cpp b/Engine/Source/Runtime/Engine/Private/Materials/MaterialHLSLEmitter.cpp
--- a/Engine/Source/Runtime/Engine/Private/Materials/MaterialHLSLEmitter.cpp
+++ b/Engine/Source/Runtime/Engine/Private/Materials/MaterialHLSLEmitter.cpp
@@ -50,8 +50,8 @@ static FString GenerateMaterialTemplateHLSL(EShaderPlatform ShaderPlatform,
 	const TCHAR* DeclarationsCode,
 	const TCHAR* SharedShaderCode,
 	const TCHAR* VertexShaderCode,
-	const TCHAR* PixelShaderCodePhase0,
-	const TCHAR* PixelShaderCodePhase1)
+	const TCHAR* PixelShaderCodePhase0[2],
+	const TCHAR* PixelShaderCodePhase1[2])
 {
 	using namespace UE::HLSLTree;
 
@@ -321,78 +321,75 @@ static FString GenerateMaterialTemplateHLSL(EShaderPlatform ShaderPlatform,
 		EvaluateMaterialDeclaration += TEXT("   return Result;" LINE_TERMINATOR);
 		EvaluateMaterialDeclaration += TEXT("}" LINE_TERMINATOR);
 
-		// Set default texcoords in the VS
-		/*for (uint32 TexCoordIndex = 0; TexCoordIndex < NumUserVertexTexCoords; ++TexCoordIndex)
+		FString EvaluateMaterialAttributesPhase0[2];
+		FString EvaluateMaterialAttributesPhase1[2];
+		for (int32 DerivativeIndex = 0; DerivativeIndex < 2; ++DerivativeIndex)
 		{
-			const FString AttributeName = FMaterialAttributeDefinitionMap::GetAttributeName((EMaterialProperty)(MP_CustomizedUVs0 + TexCoordIndex));
+			const TCHAR* DerivativeName = (DerivativeIndex == 0) ? TEXT("HWDerivative") : TEXT("AnalyticDerivative");
 
-			EvaluateMaterialDeclaration += FString::Printf(TEXT("\tDefaultMaterialAttributes.%s = Parameters.TexCoords[%d];") LINE_TERMINATOR, *AttributeName, TexCoordIndex);
-		}*/
-
-		FString EvaluateMaterialAttributesPhase0 = TEXT("    FMaterialAttributes MaterialAttributesPhase0 = EvaluatePixelMaterialAttributesPhase0(Parameters);" LINE_TERMINATOR);
-		FString EvaluateMaterialAttributesPhase1;
-
-		EvaluateMaterialDeclaration += TEXT("FMaterialAttributes EvaluatePixelMaterialAttributesPhase0(in out FMaterialPixelParameters Parameters)" LINE_TERMINATOR);
-		EvaluateMaterialDeclaration += TEXT("{" LINE_TERMINATOR);
-		EvaluateMaterialDeclaration += PixelShaderCodePhase0;
-		EvaluateMaterialDeclaration += TEXT("}" LINE_TERMINATOR);
+			EvaluateMaterialAttributesPhase0[DerivativeIndex] = FString::Printf(TEXT("    FMaterialAttributes MaterialAttributesPhase0 = EvaluatePixelMaterialAttributesPhase0_%s(Parameters);" LINE_TERMINATOR), DerivativeName);
 
-		if (PixelShaderCodePhase1)
-		{
-			EvaluateMaterialDeclaration += TEXT("FMaterialAttributes EvaluatePixelMaterialAttributesPhase1(in out FMaterialPixelParameters Parameters)" LINE_TERMINATOR);
+			EvaluateMaterialDeclaration += FString::Printf(TEXT("FMaterialAttributes EvaluatePixelMaterialAttributesPhase0_%s(in out FMaterialPixelParameters Parameters)" LINE_TERMINATOR), DerivativeName);
 			EvaluateMaterialDeclaration += TEXT("{" LINE_TERMINATOR);
-			EvaluateMaterialDeclaration += PixelShaderCodePhase1;
+			EvaluateMaterialDeclaration += PixelShaderCodePhase0[DerivativeIndex];
 			EvaluateMaterialDeclaration += TEXT("}" LINE_TERMINATOR);
 
-			EvaluateMaterialAttributesPhase1  = TEXT("    FMaterialAttributes MaterialAttributesPhase1 = EvaluatePixelMaterialAttributesPhase1(Parameters);" LINE_TERMINATOR);
-			EvaluateMaterialAttributesPhase1 += TEXT("    Parameters.MaterialAttributes = MaterialAttributesPhase1;" LINE_TERMINATOR);
-		}
-		else
-		{
-			EvaluateMaterialAttributesPhase0 += TEXT("    Parameters.MaterialAttributes = MaterialAttributesPhase0;" LINE_TERMINATOR);
+			if (PixelShaderCodePhase1[DerivativeIndex])
+			{
+				EvaluateMaterialDeclaration += FString::Printf(TEXT("FMaterialAttributes EvaluatePixelMaterialAttributesPhase1_%s(in out FMaterialPixelParameters Parameters)" LINE_TERMINATOR), DerivativeName);
+				EvaluateMaterialDeclaration += TEXT("{" LINE_TERMINATOR);
+				EvaluateMaterialDeclaration += PixelShaderCodePhase1[DerivativeIndex];
+				EvaluateMaterialDeclaration += TEXT("}" LINE_TERMINATOR);
+
+				EvaluateMaterialAttributesPhase1[DerivativeIndex] = FString::Printf(TEXT("    FMaterialAttributes MaterialAttributesPhase1 = EvaluatePixelMaterialAttributesPhase1_%s(Parameters);" LINE_TERMINATOR), DerivativeName);
+				EvaluateMaterialAttributesPhase1[DerivativeIndex] += TEXT("    Parameters.MaterialAttributes = MaterialAttributesPhase1;" LINE_TERMINATOR);
+			}
+			else
+			{
+				EvaluateMaterialAttributesPhase0[DerivativeIndex] += TEXT("    Parameters.MaterialAttributes = MaterialAttributesPhase0;" LINE_TERMINATOR);
+			}
 		}
 
 		LazyPrintf.PushParam(*EvaluateMaterialDeclaration);
 
-		for (int32 PropertyIndex = 0; PropertyIndex < MP_MAX; ++PropertyIndex)
+		for (int32 DerivativeIndex = 0; DerivativeIndex < 2; ++DerivativeIndex)
 		{
-			// Skip non-shared properties
-			if (!SharedPixelProperties[PropertyIndex])
+			for (int32 PropertyIndex = 0; PropertyIndex < MP_MAX; ++PropertyIndex)
 			{
-				continue;
-			}
+				// Skip non-shared properties
+				if (!SharedPixelProperties[PropertyIndex])
+				{
+					continue;
+				}
 
-			const EMaterialProperty Property = (EMaterialProperty)PropertyIndex;
-			check(FMaterialAttributeDefinitionMap::GetShaderFrequency(Property) == SF_Pixel);
-			// Special case MP_SubsurfaceColor as the actual property is a combination of the color and the profile but we don't want to expose the profile
-			const FString PropertyName = FMaterialAttributeDefinitionMap::GetAttributeName(Property);
+				const EMaterialProperty Property = (EMaterialProperty)PropertyIndex;
+				check(FMaterialAttributeDefinitionMap::GetShaderFrequency(Property) == SF_Pixel);
+				// Special case MP_SubsurfaceColor as the actual property is a combination of the color and the profile but we don't want to expose the profile
+				const FString PropertyName = FMaterialAttributeDefinitionMap::GetAttributeName(Property);
 
-			// 'Normal' is always evaluated in Phase0
-			// If there is a Phase1, it will contain all other attributes
-			const int32 EvaluatePhase = (Property == MP_Normal || !PixelShaderCodePhase1) ? 0 : 1;
-			FString& EvaluateMaterialAttributesCode = (EvaluatePhase == 0) ? EvaluateMaterialAttributesPhase0 : EvaluateMaterialAttributesPhase1;
+				// 'Normal' is always evaluated in Phase0
+				// If there is a Phase1, it will contain all other attributes
+				const int32 EvaluatePhase = (Property == MP_Normal || !PixelShaderCodePhase1[DerivativeIndex]) ? 0 : 1;
+				FString& EvaluateMaterialAttributesCode = (EvaluatePhase == 0) ? EvaluateMaterialAttributesPhase0[DerivativeIndex] : EvaluateMaterialAttributesPhase1[DerivativeIndex];
 
-			if (PropertyIndex == MP_FrontMaterial)
-			{
-				EvaluateMaterialAttributesCode += FString::Printf("    PixelMaterialInputs.%s = GetInitialisedStrataData();" LINE_TERMINATOR, *PropertyName);
-			}
-			else if (PropertyIndex == MP_SubsurfaceColor)
-			{
-				// TODO - properly handle subsurface profile
-				EvaluateMaterialAttributesCode += FString::Printf("    PixelMaterialInputs.Subsurface = float4(MaterialAttributesPhase%d.%s, 0.0f);" LINE_TERMINATOR, EvaluatePhase, *PropertyName);
-			}
-			else
-			{
-				EvaluateMaterialAttributesCode += FString::Printf("    PixelMaterialInputs.%s = MaterialAttributesPhase%d.%s;" LINE_TERMINATOR, *PropertyName, EvaluatePhase, *PropertyName);
+				if (PropertyIndex == MP_FrontMaterial)
+				{
+					EvaluateMaterialAttributesCode += FString::Printf("    PixelMaterialInputs.%s = GetInitialisedStrataData();" LINE_TERMINATOR, *PropertyName);
+				}
+				else if (PropertyIndex == MP_SubsurfaceColor)
+				{
+					// TODO - properly handle subsurface profile
+					EvaluateMaterialAttributesCode += FString::Printf("    PixelMaterialInputs.Subsurface = float4(MaterialAttributesPhase%d.%s, 0.0f);" LINE_TERMINATOR, EvaluatePhase, *PropertyName);
+				}
+				else
+				{
+					EvaluateMaterialAttributesCode += FString::Printf("    PixelMaterialInputs.%s = MaterialAttributesPhase%d.%s;" LINE_TERMINATOR, *PropertyName, EvaluatePhase, *PropertyName);
+				}
 			}
-		}
 
-		// TODO - deriv
-		for (int32 Iter = 0; Iter < CompiledPDV_MAX; Iter++)
-		{
 			LazyPrintf.PushParam(TEXT(""));
-			LazyPrintf.PushParam(*EvaluateMaterialAttributesPhase0);
-			LazyPrintf.PushParam(*EvaluateMaterialAttributesPhase1);
+			LazyPrintf.PushParam(*EvaluateMaterialAttributesPhase0[DerivativeIndex]);
+			LazyPrintf.PushParam(*EvaluateMaterialAttributesPhase1[DerivativeIndex]);
 		}
 	}
 
@@ -816,14 +813,17 @@ bool MaterialEmitHLSL(const FMaterialCompileTargetParameters& InCompilerTarget,
 	EmitContext.StaticParameters = &InStaticParameters;
 	EmitContext.MaterialCompilationOutput = &OutCompilationOutput;
 
+	const FStructField* NormalField = Generator.GetMaterialAttributesType()->FindFieldByName(*FMaterialAttributeDefinitionMap::GetAttributeName(MP_Normal));
+	check(NormalField);
+
 	// Prepare pixel shader code
-	FStringBuilderMemstack PixelCodePhase0(Allocator, 1024 * 1024);
-	FStringBuilderMemstack PixelCodePhase1(Allocator, 1024 * 1024);
-	{
-		const FStructField* NormalField = Generator.GetMaterialAttributesType()->FindFieldByName(*FMaterialAttributeDefinitionMap::GetAttributeName(MP_Normal));
-		check(NormalField);
+	FStringBuilderBase* PixelCodePhase0[2] = { nullptr };
+	FStringBuilderBase* PixelCodePhase1[2] = { nullptr };
 
+	for(int32 DerivativeIndex = 0; DerivativeIndex < 2; ++DerivativeIndex)
+	{
 		EmitContext.ShaderFrequency = SF_Pixel;
+		EmitContext.bUseAnalyticDerivatives = (DerivativeIndex == 1);
 		FEmitScope* EmitResultScope = EmitContext.PrepareScope(&Generator.GetResultStatement()->GetParentScope());
 
 		// Prepare all fields *except* normal
@@ -848,7 +848,10 @@ bool MaterialEmitHLSL(const FMaterialCompileTargetParameters& InCompilerTarget,
 			return false;
 		}
 
-		bUsesPixelDepthOffset = IsAttributeUsed(Generator, EmitContext, *EmitResultScope, PixelResultType0, Generator.GetResultExpression(), MP_PixelDepthOffset);
+		if (IsAttributeUsed(Generator, EmitContext, *EmitResultScope, PixelResultType0, Generator.GetResultExpression(), MP_PixelDepthOffset))
+		{
+			bUsesPixelDepthOffset = true;
+		}
 
 		if (!EmitContext.bReadMaterialNormal)
 		{
@@ -860,14 +863,15 @@ bool MaterialEmitHLSL(const FMaterialCompileTargetParameters& InCompilerTarget,
 				return false;
 			}
 
-			HLSLTree->EmitShader(EmitContext, PixelCodePhase0);
+			PixelCodePhase0[DerivativeIndex] = new(Allocator) FStringBuilderMemstack(Allocator, 128 * 1024);
+			HLSLTree->EmitShader(EmitContext, *PixelCodePhase0[DerivativeIndex]);
+			HLSLTree->ResetNodes();
 		}
 		else
 		{
 			// Execute everything *except* normal in phase1
-			HLSLTree->EmitShader(EmitContext, PixelCodePhase1);
-
-			// Reset state
+			PixelCodePhase1[DerivativeIndex] = new(Allocator) FStringBuilderMemstack(Allocator, 128 * 1024);
+			HLSLTree->EmitShader(EmitContext, *PixelCodePhase1[DerivativeIndex]);
 			HLSLTree->ResetNodes();
 
 			EmitResultScope = EmitContext.PrepareScope(&Generator.GetResultStatement()->GetParentScope());
@@ -882,12 +886,14 @@ bool MaterialEmitHLSL(const FMaterialCompileTargetParameters& InCompilerTarget,
 			}
 
 			// Execute the normal in phase0
-			HLSLTree->EmitShader(EmitContext, PixelCodePhase0);
+			PixelCodePhase0[DerivativeIndex] = new(Allocator) FStringBuilderMemstack(Allocator, 128 * 1024);
+			HLSLTree->EmitShader(EmitContext, *PixelCodePhase0[DerivativeIndex]);
+			HLSLTree->ResetNodes();
 		}
 	}
 
 	// Prepare vertex shader code
-	FStringBuilderMemstack VertexCode(Allocator, 1024 * 1024);
+	FStringBuilderMemstack VertexCode(Allocator, 128 * 1024);
 	{
 		FRequestedType RequestedVertexAttributesType;
 		Generator.SetRequestedFields(SF_Vertex, RequestedVertexAttributesType);
@@ -905,8 +911,6 @@ bool MaterialEmitHLSL(const FMaterialCompileTargetParameters& InCompilerTarget,
 		}
 		RequestedVertexAttributesType.SetFieldRequested(Generator.GetMaterialAttributesType()->FindFieldByName(TEXT("PrevWorldPositionOffset")));
 
-		HLSLTree->ResetNodes();
-
 		EmitContext.ShaderFrequency = SF_Vertex;
 		FEmitScope* EmitResultScope = EmitContext.PrepareScope(&Generator.GetResultStatement()->GetParentScope());
 
@@ -930,21 +934,34 @@ bool MaterialEmitHLSL(const FMaterialCompileTargetParameters& InCompilerTarget,
 	OutCompilationOutput.bUsesWorldPositionOffset = bUsesWorldPositionOffset;
 	OutCompilationOutput.bUsesPixelDepthOffset = bUsesPixelDepthOffset;
 
-	FStringBuilderMemstack Declarations(Allocator, 64 * 1024);
+	FStringBuilderMemstack Declarations(Allocator, 32 * 1024);
 	TypeRegistry.EmitDeclarationsCode(Declarations);
 
-	FStringBuilderMemstack SharedCode(Allocator, 64 * 1024);
+	FStringBuilderMemstack SharedCode(Allocator, 32 * 1024);
 	Generator.EmitSharedCode(SharedCode);
 	EmitContext.EmitDeclarationsCode(SharedCode);
-	
-	FString MaterialTemplateSource = GenerateMaterialTemplateHLSL(InCompilerTarget.ShaderPlatform,
-		InOutMaterial,
-		EmitContext,
-		Declarations.ToString(),
-		SharedCode.ToString(),
-		VertexCode.ToString(),
-		PixelCodePhase0.ToString(),
-		EmitContext.bReadMaterialNormal ? PixelCodePhase1.ToString() : nullptr);
+
+	FString MaterialTemplateSource;
+	{
+		const TCHAR* InputPixelCodePhase0[2] =
+		{
+			PixelCodePhase0[0]->ToString(),
+			PixelCodePhase0[1]->ToString(),
+		};
+		const TCHAR* InputPixelCodePhase1[2] =
+		{
+			EmitContext.bReadMaterialNormal ? PixelCodePhase1[0]->ToString() : nullptr,
+			EmitContext.bReadMaterialNormal ? PixelCodePhase1[1]->ToString() : nullptr,
+		};
+		MaterialTemplateSource = GenerateMaterialTemplateHLSL(InCompilerTarget.ShaderPlatform,
+			InOutMaterial,
+			EmitContext,
+			Declarations.ToString(),
+			SharedCode.ToString(),
+			VertexCode.ToString(),
+			InputPixelCodePhase0,
+			InputPixelCodePhase1);
+	}
 
 	OutMaterialEnvironment = new FSharedShaderCompilerEnvironment();
 	OutMaterialEnvironment->TargetPlatform = InCompilerTarget.TargetPlatform;
diff --git a/Engine/Source/Runtime/Engine/Private/Shader/ShaderValue.cpp b/Engine/Source/Runtime/Engine/Private/Shader/ShaderValue.cpp
--- a/Engine/Source/Runtime/Engine/Private/Shader/ShaderValue.cpp
+++ b/Engine/Source/Runtime/Engine/Private/Shader/ShaderValue.cpp
@@ -33,12 +33,7 @@ FType FType::GetDerivativeType() const
 	}
 
 	check(ValueType != EValueType::Struct);
-	const FValueTypeDescription TypeDesc = GetValueTypeDescription(ValueType);
-	if (IsNumericType(TypeDesc.ComponentType))
-	{
-		return MakeValueType(EValueComponentType::Float, TypeDesc.NumComponents);
-	}
-	return EValueType::Void;
+	return MakeDerivativeType(ValueType);
 }
 
 int32 FType::GetNumComponents() const
@@ -736,6 +731,16 @@ EValueType MakeNonLWCType(EValueType Type)
 	return MakeValueType(MakeNonLWCType(TypeDesc.ComponentType), TypeDesc.NumComponents);
 }
 
+EValueType MakeDerivativeType(EValueType Type)
+{
+	const FValueTypeDescription TypeDesc = GetValueTypeDescription(Type);
+	if (IsNumericType(TypeDesc.ComponentType))
+	{
+		return MakeValueType(EValueComponentType::Float, TypeDesc.NumComponents);
+	}
+	return EValueType::Void;
+}
+
 EValueType MakeArithmeticResultType(EValueType Lhs, EValueType Rhs, FString& OutErrorMessage)
 {
 	const FValueTypeDescription LhsDesc = GetValueTypeDescription(Lhs);
diff --git a/Engine/Source/Runtime/Engine/Public/HLSLTree/HLSLTree.h b/Engine/Source/Runtime/Engine/Public/HLSLTree/HLSLTree.h
--- a/Engine/Source/Runtime/Engine/Public/HLSLTree/HLSLTree.h
+++ b/Engine/Source/Runtime/Engine/Public/HLSLTree/HLSLTree.h
@@ -569,6 +569,8 @@ public:
 	FExpression* NewRcp(FExpression* Input) { return NewUnaryOp(EOperation::Rcp, Input); }
 	FExpression* NewSqrt(FExpression* Input) { return NewUnaryOp(EOperation::Sqrt, Input); }
 	FExpression* NewRsqrt(FExpression* Input) { return NewUnaryOp(EOperation::Rsqrt, Input); }
+	FExpression* NewLog2(FExpression* Input) { return NewUnaryOp(EOperation::Log2, Input); }
+	FExpression* NewExp2(FExpression* Input) { return NewUnaryOp(EOperation::Exp2, Input); }
 	FExpression* NewFrac(FExpression* Input) { return NewUnaryOp(EOperation::Frac, Input); }
 	FExpression* NewLength(FExpression* Input) { return NewUnaryOp(EOperation::Length, Input); }
 	FExpression* NewNormalize(FExpression* Input) { return NewUnaryOp(EOperation::Normalize, Input); }
diff --git a/Engine/Source/Runtime/Engine/Public/HLSLTree/HLSLTreeCommon.h b/Engine/Source/Runtime/Engine/Public/HLSLTree/HLSLTreeCommon.h
--- a/Engine/Source/Runtime/Engine/Public/HLSLTree/HLSLTreeCommon.h
+++ b/Engine/Source/Runtime/Engine/Public/HLSLTree/HLSLTreeCommon.h
@@ -280,11 +280,13 @@ class FExpressionTextureSample : public FExpression
 public:
 	FExpressionTextureSample(FExpression* InTextureExpression,
 		FExpression* InTexCoordExpression,
+		FExpression* InMipValueExpression,
 		const FExpressionDerivatives& InTexCoordDerivatives,
 		ESamplerSourceMode InSamplerSource,
 		ETextureMipValueMode InMipValueMode)
 		: TextureExpression(InTextureExpression)
 		, TexCoordExpression(InTexCoordExpression)
+		, MipValueExpression(InMipValueExpression)
 		, TexCoordDerivatives(InTexCoordDerivatives)
 		, SamplerSource(InSamplerSource)
 		, MipValueMode(InMipValueMode)
@@ -292,6 +294,7 @@ public:
 
 	FExpression* TextureExpression;
 	FExpression* TexCoordExpression;
+	FExpression* MipValueExpression;
 	FExpressionDerivatives TexCoordDerivatives;
 	ESamplerSourceMode SamplerSource;
 	ETextureMipValueMode MipValueMode;
@@ -359,6 +362,8 @@ public:
 	FExpression* TrueExpression;
 	FExpression* FalseExpression;
 
+	virtual void ComputeAnalyticDerivatives(FTree& Tree, FExpressionDerivatives& OutResult) const override;
+	virtual FExpression* ComputePreviousFrame(FTree& Tree, const FRequestedType& RequestedType) const override;
 	virtual bool PrepareValue(FEmitContext& Context, FEmitScope& Scope, const FRequestedType& RequestedType, FPrepareValueResult& OutResult) const override;
 	virtual void EmitValueShader(FEmitContext& Context, FEmitScope& Scope, const FRequestedType& RequestedType, FEmitValueShaderResult& OutResult) const override;
 	virtual void EmitValuePreshader(FEmitContext& Context, FEmitScope& Scope, const FRequestedType& RequestedType, FEmitValuePreshaderResult& OutResult) const override;
diff --git a/Engine/Source/Runtime/Engine/Public/HLSLTree/HLSLTreeEmit.h b/Engine/Source/Runtime/Engine/Public/HLSLTree/HLSLTreeEmit.h
--- a/Engine/Source/Runtime/Engine/Public/HLSLTree/HLSLTreeEmit.h
+++ b/Engine/Source/Runtime/Engine/Public/HLSLTree/HLSLTreeEmit.h
@@ -430,6 +430,7 @@ public:
 	FErrorHandlerInterface* Errors = nullptr;
 	const Shader::FStructTypeRegistry* TypeRegistry = nullptr;
 	EShaderFrequency ShaderFrequency = SF_Pixel;
+	bool bUseAnalyticDerivatives = false;
 
 	TArray<FEmitShaderNode*> EmitNodes;
 	TMap<const FScope*, FEmitScope*> EmitScopeMap;
diff --git a/Engine/Source/Runtime/Engine/Public/HLSLTree/HLSLTreeTypes.h b/Engine/Source/Runtime/Engine/Public/HLSLTree/HLSLTreeTypes.h
--- a/Engine/Source/Runtime/Engine/Public/HLSLTree/HLSLTreeTypes.h
+++ b/Engine/Source/Runtime/Engine/Public/HLSLTree/HLSLTreeTypes.h
@@ -72,6 +72,7 @@ enum class EOperation : uint8
 	Sqrt,
 	Rsqrt,
 	Log2,
+	Exp2,
 	Frac,
 	Floor,
 	Ceil,
diff --git a/Engine/Source/Runtime/Engine/Public/Shader/ShaderTypes.h b/Engine/Source/Runtime/Engine/Public/Shader/ShaderTypes.h
--- a/Engine/Source/Runtime/Engine/Public/Shader/ShaderTypes.h
+++ b/Engine/Source/Runtime/Engine/Public/Shader/ShaderTypes.h
@@ -151,6 +151,7 @@ EValueType MakeValueType(EValueComponentType ComponentType, int32 NumComponents)
 EValueType MakeValueType(EValueType BaseType, int32 NumComponents);
 EValueType MakeValueTypeWithRequestedNumComponents(EValueType BaseType, int8 RequestedNumComponents);
 EValueType MakeNonLWCType(EValueType Type);
+EValueType MakeDerivativeType(EValueType Type);
 EValueType MakeArithmeticResultType(EValueType Lhs, EValueType Rhs, FString& OutErrorMessage);
 EValueType MakeComparisonResultType(EValueType Lhs, EValueType Rhs, FString& OutErrorMessage);
 
