Strata - fixed crash when water is made translucent

diff --git a/Engine/Source/Runtime/Engine/Private/Materials/HLSLMaterialTranslator.cpp b/Engine/Source/Runtime/Engine/Private/Materials/HLSLMaterialTranslator.cpp
--- a/Engine/Source/Runtime/Engine/Private/Materials/HLSLMaterialTranslator.cpp
+++ b/Engine/Source/Runtime/Engine/Private/Materials/HLSLMaterialTranslator.cpp
@@ -1333,6 +1333,7 @@ bool FHLSLMaterialTranslator::Translate()
 
 		if (bStrataFrontMaterialIsValid)
 		{
+			EStrataBlendMode StrataBlendMode = Material->GetStrataBlendMode();
 			bMaterialIsStrata = true;
 			const FStrataMaterialCompilationInfo& StrataCompilationInfo = CodeChunkToStrataCompilationInfoMap[StrataValidFrontMaterialCodeChunkPostTranslate];
 
@@ -1342,7 +1343,7 @@ bool FHLSLMaterialTranslator::Translate()
 				// If the unlit node is used, it must be the only one used
 				if (!StrataIsVolumetricFogCloudOnly(this, StrataCompilationInfo))
 				{
-					FString ErrorMsg = FString::Printf(TEXT("Material %s contains Unlit BSDF but it is not the only one representing the material asset: %s. It must be the single BSDF.\r\n"), *Material->GetDebugName(), *Material->GetAssetPath().ToString());
+					FString ErrorMsg = FString::Printf(TEXT("Material %s contains Unlit BSDF but it is not the single BSDF representing the material asset: %s. It must be the single BSDF.\r\n"), *Material->GetDebugName(), *Material->GetAssetPath().ToString());
 					Error(*ErrorMsg);
 				}
 			}
@@ -1353,7 +1354,7 @@ bool FHLSLMaterialTranslator::Translate()
 				// If the unlit node is used, it must be the only one used
 				if (!StrataIsUnlitOnly(this, StrataCompilationInfo))
 				{
-					FString ErrorMsg = FString::Printf(TEXT("Material %s contains Unlit BSDF but it is not the only one representing the material asset: %s. It must be the single BSDF.\r\n"), *Material->GetDebugName(), *Material->GetAssetPath().ToString());
+					FString ErrorMsg = FString::Printf(TEXT("Material %s contains Unlit BSDF but it is not the single BSDF representing the material asset: %s. It must be the single BSDF.\r\n"), *Material->GetDebugName(), *Material->GetAssetPath().ToString());
 					Error(*ErrorMsg);
 				}
 			}
@@ -1364,7 +1365,7 @@ bool FHLSLMaterialTranslator::Translate()
 				// If the unlit node is used, it must be the only one used
 				if (!StrataIsHairOnly(this, StrataCompilationInfo))
 				{
-					FString ErrorMsg = FString::Printf(TEXT("Material %s contains hair BSDF but it is not the only representing the material asset: %s. It must be the single BSDF.\r\n"), *Material->GetDebugName(), *Material->GetAssetPath().ToString());
+					FString ErrorMsg = FString::Printf(TEXT("Material %s contains hair BSDF but it is not the single BSDF representing the material asset: %s. It must be the single BSDF.\r\n"), *Material->GetDebugName(), *Material->GetAssetPath().ToString());
 					Error(*ErrorMsg);
 				}
 			}
@@ -1375,9 +1376,14 @@ bool FHLSLMaterialTranslator::Translate()
 				// If the unlit node is used, it must be the only one used
 				if (!StrataIsSingleLayerWaterOnly(this, StrataCompilationInfo))
 				{
-					FString ErrorMsg = FString::Printf(TEXT("Material %s contains water BSDF but it is not the only representing the material asset: %s. It must be the single BSDF.\r\n"), *Material->GetDebugName(), *Material->GetAssetPath().ToString());
+					FString ErrorMsg = FString::Printf(TEXT("Material %s contains water BSDF but it is not the single BSDF representing the material asset: %s. It must be the single BSDF.\r\n"), *Material->GetDebugName(), *Material->GetAssetPath().ToString());
 					Error(*ErrorMsg);
 				}
+
+				if (StrataBlendMode != SBM_Opaque && StrataBlendMode != SBM_Masked)
+				{
+					Errorf(TEXT("SingleLayerWater materials must be opaque or masked."));
+				}
 			}
 
 			if (!bStrataUsesConversionFromLegacy)
