fix Texture GetMipData returning true when it didn't get the mip

diff --git a/Engine/Source/Runtime/Engine/Private/Texture.cpp b/Engine/Source/Runtime/Engine/Private/Texture.cpp
--- a/Engine/Source/Runtime/Engine/Private/Texture.cpp
+++ b/Engine/Source/Runtime/Engine/Private/Texture.cpp
@@ -119,6 +119,8 @@ UTexture::UTexture(const FObjectInitializer& ObjectInitializer)
 	MaxTextureSize = 0; // means no limitation
 	MipGenSettings = TMGS_FromTextureGroup;
 	CompositeTextureMode = CTM_NormalRoughnessToAlpha;
+	// @todo Oodle : this should have defaulted to CTM_Disabled
+	//	but it's hard to change now because of UPROPERTY writing deltas to default
 	CompositePower = 1.0f;
 	bUseLegacyGamma = false;
 	bIsImporting = false;
@@ -2634,14 +2636,12 @@ bool FTextureSource::FMipData::GetMipData(TArray64<uint8>& OutMipData, int32 Blo
 				(const uint8*)MipData.GetData() + MipOffset,
 				MipSize
 			);
-		}
 
-		return true;
-	}
-	else
-	{
-		return false;
+			return true;
+		}
 	}
+
+	return false;
 }
 
 #endif //WITH_EDITOR
