small improvements to Color loops in ImageUtils

diff --git a/Engine/Source/Runtime/Engine/Private/ImageUtils.cpp b/Engine/Source/Runtime/Engine/Private/ImageUtils.cpp
--- a/Engine/Source/Runtime/Engine/Private/ImageUtils.cpp
+++ b/Engine/Source/Runtime/Engine/Private/ImageUtils.cpp
@@ -285,6 +285,9 @@ UTexture2D* FImageUtils::CreateTexture2D(int32 SrcWidth, int32 SrcHeight, const
 	Tex2D = NewObject<UTexture2D>(Outer, FName(*Name), Flags);
 	Tex2D->Source.Init(SrcWidth, SrcHeight, /*NumSlices=*/ 1, /*NumMips=*/ 1, TSF_BGRA8);
 	
+	// if bUseAlpha is off, alpha is changed to 255
+	uint8 AlphaOr = InParams.bUseAlpha ? 0 : 0xFF;
+
 	// Create base mip for the texture we created.
 	uint8* MipData = Tex2D->Source.LockMip(0);
 	for( int32 y=0; y<SrcHeight; y++ )
@@ -296,14 +299,7 @@ UTexture2D* FImageUtils::CreateTexture2D(int32 SrcWidth, int32 SrcHeight, const
 			*DestPtr++ = SrcPtr->B;
 			*DestPtr++ = SrcPtr->G;
 			*DestPtr++ = SrcPtr->R;
-			if( InParams.bUseAlpha )
-			{
-				*DestPtr++ = SrcPtr->A;
-			}
-			else
-			{
-				*DestPtr++ = 0xFF;
-			}
+			*DestPtr++ = SrcPtr->A | AlphaOr;
 			SrcPtr++;
 		}
 	}
@@ -908,20 +904,20 @@ UTexture2D* FImageUtils::ImportFileAsTexture2D(const FString& Filename)
 					if (NewTexture)
 					{
 						uint8* MipData = static_cast<uint8*>(NewTexture->GetPlatformData()->Mips[0].BulkData.Lock(LOCK_READ_WRITE));
+						int64 MipDataSize = NewTexture->GetPlatformData()->Mips[0].BulkData.GetBulkDataSize();
 
 						TArrayView64<FColor> SourceColors(reinterpret_cast<FColor*>(BGREImage.GetData()), BGREImage.Num() / sizeof(FColor));
 
 						// Bulk data was already allocated for the correct size when we called CreateTransient above
-						TArrayView64<FFloat16> Destination(reinterpret_cast<FFloat16*>(MipData), NewTexture->GetPlatformData()->Mips[0].BulkData.GetBulkDataSize() / sizeof(FFloat16));
+						TArrayView64<FFloat16Color> DestinationColors(reinterpret_cast<FFloat16Color*>(MipData), MipDataSize / sizeof(FFloat16Color));
+
+						check( SourceColors.Num() == DestinationColors.Num() );
 
 						int64 DestinationIndex = 0;
 						for (const FColor& Color: SourceColors)
 						{
 							FLinearColor LinearColor = Color.FromRGBE();
-							Destination[DestinationIndex++].Set(LinearColor.R);
-							Destination[DestinationIndex++].Set(LinearColor.G);
-							Destination[DestinationIndex++].Set(LinearColor.B);
-							Destination[DestinationIndex++].Set(LinearColor.A);
+							DestinationColors[DestinationIndex++] = FFloat16Color(LinearColor);
 						}
 
 						NewTexture->GetPlatformData()->Mips[0].BulkData.Unlock();
@@ -1001,9 +997,11 @@ UTexture2D* FImageUtils::ImportBufferAsTexture2D(TArrayView64<const uint8> Buffe
 			{
 				NewTexture->bNotOfflineProcessed = true;
 				uint8* MipData = static_cast<uint8*>(NewTexture->GetPlatformData()->Mips[0].BulkData.Lock(LOCK_READ_WRITE));
-				
+				int64 MipDataSize = NewTexture->GetPlatformData()->Mips[0].BulkData.GetBulkDataSize();
+
 				// Bulk data was already allocated for the correct size when we called CreateTransient above
-				FMemory::Memcpy(MipData, UncompressedData.GetData(), NewTexture->GetPlatformData()->Mips[0].BulkData.GetBulkDataSize());
+				check( UncompressedData.Num() >= MipDataSize );
+				FMemory::Memcpy(MipData, UncompressedData.GetData(), MipDataSize);
 				
 				NewTexture->GetPlatformData()->Mips[0].BulkData.Unlock();
 
