Fix skeletal mesh LOD changing unexpectedly in editor viewport with Realtime viewport update on when
1) hovering mouse over level thumbnail in content browser (with bAllowWorldThumbnails enabled in DefaultEditor.ini)
2) drag a rotation axis in Details panel

diff --git a/Engine/Source/Editor/UnrealEd/Private/ThumbnailRendering/WorldThumbnailRenderer.cpp b/Engine/Source/Editor/UnrealEd/Private/ThumbnailRendering/WorldThumbnailRenderer.cpp
--- a/Engine/Source/Editor/UnrealEd/Private/ThumbnailRendering/WorldThumbnailRenderer.cpp
+++ b/Engine/Source/Editor/UnrealEd/Private/ThumbnailRendering/WorldThumbnailRenderer.cpp
@@ -45,6 +45,8 @@ void UWorldThumbnailRenderer::Draw(UObject* Object, int32 X, int32 Y, uint32 Wid
 			.SetTime(UThumbnailRenderer::GetTime())
 			.SetAdditionalViewFamily(bAdditionalViewFamily));
 
+		ViewFamily.bThumbnailRendering = true;
+
 		ViewFamily.EngineShowFlags.SetDiffuse(true);
 		ViewFamily.EngineShowFlags.SetSkeletalMeshes(true);
 		ViewFamily.EngineShowFlags.SetTranslucency(true);
diff --git a/Engine/Source/Runtime/Engine/Private/Components/SkinnedMeshComponent.cpp b/Engine/Source/Runtime/Engine/Private/Components/SkinnedMeshComponent.cpp
--- a/Engine/Source/Runtime/Engine/Private/Components/SkinnedMeshComponent.cpp
+++ b/Engine/Source/Runtime/Engine/Private/Components/SkinnedMeshComponent.cpp
@@ -3194,10 +3194,10 @@ bool USkinnedMeshComponent::UpdateLODStatus_Internal(int32 InMasterPoseComponent
 			{
 				NewPredictedLODLevel = FMath::Clamp(MeshObject->MinDesiredLODLevel + LODBias, 0, MaxLODIndex);
 			}
-			// If no MeshObject - just assume lowest LOD.
+			// If no MeshObject - just reuse old predicted LOD.
 			else
 			{
-				NewPredictedLODLevel = MinLodIndex;
+				NewPredictedLODLevel = FMath::Clamp(OldPredictedLODLevel, MinLodIndex, MaxLODIndex);
 			}
 
 			// now check to see if we have a MinLODLevel and apply it
diff --git a/Engine/Source/Runtime/Engine/Private/SkeletalRender.cpp b/Engine/Source/Runtime/Engine/Private/SkeletalRender.cpp
--- a/Engine/Source/Runtime/Engine/Private/SkeletalRender.cpp
+++ b/Engine/Source/Runtime/Engine/Private/SkeletalRender.cpp
@@ -33,7 +33,7 @@ FSkeletalMeshObject
 -----------------------------------------------------------------------------*/
 
 FSkeletalMeshObject::FSkeletalMeshObject(USkinnedMeshComponent* InMeshComponent, FSkeletalMeshRenderData* InSkelMeshRenderData, ERHIFeatureLevel::Type InFeatureLevel)
-:	MinDesiredLODLevel(InSkelMeshRenderData->CurrentFirstLODIdx)
+:	MinDesiredLODLevel(FMath::Max<int32>(InMeshComponent->GetPredictedLODLevel(), InSkelMeshRenderData->CurrentFirstLODIdx))
 ,	MaxDistanceFactor(0.f)
 ,	WorkingMinDesiredLODLevel(MinDesiredLODLevel)
 ,	WorkingMaxDistanceFactor(0.f)
@@ -89,6 +89,12 @@ FSkeletalMeshObject::~FSkeletalMeshObject()
 
 void FSkeletalMeshObject::UpdateMinDesiredLODLevel(const FSceneView* View, const FBoxSphereBounds& Bounds, int32 FrameNumber, uint8 CurFirstLODIdx)
 {
+	// Thumbnail rendering doesn't contribute to MinDesiredLODLevel calculation
+	if (View->Family && View->Family->bThumbnailRendering)
+	{
+		return;
+	}
+
 	static const auto* SkeletalMeshLODRadiusScale = IConsoleManager::Get().FindTConsoleVariableDataFloat(TEXT("r.SkeletalMeshLODRadiusScale"));
 	float LODScale = FMath::Clamp(SkeletalMeshLODRadiusScale->GetValueOnRenderThread(), 0.25f, 1.0f);
 
@@ -102,7 +108,7 @@ void FSkeletalMeshObject::UpdateMinDesiredLODLevel(const FSceneView* View, const
 
 	int32 NewLODLevel = 0;
 
-	// Look for a lower LOD if the EngineShowFlags is enabled - Thumbnail rendering disables LODs
+	// Look for a lower LOD if the EngineShowFlags is enabled
 	if( View->Family && 1==View->Family->EngineShowFlags.LOD )
 	{
 		// Iterate from worst to best LOD
diff --git a/Engine/Source/Runtime/Engine/Public/SceneView.h b/Engine/Source/Runtime/Engine/Public/SceneView.h
--- a/Engine/Source/Runtime/Engine/Public/SceneView.h
+++ b/Engine/Source/Runtime/Engine/Public/SceneView.h
@@ -1734,6 +1734,9 @@ public:
 	/** When enabled, the scene capture will composite into the render target instead of overwriting its contents. */
 	ESceneCaptureCompositeMode SceneCaptureCompositeMode;
 
+	/** Whether this view is for thumbnail rendering */
+	bool bThumbnailRendering = false;
+
 	/**
 	 * GetWorld->IsPaused() && !Simulate
 	 * Simulate is excluded as the camera can move which invalidates motionblur
