Additive transform curve gives incorrect result when evaluating Animation Sequence
#fix ensure additive components (translation,rotation,scale) are applied in isolation

diff --git a/Engine/Source/Runtime/Engine/Private/Animation/BonePose.cpp b/Engine/Source/Runtime/Engine/Private/Animation/BonePose.cpp
--- a/Engine/Source/Runtime/Engine/Private/Animation/BonePose.cpp
+++ b/Engine/Source/Runtime/Engine/Private/Animation/BonePose.cpp
@@ -231,16 +231,18 @@ void BuildPoseFromRawDataInternal(const TArray<FRawAnimSequenceTrack>& InAnimati
 					if (AdditiveBoneTransformCurve)
 					{
 						const FTransform PoseOneAdditive = (*AdditiveBoneTransformCurve)->Evaluate(KeyIndex1 * TimePerKey, 1.f);
-						InOutPose[PoseBoneIndex].SetRotation(InOutPose[PoseBoneIndex].GetRotation() * PoseOneAdditive.GetRotation());
-						InOutPose[PoseBoneIndex].SetTranslation(InOutPose[PoseBoneIndex].TransformPosition(PoseOneAdditive.GetTranslation()));
-						InOutPose[PoseBoneIndex].SetScale3D(InOutPose[PoseBoneIndex].GetScale3D() * PoseOneAdditive.GetScale3D());
+						const FTransform PoseOneLocalTransform = InOutPose[PoseBoneIndex];
+						InOutPose[PoseBoneIndex].SetRotation(PoseOneLocalTransform.GetRotation() * PoseOneAdditive.GetRotation());
+						InOutPose[PoseBoneIndex].SetTranslation(PoseOneLocalTransform.TransformPosition(PoseOneAdditive.GetTranslation()));
+						InOutPose[PoseBoneIndex].SetScale3D(PoseOneLocalTransform.GetScale3D() * PoseOneAdditive.GetScale3D());
 
 						if (bInterpolateT)
 						{
 							const FTransform PoseTwoAdditive = (*AdditiveBoneTransformCurve)->Evaluate(KeyIndex2 * TimePerKey, 1.f);
-							Key2Pose[PoseBoneIndex].SetRotation(Key2Pose[PoseBoneIndex].GetRotation() * PoseTwoAdditive.GetRotation());
-							Key2Pose[PoseBoneIndex].SetTranslation(Key2Pose[PoseBoneIndex].TransformPosition(PoseTwoAdditive.GetTranslation()));
-							Key2Pose[PoseBoneIndex].SetScale3D(Key2Pose[PoseBoneIndex].GetScale3D() * PoseTwoAdditive.GetScale3D());
+							const FTransform PoseTwoLocalTransform = Key2Pose[PoseBoneIndex];
+							Key2Pose[PoseBoneIndex].SetRotation(PoseTwoLocalTransform.GetRotation() * PoseTwoAdditive.GetRotation());
+							Key2Pose[PoseBoneIndex].SetTranslation(PoseTwoLocalTransform.TransformPosition(PoseTwoAdditive.GetTranslation()));
+							Key2Pose[PoseBoneIndex].SetScale3D(PoseTwoLocalTransform.GetScale3D() * PoseTwoAdditive.GetScale3D());
 						}
 					}
 				}
