World Partition: retain actor descriptors bounding volume data from existing descriptor when the saved actor is not registered in the world.

#rnx

diff --git a/Engine/Source/Runtime/Engine/Private/Actor.cpp b/Engine/Source/Runtime/Engine/Private/Actor.cpp
--- a/Engine/Source/Runtime/Engine/Private/Actor.cpp
+++ b/Engine/Source/Runtime/Engine/Private/Actor.cpp
@@ -54,6 +54,7 @@
 
 #if WITH_EDITOR
 #include "FoliageHelper.h"
+#include "AssetRegistryModule.h"
 #endif
 
 #include "WorldPartition/WorldPartition.h"
@@ -475,6 +476,40 @@ void AActor::GetExternalActorExtendedAssetRegistryTags(TArray<FAssetRegistryTag>
 	{
 		TUniquePtr<FWorldPartitionActorDesc> ActorDesc(CreateActorDesc());
 
+		// If the actor is not added to a world, we can't retrieve its bounding volume, so try to get the existing one
+		if (ULevel* Level = GetLevel(); !Level || !Level->Actors.Contains(this))
+		{
+			IAssetRegistry& AssetRegistry = FModuleManager::LoadModuleChecked<FAssetRegistryModule>(TEXT("AssetRegistry")).Get();
+
+			FARFilter Filter;
+			Filter.bIncludeOnlyOnDiskAssets = true;
+			Filter.PackageNames.Add(GetPackage()->GetFName());
+
+			TArray<FAssetData> Assets;
+			AssetRegistry.GetAssets(Filter, Assets);
+
+			if (Assets.Num() == 1)
+			{
+				const FAssetData& InAssetData = Assets[0];
+
+				FString ActorMetaDataStr;
+				static FName NAME_ActorMetaData(TEXT("ActorMetaData"));
+				if (InAssetData.GetTagValue(NAME_ActorMetaData, ActorMetaDataStr))
+				{
+					FWorldPartitionActorDescInitData ActorDescInitData;
+					ActorDescInitData.NativeClass = AActor::StaticClass();
+					ActorDescInitData.PackageName = InAssetData.PackageName;
+					ActorDescInitData.ActorPath = InAssetData.ObjectPath;
+					FBase64::Decode(ActorMetaDataStr, ActorDescInitData.SerializedData);
+
+					TUniquePtr<FWorldPartitionActorDesc> NewActorDesc(AActor::StaticCreateClassActorDesc(ActorDescInitData.NativeClass));
+					NewActorDesc->Init(ActorDescInitData);
+
+					ActorDesc->TransferWorldData(NewActorDesc.Get());
+				}
+			}
+		}
+
 		const FString ActorMetaDataClass = GetParentNativeClass(GetClass())->GetPathName();
 		static FName NAME_ActorMetaDataClass(TEXT("ActorMetaDataClass"));
 		OutTags.Add(UObject::FAssetRegistryTag(NAME_ActorMetaDataClass, ActorMetaDataClass, UObject::FAssetRegistryTag::TT_Hidden));
diff --git a/Engine/Source/Runtime/Engine/Private/WorldPartition/WorldPartition.cpp b/Engine/Source/Runtime/Engine/Private/WorldPartition/WorldPartition.cpp
--- a/Engine/Source/Runtime/Engine/Private/WorldPartition/WorldPartition.cpp
+++ b/Engine/Source/Runtime/Engine/Private/WorldPartition/WorldPartition.cpp
@@ -462,7 +462,7 @@ void UWorldPartition::Initialize(UWorld* InWorld, const FTransform& InTransform)
 			{
 				if (bIsInstanced)
 				{
-					const FString LongActorPackageName = ActorDescIterator->ActorPackage.ToString();
+					const FString LongActorPackageName = ActorDescIterator->GetActorPackage().ToString();
 					const FString ActorPackageName = FPaths::GetBaseFilename(LongActorPackageName);
 					const FString InstancedName = FString::Printf(TEXT("%s_InstanceOf_%s"), *LevelPackage->GetName(), *ActorPackageName);
 
diff --git a/Engine/Source/Runtime/Engine/Private/WorldPartition/WorldPartitionActorDesc.cpp b/Engine/Source/Runtime/Engine/Private/WorldPartition/WorldPartitionActorDesc.cpp
--- a/Engine/Source/Runtime/Engine/Private/WorldPartition/WorldPartitionActorDesc.cpp
+++ b/Engine/Source/Runtime/Engine/Private/WorldPartition/WorldPartitionActorDesc.cpp
@@ -96,7 +96,6 @@ void FWorldPartitionActorDesc::Init(const FWorldPartitionActorDescInitData& Desc
 	Serialize(MetadataAr);
 
 	Container = nullptr;
-	ActorPtr = FindObject<AActor>(nullptr, *ActorPath.ToString());
 }
 
 bool FWorldPartitionActorDesc::Equals(const FWorldPartitionActorDesc* Other) const
@@ -293,11 +292,21 @@ FName FWorldPartitionActorDesc::GetActorLabelOrName() const
 
 bool FWorldPartitionActorDesc::IsLoaded(bool bEvenIfPendingKill) const
 {
+	if (ActorPtr.IsExplicitlyNull())
+	{
+		ActorPtr = FindObject<AActor>(nullptr, *ActorPath.ToString());
+	}
+
 	return ActorPtr.IsValid(bEvenIfPendingKill);
 }
 
 AActor* FWorldPartitionActorDesc::GetActor(bool bEvenIfPendingKill, bool bEvenIfUnreachable) const
 {
+	if (ActorPtr.IsExplicitlyNull())
+	{
+		ActorPtr = FindObject<AActor>(nullptr, *ActorPath.ToString());
+	}
+
 	return bEvenIfUnreachable ? ActorPtr.GetEvenIfUnreachable() : ActorPtr.Get(bEvenIfPendingKill);
 }
 
@@ -309,7 +318,7 @@ AActor* FWorldPartitionActorDesc::Load() const
 		ActorPtr = FindObject<AActor>(nullptr, *ActorPath.ToString());
 	}
 
-	// The, if the actor isn't loaded, load it
+	// Then, if the actor isn't loaded, load it
 	if (ActorPtr.IsExplicitlyNull())
 	{
 		const FLinkerInstancingContext* InstancingContext = nullptr;
diff --git a/Engine/Source/Runtime/Engine/Public/WorldPartition/WorldPartitionActorDesc.h b/Engine/Source/Runtime/Engine/Public/WorldPartition/WorldPartitionActorDesc.h
--- a/Engine/Source/Runtime/Engine/Public/WorldPartition/WorldPartitionActorDesc.h
+++ b/Engine/Source/Runtime/Engine/Public/WorldPartition/WorldPartitionActorDesc.h
@@ -27,13 +27,6 @@ enum class EContainerClusterMode : uint8
 };
 #endif
 
-#if WITH_DEV_AUTOMATION_TESTS
-namespace WorldPartitionTests
-{
-	class FWorldPartitionSoftRefTest;
-}
-#endif
-
 /**
  * Represents a potentially unloaded actor (editor-only)
  */
@@ -41,17 +34,9 @@ class ENGINE_API FWorldPartitionActorDesc
 {
 #if WITH_EDITOR
 	friend class AActor;
-	friend class UWorldPartition;
 	friend class UActorDescContainer;
-	friend class UWorldPartitionRuntimeHash;
-	friend class FWorldPartitionStreamingGenerator;
 	friend struct FWorldPartitionHandleImpl;
 	friend struct FWorldPartitionReferenceImpl;
-	friend struct FWorldPartitionHandleUtils;
-
-#if WITH_DEV_AUTOMATION_TESTS
-	friend class WorldPartitionTests::FWorldPartitionSoftRefTest;
-#endif
 
 public:
 	virtual ~FWorldPartitionActorDesc() {}
@@ -77,7 +62,6 @@ public:
 	inline const FGuid& GetParentActor() const { return ParentActor; }
 
 	FName GetActorName() const;
-
 	FName GetActorLabelOrName() const;
 
 	virtual bool GetContainerInstance(const UActorDescContainer*& OutLevelContainer, FTransform& OutLevelTransform, EContainerClusterMode& OutClusterMode) const { return false; }
@@ -104,11 +88,6 @@ protected:
 		return --SoftRefCount;
 	}
 
-	inline uint32 GetSoftRefCount() const
-	{
-		return SoftRefCount;
-	}
-
 	inline uint32 IncHardRefCount() const
 	{
 		return ++HardRefCount;
@@ -120,11 +99,6 @@ protected:
 		return --HardRefCount;
 	}
 
-	inline uint32 GetHardRefCount() const
-	{
-		return HardRefCount;
-	}
-
 	virtual void SetContainer(UActorDescContainer* InContainer)
 	{
 		check(!Container || !InContainer);
@@ -132,6 +106,16 @@ protected:
 	}
 
 public:
+	inline uint32 GetSoftRefCount() const
+	{
+		return SoftRefCount;
+	}
+
+	inline uint32 GetHardRefCount() const
+	{
+		return HardRefCount;
+	}
+
 	const TArray<FGuid>& GetReferences() const
 	{
 		return References;
@@ -160,11 +144,11 @@ public:
 
 	void SerializeTo(TArray<uint8>& OutData);
 
+	void TransformInstance(const FString& From, const FString& To, const FTransform& InstanceTransform);
+
 protected:
 	FWorldPartitionActorDesc();
 
-	void TransformInstance(const FString& From, const FString& To, const FTransform& InstanceTransform);
-
 	virtual void TransferFrom(const FWorldPartitionActorDesc* From)
 	{
 		Container = From->Container;
@@ -172,6 +156,12 @@ protected:
 		HardRefCount = From->HardRefCount;
 	}
 
+	virtual void TransferWorldData(const FWorldPartitionActorDesc* From)
+	{
+		BoundsLocation = From->BoundsLocation;
+		BoundsExtent = From->BoundsExtent;
+	}
+
 	virtual void Serialize(FArchive& Ar);
 
 	// Persistent
