Enable Editor ODSC.

- Drastically reduce the number of shaders compiled when working in the editor.
- On average a 1.6x speedup when opening the editor and a map in various projects.
- Best case a 2.23x speedup.

diff --git a/Engine/Source/Runtime/Engine/Private/ShaderCompiler/ShaderCompiler.cpp b/Engine/Source/Runtime/Engine/Private/ShaderCompiler/ShaderCompiler.cpp
--- a/Engine/Source/Runtime/Engine/Private/ShaderCompiler/ShaderCompiler.cpp
+++ b/Engine/Source/Runtime/Engine/Private/ShaderCompiler/ShaderCompiler.cpp
@@ -116,7 +116,7 @@ static FAutoConsoleVariableRef CVarShaderCompilerMaxJobCacheMemoryPercent(
 
 static TAutoConsoleVariable<bool> CVarJobCacheDDC(
 	TEXT("r.ShaderCompiler.JobCacheDDC"),
-	false,
+	true,
 	TEXT("Skips compilation of all shaders on Material and Material Instance PostLoad and relies on on-demand shader compilation to compile what is needed."),
 	ECVF_ReadOnly);
 
