Change how UNetDriver::bInTick is set, so it has the correct value during derived class' tick functions

diff --git a/Engine/Source/Runtime/Engine/Private/NetDriver.cpp b/Engine/Source/Runtime/Engine/Private/NetDriver.cpp
--- a/Engine/Source/Runtime/Engine/Private/NetDriver.cpp
+++ b/Engine/Source/Runtime/Engine/Private/NetDriver.cpp
@@ -705,7 +705,7 @@ void UNetDriver::TickFlush(float DeltaSeconds)
 {
 	LLM_SCOPE_BYTAG(NetDriver);
 
-	TGuardValue<bool> GuardInNetTick(bInTick, true);
+	bInTick = true;
 
 	CSV_SCOPED_TIMING_STAT_EXCLUSIVE(NetworkOutgoing);
 	SCOPE_CYCLE_COUNTER(STAT_NetTickFlush);
@@ -1541,6 +1541,7 @@ void UNetDriver::PostTickFlush()
 		UOnlineEngineInterface::Get()->ClearVoicePackets(World);
 	}
 
+	bInTick = false;
 	if (bPendingDestruction)
 	{
 		if (World)
@@ -1916,7 +1917,7 @@ struct FRPCCSVTracker
 
 void UNetDriver::TickDispatch( float DeltaTime )
 {
-	TGuardValue<bool> GuardInNetTick(bInTick, true);
+	bInTick = true;
 
 	SendCycles=0;
 
@@ -2030,6 +2031,7 @@ void UNetDriver::PostTickDispatch()
 		GReceiveRPCTimingEnabled = false;
 	}
 
+	bInTick = false;
 	if (bPendingDestruction)
 	{
 		if (World)
