fix VirtualTextureDataBuilder TBSettings missing some fields
comment about strange behavior

diff --git a/Engine/Source/Runtime/Engine/Private/VT/VirtualTextureDataBuilder.cpp b/Engine/Source/Runtime/Engine/Private/VT/VirtualTextureDataBuilder.cpp
--- a/Engine/Source/Runtime/Engine/Private/VT/VirtualTextureDataBuilder.cpp
+++ b/Engine/Source/Runtime/Engine/Private/VT/VirtualTextureDataBuilder.cpp
@@ -690,6 +690,9 @@ void FVirtualTextureDataBuilder::BuildTiles(const TArray<FVTSourceTileEntry>& Ti
 		// Create settings for building the tile. These should be simple, "clean" settings
 		// just compressing the style to a GPU format not adding things like colour correction, ... 
 		// as these settings were already baked into the SourcePixels.
+
+		// TBSettings starts default constructed (no processing)
+		//	then we copy over just the compression options without the color-change processing
 		FTextureBuildSettings TBSettings;
 		TBSettings.MaxTextureResolution = TNumericLimits<uint32>::Max();
 		TBSettings.TextureFormatName = LayerData.TextureFormatName;
@@ -697,11 +700,14 @@ void FVirtualTextureDataBuilder::BuildTiles(const TArray<FVTSourceTileEntry>& Ti
 		TBSettings.bUseLegacyGamma = BuildSettingsForLayer.bUseLegacyGamma;
 		TBSettings.MipGenSettings = TMGS_NoMipmaps;
 		TBSettings.bForceAlphaChannel = BuildSettingsForLayer.bForceAlphaChannel;
+		TBSettings.bForceNoAlphaChannel = BuildSettingsForLayer.bForceNoAlphaChannel;
+		TBSettings.bHDRSource = BuildSettingsForLayer.bHDRSource;
 		TBSettings.bVirtualStreamable = true;
 
 		// Encode speed must be resolved before we get here.
-		TBSettings.OodleEncodeEffort = BuildSettingsForLayer.OodleEncodeEffort;
 		TBSettings.LossyCompressionAmount = BuildSettingsForLayer.LossyCompressionAmount;
+		TBSettings.CompressionQuality = BuildSettingsForLayer.CompressionQuality;
+		TBSettings.OodleEncodeEffort = BuildSettingsForLayer.OodleEncodeEffort;
 		TBSettings.OodleUniversalTiling = BuildSettingsForLayer.OodleUniversalTiling;
 		TBSettings.bOodleUsesRDO = BuildSettingsForLayer.bOodleUsesRDO;
 		TBSettings.OodleRDO = BuildSettingsForLayer.OodleRDO;
@@ -1025,15 +1031,32 @@ void FVirtualTextureDataBuilder::BuildSourcePixels(const FTextureSourceData& Sou
 
 			// Adjust the build settings to generate an uncompressed texture with mips but leave other settings
 			// like color correction, ... in place
+
+			// TBSettings starts with the full Texture settings, so we get all options
+			//  then we change FormatName to be == Source format, so no Compression is done
 			FTextureBuildSettings TBSettings = SettingsPerLayer[0];
 			//TBSettings.MaxTextureResolution = TNumericLimits<uint32>::Max();
 			TBSettings.TextureFormatName = LayerData.FormatName;
-			TBSettings.bSRGB = BuildSettingsForLayer.bSRGB;
-			TBSettings.bUseLegacyGamma = BuildSettingsForLayer.bUseLegacyGamma;
-			TBSettings.bForceAlphaChannel = BuildSettingsForLayer.bForceAlphaChannel;
-			TBSettings.bApplyYCoCgBlockScale = BuildSettingsForLayer.bApplyYCoCgBlockScale;
-			TBSettings.bReplicateRed = BuildSettingsForLayer.bReplicateRed;
-			TBSettings.bReplicateAlpha = BuildSettingsForLayer.bReplicateAlpha;
+			
+			if ( LayerIndex != 0 )
+			{
+				// @todo Oodle : this looks like a mess
+				//	some of the processing options are copied from BuildSettingsForLayer
+				//	but some are NOT
+				//	it seems semi-random
+				//  In the common case of NumLayers==1 , then it doesn't matter
+				//	so this would be rarely observed
+
+				// TBSettings was set from Layer 0, copy in some settings from this Layer ?
+				TBSettings.bSRGB = BuildSettingsForLayer.bSRGB;
+				TBSettings.bUseLegacyGamma = BuildSettingsForLayer.bUseLegacyGamma;
+				TBSettings.bForceAlphaChannel = BuildSettingsForLayer.bForceAlphaChannel;
+				TBSettings.bForceNoAlphaChannel = BuildSettingsForLayer.bForceNoAlphaChannel;
+				TBSettings.bHDRSource = BuildSettingsForLayer.bHDRSource;
+				TBSettings.bApplyYCoCgBlockScale = BuildSettingsForLayer.bApplyYCoCgBlockScale;
+				TBSettings.bReplicateRed = BuildSettingsForLayer.bReplicateRed;
+				TBSettings.bReplicateAlpha = BuildSettingsForLayer.bReplicateAlpha;
+			}
 
 			// Make sure the output of the texture builder is in the same gamma space as we expect it.
 			check(TBSettings.GetGammaSpace() == BuildSettingsForLayer.GetGammaSpace());
