SubObject API
* ActorComponent's that support the registered list can now replicate their subobjects using the new method even if their  owner actor does not support the new list.

diff --git a/Engine/Source/Runtime/Engine/Classes/Engine/ActorChannel.h b/Engine/Source/Runtime/Engine/Classes/Engine/ActorChannel.h
--- a/Engine/Source/Runtime/Engine/Classes/Engine/ActorChannel.h
+++ b/Engine/Source/Runtime/Engine/Classes/Engine/ActorChannel.h
@@ -8,12 +8,14 @@
 #include "Misc/NetworkGuid.h"
 #include "Engine/Channel.h"
 #include "Net/DataReplication.h"
+#include "Containers/StaticBitArray.h"
 #include "ActorChannel.generated.h"
 
 #ifndef NET_ENABLE_SUBOBJECT_REPKEYS
 #define NET_ENABLE_SUBOBJECT_REPKEYS 1
 #endif // NET_ENABLE_SUBOBJECT_REPKEYS
 
+class UActorComponent;
 class AActor;
 class FInBunch;
 class FNetFieldExportGroup;
@@ -269,8 +271,18 @@ public:
 	//	
 	// --------------------------------
 
+    /** 
+    * Sets the owner of the next replicated subobjects that will be passed to ReplicateSubObject.
+    * The ActorComponent version will choose not to replicate the future subobjects if the component opted to use the registered list.
+    */
+	static void SetCurrentSubObjectOwner(AActor* SubObjectOwner);
+	static void SetCurrentSubObjectOwner(UActorComponent* SubObjectOwner);
+
 	/** Replicates given subobject on this actor channel */
-	bool ReplicateSubobject(UObject *Obj, FOutBunch &Bunch, const FReplicationFlags &RepFlags);
+	bool ReplicateSubobject(UObject* Obj, FOutBunch& Bunch, FReplicationFlags RepFlags);
+
+	/** Replicates an ActorComponent subobject. Used to catch ActorChannels who are using the registration list while their Owner is not flagged to use it */
+	bool ReplicateSubobject(UActorComponent* ActorChannel, FOutBunch& Bunch, FReplicationFlags RepFlags);
 	
 	/** Custom implementation for ReplicateSubobject when RepFlags.bUseCustomSubobjectReplication is true */
 	virtual bool ReplicateSubobjectCustom(UObject* Obj, FOutBunch& Bunch, const FReplicationFlags& RepFlags) { return true;  }
@@ -372,13 +384,16 @@ protected:
 private:
 
 	/** Handle the replication of subobjects for this actor. Returns true if data was written into the Bunch. */
-	bool DoSubObjectReplication(FOutBunch& Bunch, const FReplicationFlags& RepFlags);
+	bool DoSubObjectReplication(FOutBunch& Bunch, FReplicationFlags& OutRepFlags);
 
 	/** Replicate Subobjects using the actor's registered list and its replicated actor component list */
-	bool ReplicateRegisteredSubObjects(FOutBunch& Bunch, const FReplicationFlags& RepFlags);
+	bool ReplicateRegisteredSubObjects(FOutBunch& Bunch, FReplicationFlags RepFlags);
 
     /** Write the replicated bits into the bunch data */
-	bool WriteSubObjectInBunch(UObject* Obj, FOutBunch& Bunch, const FReplicationFlags& RepFlags);
+	bool WriteSubObjectInBunch(UObject* Obj, FOutBunch& Bunch, FReplicationFlags RepFlags);
+
+	/** Find the replicated subobjects of the component and write them into the bunch */
+	bool WriteSubObjects(UActorComponent* Component, FOutBunch& Bunch, FReplicationFlags RepFlags, const TStaticBitArray<COND_Max>& ConditionMap);
 
 private:
 
diff --git a/Engine/Source/Runtime/Engine/Private/ActorReplication.cpp b/Engine/Source/Runtime/Engine/Private/ActorReplication.cpp
--- a/Engine/Source/Runtime/Engine/Private/ActorReplication.cpp
+++ b/Engine/Source/Runtime/Engine/Private/ActorReplication.cpp
@@ -476,6 +476,7 @@ bool AActor::ReplicateSubobjects(UActorChannel *Channel, FOutBunch *Bunch, FRepl
 	{
 		if (ActorComp && ActorComp->GetIsReplicated())
 		{
+			UActorChannel::SetCurrentSubObjectOwner(ActorComp);
 			WroteSomething |= ActorComp->ReplicateSubobjects(Channel, Bunch, RepFlags);		// Lets the component add subobjects before replicating its own properties.
 			WroteSomething |= Channel->ReplicateSubobject(ActorComp, *Bunch, *RepFlags);	// (this makes those subobjects 'supported', and from here on those objects may have reference replicated)		
 		}
diff --git a/Engine/Source/Runtime/Engine/Private/Components/ActorComponent.cpp b/Engine/Source/Runtime/Engine/Private/Components/ActorComponent.cpp
--- a/Engine/Source/Runtime/Engine/Private/Components/ActorComponent.cpp
+++ b/Engine/Source/Runtime/Engine/Private/Components/ActorComponent.cpp
@@ -21,6 +21,7 @@
 #include "UObject/PropertyPortFlags.h"
 #include "UObject/UObjectHash.h"
 #include "Engine/NetDriver.h"
+#include "Engine/ActorChannel.h"
 #include "Net/UnrealNetwork.h"
 #include "Logging/TokenizedMessage.h"
 #include "Logging/MessageLog.h"
@@ -2095,8 +2096,9 @@ void UActorComponent::RemoveReplicatedSubObject(UObject* SubObject)
 	}
 }
 
-bool UActorComponent::ReplicateSubobjects(class UActorChannel* Channel, class FOutBunch* Bunch, FReplicationFlags* RepFlags)
+bool UActorComponent::ReplicateSubobjects(UActorChannel* Channel, FOutBunch* Bunch, FReplicationFlags* RepFlags)
 {
+	UActorChannel::SetCurrentSubObjectOwner(this);
 	return false;
 }
 
diff --git a/Engine/Source/Runtime/Engine/Private/DataChannel.cpp b/Engine/Source/Runtime/Engine/Private/DataChannel.cpp
--- a/Engine/Source/Runtime/Engine/Private/DataChannel.cpp
+++ b/Engine/Source/Runtime/Engine/Private/DataChannel.cpp
@@ -141,8 +141,30 @@ public:
 		AActor::FReplicatedComponentInfo* ComponentInfo = InActor->ReplicatedComponentsInfo.FindByKey(InActorComp);
 		return ComponentInfo ? &(ComponentInfo->SubObjects.GetRegistryList()) : nullptr;
 	}
+
+	static bool IsSubObjectInRegistry(AActor* InActor, UActorComponent* InActorComp, UObject* InSubObject)
+	{
+		AActor::FReplicatedComponentInfo* ComponentInfo = InActor->ReplicatedComponentsInfo.FindByKey(InActorComp);
+		return ComponentInfo ? ComponentInfo->SubObjects.IsSubObjectInRegistry(InSubObject) : false;
+	}
 };
 
+namespace DataChannelInternal
+{
+	/**
+	* Track the owner of the next subobjects replicated via ReplicateSubObject.
+	* Nulled out after every call to UActorChannel::ReplicateActor()
+	*/
+	UObject* CurrentSubObjectOwner = nullptr;
+
+	/**
+	* When enabled we will ignore requests to replicate a subobject via the public ReplicateSubObject function.
+	* This occurs when the subobjects belong to an ActorComponent converted to the registered list, but the owner does not support it yet so we still call ReplicateSubObjects on the actor.
+	* This flag is controlled via UActorChannel::SetCurrentSubObjectOwner, so it is important for users set the owner properly if they are controlling component replication themselves.
+	*/
+	bool bIgnoreLegacyReplicateSubObject = false;
+}
+
 /*-----------------------------------------------------------------------------
 	UChannel implementation.
 -----------------------------------------------------------------------------*/
@@ -3201,7 +3223,6 @@ int64 UActorChannel::ReplicateActor()
 	{
 		Bunch.bReliable = true;
 		Bunch.bIsReplicationPaused = true;
-
 	}
 
 #if !(UE_BUILD_SHIPPING || UE_BUILD_TEST)
@@ -3423,25 +3444,25 @@ int64 UActorChannel::ReplicateActor()
 	return NumBitsWrote;
 }
 
-bool UActorChannel::DoSubObjectReplication(FOutBunch& Bunch, const FReplicationFlags& RepFlags)
+bool UActorChannel::DoSubObjectReplication(FOutBunch& Bunch, FReplicationFlags& OutRepFlags)
 {
 	bool bWroteSomethingImportant = false;
 
 	if (Actor->IsUsingRegisteredSubObjectList())
 	{
 		// If the actor replicates it's subobjects and those of it's components via the list.
-		bWroteSomethingImportant |= ReplicateRegisteredSubObjects(Bunch, RepFlags);
+		bWroteSomethingImportant |= ReplicateRegisteredSubObjects(Bunch, OutRepFlags);
 	}
 	else
 	{
 		// Replicate the subobjects using the virtual method.
-		bWroteSomethingImportant |= Actor->ReplicateSubobjects(this, &Bunch, const_cast<FReplicationFlags*>(&RepFlags));
+		bWroteSomethingImportant |= Actor->ReplicateSubobjects(this, &Bunch, &OutRepFlags);
 	}
 
 	return bWroteSomethingImportant;
 }
 
-bool UActorChannel::ReplicateRegisteredSubObjects(FOutBunch& Bunch, const FReplicationFlags& RepFlags)
+bool UActorChannel::ReplicateRegisteredSubObjects(FOutBunch& Bunch, FReplicationFlags RepFlags)
 {
 	check(Actor->IsUsingRegisteredSubObjectList());
 
@@ -3451,6 +3472,7 @@ bool UActorChannel::ReplicateRegisteredSubObjects(FOutBunch& Bunch, const FRepli
 
 	// Start with the Actor's subobjects
 	{
+		SetCurrentSubObjectOwner(Actor);
 		const TArray<UE::Net::FSubObjectRegistry::FEntry>& ActorSubObjects = FSubObjectGetter::GetSubObjects(Actor);
 		for (const UE::Net::FSubObjectRegistry::FEntry& SubObjectInfo : ActorSubObjects)
 		{
@@ -3466,31 +3488,118 @@ bool UActorChannel::ReplicateRegisteredSubObjects(FOutBunch& Bunch, const FRepli
 	{
 		check(ReplicatedComponent);
 
+		SetCurrentSubObjectOwner(ReplicatedComponent);
+
 		if (ReplicatedComponent->IsUsingRegisteredSubObjectList())
 		{
-			if (const TArray<UE::Net::FSubObjectRegistry::FEntry>* ComponentSubObjects = FSubObjectGetter::GetSubObjectsOfActorCompoment(Actor, ReplicatedComponent))
-			{
-				for (const UE::Net::FSubObjectRegistry::FEntry& SubObjectInfo : *ComponentSubObjects)
-				{
-					checkf(IsValid(SubObjectInfo.SubObject), TEXT("Found invalid subobject (%s) registered in %s"), *GetNameSafe(SubObjectInfo.SubObject), *GetNameSafe(ReplicatedComponent));
-
-					if (ConditionMap[SubObjectInfo.NetCondition])
-					{
-						bWroteSomethingImportant |= WriteSubObjectInBunch(SubObjectInfo.SubObject, Bunch, RepFlags);
-					}
-				}
-			}
+			bWroteSomethingImportant |= WriteSubObjects(ReplicatedComponent, Bunch, RepFlags, ConditionMap);
 		}
-		// If a component is not using the registered list, collect the subobjects to replicate via the virtual function.
 		else
 		{
-			// Replicate the subobjects using the virtual method
+			// Replicate the component's SubObjects using the virtual method
 			ReplicatedComponent->ReplicateSubobjects(this, &Bunch, const_cast<FReplicationFlags*>(&RepFlags));
 		}
 
-		// Finally replicate the component itself after it's subobjects. They need to be created before the component on the receiving end.
+		// Finally replicate the component itself at the end. SubObjects have to be created before the component on the receiving end.
+		SetCurrentSubObjectOwner(Actor);
+		bWroteSomethingImportant |= WriteSubObjectInBunch(ReplicatedComponent, Bunch, RepFlags);
+	}
+
+	return bWroteSomethingImportant;
+}
+
+
+void UActorChannel::SetCurrentSubObjectOwner(AActor* SubObjectOwner)
+{
+	DataChannelInternal::CurrentSubObjectOwner = SubObjectOwner;
+	DataChannelInternal::bIgnoreLegacyReplicateSubObject = false;
+}
+
+void UActorChannel::SetCurrentSubObjectOwner(UActorComponent* SubObjectOwner)
+{
+	DataChannelInternal::CurrentSubObjectOwner = SubObjectOwner;
+	DataChannelInternal::bIgnoreLegacyReplicateSubObject = SubObjectOwner && SubObjectOwner->IsUsingRegisteredSubObjectList();
+}
+
+bool UActorChannel::ReplicateSubobject(UObject* SubObj, FOutBunch& Bunch, FReplicationFlags RepFlags)
+{
+	if (!IsValid(SubObj))
+	{
+		return false;
+	}
+
+	bool bWroteSomethingImportant = false;
+
+	if (!DataChannelInternal::bIgnoreLegacyReplicateSubObject)
+	{
+		bWroteSomethingImportant = WriteSubObjectInBunch(SubObj, Bunch, RepFlags);
+	}
+	else
+	{
+#if !UE_BUILD_SHIPPING
+		// Make sure the SubObjectOwner actually has the SubObject in his list and will replicate it later.
+		UActorComponent* Component = CastChecked<UActorComponent>(DataChannelInternal::CurrentSubObjectOwner);
+		const bool bIsInRegistry = FSubObjectGetter::IsSubObjectInRegistry(Actor, Component, SubObj);
+		ensureMsgf(bIsInRegistry, TEXT("ReplicatedSubObject %s was replicated using the legacy method but it is not in %s::%s registered list. It won't be replicated at all."), 
+			*SubObj->GetName(), *Actor->GetName(), *Component->GetName());
+#endif
+	}
+
+	return bWroteSomethingImportant;
+}
+
+
+bool UActorChannel::ReplicateSubobject(UActorComponent* ReplicatedComponent, FOutBunch& Bunch, FReplicationFlags RepFlags)
+{
+    // This function traps ActorComponents using the registration list even if their owning Actor isn't using the list itself.
+
+	if (!IsValid(ReplicatedComponent))
+	{
+		return false;
+	}
+
+	bool bWroteSomethingImportant = false;
+
+	if (ReplicatedComponent->IsUsingRegisteredSubObjectList())
+	{
+		const TStaticBitArray<COND_Max> ConditionMap = FSendingRepState::BuildConditionMapFromRepFlags(RepFlags);
+
+		checkf(Actor->IsUsingRegisteredSubObjectList() == false, TEXT("This code should only be hit when an Actor that does NOT support the SubObjectList is replicating an ActorComponent that does: %s replicating %s."), 
+			*Actor->GetName(), *ReplicatedComponent->GetName());
+
+		ensureMsgf(DataChannelInternal::CurrentSubObjectOwner == ReplicatedComponent, TEXT("%s (owner %s) was not the CurrentSubObjectOwner prior to replicating his subobjects. Make sure to call UActorChannel::SetCurrentSubObjectOwner BEFORE replicating any of the component's subobjects."),
+			*ReplicatedComponent->GetName(), *Actor->GetName());
+
+        // Write all the subobjects of the component
+		bWroteSomethingImportant |= WriteSubObjects(ReplicatedComponent, Bunch, RepFlags, ConditionMap);
+
+        // Finally write the component itself
+		SetCurrentSubObjectOwner(Actor);
 		bWroteSomethingImportant |= WriteSubObjectInBunch(ReplicatedComponent, Bunch, RepFlags);
 	}
+	else
+	{
+		return ReplicateSubobject(StaticCast<UObject*>(ReplicatedComponent), Bunch, RepFlags);
+	}
+
+	return bWroteSomethingImportant;
+}
+
+bool UActorChannel::WriteSubObjects(UActorComponent* ReplicatedComponent, FOutBunch& Bunch, FReplicationFlags RepFlags, const TStaticBitArray<COND_Max>& ConditionMap)
+{
+	bool bWroteSomethingImportant = false;
+	if (const TArray<UE::Net::FSubObjectRegistry::FEntry>* ComponentSubObjects = FSubObjectGetter::GetSubObjectsOfActorCompoment(Actor, ReplicatedComponent))
+	{
+		for (const UE::Net::FSubObjectRegistry::FEntry& SubObjectInfo : *ComponentSubObjects)
+		{
+			checkf(IsValid(SubObjectInfo.SubObject), TEXT("Found invalid subobject (%s) registered in %s"), *GetNameSafe(SubObjectInfo.SubObject), *GetNameSafe(ReplicatedComponent));
+
+			if (ConditionMap[SubObjectInfo.NetCondition])
+			{
+				bWroteSomethingImportant |= WriteSubObjectInBunch(SubObjectInfo.SubObject, Bunch, RepFlags);
+			}
+		}
+	}
 
 	return bWroteSomethingImportant;
 }
@@ -4484,19 +4593,7 @@ void UActorChannel::AddedToChannelPool()
 #endif // NET_ENABLE_SUBOBJECT_REPKEYS
 }
 
-bool UActorChannel::ReplicateSubobject(UObject* Obj, FOutBunch& Bunch, const FReplicationFlags& RepFlags)
-{
-	if (!IsValid(Obj))
-	{
-		return false;
-	}
-
-	const bool bWroteSomethingImportant = WriteSubObjectInBunch(Obj, Bunch, RepFlags);
-
-	return bWroteSomethingImportant;
-}
-
-bool UActorChannel::WriteSubObjectInBunch(UObject * Obj, FOutBunch & Bunch, const FReplicationFlags & RepFlags)
+bool UActorChannel::WriteSubObjectInBunch(UObject* Obj, FOutBunch& Bunch, FReplicationFlags RepFlags)
 {
 	if (!IsValid(Obj))
 	{
diff --git a/Engine/Source/Runtime/Net/Core/Private/Net/Core/Misc/NetSubObjectRegistry.cpp b/Engine/Source/Runtime/Net/Core/Private/Net/Core/Misc/NetSubObjectRegistry.cpp
--- a/Engine/Source/Runtime/Net/Core/Private/Net/Core/Misc/NetSubObjectRegistry.cpp
+++ b/Engine/Source/Runtime/Net/Core/Private/Net/Core/Misc/NetSubObjectRegistry.cpp
@@ -30,7 +30,7 @@ FSubObjectRegistry::EResult FSubObjectRegistry::AddSubObjectUnique(UObject* InSu
 
 ELifetimeCondition FSubObjectRegistry::GetNetCondition(UObject* InSubObject) const
 {
-	const int32 CurrentIndex = Registry.Find(FEntry{ InSubObject });
+	const int32 CurrentIndex = Registry.IndexOfByKey(InSubObject);
 	return CurrentIndex == INDEX_NONE ? COND_Max : Registry[CurrentIndex].NetCondition;
 }
 
@@ -39,5 +39,10 @@ bool FSubObjectRegistry::RemoveSubObject(UObject* InSubObject)
 	return Registry.RemoveSingle(FEntry{ InSubObject }) == 1;
 }
 
+bool FSubObjectRegistry::IsSubObjectInRegistry(UObject* SubObject) const
+{
+	return Registry.FindByKey(SubObject) != nullptr;
+}
+
 } //namespace Net
 } //namespace UE
\ No newline at end of file
diff --git a/Engine/Source/Runtime/Net/Core/Public/Net/Core/Misc/NetSubObjectRegistry.h b/Engine/Source/Runtime/Net/Core/Public/Net/Core/Misc/NetSubObjectRegistry.h
--- a/Engine/Source/Runtime/Net/Core/Public/Net/Core/Misc/NetSubObjectRegistry.h
+++ b/Engine/Source/Runtime/Net/Core/Public/Net/Core/Misc/NetSubObjectRegistry.h
@@ -43,11 +43,15 @@ public:
 		ELifetimeCondition NetCondition = COND_None;
 
 		bool operator==(const FEntry& rhs) const { return SubObject == rhs.SubObject; }
+		bool operator==(UObject* rhs) const { return SubObject == rhs; }
 	};
 
 	/** Returns the list of registered subobjects */
 	const TArray<FSubObjectRegistry::FEntry>& GetRegistryList() const { return Registry; }
 
+	/** Returns true if the subobject is contained in this list */
+	bool IsSubObjectInRegistry(UObject* SubObject) const;
+
 private:
 
 	TArray<FEntry> Registry;
