Fix EngineNetTest failing due to missing WroteSomethingImportant assignement when replicating ActorComponent not using subobject list

diff --git a/Engine/Source/Runtime/Engine/Private/DataChannel.cpp b/Engine/Source/Runtime/Engine/Private/DataChannel.cpp
--- a/Engine/Source/Runtime/Engine/Private/DataChannel.cpp
+++ b/Engine/Source/Runtime/Engine/Private/DataChannel.cpp
@@ -3497,7 +3497,7 @@ bool UActorChannel::ReplicateRegisteredSubObjects(FOutBunch& Bunch, FReplication
 		else
 		{
 			// Replicate the component's SubObjects using the virtual method
-			ReplicatedComponent->ReplicateSubobjects(this, &Bunch, const_cast<FReplicationFlags*>(&RepFlags));
+			bWroteSomethingImportant |= ReplicatedComponent->ReplicateSubobjects(this, &Bunch, const_cast<FReplicationFlags*>(&RepFlags));
 		}
 
 		// Finally replicate the component itself at the end. SubObjects have to be created before the component on the receiving end.
