Granular mem tracking improvements to replayhelper

diff --git a/Engine/Source/Runtime/Engine/Private/ReplayHelper.cpp b/Engine/Source/Runtime/Engine/Private/ReplayHelper.cpp
--- a/Engine/Source/Runtime/Engine/Private/ReplayHelper.cpp
+++ b/Engine/Source/Runtime/Engine/Private/ReplayHelper.cpp
@@ -1613,6 +1613,29 @@ bool FReplayHelper::SetExternalDataForObject(UNetConnection* Connection, UObject
 	return false;
 }
 
+void FDeltaCheckpointData::CountBytes(FArchive& Ar) const
+{
+	GRANULAR_NETWORK_MEMORY_TRACKING_INIT(Ar, "FDeltaCheckpointData::CountBytes");
+
+	GRANULAR_NETWORK_MEMORY_TRACKING_TRACK("RecordingDeletedNetStartupActors", RecordingDeletedNetStartupActors.CountBytes(Ar));
+	GRANULAR_NETWORK_MEMORY_TRACKING_TRACK("DestroyedNetStartupActors", DestroyedNetStartupActors.CountBytes(Ar));
+	GRANULAR_NETWORK_MEMORY_TRACKING_TRACK("DestroyedDynamicActors", DestroyedDynamicActors.CountBytes(Ar));
+	GRANULAR_NETWORK_MEMORY_TRACKING_TRACK("ChannelsToClose", ChannelsToClose.CountBytes(Ar));
+}
+
+void FReplayHelper::FCheckpointSaveStateContext::CountBytes(FArchive& Ar) const
+{
+	GRANULAR_NETWORK_MEMORY_TRACKING_INIT(Ar, "FCheckpointSaveStateContext::CountBytes");
+
+	GRANULAR_NETWORK_MEMORY_TRACKING_TRACK("CheckpointAckState", CheckpointAckState.CountBytes(Ar));
+	GRANULAR_NETWORK_MEMORY_TRACKING_TRACK("PendingCheckpointActors", PendingCheckpointActors.CountBytes(Ar));
+	GRANULAR_NETWORK_MEMORY_TRACKING_TRACK("DeltaCheckpointData", DeltaCheckpointData.CountBytes(Ar));
+	GRANULAR_NETWORK_MEMORY_TRACKING_TRACK("DeltaChannelCloseKeys", DeltaChannelCloseKeys.CountBytes(Ar));
+	GRANULAR_NETWORK_MEMORY_TRACKING_TRACK("NetGuidCacheSnapshot", NetGuidCacheSnapshot.CountBytes(Ar));
+	GRANULAR_NETWORK_MEMORY_TRACKING_TRACK("CheckpointDeletedNetStartupActors", CheckpointDeletedNetStartupActors.CountBytes(Ar));
+	GRANULAR_NETWORK_MEMORY_TRACKING_TRACK("NameTableMap", NameTableMap.CountBytes(Ar));
+}
+
 void FReplayHelper::Serialize(FArchive& Ar)
 {
 	GRANULAR_NETWORK_MEMORY_TRACKING_INIT(Ar, "FReplayHelper::Serialize");
diff --git a/Engine/Source/Runtime/Engine/Public/ReplayHelper.h b/Engine/Source/Runtime/Engine/Public/ReplayHelper.h
--- a/Engine/Source/Runtime/Engine/Public/ReplayHelper.h
+++ b/Engine/Source/Runtime/Engine/Public/ReplayHelper.h
@@ -302,16 +302,7 @@ private:
 
 		TMap<FName, uint32> NameTableMap;
 
-		void CountBytes(FArchive& Ar) const
-		{
-			CheckpointAckState.CountBytes(Ar);
-			PendingCheckpointActors.CountBytes(Ar);
-			DeltaCheckpointData.CountBytes(Ar);
-			DeltaChannelCloseKeys.CountBytes(Ar);
-			NetGuidCacheSnapshot.CountBytes(Ar);
-			CheckpointDeletedNetStartupActors.CountBytes(Ar);
-			NameTableMap.CountBytes(Ar);
-		}
+		void CountBytes(FArchive& Ar) const;
 	};
 
 	FCheckpointSaveStateContext CheckpointSaveContext;
diff --git a/Engine/Source/Runtime/Engine/Public/ReplayTypes.h b/Engine/Source/Runtime/Engine/Public/ReplayTypes.h
--- a/Engine/Source/Runtime/Engine/Public/ReplayTypes.h
+++ b/Engine/Source/Runtime/Engine/Public/ReplayTypes.h
@@ -253,13 +253,7 @@ struct FDeltaCheckpointData
 	/** Channels closed that were open in the previous checkpoint, and the reason why */
 	TMap<FNetworkGUID, EChannelCloseReason> ChannelsToClose;
 
-	void CountBytes(FArchive& Ar) const
-	{
-		RecordingDeletedNetStartupActors.CountBytes(Ar);
-		DestroyedNetStartupActors.CountBytes(Ar);
-		DestroyedDynamicActors.CountBytes(Ar);
-		ChannelsToClose.CountBytes(Ar);
-	}
+	void CountBytes(FArchive& Ar) const;
 };
 
 class FRepActorsCheckpointParams
