Support material layers with new HLSL translator

diff --git a/Engine/Source/Runtime/Engine/Classes/Materials/MaterialExpressionBlendMaterialAttributes.h b/Engine/Source/Runtime/Engine/Classes/Materials/MaterialExpressionBlendMaterialAttributes.h
--- a/Engine/Source/Runtime/Engine/Classes/Materials/MaterialExpressionBlendMaterialAttributes.h
+++ b/Engine/Source/Runtime/Engine/Classes/Materials/MaterialExpressionBlendMaterialAttributes.h
@@ -54,6 +54,7 @@ class UMaterialExpressionBlendMaterialAttributes : public UMaterialExpression
 	virtual void GatherStrataMaterialInfo(FStrataMaterialInfo& StrataMaterialInfo, int32 OutputIndex) override;
 	virtual FStrataOperator* StrataGenerateMaterialTopologyTree(class FMaterialCompiler* Compiler, class UMaterialExpression* Parent, int32 OutputIndex) override;
 	virtual uint32 GetInputType(int32 InputIndex) override {return InputIndex == 2 ? MCT_Float1 : MCT_MaterialAttributes;}
+	virtual bool GenerateHLSLExpression(FMaterialHLSLGenerator& Generator, UE::HLSLTree::FScope& Scope, int32 OutputIndex, UE::HLSLTree::FExpression*& OutExpression) override;
 #endif
 	//~ End UMaterialExpression Interface
 };
diff --git a/Engine/Source/Runtime/Engine/Classes/Materials/MaterialExpressionMaterialAttributeLayers.h b/Engine/Source/Runtime/Engine/Classes/Materials/MaterialExpressionMaterialAttributeLayers.h
--- a/Engine/Source/Runtime/Engine/Classes/Materials/MaterialExpressionMaterialAttributeLayers.h
+++ b/Engine/Source/Runtime/Engine/Classes/Materials/MaterialExpressionMaterialAttributeLayers.h
@@ -111,6 +111,7 @@ class UMaterialExpressionMaterialAttributeLayers : public UMaterialExpression
 	virtual bool IsInputConnectionRequired(int32 InputIndex) const override {return false;}
 	virtual uint32 GetInputType(int32 InputIndex) override;
 	virtual bool IsResultMaterialAttributes(int32 OutputIndex) override {return true;}
+	virtual bool GenerateHLSLExpression(FMaterialHLSLGenerator& Generator, UE::HLSLTree::FScope& Scope, int32 OutputIndex, UE::HLSLTree::FExpression*& OutExpression) override;
 #endif
 	//~ End UMaterialExpression Interface
 
diff --git a/Engine/Source/Runtime/Engine/Classes/Materials/MaterialExpressionStaticSwitchParameter.h b/Engine/Source/Runtime/Engine/Classes/Materials/MaterialExpressionStaticSwitchParameter.h
--- a/Engine/Source/Runtime/Engine/Classes/Materials/MaterialExpressionStaticSwitchParameter.h
+++ b/Engine/Source/Runtime/Engine/Classes/Materials/MaterialExpressionStaticSwitchParameter.h
@@ -30,6 +30,7 @@ class UMaterialExpressionStaticSwitchParameter : public UMaterialExpressionStati
 	virtual bool IsResultMaterialAttributes(int32 OutputIndex) override;
 	virtual uint32 GetInputType(int32 InputIndex) override {return MCT_Unknown;}
 	virtual uint32 GetOutputType(int32 OutputIndex) override {return MCT_Unknown;}
+	virtual bool GenerateHLSLExpression(FMaterialHLSLGenerator& Generator, UE::HLSLTree::FScope& Scope, int32 OutputIndex, UE::HLSLTree::FExpression*& OutExpression) override;
 #endif // WITH_EDITOR
 	//~ End UMaterialExpression Interface
 };
diff --git a/Engine/Source/Runtime/Engine/Private/HLSLTree/HLSLTree.cpp b/Engine/Source/Runtime/Engine/Private/HLSLTree/HLSLTree.cpp
--- a/Engine/Source/Runtime/Engine/Private/HLSLTree/HLSLTree.cpp
+++ b/Engine/Source/Runtime/Engine/Private/HLSLTree/HLSLTree.cpp
@@ -683,11 +683,9 @@ int32 FPreparedType::GetNumComponents() const
 	}
 	else if (ValueComponentType != Shader::EValueComponentType::Void)
 	{
-		const int32 MaxComponentIndex = PreparedComponents.FindLastByPredicate([](const FPreparedComponent& InComponent) { return InComponent.Evaluation != EExpressionEvaluation::None; });
-		if (MaxComponentIndex != INDEX_NONE)
-		{
-			return MaxComponentIndex + 1;
-		}
+		auto Predicate = [](const FPreparedComponent& InComponent) { return InComponent.IsRequested(); };
+		const int32 MaxComponentIndex = PreparedComponents.FindLastByPredicate(Predicate);
+		return (MaxComponentIndex != INDEX_NONE) ? MaxComponentIndex + 1 : 1;
 	}
 	return 0;
 }
@@ -708,17 +706,13 @@ Shader::FType FPreparedType::GetType() const
 
 FRequestedType FPreparedType::GetRequestedType() const
 {
-	const int32 NumComponents = GetNumComponents();
 	FRequestedType Result;
-	if (NumComponents > 0)
+	for (int32 Index = 0; Index < PreparedComponents.Num(); ++Index)
 	{
-		for (int32 Index = 0; Index < NumComponents; ++Index)
+		const FPreparedComponent& Component = PreparedComponents[Index];
+		if (Component.IsRequested())
 		{
-			const FPreparedComponent Component = GetComponent(Index);
-			if (!Component.IsNone())
-			{
-				Result.SetComponentRequest(Index);
-			}
+			Result.SetComponentRequest(Index);
 		}
 	}
 	return Result;
@@ -738,6 +732,10 @@ FPreparedType FPreparedType::GetTypeForRequest(const FRequestedType& RequestedTy
 		{
 			Result.SetComponent(Index, GetComponent(Index));
 		}
+		else
+		{
+			Result.SetComponent(Index, EExpressionEvaluation::ConstantZero);
+		}
 	}
 	return Result;
 }
diff --git a/Engine/Source/Runtime/Engine/Private/HLSLTree/HLSLTreeCommon.cpp b/Engine/Source/Runtime/Engine/Private/HLSLTree/HLSLTreeCommon.cpp
--- a/Engine/Source/Runtime/Engine/Private/HLSLTree/HLSLTreeCommon.cpp
+++ b/Engine/Source/Runtime/Engine/Private/HLSLTree/HLSLTreeCommon.cpp
@@ -308,7 +308,7 @@ void FExpressionMaterialParameter::EmitValueShader(FEmitContext& Context, FEmitS
 			TextureTypeName = TEXT("ExternalTexture");
 
 			FMaterialExternalTextureParameterInfo TextureParameterInfo;
-			TextureParameterInfo.ParameterName = NameToScriptName(ParameterName);
+			TextureParameterInfo.ParameterName = NameToScriptName(ParameterInfo.Name);
 			TextureParameterInfo.ExternalTextureGuid = TextureValue->ExternalTextureGuid;
 			if (TextureValue->Texture)
 			{
@@ -352,7 +352,7 @@ void FExpressionMaterialParameter::EmitValueShader(FEmitContext& Context, FEmitS
 			}
 
 			FMaterialTextureParameterInfo TextureParameterInfo;
-			TextureParameterInfo.ParameterInfo = ParameterName;
+			TextureParameterInfo.ParameterInfo = ParameterInfo;
 			TextureParameterInfo.TextureIndex = Context.Material->GetReferencedTextures().Find(TextureValue->Texture);
 			TextureParameterInfo.SamplerSource = SSM_FromTextureAsset; // TODO - Is this needed?
 			check(TextureParameterInfo.TextureIndex != INDEX_NONE);
@@ -371,7 +371,6 @@ void FExpressionMaterialParameter::EmitValuePreshader(FEmitContext& Context, FEm
 	OutResult.Type = GetShaderValueType(ParameterType);
 	if (ParameterType == EMaterialParameterType::StaticSwitch)
 	{
-		const FMaterialParameterInfo ParameterInfo(ParameterName);
 		Shader::FValue Value = DefaultValue;
 		for (const FStaticSwitchParameter& Parameter : Context.StaticParameters->StaticSwitchParameters)
 		{
@@ -396,7 +395,7 @@ void FExpressionMaterialParameter::EmitValuePreshader(FEmitContext& Context, FEm
 			DefaultOffset = Context.MaterialCompilationOutput->UniformExpressionSet.AddDefaultParameterValue(DefaultValue);
 			Context.DefaultUniformValues.Add(DefaultValue, DefaultOffset);
 		}
-		const int32 ParameterIndex = Context.MaterialCompilationOutput->UniformExpressionSet.FindOrAddNumericParameter(ParameterType, ParameterName, DefaultOffset);
+		const int32 ParameterIndex = Context.MaterialCompilationOutput->UniformExpressionSet.FindOrAddNumericParameter(ParameterType, ParameterInfo, DefaultOffset);
 		check(ParameterIndex >= 0 && ParameterIndex <= 0xffff);
 		OutResult.Preshader.WriteOpcode(Shader::EPreshaderOpcode::Parameter).Write((uint16)ParameterIndex);
 	}
diff --git a/Engine/Source/Runtime/Engine/Private/HLSLTree/HLSLTreeOperation.cpp b/Engine/Source/Runtime/Engine/Private/HLSLTree/HLSLTreeOperation.cpp
--- a/Engine/Source/Runtime/Engine/Private/HLSLTree/HLSLTreeOperation.cpp
+++ b/Engine/Source/Runtime/Engine/Private/HLSLTree/HLSLTreeOperation.cpp
@@ -180,7 +180,7 @@ FOperationRequestedTypes GetOperationRequestedTypes(EOperation Op, const FReques
 	return Types;
 }
 
-FOperationTypes GetOperationTypes(EOperation Op, TConstArrayView<FPreparedType> InputTypes)
+FOperationTypes GetOperationTypes(EOperation Op, TConstArrayView<FExpression*> Inputs)
 {
 	FOperationTypes Types;
 
@@ -189,12 +189,13 @@ FOperationTypes GetOperationTypes(EOperation Op, TConstArrayView<FPreparedType>
 		Op == EOperation::Matrix3MulVec ||
 		Op == EOperation::Matrix4MulVec)
 	{
-		FPreparedComponent IntermediateComponent = InputTypes[0].GetMergedComponent();
-		Shader::EValueComponentType IntermediateComponentType = InputTypes[0].ValueComponentType;
-		for (int32 Index = 1; Index < InputTypes.Num(); ++Index)
+		FPreparedComponent IntermediateComponent;
+		Shader::EValueComponentType IntermediateComponentType = Shader::EValueComponentType::Void;
+		for (int32 Index = 0; Index < Inputs.Num(); ++Index)
 		{
-			IntermediateComponent = CombineComponents(IntermediateComponent, InputTypes[Index].GetMergedComponent());
-			IntermediateComponentType = Shader::CombineComponentTypes(IntermediateComponentType, InputTypes[Index].ValueComponentType);
+			const FPreparedType& InputType = Inputs[Index]->GetPreparedType();
+			IntermediateComponent = CombineComponents(IntermediateComponent, InputType.GetMergedComponent());
+			IntermediateComponentType = Shader::CombineComponentTypes(IntermediateComponentType, InputType.ValueComponentType);
 		}
 
 		switch (Op)
@@ -221,14 +222,15 @@ FOperationTypes GetOperationTypes(EOperation Op, TConstArrayView<FPreparedType>
 	}
 	else
 	{
-		FPreparedType IntermediateType = InputTypes[0];
-		for (int32 Index = 1; Index < InputTypes.Num(); ++Index)
+		FPreparedType IntermediateType;
+		for (int32 Index = 0; Index < Inputs.Num(); ++Index)
 		{
-			IntermediateType = MergePreparedTypes(IntermediateType, InputTypes[Index]);
+			const FPreparedType& InputType = Inputs[Index]->GetPreparedType();
+			IntermediateType = MergePreparedTypes(IntermediateType, InputType);
 		}
 
 		const Shader::EValueType IntermediateValueType = IntermediateType.GetType();
-		for (int32 Index = 0; Index < InputTypes.Num(); ++Index)
+		for (int32 Index = 0; Index < Inputs.Num(); ++Index)
 		{
 			Types.InputType[Index] = IntermediateValueType;
 		}
@@ -297,13 +299,13 @@ FOperationTypes GetOperationTypes(EOperation Op, TConstArrayView<FPreparedType>
 		case EOperation::Min:
 			for (int32 Index = 0; Index < Types.ResultType.PreparedComponents.Num(); ++Index)
 			{
-				Types.ResultType.SetComponentBounds(Index, Shader::MinBound(InputTypes[0].GetComponentBounds(Index), InputTypes[1].GetComponentBounds(Index)));
+				Types.ResultType.SetComponentBounds(Index, Shader::MinBound(Inputs[0]->GetPreparedType().GetComponentBounds(Index), Inputs[1]->GetPreparedType().GetComponentBounds(Index)));
 			}
 			break;
 		case EOperation::Max:
 			for (int32 Index = 0; Index < Types.ResultType.PreparedComponents.Num(); ++Index)
 			{
-				Types.ResultType.SetComponentBounds(Index, Shader::MaxBound(InputTypes[0].GetComponentBounds(Index), InputTypes[1].GetComponentBounds(Index)));
+				Types.ResultType.SetComponentBounds(Index, Shader::MaxBound(Inputs[0]->GetPreparedType().GetComponentBounds(Index), Inputs[1]->GetPreparedType().GetComponentBounds(Index)));
 			}
 			break;
 		default:
@@ -520,41 +522,57 @@ bool FExpressionOperation::PrepareValue(FEmitContext& Context, FEmitScope& Scope
 	const FOperationDescription OpDesc = GetOperationDescription(Op);
 	const Private::FOperationRequestedTypes RequestedTypes = Private::GetOperationRequestedTypes(Op, RequestedType);
 
-	FPreparedType InputTypes[MaxInputs];
+	Shader::FValue ConstantInput[MaxInputs];
+	bool bConstantZeroInput[MaxInputs] = { false };
 	for (int32 Index = 0; Index < OpDesc.NumInputs; ++Index)
 	{
-		InputTypes[Index] = Context.PrepareExpression(Inputs[Index], Scope, RequestedTypes.InputType[Index]);
-		if (InputTypes[Index].IsVoid())
+		const FPreparedType& InputType = Context.PrepareExpression(Inputs[Index], Scope, RequestedTypes.InputType[Index]);
+		if (InputType.IsVoid())
 		{
 			return false;
 		}
 
-		if (!InputTypes[Index].IsNumeric())
+		if (!InputType.IsNumeric())
 		{
 			return Context.Errors->AddError(TEXT("Invalid arithmetic between non-numeric types"));
 		}
+		
+		const EExpressionEvaluation InputEvaluation = InputType.GetEvaluation(Scope, RequestedType);
+		if (InputEvaluation == EExpressionEvaluation::Constant)
+		{
+			ConstantInput[Index] = Inputs[Index]->GetValueConstant(Context, Scope, RequestedType);
+			bConstantZeroInput[Index] = ConstantInput[Index].IsZero();
+		}
 	}
 
-	Private::FOperationTypes Types = Private::GetOperationTypes(Op, MakeArrayView(InputTypes, OpDesc.NumInputs));
+	Private::FOperationTypes Types = Private::GetOperationTypes(Op, MakeArrayView(Inputs, OpDesc.NumInputs));
 	if (OpDesc.PreshaderOpcode == Shader::EPreshaderOpcode::Nop)
 	{
 		// No preshader support
 		Types.ResultType.SetEvaluation(EExpressionEvaluation::Shader);
 	}
 
+	switch (Op)
+	{
+	case EOperation::Mul:
+		if (bConstantZeroInput[0] || bConstantZeroInput[1])
+		{
+			// X * 0 == 0
+			Types.ResultType.SetEvaluation(EExpressionEvaluation::ConstantZero);
+		}
+		break;
+	default:
+		break;
+	}
+
 	return OutResult.SetType(Context, RequestedType, Types.ResultType);
 }
 
 void FExpressionOperation::EmitValueShader(FEmitContext& Context, FEmitScope& Scope, const FRequestedType& RequestedType, FEmitValueShaderResult& OutResult) const
 {
 	const FOperationDescription OpDesc = GetOperationDescription(Op);
-	FPreparedType InputTypes[MaxInputs];
-	for (int32 Index = 0; Index < OpDesc.NumInputs; ++Index)
-	{
-		InputTypes[Index] = Inputs[Index]->GetPreparedType();
-	}
 	const Private::FOperationRequestedTypes RequestedTypes = Private::GetOperationRequestedTypes(Op, RequestedType);
-	const Private::FOperationTypes Types = Private::GetOperationTypes(Op, MakeArrayView(InputTypes, OpDesc.NumInputs));
+	const Private::FOperationTypes Types = Private::GetOperationTypes(Op, MakeArrayView(Inputs, OpDesc.NumInputs));
 	FEmitShaderExpression* InputValue[MaxInputs] = { nullptr };
 	for (int32 Index = 0; Index < OpDesc.NumInputs; ++Index)
 	{
@@ -653,13 +671,8 @@ void FExpressionOperation::EmitValueShader(FEmitContext& Context, FEmitScope& Sc
 void FExpressionOperation::EmitValuePreshader(FEmitContext& Context, FEmitScope& Scope, const FRequestedType& RequestedType, FEmitValuePreshaderResult& OutResult) const
 {
 	const FOperationDescription OpDesc = GetOperationDescription(Op);
-	FPreparedType InputTypes[MaxInputs];
-	for (int32 Index = 0; Index < OpDesc.NumInputs; ++Index)
-	{
-		InputTypes[Index] = Inputs[Index]->GetPreparedType();
-	}
 	const Private::FOperationRequestedTypes RequestedTypes = Private::GetOperationRequestedTypes(Op, RequestedType);
-	const Private::FOperationTypes Types = Private::GetOperationTypes(Op, MakeArrayView(InputTypes, OpDesc.NumInputs));
+	const Private::FOperationTypes Types = Private::GetOperationTypes(Op, MakeArrayView(Inputs, OpDesc.NumInputs));
 	check(OpDesc.PreshaderOpcode != Shader::EPreshaderOpcode::Nop);
 
 	for (int32 Index = 0; Index < OpDesc.NumInputs; ++Index)
diff --git a/Engine/Source/Runtime/Engine/Private/Materials/MaterialExpressionHLSL.cpp b/Engine/Source/Runtime/Engine/Private/Materials/MaterialExpressionHLSL.cpp
--- a/Engine/Source/Runtime/Engine/Private/Materials/MaterialExpressionHLSL.cpp
+++ b/Engine/Source/Runtime/Engine/Private/Materials/MaterialExpressionHLSL.cpp
@@ -24,6 +24,7 @@
 #include "Materials/MaterialExpressionDoubleVectorParameter.h"
 #include "Materials/MaterialExpressionScalarParameter.h"
 #include "Materials/MaterialExpressionStaticBoolParameter.h"
+#include "Materials/MaterialExpressionStaticSwitchParameter.h"
 #include "Materials/MaterialExpressionPixelDepth.h"
 #include "Materials/MaterialExpressionWorldPosition.h"
 #include "Materials/MaterialExpressionCameraPositionWS.h"
@@ -58,6 +59,8 @@
 #include "Materials/MaterialExpressionFunctionInput.h"
 #include "Materials/MaterialExpressionFunctionOutput.h"
 #include "Materials/MaterialExpressionMaterialFunctionCall.h"
+#include "Materials/MaterialExpressionMaterialAttributeLayers.h"
+#include "Materials/MaterialExpressionBlendMaterialAttributes.h"
 #include "Materials/MaterialExpressionStaticSwitch.h"
 #include "Materials/MaterialExpressionFeatureLevelSwitch.h"
 #include "Materials/MaterialExpressionShadingPathSwitch.h"
@@ -286,25 +289,40 @@ bool UMaterialExpressionGetLocal::GenerateHLSLExpression(FMaterialHLSLGenerator&
 
 bool UMaterialExpressionVectorParameter::GenerateHLSLExpression(FMaterialHLSLGenerator& Generator, UE::HLSLTree::FScope& Scope, int32 OutputIndex, UE::HLSLTree::FExpression*& OutExpression)
 {
-	OutExpression = Generator.GetTree().NewExpression<UE::HLSLTree::FExpressionMaterialParameter>(EMaterialParameterType::Vector, ParameterName, DefaultValue);
+	OutExpression = Generator.GenerateMaterialParameter(EMaterialParameterType::Vector, ParameterName, DefaultValue);
 	return true;
 }
 
 bool UMaterialExpressionDoubleVectorParameter::GenerateHLSLExpression(FMaterialHLSLGenerator& Generator, UE::HLSLTree::FScope& Scope, int32 OutputIndex, UE::HLSLTree::FExpression*& OutExpression)
 {
-	OutExpression = Generator.GetTree().NewExpression<UE::HLSLTree::FExpressionMaterialParameter>(EMaterialParameterType::DoubleVector, ParameterName, DefaultValue);
+	OutExpression = Generator.GenerateMaterialParameter(EMaterialParameterType::DoubleVector, ParameterName, DefaultValue);
 	return true;
 }
 
 bool UMaterialExpressionScalarParameter::GenerateHLSLExpression(FMaterialHLSLGenerator& Generator, UE::HLSLTree::FScope& Scope, int32 OutputIndex, UE::HLSLTree::FExpression*& OutExpression)
 {
-	OutExpression = Generator.GetTree().NewExpression<UE::HLSLTree::FExpressionMaterialParameter>(EMaterialParameterType::Scalar, ParameterName, DefaultValue);
+	OutExpression = Generator.GenerateMaterialParameter(EMaterialParameterType::Scalar, ParameterName, DefaultValue);
 	return true;
 }
 
 bool UMaterialExpressionStaticBoolParameter::GenerateHLSLExpression(FMaterialHLSLGenerator& Generator, UE::HLSLTree::FScope& Scope, int32 OutputIndex, UE::HLSLTree::FExpression*& OutExpression)
 {
-	OutExpression = Generator.GetTree().NewExpression<UE::HLSLTree::FExpressionMaterialParameter>(EMaterialParameterType::StaticSwitch, ParameterName, (bool)DefaultValue);
+	OutExpression = Generator.GenerateMaterialParameter(EMaterialParameterType::StaticSwitch, ParameterName, (bool)DefaultValue);
+	return true;
+}
+
+bool UMaterialExpressionStaticSwitchParameter::GenerateHLSLExpression(FMaterialHLSLGenerator& Generator, UE::HLSLTree::FScope& Scope, int32 OutputIndex, UE::HLSLTree::FExpression*& OutExpression)
+{
+	using namespace UE::HLSLTree;
+	FExpression* ExpressionA = A.AcquireHLSLExpression(Generator, Scope);
+	FExpression* ExpressionB = B.AcquireHLSLExpression(Generator, Scope);
+	if (!ExpressionA || !ExpressionB)
+	{
+		return false;
+	}
+
+	FExpression* ExpressionSwitch = Generator.GenerateMaterialParameter(EMaterialParameterType::StaticSwitch, ParameterName, (bool)DefaultValue);
+	OutExpression = Generator.GetTree().NewExpression <FExpressionSelect>(ExpressionSwitch, ExpressionA, ExpressionB);
 	return true;
 }
 
@@ -551,7 +569,7 @@ bool UMaterialExpressionTextureObject::GenerateHLSLExpression(FMaterialHLSLGener
 	using namespace UE::HLSLTree;
 	using namespace UE::Shader;
 	const FTextureValue* TextureValue = Generator.AcquireTextureValue(FTextureValue(Texture, SamplerType));
-	OutExpression = Generator.GetTree().NewExpression<FExpressionMaterialParameter>(EMaterialParameterType::Texture, FName(), TextureValue);
+	OutExpression = Generator.GenerateMaterialParameter(EMaterialParameterType::Texture, FName(), TextureValue);
 	return true;
 }
 
@@ -560,7 +578,7 @@ bool UMaterialExpressionTextureObjectParameter::GenerateHLSLExpression(FMaterial
 	using namespace UE::HLSLTree;
 	using namespace UE::Shader;
 	const FTextureValue* TextureValue = Generator.AcquireTextureValue(FTextureValue(Texture, SamplerType));
-	OutExpression = Generator.GetTree().NewExpression<FExpressionMaterialParameter>(EMaterialParameterType::Texture, ParameterName, TextureValue);
+	OutExpression = Generator.GenerateMaterialParameter(EMaterialParameterType::Texture, ParameterName, TextureValue);
 	return true;
 }
 
@@ -617,7 +635,7 @@ bool UMaterialExpressionTextureSample::GenerateHLSLExpression(FMaterialHLSLGener
 	else if (Texture)
 	{
 		const FTextureValue* TextureValue = Generator.AcquireTextureValue(FTextureValue(Texture, SamplerType));
-		TextureExpression = Generator.GetTree().NewExpression<FExpressionMaterialParameter>(EMaterialParameterType::Texture, FName(), TextureValue);
+		TextureExpression = Generator.GenerateMaterialParameter(EMaterialParameterType::Texture, FName(), TextureValue);
 	}
 
 	return GenerateHLSLExpressionBase(Generator, Scope, TextureExpression, OutExpression);
@@ -635,7 +653,7 @@ bool UMaterialExpressionTextureSampleParameter::GenerateHLSLExpression(FMaterial
 	else if (Texture)
 	{
 		const FTextureValue* TextureValue = Generator.AcquireTextureValue(FTextureValue(Texture, SamplerType));
-		TextureExpression = Generator.GetTree().NewExpression<FExpressionMaterialParameter>(EMaterialParameterType::Texture, ParameterName, TextureValue);
+		TextureExpression = Generator.GenerateMaterialParameter(EMaterialParameterType::Texture, ParameterName, TextureValue);
 	}
 
 	return GenerateHLSLExpressionBase(Generator, Scope, TextureExpression, OutExpression);
@@ -643,12 +661,12 @@ bool UMaterialExpressionTextureSampleParameter::GenerateHLSLExpression(FMaterial
 
 UE::HLSLTree::FExpression* UMaterialExpressionFontSample::GenerateHLSLTextureExpression(FMaterialHLSLGenerator& Generator, const UE::Shader::FTextureValue* TextureValue)
 {
-	return Generator.GetTree().NewExpression<UE::HLSLTree::FExpressionMaterialParameter>(EMaterialParameterType::Texture, FName(), TextureValue);
+	return Generator.GenerateMaterialParameter(EMaterialParameterType::Texture, FName(), TextureValue);
 }
 
 UE::HLSLTree::FExpression* UMaterialExpressionFontSampleParameter::GenerateHLSLTextureExpression(FMaterialHLSLGenerator& Generator, const UE::Shader::FTextureValue* TextureValue)
 {
-	return Generator.GetTree().NewExpression<UE::HLSLTree::FExpressionMaterialParameter>(EMaterialParameterType::Texture, ParameterName, TextureValue);
+	return Generator.GenerateMaterialParameter(EMaterialParameterType::Texture, ParameterName, TextureValue);
 }
 
 bool UMaterialExpressionFontSample::GenerateHLSLExpression(FMaterialHLSLGenerator& Generator, UE::HLSLTree::FScope& Scope, int32 OutputIndex, UE::HLSLTree::FExpression*& OutExpression)
@@ -1391,7 +1409,183 @@ bool UMaterialExpressionFunctionInput::GenerateHLSLExpression(FMaterialHLSLGener
 
 bool UMaterialExpressionMaterialFunctionCall::GenerateHLSLExpression(FMaterialHLSLGenerator& Generator, UE::HLSLTree::FScope& Scope, int32 OutputIndex, UE::HLSLTree::FExpression*& OutExpression)
 {
-	OutExpression = Generator.GenerateFunctionCall(Scope, MaterialFunction, FunctionInputs, OutputIndex);
+	using namespace UE::HLSLTree;
+	TArray<FExpression*, TInlineAllocator<16>> ConnectedInputs;
+	ConnectedInputs.Reserve(FunctionInputs.Num());
+	for (int32 InputIndex = 0; InputIndex < FunctionInputs.Num(); ++InputIndex)
+	{
+		// ConnectedInputs are the inputs from the UMaterialFunctionCall object
+		// We want to connect the UMaterialExpressionFunctionInput from the UMaterialFunction to whatever UMaterialExpression is passed to the UMaterialFunctionCall
+		FExpression* ConnectedInput = FunctionInputs[InputIndex].Input.TryAcquireHLSLExpression(Generator, Scope);
+		ConnectedInputs.Add(ConnectedInput);
+	}
+
+	OutExpression = Generator.GenerateFunctionCall(Scope, MaterialFunction, GlobalParameter, INDEX_NONE, ConnectedInputs, OutputIndex);
+	return true;
+}
+
+bool UMaterialExpressionMaterialAttributeLayers::GenerateHLSLExpression(FMaterialHLSLGenerator& Generator, UE::HLSLTree::FScope& Scope, int32 OutputIndex, UE::HLSLTree::FExpression*& OutExpression)
+{
+	using namespace UE::HLSLTree;
+
+	const FStaticParameterSet& StaticParameters = Generator.GetStaticParameters();
+	const FMaterialLayersFunctions& MaterialLayers = StaticParameters.bHasMaterialLayers ? StaticParameters.MaterialLayers : DefaultLayers;
+	if (MaterialLayers.Layers.Num() == 0)
+	{
+		return Generator.GetErrors().AddError(TEXT("No layers"));
+	}
+
+	TArray<FFunctionExpressionInput> FunctionInputs;
+	TArray<FFunctionExpressionOutput> FunctionOutputs;
+
+	FExpression* ExpressionLayerInput = Input.AcquireHLSLExpressionOrConstant(Generator, Scope, Generator.GetMaterialAttributesDefaultValue());
+	TArray<FExpression*, TInlineAllocator<1>> LayerInputExpressions;
+
+	TArray<FExpression*, TInlineAllocator<16>> LayerExpressions;
+	LayerExpressions.Reserve(MaterialLayers.Layers.Num());
+	for (int32 LayerIndex = 0; LayerIndex < MaterialLayers.Layers.Num(); ++LayerIndex)
+	{
+		UMaterialFunctionInterface* LayerFunction = MaterialLayers.Layers[LayerIndex];
+		FExpression* LayerExpression = nullptr;
+		if (LayerFunction && MaterialLayers.LayerStates[LayerIndex])
+		{
+			const EMaterialFunctionUsage Usage = LayerFunction->GetMaterialFunctionUsage();
+			if (Usage != EMaterialFunctionUsage::MaterialLayer)
+			{
+				return Generator.GetErrors().AddErrorf(TEXT("Layer function %s is not a UMaterialFunctionMaterialLayer"),
+					*LayerFunction->GetName());
+			}
+
+			FunctionInputs.Reset();
+			FunctionOutputs.Reset();
+			LayerFunction->GetInputsAndOutputs(FunctionInputs, FunctionOutputs);
+			if (FunctionInputs.Num() > 1 || FunctionOutputs.Num() != 1)
+			{
+				return Generator.GetErrors().AddErrorf(TEXT("Layer function %s expected to have 0 or 1 inputs and 1 output, found %d inputs and %d outputs"),
+					*LayerFunction->GetName(), FunctionInputs.Num(), FunctionOutputs.Num());
+			}
+
+			LayerInputExpressions.Empty(1);
+			if (FunctionInputs.Num() == 1)
+			{
+				LayerInputExpressions.Add(ExpressionLayerInput);
+			}
+			LayerExpression = Generator.GenerateFunctionCall(Scope, LayerFunction, LayerParameter, LayerIndex, LayerInputExpressions, 0);
+		}
+		LayerExpressions.Add(LayerExpression);
+	}
+
+	FExpression* BottomLayerExpression = LayerExpressions[0];
+	if (!BottomLayerExpression)
+	{
+		return Generator.GetErrors().AddError(TEXT("No layers"));
+	}
+
+	TArray<FExpression*, TInlineAllocator<2>> BlendInputExpressions;
+	for (int32 BlendIndex = 0; BlendIndex < MaterialLayers.Blends.Num(); ++BlendIndex)
+	{
+		const int32 LayerIndex = BlendIndex + 1;
+		if (!MaterialLayers.Layers.IsValidIndex(LayerIndex))
+		{
+			return Generator.GetErrors().AddErrorf(TEXT("Invalid number of layers (%d) and blends (%d)"), MaterialLayers.Layers.Num(), MaterialLayers.Blends.Num());
+		}
+
+		if (MaterialLayers.Layers[LayerIndex] && MaterialLayers.LayerStates[LayerIndex])
+		{
+			FExpression* LayerExpression = LayerExpressions[LayerIndex];
+			if (!LayerExpression)
+			{
+				return Generator.GetErrors().AddErrorf(TEXT("Missing layer %d"), LayerIndex);
+			}
+
+			UMaterialFunctionInterface* BlendFunction = MaterialLayers.Blends[BlendIndex];
+			if (BlendFunction)
+			{
+				const EMaterialFunctionUsage Usage = BlendFunction->GetMaterialFunctionUsage();
+				if (Usage != EMaterialFunctionUsage::MaterialLayerBlend)
+				{
+					return Generator.GetErrors().AddErrorf(TEXT("Blend function %s is not a UMaterialFunctionMaterialBlend"),
+						*BlendFunction->GetName());
+				}
+
+				FunctionInputs.Reset();
+				FunctionOutputs.Reset();
+				BlendFunction->GetInputsAndOutputs(FunctionInputs, FunctionOutputs);
+				if (FunctionInputs.Num() != 2 || FunctionOutputs.Num() != 1)
+				{
+					return Generator.GetErrors().AddErrorf(TEXT("Blend function %s expected to have 2 inputs and 1 output, found %d inputs and %d outputs"),
+						*BlendFunction->GetName(), FunctionInputs.Num(), FunctionOutputs.Num());
+				}
+
+				BlendInputExpressions.Empty(2);
+				BlendInputExpressions.Add(BottomLayerExpression);
+				BlendInputExpressions.Add(LayerExpression);
+				BottomLayerExpression = Generator.GenerateFunctionCall(Scope, BlendFunction, BlendParameter, BlendIndex, BlendInputExpressions, 0);
+			}
+			else
+			{
+				BottomLayerExpression = LayerExpression;
+			}
+		}
+	}
+
+	OutExpression = BottomLayerExpression;
+	return true;
+}
+
+bool UMaterialExpressionBlendMaterialAttributes::GenerateHLSLExpression(FMaterialHLSLGenerator& Generator, UE::HLSLTree::FScope& Scope, int32 OutputIndex, UE::HLSLTree::FExpression*& OutExpression)
+{
+	using namespace UE::HLSLTree;
+	using namespace UE::Shader;
+
+	FExpression* ExpressionA = A.AcquireHLSLExpressionOrConstant(Generator, Scope, Generator.GetMaterialAttributesDefaultValue());
+	FExpression* ExpressionB = B.AcquireHLSLExpressionOrConstant(Generator, Scope, Generator.GetMaterialAttributesDefaultValue());
+	FExpression* ExpressionAlpha = Alpha.AcquireHLSLExpression(Generator, Scope);
+	if (!ExpressionA || !ExpressionB || !ExpressionAlpha)
+	{
+		return false;
+	}
+
+	const UE::Shader::FStructType* MaterialAttributesType = Generator.GetMaterialAttributesType();
+
+	FExpression* ExpressionResult = Generator.GetTree().NewConstant(Generator.GetMaterialAttributesDefaultValue());
+	const TArray<FGuid>& OrderedVisibleAttributes = FMaterialAttributeDefinitionMap::GetOrderedVisibleAttributeList();
+	for (const FGuid& AttributeID : OrderedVisibleAttributes)
+	{
+		const FString& FieldName = FMaterialAttributeDefinitionMap::GetAttributeName(AttributeID);
+		const FStructField* Field = MaterialAttributesType->FindFieldByName(*FieldName);
+		if (!Field)
+		{
+			continue;
+		}
+
+		const EShaderFrequency AttributeFrequency = FMaterialAttributeDefinitionMap::GetShaderFrequency(AttributeID);
+		EMaterialAttributeBlend::Type AttributeBlend = EMaterialAttributeBlend::Blend;
+		switch (AttributeFrequency)
+		{
+		case SF_Vertex: AttributeBlend = VertexAttributeBlendType; break;
+		default: AttributeBlend = PixelAttributeBlendType; break;
+		}
+
+		FExpression* ExpressionFieldResult = nullptr;
+		if (AttributeBlend == EMaterialAttributeBlend::UseA)
+		{
+			ExpressionFieldResult = Generator.GetTree().NewExpression<FExpressionGetStructField>(MaterialAttributesType, Field, ExpressionA);
+		}
+		else if (AttributeBlend == EMaterialAttributeBlend::UseB)
+		{
+			ExpressionFieldResult = Generator.GetTree().NewExpression<FExpressionGetStructField>(MaterialAttributesType, Field, ExpressionB);
+		}
+		else
+		{
+			FExpression* ExpressionFieldA = Generator.GetTree().NewExpression<FExpressionGetStructField>(MaterialAttributesType, Field, ExpressionA);
+			FExpression* ExpressionFieldB = Generator.GetTree().NewExpression<FExpressionGetStructField>(MaterialAttributesType, Field, ExpressionB);
+			ExpressionFieldResult = Generator.GetTree().NewLerp(ExpressionFieldA, ExpressionFieldB, ExpressionAlpha);
+		}
+		ExpressionResult = Generator.GetTree().NewExpression<FExpressionSetStructField>(MaterialAttributesType, Field, ExpressionResult, ExpressionFieldResult);
+	}
+
+	OutExpression = ExpressionResult;
 	return true;
 }
 
diff --git a/Engine/Source/Runtime/Engine/Private/Materials/MaterialHLSLEmitter.cpp b/Engine/Source/Runtime/Engine/Private/Materials/MaterialHLSLEmitter.cpp
--- a/Engine/Source/Runtime/Engine/Private/Materials/MaterialHLSLEmitter.cpp
+++ b/Engine/Source/Runtime/Engine/Private/Materials/MaterialHLSLEmitter.cpp
@@ -799,7 +799,7 @@ bool MaterialEmitHLSL(const FMaterialCompileTargetParameters& InCompilerTarget,
 	FStructTypeRegistry TypeRegistry(Allocator);
 	FTree* HLSLTree = FTree::Create(Allocator);
 
-	FMaterialHLSLGenerator Generator(InCompilerTarget, InOutMaterial, TypeRegistry, *HLSLTree);
+	FMaterialHLSLGenerator Generator(InCompilerTarget, InStaticParameters, InOutMaterial, TypeRegistry, *HLSLTree);
 	if (!Generator.Generate())
 	{
 		return false;
@@ -912,6 +912,7 @@ bool MaterialEmitHLSL(const FMaterialCompileTargetParameters& InCompilerTarget,
 		RequestedVertexAttributesType.SetFieldRequested(Generator.GetMaterialAttributesType()->FindFieldByName(TEXT("PrevWorldPositionOffset")));
 
 		EmitContext.ShaderFrequency = SF_Vertex;
+		EmitContext.bUseAnalyticDerivatives = false;
 		FEmitScope* EmitResultScope = EmitContext.PrepareScope(&Generator.GetResultStatement()->GetParentScope());
 
 		const FPreparedType VertexResultType = EmitContext.PrepareExpression(Generator.GetResultExpression(), *EmitResultScope, RequestedVertexAttributesType);
diff --git a/Engine/Source/Runtime/Engine/Private/Materials/MaterialHLSLGenerator.cpp b/Engine/Source/Runtime/Engine/Private/Materials/MaterialHLSLGenerator.cpp
--- a/Engine/Source/Runtime/Engine/Private/Materials/MaterialHLSLGenerator.cpp
+++ b/Engine/Source/Runtime/Engine/Private/Materials/MaterialHLSLGenerator.cpp
@@ -81,10 +81,12 @@ void FMaterialHLSLErrorHandler::AddErrorInternal(UObject* InOwner, FStringView I
 }
 
 FMaterialHLSLGenerator::FMaterialHLSLGenerator(const FMaterialCompileTargetParameters& InCompileTarget,
+	const FStaticParameterSet& InStaticParameters,
 	FMaterial& InOutMaterial,
 	UE::Shader::FStructTypeRegistry& InOutTypeRegistry,
 	UE::HLSLTree::FTree& InOutTree)
 	: CompileTarget(InCompileTarget)
+	, StaticParameters(InStaticParameters)
 	, Errors(InOutMaterial)
 	, HLSLTree(&InOutTree)
 	, TypeRegistry(&InOutTypeRegistry)
@@ -612,7 +614,34 @@ bool FMaterialHLSLGenerator::GenerateStatements(UE::HLSLTree::FScope& Scope, UMa
 	return bResult;
 }
 
-UE::HLSLTree::FExpression* FMaterialHLSLGenerator::GenerateFunctionCall(UE::HLSLTree::FScope& Scope, UMaterialFunctionInterface* MaterialFunction, TArrayView<const FFunctionExpressionInput> ConnectedInputs, int32 OutputIndex)
+UE::HLSLTree::FExpression* FMaterialHLSLGenerator::GenerateMaterialParameter(EMaterialParameterType InType, FName InParameterName, const UE::Shader::FValue& InDefaultValue)
+{
+	UE::Shader::FValue DefaultValue(InDefaultValue);
+
+	FMaterialParameterMetadata Meta;
+	if (GetParameterOverrideValueForCurrentFunction(InType, InParameterName, Meta))
+	{
+		if (DefaultValue.Type.IsTexture())
+		{
+			UE::Shader::FTextureValue TextureValue(*DefaultValue.AsTexture());
+			TextureValue.Texture = Meta.Value.Texture;
+			DefaultValue = AcquireTextureValue(TextureValue);
+		}
+		else
+		{
+			DefaultValue = Meta.Value.AsShaderValue();
+		}
+	}
+
+	return HLSLTree->NewExpression<UE::HLSLTree::FExpressionMaterialParameter>(InType, GetParameterInfo(InParameterName), DefaultValue);
+}
+
+UE::HLSLTree::FExpression* FMaterialHLSLGenerator::GenerateFunctionCall(UE::HLSLTree::FScope& Scope,
+	UMaterialFunctionInterface* MaterialFunction,
+	EMaterialParameterAssociation InParameterAssociation,
+	int32 InParameterIndex,
+	TArrayView<UE::HLSLTree::FExpression*> ConnectedInputs,
+	int32 OutputIndex)
 {
 	using namespace UE::HLSLTree;
 
@@ -639,28 +668,40 @@ UE::HLSLTree::FExpression* FMaterialHLSLGenerator::GenerateFunctionCall(UE::HLSL
 		return nullptr;
 	}
 
-	const bool bInlineFunction = !MaterialFunction->IsUsingControlFlow();
-
-	FSHA1 Hasher;
-	Hasher.Update((uint8*)&MaterialFunction, sizeof(UMaterialFunctionInterface*));
+	EMaterialParameterAssociation ParameterAssociation = InParameterAssociation;
+	int32 ParameterIndex = InParameterIndex;
+	if (InParameterAssociation == GlobalParameter)
+	{
+		// If this is a global function, inherit the parameter association from the previous function
+		const FFunctionCallEntry* PrevFunctionEntry = FunctionCallStack.Last();
+		ParameterAssociation = PrevFunctionEntry->ParameterAssociation;
+		ParameterIndex = PrevFunctionEntry->ParameterIndex;
+	}
 
+	FSHAHash Hash;
 	FFunctionInputArray LocalFunctionInputs;
-	FConnectedInputArray LocalConnectedInputs;
-	for (int32 InputIndex = 0; InputIndex < ConnectedInputs.Num(); ++InputIndex)
 	{
-		// FunctionInputs are the inputs from the UMaterialFunction object
-		const FFunctionExpressionInput& FunctionInput = FunctionInputs[InputIndex];
+		FSHA1 Hasher;
+		Hasher.Update((uint8*)&MaterialFunction, sizeof(UMaterialFunctionInterface*));
+		Hasher.Update((uint8*)&ParameterAssociation, sizeof(EMaterialParameterAssociation));
+		Hasher.Update((uint8*)&ParameterIndex, sizeof(int32));
+
+		for (int32 InputIndex = 0; InputIndex < ConnectedInputs.Num(); ++InputIndex)
+		{
+			// FunctionInputs are the inputs from the UMaterialFunction object
+			const FFunctionExpressionInput& FunctionInput = FunctionInputs[InputIndex];
 
-		// ConnectedInputs are the inputs from the UMaterialFunctionCall object
-		// We want to connect the UMaterialExpressionFunctionInput from the UMaterialFunction to whatever UMaterialExpression is passed to the UMaterialFunctionCall
-		FExpression* ConnectedInput = ConnectedInputs[InputIndex].Input.TryAcquireHLSLExpression(*this, Scope);
+			// ConnectedInputs are the inputs from the UMaterialFunctionCall object
+			// We want to connect the UMaterialExpressionFunctionInput from the UMaterialFunction to whatever UMaterialExpression is passed to the UMaterialFunctionCall
+			FExpression* ConnectedInput = ConnectedInputs[InputIndex];
 
-		LocalFunctionInputs.Add(FunctionInput.ExpressionInput);
-		LocalConnectedInputs.Add(ConnectedInput);
-		Hasher.Update((uint8*)&ConnectedInput, sizeof(ConnectedInput));
+			LocalFunctionInputs.Add(FunctionInput.ExpressionInput);
+			Hasher.Update((uint8*)&ConnectedInput, sizeof(ConnectedInput));
+		}
+		Hash = Hasher.Finalize();
 	}
-	const FSHAHash Hash = Hasher.Finalize();
 
+	const bool bInlineFunction = !MaterialFunction->IsUsingControlFlow();
 	FFunctionCallEntry** ExistingFunctionCall = FunctionCallMap.Find(Hash);
 	FFunctionCallEntry* FunctionCall = ExistingFunctionCall ? *ExistingFunctionCall : nullptr;
 	if (!FunctionCall)
@@ -669,9 +710,11 @@ UE::HLSLTree::FExpression* FMaterialHLSLGenerator::GenerateFunctionCall(UE::HLSL
 		FFunction* HLSLFunction = !bInlineFunction ? HLSLTree->NewFunction() : nullptr;
 		FunctionCall = new(HLSLTree->GetAllocator()) FFunctionCallEntry();
 		FunctionCall->MaterialFunction = MaterialFunction;
+		FunctionCall->ParameterAssociation = ParameterAssociation;
+		FunctionCall->ParameterIndex = ParameterIndex;
 		FunctionCall->HLSLFunction = HLSLFunction;
 		FunctionCall->FunctionInputs = MoveTemp(LocalFunctionInputs);
-		FunctionCall->ConnectedInputs = MoveTemp(LocalConnectedInputs);
+		FunctionCall->ConnectedInputs = ConnectedInputs;
 		FunctionCall->FunctionOutputs.Reserve(FunctionOutputs.Num());
 		for (const FFunctionExpressionOutput& Output : FunctionOutputs)
 		{
@@ -715,6 +758,41 @@ UE::HLSLTree::FExpression* FMaterialHLSLGenerator::GenerateFunctionCall(UE::HLSL
 	return Result;
 }
 
+bool FMaterialHLSLGenerator::GetParameterOverrideValueForCurrentFunction(EMaterialParameterType ParameterType, FName ParameterName, FMaterialParameterMetadata& OutResult) const
+{
+	bool bResult = false;
+	if (!ParameterName.IsNone())
+	{
+		// Give every function in the callstack on opportunity to override the parameter value
+		// Parameters in outer functions take priority
+		// For example, if a layer instance calls a function instance that includes an overriden parameter, we want to use the value from the layer instance rather than the function instance
+		for (const FFunctionCallEntry* FunctionEntry : FunctionCallStack)
+		{
+			const UMaterialFunctionInterface* CurrentFunction = FunctionEntry->MaterialFunction;
+			if (CurrentFunction)
+			{
+				if (CurrentFunction->GetParameterOverrideValue(ParameterType, ParameterName, OutResult))
+				{
+					bResult = true;
+					break;
+				}
+			}
+		}
+	}
+	return bResult;
+}
+
+FMaterialParameterInfo FMaterialHLSLGenerator::GetParameterInfo(const FName& ParameterName) const
+{
+	if (ParameterName.IsNone())
+	{
+		return FMaterialParameterInfo();
+	}
+
+	const FFunctionCallEntry* FunctionEntry = FunctionCallStack.Last();
+	return FMaterialParameterInfo(ParameterName, FunctionEntry->ParameterAssociation, FunctionEntry->ParameterIndex);
+}
+
 void FMaterialHLSLGenerator::InternalRegisterExpressionData(const FName& Type, const UMaterialExpression* MaterialExpression, void* Data)
 {
 	const FExpressionDataKey Key(Type, MaterialExpression);
diff --git a/Engine/Source/Runtime/Engine/Private/Shader/ShaderValue.cpp b/Engine/Source/Runtime/Engine/Private/Shader/ShaderValue.cpp
--- a/Engine/Source/Runtime/Engine/Private/Shader/ShaderValue.cpp
+++ b/Engine/Source/Runtime/Engine/Private/Shader/ShaderValue.cpp
@@ -418,6 +418,23 @@ bool FValue::AsBoolScalar() const
 	return false;
 }
 
+bool FValue::IsZero() const
+{
+	bool bIsZero = Type.IsNumeric();
+	if (bIsZero)
+	{
+		for (const FValueComponent& Comp : Component)
+		{
+			if (Comp.Packed)
+			{
+				bIsZero = false;
+				break;
+			}
+		}
+	}
+	return bIsZero;
+}
+
 const FTextureValue* FValue::AsTexture() const
 {
 	if (Type.IsTexture() && Component.Num() > 0)
diff --git a/Engine/Source/Runtime/Engine/Public/HLSLTree/HLSLTree.h b/Engine/Source/Runtime/Engine/Public/HLSLTree/HLSLTree.h
--- a/Engine/Source/Runtime/Engine/Public/HLSLTree/HLSLTree.h
+++ b/Engine/Source/Runtime/Engine/Public/HLSLTree/HLSLTree.h
@@ -239,6 +239,7 @@ struct FPreparedComponent
 	FPreparedComponent(EExpressionEvaluation InEvaluation) : Evaluation(InEvaluation) {}
 
 	inline bool IsNone() const { return Evaluation == EExpressionEvaluation::None; }
+	inline bool IsRequested() const { return IsRequestedEvaluation(Evaluation); }
 
 	EExpressionEvaluation GetEvaluation(const FEmitScope& Scope) const;
 
@@ -291,10 +292,12 @@ public:
 	EExpressionEvaluation GetEvaluation(const FEmitScope& Scope) const;
 	EExpressionEvaluation GetEvaluation(const FEmitScope& Scope, const FRequestedType& RequestedType) const;
 	EExpressionEvaluation GetFieldEvaluation(const FEmitScope& Scope, const FRequestedType& RequestedType, int32 ComponentIndex, int32 NumComponents) const;
-	Shader::FComponentBounds GetComponentBounds(int32 Index) const;
 	Shader::FComponentBounds GetBounds(const FRequestedType& RequestedType) const;
 	FPreparedComponent GetMergedComponent() const;
+
 	FPreparedComponent GetComponent(int32 Index) const;
+	EExpressionEvaluation GetComponentEvaluation(int32 Index) const { return GetComponent(Index).Evaluation; }
+	Shader::FComponentBounds GetComponentBounds(int32 Index) const;
 
 	void SetComponent(int32 Index, const FPreparedComponent& InComponent);
 	void SetComponentBounds(int32 Index, const Shader::FComponentBounds Bounds);
diff --git a/Engine/Source/Runtime/Engine/Public/HLSLTree/HLSLTreeCommon.h b/Engine/Source/Runtime/Engine/Public/HLSLTree/HLSLTreeCommon.h
--- a/Engine/Source/Runtime/Engine/Public/HLSLTree/HLSLTreeCommon.h
+++ b/Engine/Source/Runtime/Engine/Public/HLSLTree/HLSLTreeCommon.h
@@ -2,6 +2,7 @@
 #pragma once
 
 #include "HLSLTree/HLSLTree.h"
+#include "Materials/MaterialLayersFunctions.h" // TODO - split material expressions into separate file/module
 
 enum class EMaterialParameterType : uint8;
 
@@ -39,15 +40,22 @@ public:
 	virtual void EmitValuePreshader(FEmitContext& Context, FEmitScope& Scope, const FRequestedType& RequestedType, FEmitValuePreshaderResult& OutResult) const override;
 };
 
+inline void AppendHash(FHasher& Hasher, const FMaterialParameterInfo& Value)
+{
+	AppendHash(Hasher, Value.Name);
+	AppendHash(Hasher, Value.Index);
+	AppendHash(Hasher, Value.Association);
+}
+
 class FExpressionMaterialParameter : public FExpression
 {
 public:
-	explicit FExpressionMaterialParameter(EMaterialParameterType InType, const FName& InName, const Shader::FValue& InDefaultValue)
-		: ParameterName(InName), DefaultValue(InDefaultValue), ParameterType(InType)
+	explicit FExpressionMaterialParameter(EMaterialParameterType InType, const FMaterialParameterInfo& InParameterInfo, const Shader::FValue& InDefaultValue)
+		: ParameterInfo(InParameterInfo), DefaultValue(InDefaultValue), ParameterType(InType)
 	{
 	}
 
-	FName ParameterName;
+	FMaterialParameterInfo ParameterInfo;
 	Shader::FValue DefaultValue;
 	EMaterialParameterType ParameterType;
 
diff --git a/Engine/Source/Runtime/Engine/Public/HLSLTree/HLSLTreeTypes.h b/Engine/Source/Runtime/Engine/Public/HLSLTree/HLSLTreeTypes.h
--- a/Engine/Source/Runtime/Engine/Public/HLSLTree/HLSLTreeTypes.h
+++ b/Engine/Source/Runtime/Engine/Public/HLSLTree/HLSLTreeTypes.h
@@ -61,6 +61,11 @@ inline bool IsLoopEvaluation(EExpressionEvaluation Evaluation)
 	return Evaluation == EExpressionEvaluation::PreshaderLoop || Evaluation == EExpressionEvaluation::ConstantLoop;
 }
 
+inline bool IsRequestedEvaluation(EExpressionEvaluation Evaluation)
+{
+	return Evaluation != EExpressionEvaluation::None && Evaluation != EExpressionEvaluation::ConstantZero;
+}
+
 enum class EOperation : uint8
 {
 	None,
diff --git a/Engine/Source/Runtime/Engine/Public/MaterialHLSLGenerator.h b/Engine/Source/Runtime/Engine/Public/MaterialHLSLGenerator.h
--- a/Engine/Source/Runtime/Engine/Public/MaterialHLSLGenerator.h
+++ b/Engine/Source/Runtime/Engine/Public/MaterialHLSLGenerator.h
@@ -10,12 +10,14 @@
 #include "Misc/EnumClassFlags.h"
 #include "Templates/RefCounting.h"
 #include "RHIDefinitions.h"
+#include "Materials/MaterialLayersFunctions.h"
 #include "HLSLTree/HLSLTree.h"
 
 class FMaterial;
 class FMaterialCompilationOutput;
 struct FSharedShaderCompilerEnvironment;
 struct FMaterialCompileTargetParameters;
+struct FMaterialParameterMetadata;
 class UMaterial;
 class UMaterialFunctionInterface;
 class UMaterialExpression;
@@ -24,6 +26,7 @@ class UMaterialExpressionFunctionOutput;
 class UMaterialExpressionCustomOutput;
 struct FFunctionExpressionInput;
 class ITargetPlatform;
+enum class EMaterialParameterType : uint8;
 
 namespace UE
 {
@@ -68,11 +71,13 @@ class FMaterialHLSLGenerator
 {
 public:
 	FMaterialHLSLGenerator(const FMaterialCompileTargetParameters& InCompilerTarget,
+		const FStaticParameterSet& InStaticParameters,
 		FMaterial& InOutMaterial,
 		UE::Shader::FStructTypeRegistry& InOutTypeRegistry,
 		UE::HLSLTree::FTree& InOutTree);
 
 	const FMaterialCompileTargetParameters& GetCompileTarget() const { return CompileTarget; }
+	const FStaticParameterSet& GetStaticParameters() const { return StaticParameters; }
 	
 	bool Generate();
 
@@ -110,7 +115,14 @@ public:
 
 	bool GenerateStatements(UE::HLSLTree::FScope& Scope, UMaterialExpression* MaterialExpression);
 
-	UE::HLSLTree::FExpression* GenerateFunctionCall(UE::HLSLTree::FScope& Scope, UMaterialFunctionInterface* Function, TArrayView<const FFunctionExpressionInput> ConnectedInputs, int32 OutputIndex);
+	UE::HLSLTree::FExpression* GenerateMaterialParameter(EMaterialParameterType InType, FName InParameterName, const UE::Shader::FValue& InDefaultValue);
+
+	UE::HLSLTree::FExpression* GenerateFunctionCall(UE::HLSLTree::FScope& Scope,
+		UMaterialFunctionInterface* Function,
+		EMaterialParameterAssociation ParameterAssociation,
+		int32 ParameterIndex,
+		TArrayView<UE::HLSLTree::FExpression*> ConnectedInputs,
+		int32 OutputIndex);
 
 	template<typename T, typename... ArgTypes>
 	T* NewExpressionData(const UMaterialExpression* MaterialExpression, ArgTypes... Args)
@@ -141,6 +153,9 @@ public:
 	const UE::Shader::FStructType* GetMaterialAttributesType() const { return MaterialAttributesType; }
 	const UE::Shader::FValue& GetMaterialAttributesDefaultValue() const { return MaterialAttributesDefaultValue; }
 
+	bool GetParameterOverrideValueForCurrentFunction(EMaterialParameterType ParameterType, FName ParameterName, FMaterialParameterMetadata& OutResult) const;
+	FMaterialParameterInfo GetParameterInfo(const FName& ParameterName) const;
+
 private:
 	static constexpr int32 MaxNumPreviousScopes = UE::HLSLTree::MaxNumPreviousScopes;
 	
@@ -173,6 +188,8 @@ private:
 		FFunctionInputArray FunctionInputs;
 		FFunctionOutputArray FunctionOutputs;
 		FConnectedInputArray ConnectedInputs;
+		EMaterialParameterAssociation ParameterAssociation = GlobalParameter;
+		int32 ParameterIndex = INDEX_NONE;
 		bool bGeneratedResult = false;
 	};
 
@@ -186,6 +203,7 @@ private:
 	void* InternalFindExpressionData(const FName& Type, const UMaterialExpression* MaterialExpression);
 
 	const FMaterialCompileTargetParameters& CompileTarget;
+	const FStaticParameterSet& StaticParameters;
 	UMaterial* TargetMaterial = nullptr;
 	FMaterialHLSLErrorHandler Errors;
 
diff --git a/Engine/Source/Runtime/Engine/Public/Shader/ShaderTypes.h b/Engine/Source/Runtime/Engine/Public/Shader/ShaderTypes.h
--- a/Engine/Source/Runtime/Engine/Public/Shader/ShaderTypes.h
+++ b/Engine/Source/Runtime/Engine/Public/Shader/ShaderTypes.h
@@ -514,6 +514,8 @@ struct FValue
 	float AsFloatScalar() const;
 	bool AsBoolScalar() const;
 
+	bool IsZero() const;
+
 	const TCHAR* ToString(EValueStringFormat Format, FStringBuilderBase& OutString) const;
 
 	FType Type;
