#UE Check all parent classes for duplicate SparseClassDataStruct objects before calling Link() on them in UAnimBlueprintGeneratedClass::Link

diff --git a/Engine/Source/Runtime/Engine/Private/Animation/AnimBlueprintGeneratedClass.cpp b/Engine/Source/Runtime/Engine/Private/Animation/AnimBlueprintGeneratedClass.cpp
--- a/Engine/Source/Runtime/Engine/Private/Animation/AnimBlueprintGeneratedClass.cpp
+++ b/Engine/Source/Runtime/Engine/Private/Animation/AnimBlueprintGeneratedClass.cpp
@@ -449,14 +449,37 @@ void UAnimBlueprintGeneratedClass::Link(FArchive& Ar, bool bRelinkExistingProper
 	{
 		if(UScriptStruct* ParentSparseClassDataStruct = ParentClass->GetSparseClassDataStruct())
 		{
-			if(SparseClassDataStruct && ParentSparseClassDataStruct != SparseClassDataStruct)
+			if(SparseClassDataStruct)
 			{
-				// Ensure parent is linked before setting super
-				ParentSparseClassDataStruct->Link(Ar, bRelinkExistingProperties);
-				SparseClassDataStruct->SetSuperStruct(ParentSparseClassDataStruct);
+				// Make sure there are no duplicate SparseClassDataStructs in the parent heirarchy
+				bool bAllSparseClassDataStructsUnique = true;
+				TArray<UScriptStruct*, TInlineAllocator<4>> HierarchySparseClassDataStructs;
+				HierarchySparseClassDataStructs.Add(SparseClassDataStruct);
+				for (UClass* CurClass = ParentClass; CurClass; CurClass = CurClass->GetSuperClass())
+				{
+					if (UScriptStruct* CurSparseClassDataStruct = CurClass->GetSparseClassDataStruct())
+					{
+						if (HierarchySparseClassDataStructs.Contains(CurSparseClassDataStruct))
+						{
+							bAllSparseClassDataStructsUnique = false;
+							break;
+						}
+						else
+						{
+							HierarchySparseClassDataStructs.Add(CurSparseClassDataStruct);
+						}
+					}
+				}
 
-				// Link sparse class data now it is correctly parented
-				SparseClassDataStruct->Link(Ar, bRelinkExistingProperties);
+				if (bAllSparseClassDataStructsUnique)
+				{
+					// Ensure parent is linked before setting super
+					ParentSparseClassDataStruct->Link(Ar, bRelinkExistingProperties);
+					SparseClassDataStruct->SetSuperStruct(ParentSparseClassDataStruct);
+
+					// Link sparse class data now it is correctly parented
+					SparseClassDataStruct->Link(Ar, bRelinkExistingProperties);
+				}
 			}
 		}
 	}
