#UE Expose GetBlendMode() to BPs. You could already get the blend mode for UMaterials, but this also allows material instances access to the function

diff --git a/Engine/Source/Runtime/Engine/Classes/Materials/MaterialInterface.h b/Engine/Source/Runtime/Engine/Classes/Materials/MaterialInterface.h
--- a/Engine/Source/Runtime/Engine/Classes/Materials/MaterialInterface.h
+++ b/Engine/Source/Runtime/Engine/Classes/Materials/MaterialInterface.h
@@ -731,6 +731,7 @@ public:
 	*/
 	ENGINE_API virtual float GetOpacityMaskClipValue() const;
 	ENGINE_API virtual bool GetCastDynamicShadowAsMasked() const;
+	UFUNCTION(BlueprintCallable, Category = "Rendering|Material")
 	ENGINE_API virtual EBlendMode GetBlendMode() const;
 	ENGINE_API virtual EStrataBlendMode GetStrataBlendMode() const;
 	ENGINE_API virtual FMaterialShadingModelField GetShadingModels() const;
