Add some extra logging in PhysicsReplication.

diff --git a/Engine/Source/Runtime/Engine/Private/PhysicsEngine/PhysicsReplication.cpp b/Engine/Source/Runtime/Engine/Private/PhysicsEngine/PhysicsReplication.cpp
--- a/Engine/Source/Runtime/Engine/Private/PhysicsEngine/PhysicsReplication.cpp
+++ b/Engine/Source/Runtime/Engine/Private/PhysicsEngine/PhysicsReplication.cpp
@@ -370,9 +370,9 @@ bool FPhysicsReplication::ApplyRigidBodyState(float DeltaSeconds, FBodyInstance*
 #if !UE_BUILD_SHIPPING
 			if (PhysicsReplicationCVars::LogPhysicsReplicationHardSnaps && GetOwningWorld())
 			{
-				UE_LOG(LogTemp, Warning, TEXT("Simulated HARD SNAP - - %s, \nCurrent Pos - %s, Target Pos - %s\n CurrentState.LinVel - %s, New Lin Vel - %s\nTarget Extrapolation Delta - %s, Is Replay? - %d, Is Asleep - %d"),
+				UE_LOG(LogTemp, Warning, TEXT("Simulated HARD SNAP - \nCurrent Pos - %s, Target Pos - %s\n CurrentState.LinVel - %s, New Lin Vel - %s\nTarget Extrapolation Delta - %s, Is Replay? - %d, Is Asleep - %d, Prev Progress - %f, Prev Similarity - %f"),
 					*CurrentState.Position.ToString(), *TargetPos.ToString(), *CurrentState.LinVel.ToString(), *NewState.LinVel.ToString(),
-					*ExtrapolationDeltaPos.ToString(), (GetOwningWorld()->GetDemoNetDriver() && GetOwningWorld()->GetDemoNetDriver()->IsPlaying()), !BI->IsInstanceAwake());
+					*ExtrapolationDeltaPos.ToString(), (GetOwningWorld()->GetDemoNetDriver() && GetOwningWorld()->GetDemoNetDriver()->IsPlaying()), !BI->IsInstanceAwake(), PrevProgress, PrevSimilarity);
 				if (bDidHardSnap)
 				{
 					*bDidHardSnap = true;
