Fix path tracer not being able to be toggled through the show command in game viewports

The code was improperly trying to change the ViewModeIndex which had no way of being reset afterwards. It turns out this logic had no effect and could be safely removed. Setting ShowFlags.PathTracing is sufficient to enable the path tracer, the ViewModeIndex is not actually used.

This enables a mechanism for toggling the path tracer in shipping builds.

diff --git a/Engine/Source/Runtime/Engine/Private/GameViewportClient.cpp b/Engine/Source/Runtime/Engine/Private/GameViewportClient.cpp
--- a/Engine/Source/Runtime/Engine/Private/GameViewportClient.cpp
+++ b/Engine/Source/Runtime/Engine/Private/GameViewportClient.cpp
@@ -1281,14 +1281,6 @@ void UGameViewportClient::Draw(FViewport* InViewport, FCanvas* SceneCanvas)
 		return;
 	}
 
-	// Force path tracing view mode, and extern code set path tracer show flags
-	const bool bForcePathTracing = InViewport->GetClient()->GetEngineShowFlags()->PathTracing;
-	if (bForcePathTracing)
-	{
-		EngineShowFlags.SetPathTracing(true);
-		ViewModeIndex = VMI_PathTracing;
-	}
-
 	// create the view family for rendering the world scene to the viewport's render target
 	FSceneViewFamilyContext ViewFamily(FSceneViewFamily::ConstructionValues(
 		InViewport,
