Fix for materials failing to open correctly in the Material Editor on Mac.

- FMaterialExpressionPreview does not override GetBaseMaterialPathName() and defaults to returning an empty string literal (TEXT(""))
- This means that FMaterialShaderMapContent::MaterialPath would have a _empty_  `Data` member.
- Dereferencing an empty FMemoryImageString  would result in it returning an _empty_ string literal: TEXT("")
- The ensure would then compare the address of the string literal to the start of the memory image.  Depending on where this string literal was located in memory this ensure would trigger or maybe not.  On Mac it seemed to always trigger so the location must be compiler dependent.
- The intention of the ensure was to catch overflows, but the way it was comparing pointers is invalid if we are comparing against a string literal that isn't in the memory image.
- The safest, easiest fix is to remove the ensure.  In the future if we still think we have an overflow in the memory image we can write a different ensure.

diff --git a/Engine/Source/Runtime/Engine/Private/Materials/MaterialShader.cpp b/Engine/Source/Runtime/Engine/Private/Materials/MaterialShader.cpp
--- a/Engine/Source/Runtime/Engine/Private/Materials/MaterialShader.cpp
+++ b/Engine/Source/Runtime/Engine/Private/Materials/MaterialShader.cpp
@@ -2842,9 +2842,6 @@ FMaterialShaderMap* FMaterialShaderMap::AcquireFinalizedClone()
 	Clone->bCompiledSuccessfully = bCompiledSuccessfully;
 	Clone->bIsPersistent = bIsPersistent;
 	Clone->AssignCopy(*this);
-	ensureMsgf(reinterpret_cast<uint64>(Clone->GetMaterialPath()) > reinterpret_cast<uint64>(Clone->GetContent()), TEXT("String embedded in memory image is located outside of it (image start %p > string %p). Pointer/offset overflow while freezing?"),
-		Clone->GetContent(), Clone->GetMaterialPath()
-		);
 #if WITH_EDITOR
 	Clone->AssociateWithAssets(GetAssociatedAssets());
 #endif
