Rename Local Exposure "Contrast Reduction" to "Contrast Scale"
- Also support contrast scale > 1.

diff --git a/Engine/Config/BaseEngine.ini b/Engine/Config/BaseEngine.ini
--- a/Engine/Config/BaseEngine.ini
+++ b/Engine/Config/BaseEngine.ini
@@ -1008,6 +1008,9 @@ NearClipPlane=10.0
 
 +ClassRedirects=(OldName="PackedLevelInstance", NewName="/Script/Engine.PackedLevelActor")
 
++PropertyRedirects=(OldName="PostProcessSettings.LocalExposureContrastReduction", NewName="PostProcessSettings.LocalExposureContrastScale")
++PropertyRedirects=(OldName="PostProcessSettings.bOverride_LocalExposureContrastReduction", NewName="PostProcessSettings.bOverride_LocalExposureContrastScale")
+
 [CoreUObject.Metadata]
 ; Note: UnrealHeaderTool should be rerun after making changes to MetadataRedirects to catch any keys specified in class headers
 +MetadataRedirects=(OldKey="K2Protected", NewKey="BlueprintProtected")
diff --git a/Engine/Shaders/Private/PostProcessHistogramCommon.ush b/Engine/Shaders/Private/PostProcessHistogramCommon.ush
--- a/Engine/Shaders/Private/PostProcessHistogramCommon.ush
+++ b/Engine/Shaders/Private/PostProcessHistogramCommon.ush
@@ -16,7 +16,7 @@ float EyeAdaptation_ExposureSpeedDown;
 float EyeAdaptation_HistogramScale;
 float EyeAdaptation_HistogramBias;
 float EyeAdaptation_LuminanceMin;
-float EyeAdaptation_LocalExposureContrastReduction;
+float EyeAdaptation_LocalExposureContrastScale;
 float EyeAdaptation_LocalExposureDetailStrength;
 float EyeAdaptation_LocalExposureBlurredLuminanceBlend;
 float EyeAdaptation_LocalExposureMiddleGreyExposureCompensation;
diff --git a/Engine/Shaders/Private/PostProcessTonemap.usf b/Engine/Shaders/Private/PostProcessTonemap.usf
--- a/Engine/Shaders/Private/PostProcessTonemap.usf
+++ b/Engine/Shaders/Private/PostProcessTonemap.usf
@@ -465,7 +465,7 @@ float4 TonemapCommonPS(
 	float MiddleGreyLumValue = ExposureScaleVignette.y;
 
 	float BaseLogLum = CalculateBaseLogLuminance(LogLuminance, EyeAdaptation_LocalExposureBlurredLuminanceBlend, ExposureScale, FullViewUV, LumBilateralGrid, BlurredLogLum, LumBilateralGridSampler, BlurredLogLumSampler);
-	float LocalExposure = CalculateLocalExposure(LogLuminance + log2(ExposureScale), BaseLogLum, MiddleGreyLumValue, EyeAdaptation_LocalExposureContrastReduction, EyeAdaptation_LocalExposureDetailStrength);
+	float LocalExposure = CalculateLocalExposure(LogLuminance + log2(ExposureScale), BaseLogLum, MiddleGreyLumValue, EyeAdaptation_LocalExposureContrastScale, EyeAdaptation_LocalExposureDetailStrength);
 
 	// Apply before bloom
 	SceneColor.rgb *= LocalExposure;
diff --git a/Engine/Shaders/Private/PostProcessVisualizeLocalExposure.usf b/Engine/Shaders/Private/PostProcessVisualizeLocalExposure.usf
--- a/Engine/Shaders/Private/PostProcessVisualizeLocalExposure.usf
+++ b/Engine/Shaders/Private/PostProcessVisualizeLocalExposure.usf
@@ -35,7 +35,7 @@ float4 MainPS(noperspective float4 UVAndScreenPos : TEXCOORD0, float4 SvPosition
 	float MiddleGreyLumValue = log2(0.18 * ExposureScaleMiddleGrey.y * EyeAdaptation_LocalExposureMiddleGreyExposureCompensation);
 
 	float BaseLogLum = CalculateBaseLogLuminance(LogLuminance, EyeAdaptation_LocalExposureBlurredLuminanceBlend, ExposureScaleMiddleGrey.x, UV, LumBilateralGrid, BlurredLogLum, TextureSampler, TextureSampler);
-	float LocalExposure = CalculateLocalExposure(LogLuminance + log2(ExposureScaleMiddleGrey.x), BaseLogLum, MiddleGreyLumValue, EyeAdaptation_LocalExposureContrastReduction, EyeAdaptation_LocalExposureDetailStrength);
+	float LocalExposure = CalculateLocalExposure(LogLuminance + log2(ExposureScaleMiddleGrey.x), BaseLogLum, MiddleGreyLumValue, EyeAdaptation_LocalExposureContrastScale, EyeAdaptation_LocalExposureDetailStrength);
 	float DetailLogLum = (LogLuminance + log2(ExposureScaleMiddleGrey.x)) - BaseLogLum;
 
 	float3 OutValue;
diff --git a/Engine/Source/Runtime/Engine/Classes/Engine/Scene.h b/Engine/Source/Runtime/Engine/Classes/Engine/Scene.h
--- a/Engine/Source/Runtime/Engine/Classes/Engine/Scene.h
+++ b/Engine/Source/Runtime/Engine/Classes/Engine/Scene.h
@@ -916,7 +916,7 @@ struct FPostProcessSettings
 	uint8 bOverride_HistogramLogMax:1;
 
 	UPROPERTY(EditAnywhere, BlueprintReadWrite, Category=Overrides, meta=(PinHiddenByDefault, InlineEditConditionToggle))
-	uint8 bOverride_LocalExposureContrastReduction:1;
+	uint8 bOverride_LocalExposureContrastScale:1;
 
 	UPROPERTY(EditAnywhere, BlueprintReadWrite, Category=Overrides, meta=(PinHiddenByDefault, InlineEditConditionToggle))
 	uint8 bOverride_LocalExposureDetailStrength:1;
@@ -1767,8 +1767,8 @@ struct FPostProcessSettings
 	 * Value less than 1 will enable local exposure.
 	 * Good values are usually in the range 0.6 .. 1.0.
 	*/
-	UPROPERTY(interp, BlueprintReadWrite, Category="Lens|Local Exposure", meta=(UIMin = "0.0", UIMax = "1.0", editcondition = "bOverride_LocalExposureContrastReduction", DisplayName = "Contrast Reduction"))
-	float LocalExposureContrastReduction;
+	UPROPERTY(interp, BlueprintReadWrite, Category="Lens|Local Exposure", meta=(UIMin = "0.0", UIMax = "1.0", editcondition = "bOverride_LocalExposureContrastScale", DisplayName = "Contrast Scale"))
+	float LocalExposureContrastScale;
 
 	/**
 	 * Local Exposure decomposes luminance of the frame into a base layer and a detail layer.
diff --git a/Engine/Source/Runtime/Engine/Private/Scene.cpp b/Engine/Source/Runtime/Engine/Private/Scene.cpp
--- a/Engine/Source/Runtime/Engine/Private/Scene.cpp
+++ b/Engine/Source/Runtime/Engine/Private/Scene.cpp
@@ -480,7 +480,7 @@ FPostProcessSettings::FPostProcessSettings()
 
 	AutoExposureApplyPhysicalCameraExposure = 1;
 
-	LocalExposureContrastReduction = 1.0f;
+	LocalExposureContrastScale = 1.0f;
 	LocalExposureDetailStrength = 1.0f;
 	LocalExposureBlurredLuminanceBlend = 0.6f;
 	LocalExposureBlurredLuminanceKernelSizePercent = 50.0f;
@@ -697,7 +697,7 @@ FPostProcessSettings::FPostProcessSettings(const FPostProcessSettings& Settings)
 	, bOverride_AutoExposureApplyPhysicalCameraExposure(Settings.bOverride_AutoExposureApplyPhysicalCameraExposure)
 	, bOverride_HistogramLogMin(Settings.bOverride_HistogramLogMin)
 	, bOverride_HistogramLogMax(Settings.bOverride_HistogramLogMax)
-	, bOverride_LocalExposureContrastReduction(Settings.bOverride_LocalExposureContrastReduction)
+	, bOverride_LocalExposureContrastScale(Settings.bOverride_LocalExposureContrastScale)
 	, bOverride_LocalExposureDetailStrength(Settings.bOverride_LocalExposureDetailStrength)
 	, bOverride_LocalExposureBlurredLuminanceBlend(Settings.bOverride_LocalExposureBlurredLuminanceBlend)
 	, bOverride_LocalExposureBlurredLuminanceKernelSizePercent(Settings.bOverride_LocalExposureBlurredLuminanceKernelSizePercent)
@@ -913,7 +913,7 @@ FPostProcessSettings::FPostProcessSettings(const FPostProcessSettings& Settings)
 	, HistogramLogMin(Settings.HistogramLogMin)
 	, HistogramLogMax(Settings.HistogramLogMax)
 	, AutoExposureCalibrationConstant_DEPRECATED(Settings.AutoExposureCalibrationConstant_DEPRECATED)
-	, LocalExposureContrastReduction(Settings.LocalExposureContrastReduction)
+	, LocalExposureContrastScale(Settings.LocalExposureContrastScale)
 	, LocalExposureDetailStrength(Settings.LocalExposureDetailStrength)
 	, LocalExposureBlurredLuminanceBlend(Settings.LocalExposureBlurredLuminanceBlend)
 	, LocalExposureBlurredLuminanceKernelSizePercent(Settings.LocalExposureBlurredLuminanceKernelSizePercent)
diff --git a/Engine/Source/Runtime/Engine/Private/SceneView.cpp b/Engine/Source/Runtime/Engine/Private/SceneView.cpp
--- a/Engine/Source/Runtime/Engine/Private/SceneView.cpp
+++ b/Engine/Source/Runtime/Engine/Private/SceneView.cpp
@@ -1444,7 +1444,7 @@ void FSceneView::OverridePostProcessSettings(const FPostProcessSettings& Src, fl
 		LERP_PP(AutoExposureBias);
 		LERP_PP(HistogramLogMin);
 		LERP_PP(HistogramLogMax);
-		LERP_PP(LocalExposureContrastReduction);
+		LERP_PP(LocalExposureContrastScale);
 		LERP_PP(LocalExposureDetailStrength);
 		LERP_PP(LocalExposureBlurredLuminanceBlend);
 		LERP_PP(LocalExposureBlurredLuminanceKernelSizePercent);
@@ -1941,7 +1941,7 @@ void FSceneView::EndFinalPostprocessSettings(const FSceneViewInitOptions& ViewIn
 
 		if (Value <= 0)
 		{
-			FinalPostProcessSettings.LocalExposureContrastReduction = 1.0f;
+			FinalPostProcessSettings.LocalExposureContrastScale = 1.0f;
 			FinalPostProcessSettings.LocalExposureDetailStrength = 1.0f;
 		}
 	}
@@ -1964,7 +1964,7 @@ void FSceneView::EndFinalPostprocessSettings(const FSceneViewInitOptions& ViewIn
 
 	if (!Family->EngineShowFlags.LocalExposure)
 	{
-		FinalPostProcessSettings.LocalExposureContrastReduction = 1.0f;
+		FinalPostProcessSettings.LocalExposureContrastScale = 1.0f;
 		FinalPostProcessSettings.LocalExposureDetailStrength = 1.0f;
 	}
 
diff --git a/Engine/Source/Runtime/Renderer/Private/PostProcess/PostProcessEyeAdaptation.cpp b/Engine/Source/Runtime/Renderer/Private/PostProcess/PostProcessEyeAdaptation.cpp
--- a/Engine/Source/Runtime/Renderer/Private/PostProcess/PostProcessEyeAdaptation.cpp
+++ b/Engine/Source/Runtime/Renderer/Private/PostProcess/PostProcessEyeAdaptation.cpp
@@ -430,7 +430,7 @@ FEyeAdaptationParameters GetEyeAdaptationParameters(const FViewInfo& View, ERHIF
 	Parameters.HistogramScale = HistogramScale;
 	Parameters.HistogramBias = HistogramBias;
 	Parameters.LuminanceMin = LuminanceMin;
-	Parameters.LocalExposureContrastReduction = Settings.LocalExposureContrastReduction;
+	Parameters.LocalExposureContrastScale = Settings.LocalExposureContrastScale;
 	Parameters.LocalExposureDetailStrength = Settings.LocalExposureDetailStrength;
 	Parameters.LocalExposureBlurredLuminanceBlend = Settings.LocalExposureBlurredLuminanceBlend;
 	Parameters.LocalExposureMiddleGreyExposureCompensation = LocalExposureMiddleGreyExposureCompensation;
diff --git a/Engine/Source/Runtime/Renderer/Private/PostProcess/PostProcessEyeAdaptation.h b/Engine/Source/Runtime/Renderer/Private/PostProcess/PostProcessEyeAdaptation.h
--- a/Engine/Source/Runtime/Renderer/Private/PostProcess/PostProcessEyeAdaptation.h
+++ b/Engine/Source/Runtime/Renderer/Private/PostProcess/PostProcessEyeAdaptation.h
@@ -65,7 +65,7 @@ BEGIN_SHADER_PARAMETER_STRUCT(FEyeAdaptationParameters, )
 	SHADER_PARAMETER(float, HistogramScale)
 	SHADER_PARAMETER(float, HistogramBias)
 	SHADER_PARAMETER(float, LuminanceMin)
-	SHADER_PARAMETER(float, LocalExposureContrastReduction)
+	SHADER_PARAMETER(float, LocalExposureContrastScale)
 	SHADER_PARAMETER(float, LocalExposureDetailStrength)
 	SHADER_PARAMETER(float, LocalExposureBlurredLuminanceBlend)
 	SHADER_PARAMETER(float, LocalExposureMiddleGreyExposureCompensation)
diff --git a/Engine/Source/Runtime/Renderer/Private/PostProcess/PostProcessing.cpp b/Engine/Source/Runtime/Renderer/Private/PostProcess/PostProcessing.cpp
--- a/Engine/Source/Runtime/Renderer/Private/PostProcess/PostProcessing.cpp
+++ b/Engine/Source/Runtime/Renderer/Private/PostProcess/PostProcessing.cpp
@@ -450,7 +450,7 @@ void AddPostProcessingPasses(
 
 		const bool bLocalExposureEnabled =
 			EngineShowFlags.VisualizeLocalExposure || 
-			View.FinalPostProcessSettings.LocalExposureContrastReduction < 1.0 ||
+			!FMath::IsNearlyEqual(View.FinalPostProcessSettings.LocalExposureContrastScale, 1.0f) ||
 			!FMath::IsNearlyEqual(View.FinalPostProcessSettings.LocalExposureDetailStrength, 1.0f);
 
 		const bool bBloomEnabled = View.FinalPostProcessSettings.BloomIntensity > 0.0f;
diff --git a/Engine/Source/Runtime/Renderer/Private/SceneRendering.cpp b/Engine/Source/Runtime/Renderer/Private/SceneRendering.cpp
--- a/Engine/Source/Runtime/Renderer/Private/SceneRendering.cpp
+++ b/Engine/Source/Runtime/Renderer/Private/SceneRendering.cpp
@@ -3196,7 +3196,7 @@ void FSceneRenderer::RenderFinish(FRDGBuilder& GraphBuilder, FRDGTextureRef View
 
 			bNaniteEnabledButDisabledInProject = bNaniteEnabledButDisabledInProject || (WouldRenderNanite(Scene, View, /*bCheckForAtomicSupport*/ false, /*bCheckForProjectSetting*/ false) && !WouldRenderNanite(Scene, View, /*bCheckForAtomicSupport*/ false, /*bCheckForProjectSetting*/ true));
 
-			if (IsPostProcessingEnabled(View) && (View.FinalPostProcessSettings.LocalExposureContrastReduction < 1.0f || !FMath::IsNearlyEqual(View.FinalPostProcessSettings.LocalExposureDetailStrength, 1.0f)))
+			if (IsPostProcessingEnabled(View) && (!FMath::IsNearlyEqual(View.FinalPostProcessSettings.LocalExposureContrastScale, 1.0f) || !FMath::IsNearlyEqual(View.FinalPostProcessSettings.LocalExposureDetailStrength, 1.0f)))
 				bLocalExposureEnabledOnAnyView = true;
 		}
 
