Fixed velocities on static objects being moved by editor widgets, which in turn fixes Lumen GI and TAA
* Re-added bIsBeingMovedByEditor to UPrimitiveComponent so the FPrimitiveSceneProxy value is not lost on render state recreate
* FPrimitiveSceneProxy::SetIsBeingMovedByEditor updates renderer internal velocity state

diff --git a/Engine/Source/Runtime/Engine/Classes/Components/PrimitiveComponent.h b/Engine/Source/Runtime/Engine/Classes/Components/PrimitiveComponent.h
--- a/Engine/Source/Runtime/Engine/Classes/Components/PrimitiveComponent.h
+++ b/Engine/Source/Runtime/Engine/Classes/Components/PrimitiveComponent.h
@@ -592,6 +592,10 @@ public:
 	UPROPERTY()
 	uint8 bUseEditorCompositing:1;
 
+	/** Set to true while the editor is moving the component, which notifies the Renderer to track velocities even if the component is Static. */
+	UPROPERTY(Transient, DuplicateTransient)
+	uint8 bIsBeingMovedByEditor:1;
+
 	/** If true, this component will be rendered in the CustomDepth pass (usually used for outlines) */
 	UPROPERTY(EditAnywhere, AdvancedDisplay, BlueprintReadOnly, Category=Rendering, meta=(DisplayName = "Render CustomDepth Pass"))
 	uint8 bRenderCustomDepth:1;
diff --git a/Engine/Source/Runtime/Engine/Private/Components/PrimitiveComponent.cpp b/Engine/Source/Runtime/Engine/Private/Components/PrimitiveComponent.cpp
--- a/Engine/Source/Runtime/Engine/Private/Components/PrimitiveComponent.cpp
+++ b/Engine/Source/Runtime/Engine/Private/Components/PrimitiveComponent.cpp
@@ -348,6 +348,7 @@ UPrimitiveComponent::UPrimitiveComponent(const FObjectInitializer& ObjectInitial
 	bNeverDistanceCull = false;
 
 	bUseEditorCompositing = false;
+	bIsBeingMovedByEditor = false;
 
 	SetGenerateOverlapEvents(true);
 	bMultiBodyOverlap = false;
@@ -1728,6 +1729,8 @@ uint64 UPrimitiveComponent::GetHiddenEditorViews() const
 
 void UPrimitiveComponent::SetIsBeingMovedByEditor(bool bIsBeingMoved)
 {
+	bIsBeingMovedByEditor = bIsBeingMoved;
+
 	if (SceneProxy)
 	{
 		SceneProxy->SetIsBeingMovedByEditor_GameThread(bIsBeingMoved);
diff --git a/Engine/Source/Runtime/Engine/Private/PrimitiveSceneProxy.cpp b/Engine/Source/Runtime/Engine/Private/PrimitiveSceneProxy.cpp
--- a/Engine/Source/Runtime/Engine/Private/PrimitiveSceneProxy.cpp
+++ b/Engine/Source/Runtime/Engine/Private/PrimitiveSceneProxy.cpp
@@ -207,7 +207,7 @@ FPrimitiveSceneProxy::FPrimitiveSceneProxy(const UPrimitiveComponent* InComponen
 ,	bSelectable(InComponent->bSelectable)
 ,	bHasPerInstanceHitProxies(InComponent->bHasPerInstanceHitProxies)
 ,	bUseEditorCompositing(InComponent->bUseEditorCompositing)
-,	bIsBeingMovedByEditor(false)
+,	bIsBeingMovedByEditor(InComponent->bIsBeingMovedByEditor)
 ,	bReceiveMobileCSMShadows(InComponent->bReceiveMobileCSMShadows)
 ,	bRenderCustomDepth(InComponent->bRenderCustomDepth)
 ,	bVisibleInSceneCaptureOnly(InComponent->bVisibleInSceneCaptureOnly)
@@ -1048,16 +1048,10 @@ void FPrimitiveSceneProxy::SetIsBeingMovedByEditor_GameThread(bool bIsBeingMoved
 	ENQUEUE_RENDER_COMMAND(SetIsBeingMovedByEditor)(
 		[PrimitiveSceneProxy, bIsBeingMoved](FRHICommandListImmediate& RHICmdList)
 		{
-			PrimitiveSceneProxy->SetIsBeingMovedByEditor_RenderThread(bIsBeingMoved);
+			PrimitiveSceneProxy->bIsBeingMovedByEditor = bIsBeingMoved;
+			PrimitiveSceneProxy->GetScene().UpdatePrimitiveVelocityState_RenderThread(PrimitiveSceneProxy->GetPrimitiveSceneInfo(), bIsBeingMoved);
 		});
 }
-
-void FPrimitiveSceneProxy::SetIsBeingMovedByEditor_RenderThread(bool bIsBeingMoved)
-{
-	check(IsInRenderingThread());
-	bIsBeingMovedByEditor = bIsBeingMoved;
-}
-
 #endif
 
 void FPrimitiveSceneProxy::SetCollisionEnabled_GameThread(const bool bNewEnabled)
diff --git a/Engine/Source/Runtime/Engine/Public/PrimitiveSceneProxy.h b/Engine/Source/Runtime/Engine/Public/PrimitiveSceneProxy.h
--- a/Engine/Source/Runtime/Engine/Public/PrimitiveSceneProxy.h
+++ b/Engine/Source/Runtime/Engine/Public/PrimitiveSceneProxy.h
@@ -215,7 +215,6 @@ public:
 	 * Enqueue and update for the render thread to notify it that the editor is currently moving the owning component with gizmos.
 	 */
 	void SetIsBeingMovedByEditor_GameThread(bool bIsBeingMoved);
-	void SetIsBeingMovedByEditor_RenderThread(bool bIsBeingMoved);
 #endif	// WITH_EDITOR
 
 	/** @return True if the primitive is visible in the given View. */
diff --git a/Engine/Source/Runtime/Engine/Public/SceneInterface.h b/Engine/Source/Runtime/Engine/Public/SceneInterface.h
--- a/Engine/Source/Runtime/Engine/Public/SceneInterface.h
+++ b/Engine/Source/Runtime/Engine/Public/SceneInterface.h
@@ -249,6 +249,9 @@ public:
 	/** Updates the selected state values that might be cached inside the renderer */
 	virtual void UpdatePrimitiveSelectedState_RenderThread(const FPrimitiveSceneInfo* PrimitiveSceneInfo, bool bIsSelected) {}
 
+	/** Updates the velocity state values that might be cached inside the renderer */
+	virtual void UpdatePrimitiveVelocityState_RenderThread(FPrimitiveSceneInfo* PrimitiveSceneInfo, bool bIsBeingMoved) {}
+
 	/** 
 	 * Adds a new exponential height fog component to the scene
 	 * 
diff --git a/Engine/Source/Runtime/Renderer/Private/RendererScene.cpp b/Engine/Source/Runtime/Renderer/Private/RendererScene.cpp
--- a/Engine/Source/Runtime/Renderer/Private/RendererScene.cpp
+++ b/Engine/Source/Runtime/Renderer/Private/RendererScene.cpp
@@ -4211,6 +4211,25 @@ static bool ShouldPrimitiveOutputVelocity(const FPrimitiveSceneProxy* Proxy, con
 	return bPlatformSupportsVelocityRendering && bShouldPrimitiveOutputVelocity;
 }
 
+void FScene::UpdatePrimitiveVelocityState_RenderThread(FPrimitiveSceneInfo* PrimitiveSceneInfo, bool bIsBeingMoved)
+{
+	if (bIsBeingMoved)
+	{
+		if (ShouldPrimitiveOutputVelocity(PrimitiveSceneInfo->Proxy, GetShaderPlatform()))
+		{
+			PrimitiveSceneInfo->bRegisteredWithVelocityData = true;
+			// We must register the initial LocalToWorld with the velocity state. 
+			int32 PrimitiveIndex = PrimitiveSceneInfo->PackedIndex;
+			VelocityData.UpdateTransform(PrimitiveSceneInfo, PrimitiveTransforms[PrimitiveIndex], PrimitiveTransforms[PrimitiveIndex]);
+		}
+	}
+	else if (PrimitiveSceneInfo->bRegisteredWithVelocityData)
+	{
+		PrimitiveSceneInfo->bRegisteredWithVelocityData = false;
+		VelocityData.RemoveFromScene(PrimitiveSceneInfo->PrimitiveComponentId);
+	}
+}
+
 #if RHI_RAYTRACING
 void FScene::UpdateRayTracingGroupBounds_AddPrimitives(const Experimental::TRobinHoodHashSet<FPrimitiveSceneInfo*>& PrimitiveSceneInfos)
 {
diff --git a/Engine/Source/Runtime/Renderer/Private/ScenePrivate.h b/Engine/Source/Runtime/Renderer/Private/ScenePrivate.h
--- a/Engine/Source/Runtime/Renderer/Private/ScenePrivate.h
+++ b/Engine/Source/Runtime/Renderer/Private/ScenePrivate.h
@@ -3170,6 +3170,7 @@ public:
 
 	/** Updates PrimitivesSelected array for this PrimitiveSceneInfo */
 	virtual void UpdatePrimitiveSelectedState_RenderThread(const FPrimitiveSceneInfo* PrimitiveSceneInfo, bool bIsSelected) override;
+	virtual void UpdatePrimitiveVelocityState_RenderThread(FPrimitiveSceneInfo* PrimitiveSceneInfo, bool bIsBeingMoved) override;
 
 	virtual void Release() override;
 	virtual UWorld* GetWorld() const override { return World; }
