Fix race condition crash when destroying a virtual texture in editor.
Need to wait on any outstanding DDC tasks which might still hold references to the VT chunks.
#rnx

diff --git a/Engine/Source/Runtime/Engine/Private/VT/UploadingVirtualTexture.cpp b/Engine/Source/Runtime/Engine/Private/VT/UploadingVirtualTexture.cpp
--- a/Engine/Source/Runtime/Engine/Private/VT/UploadingVirtualTexture.cpp
+++ b/Engine/Source/Runtime/Engine/Private/VT/UploadingVirtualTexture.cpp
@@ -90,6 +90,11 @@ FUploadingVirtualTexture::~FUploadingVirtualTexture()
 			Codec.Reset();
 		}
 	}
+
+#if WITH_EDITOR
+	// Need to flush DDC tasks in case they reference any chunks from this.
+	GetVirtualTextureChunkDDCCache()->WaitFlushRequests_AnyThread();
+#endif
 }
 
 uint32 FUploadingVirtualTexture::GetLocalMipBias(uint8 vLevel, uint32 vAddress) const
diff --git a/Engine/Source/Runtime/Engine/Private/VT/VirtualTextureChunkDDCCache.cpp b/Engine/Source/Runtime/Engine/Private/VT/VirtualTextureChunkDDCCache.cpp
--- a/Engine/Source/Runtime/Engine/Private/VT/VirtualTextureChunkDDCCache.cpp
+++ b/Engine/Source/Runtime/Engine/Private/VT/VirtualTextureChunkDDCCache.cpp
@@ -231,18 +231,60 @@ void FVirtualTextureChunkDDCCache::Initialize()
 
 void FVirtualTextureChunkDDCCache::ShutDown()
 {
-	ActiveChunks.Empty();
+	WaitFlushRequests_AnyThread();
+
+	ENQUEUE_RENDER_COMMAND(FVirtualTextureChunkDDCCache_Shutdown)(
+		[this](FRHICommandListImmediate& RHICmdList)
+		{
+			UpdateRequests();
+			ensure(ActiveChunks.Num() == 0);
+		});
 }
 
 void FVirtualTextureChunkDDCCache::UpdateRequests()
 {
-	ActiveChunks.RemoveAll([](auto Chunk) -> bool {return Chunk->bFileAvailableInVTDDCDache == true; });
+	// Remove completed chunks from array.
+	check(IsInRenderingThread());
+	ActiveChunks.RemoveAllSwap([](auto Chunk) -> bool {return Chunk->bFileAvailableInVTDDCDache == true; }, false);
+
+	// Remove completed tasks from array.
+	// ActiveTasks needs a lock because it is accessed by WaitFlushRequests_AnyThread()
+	{
+		FScopeLock Lock(&ActiveTasksLock);
+		for (int32 TaskIndex = 0; TaskIndex < ActiveTasks.Num(); ++TaskIndex)
+		{
+			FAsyncTask<FAsyncFillCacheWorker>* Task = ActiveTasks[TaskIndex];
+			if (Task->IsWorkDone())
+			{
+				Task->EnsureCompletion();
+				delete Task;
+				ActiveTasks.RemoveAtSwap(TaskIndex--, 1, false);
+			}
+		}
+	}
+}
+
+void FVirtualTextureChunkDDCCache::WaitFlushRequests_AnyThread()
+{
+	TArray<FAsyncTask<FAsyncFillCacheWorker>*> ActiveTasksCopy;
+	{
+		FScopeLock Lock(&ActiveTasksLock);
+		ActiveTasksCopy = MoveTemp(ActiveTasks);
+	}
+	for (int32 TaskIndex = 0; TaskIndex < ActiveTasksCopy.Num(); ++TaskIndex)
+	{
+		FAsyncTask<FAsyncFillCacheWorker>* Task = ActiveTasksCopy[TaskIndex];
+		Task->EnsureCompletion();
+		delete Task;
+	}
 }
 
 bool FVirtualTextureChunkDDCCache::MakeChunkAvailable(struct FVirtualTextureDataChunk* Chunk, bool bAsync, FString& OutChunkFileName, int64& OutOffsetInFile)
 {
 	TRACE_CPUPROFILER_EVENT_SCOPE(FVirtualTextureChunkDDCCache::MakeChunkAvailable);
 
+	check(IsInRenderingThread());
+
 	const FString CachedFilePath = AbsoluteCachePath / Chunk->ShortDerivedDataKey;
 	const FString TempFilePath = AbsoluteCachePath / FGuid::NewGuid().ToString() + ".tmp";
 
@@ -272,8 +314,15 @@ bool FVirtualTextureChunkDDCCache::MakeChunkAvailable(struct FVirtualTextureData
 	// start filling it to the cache
 	if (bAsync)
 	{
+		FAsyncTask<FAsyncFillCacheWorker>* Task = (new FAsyncTask<FAsyncFillCacheWorker>(TempFilePath, CachedFilePath, Chunk));
+		
+		{
+			FScopeLock Lock(&ActiveTasksLock);
+			ActiveTasks.Add(Task);
+		}
+		
 		ActiveChunks.Add(Chunk);
-		(new FAutoDeleteAsyncTask<FAsyncFillCacheWorker>(TempFilePath, CachedFilePath, Chunk))->StartBackgroundTask();
+		Task->StartBackgroundTask();
 	}
 	else
 	{
diff --git a/Engine/Source/Runtime/Engine/Private/VT/VirtualTextureChunkDDCCache.h b/Engine/Source/Runtime/Engine/Private/VT/VirtualTextureChunkDDCCache.h
--- a/Engine/Source/Runtime/Engine/Private/VT/VirtualTextureChunkDDCCache.h
+++ b/Engine/Source/Runtime/Engine/Private/VT/VirtualTextureChunkDDCCache.h
@@ -6,19 +6,27 @@
 
 #include "CoreMinimal.h"
 
+class FAsyncFillCacheWorker;
+template<typename TTask> class FAsyncTask;
+struct FVirtualTextureDataChunk;
+
 class FVirtualTextureChunkDDCCache
 {
 public:
 	void Initialize();
 	void ShutDown();
+
 	void UpdateRequests();
+	void WaitFlushRequests_AnyThread();
 
-	bool MakeChunkAvailable(struct FVirtualTextureDataChunk* Chunk, bool bAsync, FString& OutChunkFileName, int64& OutOffsetInFile);
+	bool MakeChunkAvailable(FVirtualTextureDataChunk* Chunk, bool bAsync, FString& OutChunkFileName, int64& OutOffsetInFile);
 
-	void MakeChunkAvailable_Concurrent(struct FVirtualTextureDataChunk* Chunk);
+	void MakeChunkAvailable_Concurrent(FVirtualTextureDataChunk* Chunk);
 private:
-	TArray<struct FVirtualTextureDataChunk*> ActiveChunks;
 	FString AbsoluteCachePath;
+	TArray<FVirtualTextureDataChunk*> ActiveChunks;
+	TArray<FAsyncTask<FAsyncFillCacheWorker>*> ActiveTasks;
+	FCriticalSection ActiveTasksLock;
 };
 
 FVirtualTextureChunkDDCCache* GetVirtualTextureChunkDDCCache();
