Fix LWC replay compatibility.
Store UE package version in replay data and propagate it to the archives used for playback.
Make a best effort to fix up the case of the HISTORY_SERIALIZE_DOUBLE_VECTORS_AS_DOUBLES network version and LARGE_WORLD_COORDINATES package version so that pre-LWC replays will play back correctly. This allows vectors serialized in replays with operator<< to correctly use the UE package version for LWC backwards-compatibility.

diff --git a/Engine/Source/Runtime/Core/Private/Serialization/BitReader.cpp b/Engine/Source/Runtime/Core/Private/Serialization/BitReader.cpp
--- a/Engine/Source/Runtime/Core/Private/Serialization/BitReader.cpp
+++ b/Engine/Source/Runtime/Core/Private/Serialization/BitReader.cpp
@@ -181,8 +181,7 @@ void FBitReader::SetData( FBitReader& Src, int64 CountBits )
 	ClearError();
 
 	// Setup network version
-	this->SetEngineNetVer(Src.EngineNetVer());
-	this->SetGameNetVer(Src.GameNetVer());
+	SetNetVersionsFromArchive(Src);
 
 	Buffer.Empty();
 	Buffer.AddUninitialized( (CountBits+7)>>3 );
@@ -294,6 +293,15 @@ void FBitReader::SerializeIntPacked(uint32& OutValue)
 	OutValue = Value;
 }
 
+void FBitReader::SetNetVersionsFromArchive(FArchive& Source)
+{
+	SetEngineNetVer(Source.EngineNetVer());
+	SetGameNetVer(Source.GameNetVer());
+	SetUEVer(Source.UEVer());
+	SetLicenseeUEVer(Source.LicenseeUEVer());
+	SetEngineVer(Source.EngineVer());
+}
+
 /*-----------------------------------------------------------------------------
 	FBitReader.
 -----------------------------------------------------------------------------*/
diff --git a/Engine/Source/Runtime/Core/Public/Serialization/BitReader.h b/Engine/Source/Runtime/Core/Public/Serialization/BitReader.h
--- a/Engine/Source/Runtime/Core/Public/Serialization/BitReader.h
+++ b/Engine/Source/Runtime/Core/Public/Serialization/BitReader.h
@@ -225,6 +225,9 @@ protected:
 	int64 Num;
 	int64 Pos;
 
+	/** Copies version information used for network compatibility from Source to this archive */
+	void SetNetVersionsFromArchive(FArchive& Source);
+
 private:
 
 	FORCEINLINE uint8 Shift(uint8 Cnt)
diff --git a/Engine/Source/Runtime/Engine/Classes/Engine/NetConnection.h b/Engine/Source/Runtime/Engine/Classes/Engine/NetConnection.h
--- a/Engine/Source/Runtime/Engine/Classes/Engine/NetConnection.h
+++ b/Engine/Source/Runtime/Engine/Classes/Engine/NetConnection.h
@@ -10,6 +10,7 @@
 #include "UObject/ObjectMacros.h"
 #include "UObject/UObjectGlobals.h"
 #include "Serialization/BitWriter.h"
+#include "Misc/EngineVersion.h"
 #include "Misc/NetworkGuid.h"
 #include "GameFramework/OnlineReplStructs.h"
 #include "GameFramework/UpdateLevelVisibilityLevelInfo.h"
@@ -674,8 +675,27 @@ private:
 
 	UReplicationConnectionDriver* ReplicationConnectionDriver;
 
+	/** Engine package version for compatibility */
+	FPackageFileVersion PackageVersionUE;
+
+	/** Licensee package version for compatibility */
+	int32 PackageVersionLicenseeUE;
+
+	/** Engine version information for compatibility */
+	FEngineVersion EngineVersion;
 
 public:
+	/** Sets the UE package version for compatibility purposes */
+	void SetPackageVersionUE(FPackageFileVersion InPackageVersionUE) { PackageVersionUE = InPackageVersionUE; }
+	
+	/** Sets the licensee package version for compatibility purposes */
+	void SetPackageVersionLicenseeUE(int32 InPackageVersionLicenseeUE) { PackageVersionLicenseeUE = InPackageVersionLicenseeUE; }
+
+	/** Sets engine version information for compatibility purposes */
+	void SetEngineVersion(const FEngineVersion& InEngineVersion) { EngineVersion = InEngineVersion; }
+
+	/** Set version information on an archive that can be used for compatibility checks */
+	void SetNetVersionsOnArchive(FArchive& Ar) const;
 
 	void AddDestructionInfo(FActorDestructionInfo* DestructionInfo)
 	{
diff --git a/Engine/Source/Runtime/Engine/Private/DataBunch.cpp b/Engine/Source/Runtime/Engine/Private/DataBunch.cpp
--- a/Engine/Source/Runtime/Engine/Private/DataBunch.cpp
+++ b/Engine/Source/Runtime/Engine/Private/DataBunch.cpp
@@ -40,8 +40,7 @@ FInBunch::FInBunch( UNetConnection* InConnection, uint8* Src, int64 CountBits )
 	SetByteSwapping(Connection->bNeedsByteSwapping);
 
 	// Copy network version info
-	this->SetEngineNetVer(InConnection->EngineNetworkProtocolVersion);
-	this->SetGameNetVer(InConnection->GameNetworkProtocolVersion);
+	Connection->SetNetVersionsOnArchive(*this);
 }
 
 /** Copy constructor but with optional parameter to not copy buffer */
@@ -66,8 +65,7 @@ FInBunch::FInBunch( FInBunch &InBunch, bool CopyBuffer )
 	CloseReason = InBunch.CloseReason;
 
 	// Copy network version info
-	this->SetEngineNetVer(InBunch.EngineNetVer());
-	this->SetGameNetVer(InBunch.GameNetVer());
+	SetNetVersionsFromArchive(InBunch);
 
 	PackageMap = InBunch.PackageMap;
 
diff --git a/Engine/Source/Runtime/Engine/Private/DataReplication.cpp b/Engine/Source/Runtime/Engine/Private/DataReplication.cpp
--- a/Engine/Source/Runtime/Engine/Private/DataReplication.cpp
+++ b/Engine/Source/Runtime/Engine/Private/DataReplication.cpp
@@ -2263,6 +2263,7 @@ void FObjectReplicator::UpdateUnmappedObjects(bool & bOutHasMoreUnmapped)
 			const FFieldNetCache* FieldCache = ClassCache->GetFromIndex(Pending.RPCFieldIndex);
 
 			FNetBitReader Reader(Connection->PackageMap, Pending.Buffer.GetData(), Pending.NumBits);
+			Connection->SetNetVersionsOnArchive(Reader);
 
 			bool bIsGuidPending = false;
 
diff --git a/Engine/Source/Runtime/Engine/Private/NetConnection.cpp b/Engine/Source/Runtime/Engine/Private/NetConnection.cpp
--- a/Engine/Source/Runtime/Engine/Private/NetConnection.cpp
+++ b/Engine/Source/Runtime/Engine/Private/NetConnection.cpp
@@ -283,6 +283,8 @@ UNetConnection::UNetConnection(const FObjectInitializer& ObjectInitializer)
 ,	InitInReliable		( 0 )
 ,	EngineNetworkProtocolVersion( FNetworkVersion::GetEngineNetworkProtocolVersion() )
 ,	GameNetworkProtocolVersion( FNetworkVersion::GetGameNetworkProtocolVersion() )
+,	PackageVersionUE( GPackageFileUEVersion )
+,	PackageVersionLicenseeUE( GPackageFileLicenseeUEVersion )
 ,	ResendAllDataState( EResendAllDataState::None )
 #if !UE_BUILD_SHIPPING
 ,	ReceivedRawPacketDel()
@@ -1583,8 +1585,7 @@ void UNetConnection::ReceivedRawPacket( void* InData, int32 Count )
 			FBitReader Reader(Data, BitSize);
 
 			// Set the network version on the reader
-			Reader.SetEngineNetVer( EngineNetworkProtocolVersion );
-			Reader.SetGameNetVer( GameNetworkProtocolVersion );
+			SetNetVersionsOnArchive(Reader);
 
 			if (Handler.IsValid())
 			{
@@ -4902,6 +4903,17 @@ void UNetConnection::NotifyActorChannelCleanedUp(UActorChannel* Channel, EChanne
 	}
 }
 
+void UNetConnection::SetNetVersionsOnArchive(FArchive& Ar) const
+{
+	Ar.SetEngineNetVer(EngineNetworkProtocolVersion);
+	Ar.SetGameNetVer(GameNetworkProtocolVersion);
+	Ar.SetUEVer(PackageVersionUE);
+	Ar.SetLicenseeUEVer(PackageVersionLicenseeUE);
+	// Base archives only store FEngineVersionBase, but net connections store FEngineVersion.
+	// This will slice off the branch name and anything else stored in FEngineVersion.
+	Ar.SetEngineVer(EngineVersion);
+}
+
 /*-----------------------------------------------------------------------------
 	USimulatedClientNetConnection.
 -----------------------------------------------------------------------------*/
diff --git a/Engine/Source/Runtime/Engine/Private/RepLayout.cpp b/Engine/Source/Runtime/Engine/Private/RepLayout.cpp
--- a/Engine/Source/Runtime/Engine/Private/RepLayout.cpp
+++ b/Engine/Source/Runtime/Engine/Private/RepLayout.cpp
@@ -4188,7 +4188,7 @@ void FRepLayout::UpdateUnmappedObjects_r(
 	FReceivingRepState* RESTRICT RepState,
 	FGuidReferencesMap* GuidReferencesMap,
 	UObject* OriginalObject,
-	UPackageMap* PackageMap,
+	UNetConnection* Connection,
 	FRepShadowDataBuffer ShadowData,
 	FRepObjectDataBuffer Data,
 	const int32 MaxAbsOffset,
@@ -4232,7 +4232,7 @@ void FRepLayout::UpdateUnmappedObjects_r(
 
 				const int32 NewMaxOffset = FMath::Min(ShadowArray->Num() * Cmd.ElementSize, Array->Num() * Cmd.ElementSize);
 
-				UpdateUnmappedObjects_r(RepState, GuidReferences.Array, OriginalObject, PackageMap, ShadowArrayData, ArrayData, NewMaxOffset, bCalledPreNetReceive, bOutSomeObjectsWereMapped, bOutHasMoreUnmapped);
+				UpdateUnmappedObjects_r(RepState, GuidReferences.Array, OriginalObject, Connection, ShadowArrayData, ArrayData, NewMaxOffset, bCalledPreNetReceive, bOutSomeObjectsWereMapped, bOutHasMoreUnmapped);
 			}
 			else
 			{
@@ -4240,7 +4240,7 @@ void FRepLayout::UpdateUnmappedObjects_r(
 				FRepObjectDataBuffer ArrayData(Array->GetData());
 				const int32 NewMaxOffset = Array->Num() * Cmd.ElementSize;
 
-				UpdateUnmappedObjects_r(RepState, GuidReferences.Array, OriginalObject, PackageMap, nullptr, ArrayData, NewMaxOffset, bCalledPreNetReceive, bOutSomeObjectsWereMapped, bOutHasMoreUnmapped);
+				UpdateUnmappedObjects_r(RepState, GuidReferences.Array, OriginalObject, Connection, nullptr, ArrayData, NewMaxOffset, bCalledPreNetReceive, bOutSomeObjectsWereMapped, bOutHasMoreUnmapped);
 			}
 			continue;
 		}
@@ -4251,14 +4251,14 @@ void FRepLayout::UpdateUnmappedObjects_r(
 		{
 			const FNetworkGUID& GUID = *UnmappedIt;
 
-			if (PackageMap->IsGUIDBroken(GUID, false))
+			if (Connection->PackageMap->IsGUIDBroken(GUID, false))
 			{
 				UE_LOG(LogRep, Warning, TEXT("UpdateUnmappedObjects_r: Broken GUID. NetGuid: %s"), *GUID.ToString());
 				UnmappedIt.RemoveCurrent();
 				continue;
 			}
 
-			UObject* Object = PackageMap->GetObjectFromNetGUID(GUID, false);
+			UObject* Object = Connection->PackageMap->GetObjectFromNetGUID(GUID, false);
 
 			if (Object != nullptr)
 			{
@@ -4299,10 +4299,11 @@ void FRepLayout::UpdateUnmappedObjects_r(
 			}
 
 			// Initialize the reader with the stored buffer that we need to read from
-			FNetBitReader Reader(PackageMap, GuidReferences.Buffer.GetData(), GuidReferences.NumBufferBits);
+			FNetBitReader Reader(Connection->PackageMap, GuidReferences.Buffer.GetData(), GuidReferences.NumBufferBits);
+			Connection->SetNetVersionsOnArchive(Reader);
 
 			// Read the property
-			Cmd.Property->NetSerializeItem(Reader, PackageMap, Data + AbsOffset);
+			Cmd.Property->NetSerializeItem(Reader, Connection->PackageMap, Data + AbsOffset);
 
 			// Check to see if this property changed
 			if (bUpdateShadowState)
@@ -4356,7 +4357,7 @@ void FRepLayout::UpdateUnmappedObjects(
 			RepState,
 			&RepState->GuidReferencesMap,
 			OriginalObject,
-			PackageMap,
+			Params.Connection,
 			(uint8*)RepState->StaticBuffer.GetData(),
 			(uint8*)OriginalObject,
 			Owner->GetPropertiesSize(),
@@ -8085,7 +8086,7 @@ void FRepLayout::UpdateUnmappedGuidsForFastArray(FFastArrayDeltaSerializeParams&
 			const int32 ArrayElementOffset = ItemIndex * ElementSize;
 			FRepObjectDataBuffer ElementData(ArrayData + ArrayElementOffset);
 
-			UpdateUnmappedObjects_r(nullptr, &It.Value(), Object, PackageMap, nullptr, ElementData, ElementSize, Params.DeltaSerializeInfo.bCalledPreNetReceive, bOutSomeObjectsWereMapped, bOutHasMoreUnmapped);
+			UpdateUnmappedObjects_r(nullptr, &It.Value(), Object, DeltaSerializeInfo.Connection, nullptr, ElementData, ElementSize, Params.DeltaSerializeInfo.bCalledPreNetReceive, bOutSomeObjectsWereMapped, bOutHasMoreUnmapped);
 
 			if (bOutSomeObjectsWereMapped)
 			{
diff --git a/Engine/Source/Runtime/Engine/Private/ReplayHelper.cpp b/Engine/Source/Runtime/Engine/Private/ReplayHelper.cpp
--- a/Engine/Source/Runtime/Engine/Private/ReplayHelper.cpp
+++ b/Engine/Source/Runtime/Engine/Private/ReplayHelper.cpp
@@ -84,6 +84,11 @@ void FReplayHelper::SetPlaybackNetworkVersions(FArchive& Ar)
 {
 	Ar.SetEngineNetVer(PlaybackDemoHeader.EngineNetworkProtocolVersion);
 	Ar.SetGameNetVer(PlaybackDemoHeader.GameNetworkProtocolVersion);
+	Ar.SetUEVer(PlaybackDemoHeader.PackageVersionUE);
+	Ar.SetLicenseeUEVer(PlaybackDemoHeader.PackageVersionLicenseeUE);
+	// Base archives only store FEngineVersionBase, but the header stores FEngineVersion.
+	// This will slice off the branch name and anything else stored in FEngineVersion.
+	Ar.SetEngineVer(PlaybackDemoHeader.EngineVersion);
 }
 
 void FReplayHelper::SetPlaybackNetworkVersions(UNetConnection* Connection)
@@ -92,6 +97,9 @@ void FReplayHelper::SetPlaybackNetworkVersions(UNetConnection* Connection)
 	{
 		Connection->EngineNetworkProtocolVersion = PlaybackDemoHeader.EngineNetworkProtocolVersion;
 		Connection->GameNetworkProtocolVersion = PlaybackDemoHeader.GameNetworkProtocolVersion;
+		Connection->SetPackageVersionUE(PlaybackDemoHeader.PackageVersionUE);
+		Connection->SetPackageVersionLicenseeUE(PlaybackDemoHeader.PackageVersionLicenseeUE);
+		Connection->SetEngineVersion(PlaybackDemoHeader.EngineVersion);
 	}
 }
 
diff --git a/Engine/Source/Runtime/Engine/Public/Net/RepLayout.h b/Engine/Source/Runtime/Engine/Public/Net/RepLayout.h
--- a/Engine/Source/Runtime/Engine/Public/Net/RepLayout.h
+++ b/Engine/Source/Runtime/Engine/Public/Net/RepLayout.h
@@ -1750,7 +1750,7 @@ private:
 		FReceivingRepState* RESTRICT RepState, 
 		FGuidReferencesMap* GuidReferencesMap,
 		UObject* OriginalObject,
-		UPackageMap* PackageMap, 
+		UNetConnection* Connection,
 		FRepShadowDataBuffer ShadowData, 
 		FRepObjectDataBuffer Data, 
 		const int32 MaxAbsOffset,
diff --git a/Engine/Source/Runtime/Engine/Public/ReplayTypes.h b/Engine/Source/Runtime/Engine/Public/ReplayTypes.h
--- a/Engine/Source/Runtime/Engine/Public/ReplayTypes.h
+++ b/Engine/Source/Runtime/Engine/Public/ReplayTypes.h
@@ -4,6 +4,7 @@
 
 #include "CoreMinimal.h"
 #include "UObject/ObjectMacros.h"
+#include "UObject/ObjectVersion.h"
 #include "Misc/EngineVersion.h"
 #include "Misc/NetworkGuid.h"
 #include "Misc/NetworkVersion.h"
@@ -87,21 +88,22 @@ struct FLevelNameAndTime
 enum ENetworkVersionHistory
 {
 	HISTORY_REPLAY_INITIAL = 1,
-	HISTORY_SAVE_ABS_TIME_MS = 2,			// We now save the abs demo time in ms for each frame (solves accumulation errors)
-	HISTORY_INCREASE_BUFFER = 3,			// Increased buffer size of packets, which invalidates old replays
+	HISTORY_SAVE_ABS_TIME_MS = 2,				// We now save the abs demo time in ms for each frame (solves accumulation errors)
+	HISTORY_INCREASE_BUFFER = 3,				// Increased buffer size of packets, which invalidates old replays
 	HISTORY_SAVE_ENGINE_VERSION = 4,			// Now saving engine net version + InternalProtocolVersion
-	HISTORY_EXTRA_VERSION = 5,			// We now save engine/game protocol version, checksum, and changelist
-	HISTORY_MULTIPLE_LEVELS = 6,			// Replays support seamless travel between levels
-	HISTORY_MULTIPLE_LEVELS_TIME_CHANGES = 7,			// Save out the time that level changes happen
+	HISTORY_EXTRA_VERSION = 5,					// We now save engine/game protocol version, checksum, and changelist
+	HISTORY_MULTIPLE_LEVELS = 6,				// Replays support seamless travel between levels
+	HISTORY_MULTIPLE_LEVELS_TIME_CHANGES = 7,	// Save out the time that level changes happen
 	HISTORY_DELETED_STARTUP_ACTORS = 8,			// Save DeletedNetStartupActors inside checkpoints
-	HISTORY_HEADER_FLAGS = 9,			// Save out enum flags with demo header
+	HISTORY_HEADER_FLAGS = 9,					// Save out enum flags with demo header
 	HISTORY_LEVEL_STREAMING_FIXES = 10,			// Optional level streaming fixes.
-	HISTORY_SAVE_FULL_ENGINE_VERSION = 11,			// Now saving the entire FEngineVersion including branch name
-	HISTORY_HEADER_GUID = 12,			// Save guid to demo header
+	HISTORY_SAVE_FULL_ENGINE_VERSION = 11,		// Now saving the entire FEngineVersion including branch name
+	HISTORY_HEADER_GUID = 12,					// Save guid to demo header
 	HISTORY_CHARACTER_MOVEMENT = 13,			// Change to using replicated movement and not interpolation
-	HISTORY_CHARACTER_MOVEMENT_NOINTERP = 14,			// No longer recording interpolated movement samples
+	HISTORY_CHARACTER_MOVEMENT_NOINTERP = 14,	// No longer recording interpolated movement samples
 	HISTORY_GUID_NAMETABLE = 15,				// Added a string table for exported guids
 	HISTORY_GUIDCACHE_CHECKSUMS = 16,			// Removing guid export checksums from saved data, they are ignored during playback
+	HISTORY_SAVE_PACKAGE_VERSION_UE = 17,		// Save engine and licensee package version as well, in case serialization functions need them for compatibility
 
 	// -----<new versions can be added before this line>-------------------------------------------------
 	HISTORY_PLUS_ONE,
@@ -130,6 +132,8 @@ struct FNetworkDemoHeader
 	EReplayHeaderFlags HeaderFlags;					// Replay flags
 	TArray<FLevelNameAndTime> LevelNamesAndTimes;	// Name and time changes of levels loaded for demo
 	TArray<FString> GameSpecificData;				// Area for subclasses to write stuff
+	FPackageFileVersion PackageVersionUE;			// Engine package version on which the replay was recorded
+	int32 PackageVersionLicenseeUE;					// Licensee package version on which the replay was recorded
 
 	FNetworkDemoHeader() :
 		Magic(NETWORK_DEMO_MAGIC),
@@ -139,7 +143,9 @@ struct FNetworkDemoHeader
 		GameNetworkProtocolVersion(FNetworkVersion::GetGameNetworkProtocolVersion()),
 		Guid(),
 		EngineVersion(FEngineVersion::Current()),
-		HeaderFlags(EReplayHeaderFlags::None)
+		HeaderFlags(EReplayHeaderFlags::None),
+		PackageVersionUE(GPackageFileUEVersion),
+		PackageVersionLicenseeUE(GPackageFileLicenseeUEVersion)
 	{
 	}
 
@@ -191,6 +197,38 @@ struct FNetworkDemoHeader
 			}
 		}
 
+		if (Header.Version >= HISTORY_SAVE_PACKAGE_VERSION_UE)
+		{
+			Ar << Header.PackageVersionUE;
+			Ar << Header.PackageVersionLicenseeUE;
+		}
+		else
+		{
+			if (Ar.IsLoading())
+			{
+				// Fix up for LWC compatibility.
+				// Vectors were using operator<< to serialize in some network cases (instead of
+				// the NetSerialize function), but this operator uses EUnrealEngineObjectUE5Version
+				// and not EEngineNetworkVersionHistory for versioning. Therefore, pre-LWC replays
+				// that did not save the EUnrealEngineObjectUE5Version will try to read doubles
+				// from these vectors instead of floats.
+				//
+				// However, EEngineNetworkVersionHistory::HISTORY_SERIALIZE_DOUBLE_VECTORS_AS_DOUBLES was
+				// added around the same time as EUnrealEngineObjectUE5Version::LARGE_WORLD_COORDINATES,
+				// so we use the network version as an approximation the package version to allow
+				// most pre-LWC replays to play back correctly. The compromise is that any replays recorded
+				// between when HISTORY_SERIALIZE_DOUBLE_VECTORS_AS_DOUBLES and LARGE_WORLD_COORDINATES
+				// were added won't play back correctly since they didn't store accurate version information.
+				if (Header.EngineNetworkProtocolVersion < HISTORY_SERIALIZE_DOUBLE_VECTORS_AS_DOUBLES)
+				{
+					// If the replay was recorded before vectors were serialized with LWC, and before replays
+					// saved the UE package version, set the package version to one before LARGE_WORLD_COORDINATES
+					// so that vectors will be read properly by operator<<.
+					Header.PackageVersionUE = FPackageFileVersion(VER_LATEST_ENGINE_UE4, EUnrealEngineObjectUE5Version::OPTIONAL_RESOURCES);
+				}
+			}
+		}
+
 		if (Header.Version < HISTORY_MULTIPLE_LEVELS)
 		{
 			FString LevelName;
