UPackageMapClient is now aware of incomplete PIE packages streaming external actors:
- Original fix was to prevent discovery of incomplete PIE packages from the package map client, which resulted in potential failures to resolve external actors outer chain from the incomplete package.
- Solution is to revert the original fix and make the package map aware of incomplete PIE packages to delay resolving of external actors until the package is considered complete. This cannot rely on UPackage::IsFullyLoaded since technically, the package is fully loaded and needs to be for external actors to properly resolve their outer.

diff --git a/Engine/Source/Runtime/CoreUObject/Private/UObject/Package.cpp b/Engine/Source/Runtime/CoreUObject/Private/UObject/Package.cpp
--- a/Engine/Source/Runtime/CoreUObject/Private/UObject/Package.cpp
+++ b/Engine/Source/Runtime/CoreUObject/Private/UObject/Package.cpp
@@ -60,6 +60,8 @@ void UPackage::PostInitProperties()
 	// Mark this package as editor-only by default. As soon as something in it is accessed through a non editor-only
 	// property the flag will be removed.
 	bLoadedByEditorPropertiesOnly = !HasAnyFlags(RF_ClassDefaultObject) && !HasAnyPackageFlags(PKG_CompiledIn) && (IsRunningCommandlet());
+
+	bIsDynamicPIEPackagePending = false;
 #endif
 }
 
diff --git a/Engine/Source/Runtime/CoreUObject/Public/UObject/Package.h b/Engine/Source/Runtime/CoreUObject/Public/UObject/Package.h
--- a/Engine/Source/Runtime/CoreUObject/Public/UObject/Package.h
+++ b/Engine/Source/Runtime/CoreUObject/Public/UObject/Package.h
@@ -169,6 +169,9 @@ private:
 #if WITH_EDITORONLY_DATA
 	/** True if this package is only referenced by editor-only properties */
 	uint8 bLoadedByEditorPropertiesOnly:1;
+
+	/** True if this package is a dynamic PIE package with external objects still loading */
+	uint8 bIsDynamicPIEPackagePending:1;
 #endif
 
 public:
@@ -301,6 +304,11 @@ public:
 	void SetLoadedByEditorPropertiesOnly(bool bIsEditorOnly, bool bRecursive = false);
 	/** returns true when the package is only referenced by editor-only flag */
 	bool IsLoadedByEditorPropertiesOnly() const { return bLoadedByEditorPropertiesOnly; }
+
+	/** Sets the bIsDynamicPIEPackagePending flag */
+	void SetDynamicPIEPackagePending(bool bInIsDynamicPIEPackagePending) { bIsDynamicPIEPackagePending = bInIsDynamicPIEPackagePending; }
+	/** returns the bIsDynamicPIEPackagePending flag */
+	bool IsDynamicPIEPackagePending() const { return bIsDynamicPIEPackagePending; }
 #endif
 
 	/**
diff --git a/Engine/Source/Runtime/Engine/Private/PackageMapClient.cpp b/Engine/Source/Runtime/Engine/Private/PackageMapClient.cpp
--- a/Engine/Source/Runtime/Engine/Private/PackageMapClient.cpp
+++ b/Engine/Source/Runtime/Engine/Private/PackageMapClient.cpp
@@ -3480,10 +3480,21 @@ bool FNetGUIDCache::ShouldIgnoreWhenMissing( const FNetworkGUID& NetGUID ) const
 			// Outer is pending, don't warn
 			return true;
 		}
-		// Sometimes, other systems async load packages, which we don't track, but still must be aware of
-		if ( OutermostCacheObject->Object != NULL && !OutermostCacheObject->Object->GetOutermost()->IsFullyLoaded() )
+
+		if ( OutermostCacheObject->Object != NULL )
 		{
-			return true;
+#if WITH_EDITOR
+			// Ignore if the package is a dynamic PIE package with pending external objects still loading
+			if ( !OutermostCacheObject->Object->GetOutermost()->IsDynamicPIEPackagePending() )
+			{
+				return true;
+			}
+#endif
+			// Sometimes, other systems async load packages, which we don't track, but still must be aware of
+			if ( !OutermostCacheObject->Object->GetOutermost()->IsFullyLoaded() )
+			{
+				return true;
+			}
 		}
 	}
 
@@ -3997,4 +4008,4 @@ bool FNetGUIDCache::FIsOwnerOrPawnHelper::IsOwnerOrPawn() const
 
 	return !!CachedResult;
 }
-#endif
\ No newline at end of file
+#endif
diff --git a/Engine/Source/Runtime/Engine/Private/WorldPartition/WorldPartitionLevelHelper.cpp b/Engine/Source/Runtime/Engine/Private/WorldPartition/WorldPartitionLevelHelper.cpp
--- a/Engine/Source/Runtime/Engine/Private/WorldPartition/WorldPartitionLevelHelper.cpp
+++ b/Engine/Source/Runtime/Engine/Private/WorldPartition/WorldPartitionLevelHelper.cpp
@@ -205,6 +205,12 @@ ULevel* FWorldPartitionLevelHelper::CreateEmptyLevelForRuntimeCell(const UWorld*
 
 	// Mark the level as a runtime cell
 	NewLevel->bIsWorldPartitionRuntimeCell = true;
+	
+	// Mark the level package as fully loaded
+	CellPackage->MarkAsFullyLoaded();
+
+	// Mark the level package as containing a map
+	CellPackage->ThisContainsMap();
 
 	// Set the guids on the constructed level to something based on the generator rather than allowing indeterminism by
 	// constructing new Guids on every cook
diff --git a/Engine/Source/Runtime/Engine/Private/WorldPartition/WorldPartitionLevelStreamingDynamic.cpp b/Engine/Source/Runtime/Engine/Private/WorldPartition/WorldPartitionLevelStreamingDynamic.cpp
--- a/Engine/Source/Runtime/Engine/Private/WorldPartition/WorldPartitionLevelStreamingDynamic.cpp
+++ b/Engine/Source/Runtime/Engine/Private/WorldPartition/WorldPartitionLevelStreamingDynamic.cpp
@@ -127,12 +127,13 @@ void UWorldPartitionLevelStreamingDynamic::CreateRuntimeLevel()
 	const UWorld* World = GetWorld();
 	check(World && (World->IsGameWorld() || GetShouldBeVisibleInEditor()));
 
-	// Create streaming cell Level package (make sure this package isn't discoverable until it is fully loaded)
-	const FString WorldAssetPath = GetWorldAsset().ToString();
-	const FString WorldAssetPackageName = FPackageName::ObjectPathToPackageName(WorldAssetPath) + TEXT("_AsyncLoad.") + FPackageName::ObjectPathToObjectName(WorldAssetPath);
-	RuntimeLevel = FWorldPartitionLevelHelper::CreateEmptyLevelForRuntimeCell(World, WorldAssetPackageName);
+	// Create streaming cell Level package
+	RuntimeLevel = FWorldPartitionLevelHelper::CreateEmptyLevelForRuntimeCell(World, GetWorldAsset().ToString());
 	check(RuntimeLevel);
 
+	UPackage* RuntimeLevelPackage = RuntimeLevel->GetPackage();
+	check(RuntimeLevelPackage);
+
 	// Propagate ActorFolder flag to the runtime level and prepare its ActorFolders list
 	if (World->PersistentLevel->IsUsingActorFolders() && ActorFolders.Num())
 	{
@@ -149,6 +150,9 @@ void UWorldPartitionLevelStreamingDynamic::CreateRuntimeLevel()
 	// Set flag here as this level isn't async loaded
 	RuntimeLevel->bClientOnlyVisible = bClientOnlyVisible;
 
+	// Mark this package as a dynamic PIE package with pending external actors
+	RuntimeLevelPackage->SetDynamicPIEPackagePending(true);
+
 	// Attach ourself to Level cleanup to do our own cleanup
 	OnCleanupLevelDelegateHandle = RuntimeLevel->OnCleanupLevel.AddUObject(this, &UWorldPartitionLevelStreamingDynamic::OnCleanupLevel);
 }
@@ -381,14 +385,6 @@ void UWorldPartitionLevelStreamingDynamic::FinalizeRuntimeLevel()
 	UWorld* OuterWorld = RuntimeLevel->GetTypedOuter<UWorld>();
 	OuterWorld->bIsNameStableForNetworking = true;
 
-	UPackage* RuntimePackage = RuntimeLevel->GetPackage();
-	RuntimePackage->MarkAsFullyLoaded();
-
-	// Now that the loading is done, rename it so this level is discoverable
-	const FString WorldAssetPath = GetWorldAsset().ToString();
-	const FString WorldAssetPackageName = FPackageName::ObjectPathToPackageName(WorldAssetPath);
-	RuntimePackage->Rename(*WorldAssetPackageName, nullptr, REN_ForceNoResetLoaders | REN_DontCreateRedirectors | REN_NonTransactional | REN_DoNotDirty);
-
 	if (StreamingCell.IsValid() && !StreamingCell->GetIsHLOD())
 	{
 		TRACE_CPUPROFILER_EVENT_SCOPE(FixupIDs);
@@ -425,6 +421,9 @@ void UWorldPartitionLevelStreamingDynamic::FinalizeRuntimeLevel()
 	// Make sure this level will start to render only when it will be fully added to the world
 	check(ShouldRequireFullVisibilityToRender());
 	RuntimeLevel->bRequireFullVisibilityToRender = true;
+
+	// Mark this package as fully loaded with regards to external objects
+	RuntimeLevel->GetPackage()->SetDynamicPIEPackagePending(false);
 }
 
 /**
