High res editor screenshot with mask fix, when the viewport aspect ratio is constrained.  When a rect is specified for the screenshot, apply that rect to the mask image when merging it into the alpha of the color image, to match the rect used in the prior call to FRenderTarget::ReadPixels.

#rnx

diff --git a/Engine/Source/Editor/UnrealEd/Private/EditorViewportClient.cpp b/Engine/Source/Editor/UnrealEd/Private/EditorViewportClient.cpp
--- a/Engine/Source/Editor/UnrealEd/Private/EditorViewportClient.cpp
+++ b/Engine/Source/Editor/UnrealEd/Private/EditorViewportClient.cpp
@@ -5738,7 +5738,7 @@ bool FEditorViewportClient::ProcessScreenShots(FViewport* InViewport)
 			// Get the results of the mask rendering pass and insert into the alpha channel of the screenshot.
 			if (GIsHighResScreenshot && HighResScreenshotConfig.bMaskEnabled)
 			{
-				bWriteAlpha = HighResScreenshotConfig.MergeMaskIntoAlpha(Bitmap);
+				bWriteAlpha = HighResScreenshotConfig.MergeMaskIntoAlpha(Bitmap, CaptureRect);
 			}
 
 			// Clip the bitmap to just the capture region if valid
diff --git a/Engine/Source/Runtime/Engine/Private/GameViewportClient.cpp b/Engine/Source/Runtime/Engine/Private/GameViewportClient.cpp
--- a/Engine/Source/Runtime/Engine/Private/GameViewportClient.cpp
+++ b/Engine/Source/Runtime/Engine/Private/GameViewportClient.cpp
@@ -1947,7 +1947,7 @@ bool UGameViewportClient::ProcessScreenShots(FViewport* InViewport)
 				// global variables associated with the screenshot feature, as the application may be taking its own screenshots.
 				if (!bIsUI)
 				{
-					GetHighResScreenshotConfig().MergeMaskIntoAlpha(Bitmap);
+					GetHighResScreenshotConfig().MergeMaskIntoAlpha(Bitmap, FIntRect(0,0,0,0));
 				}
 
 				FIntRect SourceRect(0, 0, GScreenshotResolutionX, GScreenshotResolutionY);
diff --git a/Engine/Source/Runtime/Engine/Private/HighResScreenshot.cpp b/Engine/Source/Runtime/Engine/Private/HighResScreenshot.cpp
--- a/Engine/Source/Runtime/Engine/Private/HighResScreenshot.cpp
+++ b/Engine/Source/Runtime/Engine/Private/HighResScreenshot.cpp
@@ -143,20 +143,42 @@ bool FHighResScreenshotConfig::ParseConsoleCommand(const FString& InCmd, FOutput
 	return false;
 }
 
-bool FHighResScreenshotConfig::MergeMaskIntoAlpha(TArray<FColor>& InBitmap)
+bool FHighResScreenshotConfig::MergeMaskIntoAlpha(TArray<FColor>& InBitmap, const FIntRect& ViewRect)
 {
 	bool bWritten = false;
 
 	TArray<FColor>* MaskArray = FScreenshotRequest::GetHighresScreenshotMaskColorArray();
-	bool bMaskMatches = !bMaskEnabled || (MaskArray->Num() == InBitmap.Num());
+	const FIntPoint& MaskExtents = FScreenshotRequest::GetHighresScreenshotMaskExtents();
+
+	bool bMaskMatches = !bMaskEnabled || (InBitmap.Num() == MaskArray->Num()) || (InBitmap.Num() == ViewRect.Area() && ViewRect.Max.X <= MaskExtents.X && ViewRect.Max.Y <= MaskExtents.Y);
 	ensureMsgf(bMaskMatches, TEXT("Highres screenshot MaskArray doesn't match screenshot size.  Skipping Masking. MaskSize: %i, ScreenshotSize: %i"), MaskArray->Num(), InBitmap.Num());
 	if (bMaskEnabled && bMaskMatches)
 	{
 		// If this is a high resolution screenshot and we are using the masking feature,
 		// Get the results of the mask rendering pass and insert into the alpha channel of the screenshot.
-		for (int32 i = 0; i < InBitmap.Num(); ++i)
+		if (InBitmap.Num() == MaskArray->Num())
+		{
+			// Exact match, copy verbatim
+			for (int32 i = 0; i < InBitmap.Num(); ++i)
+			{
+				InBitmap[i].A = (*MaskArray)[i].R;
+			}
+		}
+		else
 		{
-			InBitmap[i].A = (*MaskArray)[i].R;
+			// Need to pull a rectangle out of the mask array
+			int32 RectOffsetX = ViewRect.Min.X;
+			int32 RectOffsetY = ViewRect.Min.Y;
+			int32 OutputOffset = 0;
+			int32 MaskStride = MaskExtents.X;
+
+			for (int32 j = ViewRect.Min.Y; j < ViewRect.Max.Y; j++)
+			{
+				for (int32 i = ViewRect.Min.X; i < ViewRect.Max.X; i++, OutputOffset++)
+				{
+					InBitmap[OutputOffset].A = (*MaskArray)[j * MaskStride + i].R;
+				}
+			}
 		}
 		bWritten = true;
 	}
diff --git a/Engine/Source/Runtime/Engine/Private/UnrealClient.cpp b/Engine/Source/Runtime/Engine/Private/UnrealClient.cpp
--- a/Engine/Source/Runtime/Engine/Private/UnrealClient.cpp
+++ b/Engine/Source/Runtime/Engine/Private/UnrealClient.cpp
@@ -378,6 +378,10 @@ TArray<FColor>* FScreenshotRequest::GetHighresScreenshotMaskColorArray()
 	return &HighresScreenshotMaskColorArray;
 }
 
+FIntPoint& FScreenshotRequest::GetHighresScreenshotMaskExtents()
+{
+	return HighresScreenshotMaskExtents;
+}
 
 // @param bAutoType true: automatically choose GB/MB/KB/... false: always use MB for easier comparisons
 FString GetMemoryString( const double Value, const bool bAutoType )
@@ -409,6 +413,7 @@ FString FScreenshotRequest::Filename;
 FString FScreenshotRequest::NextScreenshotName;
 bool FScreenshotRequest::bShowUI = false;
 TArray<FColor> FScreenshotRequest::HighresScreenshotMaskColorArray;
+FIntPoint FScreenshotRequest::HighresScreenshotMaskExtents;
 
 static TAutoConsoleVariable<int32> CVarFullSizeUnitGraph(
 	TEXT("FullSizeUnitGraph"),
diff --git a/Engine/Source/Runtime/Engine/Public/HighResScreenshot.h b/Engine/Source/Runtime/Engine/Public/HighResScreenshot.h
--- a/Engine/Source/Runtime/Engine/Public/HighResScreenshot.h
+++ b/Engine/Source/Runtime/Engine/Public/HighResScreenshot.h
@@ -54,7 +54,7 @@ struct ENGINE_API FHighResScreenshotConfig
 	/** Utility function for merging the mask buffer into the alpha channel of the supplied bitmap, if masking is enabled.
 	  * Returns true if the mask was written, and false otherwise.
 	**/
-	bool MergeMaskIntoAlpha(TArray<FColor>& InBitmap);
+	bool MergeMaskIntoAlpha(TArray<FColor>& InBitmap, const FIntRect& ViewRect);
 
 	/** Enable/disable HDR capable captures **/
 	void SetHDRCapture(bool bCaptureHDRIN);
diff --git a/Engine/Source/Runtime/Engine/Public/UnrealClient.h b/Engine/Source/Runtime/Engine/Public/UnrealClient.h
--- a/Engine/Source/Runtime/Engine/Public/UnrealClient.h
+++ b/Engine/Source/Runtime/Engine/Public/UnrealClient.h
@@ -192,6 +192,11 @@ struct ENGINE_API FScreenshotRequest
 	 */
 	static TArray<FColor>* GetHighresScreenshotMaskColorArray();
 
+	/**
+	 * Extents of array returned by function above
+	 */
+	static FIntPoint& GetHighresScreenshotMaskExtents();
+
 	static FOnScreenshotRequestProcessed& OnScreenshotRequestProcessed()
 	{
 		return ScreenshotProcessedDelegate;
@@ -210,6 +215,7 @@ private:
 	static FString Filename;
 	static bool bShowUI;
 	static TArray<FColor> HighresScreenshotMaskColorArray;
+	static FIntPoint HighresScreenshotMaskExtents;
 };
 
 // @param bAutoType true: automatically choose GB/MB/KB/... false: always use MB for easier comparisons
diff --git a/Engine/Source/Runtime/Renderer/Private/PostProcess/PostProcessMaterial.cpp b/Engine/Source/Runtime/Renderer/Private/PostProcess/PostProcessMaterial.cpp
--- a/Engine/Source/Runtime/Renderer/Private/PostProcess/PostProcessMaterial.cpp
+++ b/Engine/Source/Runtime/Renderer/Private/PostProcess/PostProcessMaterial.cpp
@@ -787,7 +787,7 @@ FScreenPassTexture AddPostProcessMaterialChain(
 	return Outputs;
 }
 
-extern void AddDumpToColorArrayPass(FRDGBuilder& GraphBuilder, FScreenPassTexture Input, TArray<FColor>* OutputColorArray);
+extern void AddDumpToColorArrayPass(FRDGBuilder& GraphBuilder, FScreenPassTexture Input, TArray<FColor>* OutputColorArray, FIntPoint* OutputExtents);
 
 bool IsHighResolutionScreenshotMaskEnabled(const FViewInfo& View)
 {
@@ -860,7 +860,7 @@ FScreenPassTexture AddHighResolutionScreenshotMaskPass(
 		PassInputs.bAllowSceneColorInputAsOutput = false;
 
 		FScreenPassTexture MaskOutput = AddPostProcessMaterialPass(GraphBuilder, View, PassInputs, Inputs.MaskMaterial);
-		AddDumpToColorArrayPass(GraphBuilder, MaskOutput, FScreenshotRequest::GetHighresScreenshotMaskColorArray());
+		AddDumpToColorArrayPass(GraphBuilder, MaskOutput, FScreenshotRequest::GetHighresScreenshotMaskColorArray(), &FScreenshotRequest::GetHighresScreenshotMaskExtents());
 
 		// The mask material pass is actually outputting to system memory. If we're the last pass in the chain
 		// and the override output is valid, we need to perform a copy of the input to the output. Since we can't
diff --git a/Engine/Source/Runtime/Renderer/Private/PostProcess/PostProcessVisualizeBuffer.cpp b/Engine/Source/Runtime/Renderer/Private/PostProcess/PostProcessVisualizeBuffer.cpp
--- a/Engine/Source/Runtime/Renderer/Private/PostProcess/PostProcessVisualizeBuffer.cpp
+++ b/Engine/Source/Runtime/Renderer/Private/PostProcess/PostProcessVisualizeBuffer.cpp
@@ -288,12 +288,12 @@ void AddDumpToFilePass(FRDGBuilder& GraphBuilder, FScreenPassTexture Input, cons
 	});
 }
 
-void AddDumpToColorArrayPass(FRDGBuilder& GraphBuilder, FScreenPassTexture Input, TArray<FColor>* OutputColorArray)
+void AddDumpToColorArrayPass(FRDGBuilder& GraphBuilder, FScreenPassTexture Input, TArray<FColor>* OutputColorArray, FIntPoint* OutputExtents)
 {
 	check(Input.IsValid());
 	check(OutputColorArray);
 	AddReadbackTexturePass(GraphBuilder, RDG_EVENT_NAME("DumpToPipe(%s)", Input.Texture->Name), Input.Texture,
-		[Input, OutputColorArray](FRHICommandListImmediate& RHICmdList)
+		[Input, OutputColorArray, OutputExtents](FRHICommandListImmediate& RHICmdList)
 	{
 		// By design, we want the whole surface, not the view rectangle, as this code is used for generating a screenshot
 		// mask surface that needs to match the corresponding screenshot color surface.  The scene may render as a viewport
@@ -305,6 +305,7 @@ void AddDumpToColorArrayPass(FRDGBuilder& GraphBuilder, FScreenPassTexture Input
 		WholeSurfaceRect.Max = Input.Texture->Desc.Extent;
 
 		RHICmdList.ReadSurfaceData(Input.Texture->GetRHI(), WholeSurfaceRect, *OutputColorArray, FReadSurfaceDataFlags());
+		*OutputExtents = Input.Texture->Desc.Extent;
 	});
 }
 
