UE-141180 - Loss of precision with String to Float Real conversion

* Added missing double/string conversion functions, including BuildString.
* Deprecated float versions, and added redirect fixup to use the double versions instead.
* Replaced direct Conv_ usage with LexToString/LexFromString.

diff --git a/Engine/Config/BaseEngine.ini b/Engine/Config/BaseEngine.ini
--- a/Engine/Config/BaseEngine.ini
+++ b/Engine/Config/BaseEngine.ini
@@ -1011,6 +1011,15 @@ NearClipPlane=10.0
 +PropertyRedirects=(OldName="PostProcessSettings.LocalExposureContrastReduction", NewName="PostProcessSettings.LocalExposureContrastScale")
 +PropertyRedirects=(OldName="PostProcessSettings.bOverride_LocalExposureContrastReduction", NewName="PostProcessSettings.bOverride_LocalExposureContrastScale")
 
+; Blueprint real number support
++FunctionRedirects=(OldName="KismetStringLibrary.Conv_FloatToString", NewName="Conv_DoubleToString")
++PropertyRedirects=(OldName="KismetStringLibrary.Conv_DoubleToString.InFloat", NewName="Conv_DoubleToString.InDouble")
+
++FunctionRedirects=(OldName="KismetStringLibrary.Conv_StringToFloat", NewName="Conv_StringToDouble")
+
++FunctionRedirects=(OldName="KismetStringLibrary.BuildString_Float", NewName="BuildString_Double")
++PropertyRedirects=(OldName="KismetStringLibrary.BuildString_Double.InFloat", NewName="BuildString_Double.InDouble")
+
 [CoreUObject.Metadata]
 ; Note: UnrealHeaderTool should be rerun after making changes to MetadataRedirects to catch any keys specified in class headers
 +MetadataRedirects=(OldKey="K2Protected", NewKey="BlueprintProtected")
diff --git a/Engine/Source/Runtime/Engine/Classes/Kismet/KismetStringLibrary.h b/Engine/Source/Runtime/Engine/Classes/Kismet/KismetStringLibrary.h
--- a/Engine/Source/Runtime/Engine/Classes/Kismet/KismetStringLibrary.h
+++ b/Engine/Source/Runtime/Engine/Classes/Kismet/KismetStringLibrary.h
@@ -13,11 +13,11 @@ class ENGINE_API UKismetStringLibrary : public UBlueprintFunctionLibrary
 	GENERATED_UCLASS_BODY()
 
 	/** Converts a float value to a string */
-	UFUNCTION(BlueprintPure, meta=(DisplayName = "To String (Float)", CompactNodeTitle = "->", BlueprintAutocast), Category="Utilities|String")
+	UE_DEPRECATED(5.0, "This method has been deprecated and will be removed. Please use FString::SanitizeFloat() or LexToSanitzedString instead.")
 	static FString Conv_FloatToString(float InFloat);
 
 	/** Converts a double value to a string */
-	UFUNCTION(BlueprintPure, meta = (DisplayName = "To String (Double)", CompactNodeTitle = "->", BlueprintAutocast), Category = "Utilities|String")
+	UFUNCTION(BlueprintPure, meta = (DisplayName = "To String (Float)", CompactNodeTitle = "->", BlueprintAutocast), Category = "Utilities|String")
 	static FString Conv_DoubleToString(double InDouble);
 
 	/** Converts an integer value to a string */
@@ -89,9 +89,13 @@ class ENGINE_API UKismetStringLibrary : public UBlueprintFunctionLibrary
 	static int32 Conv_StringToInt(const FString& InString);
 
 	/** Converts a string to a float value */
-	UFUNCTION(BlueprintPure, meta=(DisplayName = "String To Float", CompactNodeTitle = "->", BlueprintAutocast), Category="Utilities|String")
+	UE_DEPRECATED(5.0, "This method has been deprecated and will be removed. Please use LexFromString instead.")
 	static float Conv_StringToFloat(const FString& InString);
-	
+
+	/** Converts a string to a double value */
+	UFUNCTION(BlueprintPure, meta = (DisplayName = "String To Float", CompactNodeTitle = "->", BlueprintAutocast), Category = "Utilities|String")
+	static double Conv_StringToDouble(const FString& InString);
+
 	/** Convert String Back To Vector. IsValid indicates whether or not the string could be successfully converted. */
 	UFUNCTION(BlueprintPure, meta=(DisplayName = "String To Vector", CompactNodeTitle = "->"), Category="Utilities|String")
 	static void Conv_StringToVector(const FString& InString, FVector& OutConvertedVector, bool& OutIsValid);
@@ -120,9 +124,20 @@ class ENGINE_API UKismetStringLibrary : public UBlueprintFunctionLibrary
 	 * @param Suffix - A suffix to append to the end of the conversion string
 	 * @return A new string built from the passed parameters
 	 */
-	UFUNCTION(BlueprintPure, meta=(DisplayName = "Build String (Float)"), Category="Utilities|String")
+	UE_DEPRECATED(5.0, "This method has been deprecated and will be removed. Please use FString::Printf instead.")
 	static FString BuildString_Float(const FString& AppendTo, const FString& Prefix, float InFloat, const FString& Suffix);
 
+	/**
+	 * Converts a double->string, create a new string in the form AppendTo+Prefix+InDouble+Suffix
+	 * @param AppendTo - An existing string to use as the start of the conversion string
+	 * @param Prefix - A string to use as a prefix, after the AppendTo string
+	 * @param InDouble - The double value to convert
+	 * @param Suffix - A suffix to append to the end of the conversion string
+	 * @return A new string built from the passed parameters
+	 */
+	UFUNCTION(BlueprintPure, meta = (DisplayName = "Build String (Float)"), Category = "Utilities|String")
+	static FString BuildString_Double(const FString& AppendTo, const FString& Prefix, double InDouble, const FString& Suffix);
+
 	/** 
 	 * Converts a int->string, creating a new string in the form AppendTo+Prefix+InInt+Suffix
 	 * @param AppendTo - An existing string to use as the start of the conversion string
diff --git a/Engine/Source/Runtime/Engine/Private/KismetStringLibrary.cpp b/Engine/Source/Runtime/Engine/Private/KismetStringLibrary.cpp
--- a/Engine/Source/Runtime/Engine/Private/KismetStringLibrary.cpp
+++ b/Engine/Source/Runtime/Engine/Private/KismetStringLibrary.cpp
@@ -151,6 +151,11 @@ float UKismetStringLibrary::Conv_StringToFloat(const FString& InString)
 	return FCString::Atof(*InString);
 }
 
+double UKismetStringLibrary::Conv_StringToDouble(const FString& InString)
+{
+	return FCString::Atod(*InString);
+}
+
 void UKismetStringLibrary::Conv_StringToVector(const FString& InString, FVector& OutConvertedVector, bool& OutIsValid)
 {
 	OutIsValid = OutConvertedVector.InitFromString(InString);
@@ -178,14 +183,19 @@ void UKismetStringLibrary::Conv_StringToColor(const FString& InString, FLinearCo
 
 FString UKismetStringLibrary::BuildString_Float(const FString& AppendTo, const FString& Prefix, float InFloat, const FString& Suffix)
 {
-	// faster, preallocating method
-	FString const FloatStr = FString::SanitizeFloat(InFloat);
+	return BuildString_Double(AppendTo, Prefix, InFloat, Suffix);
+}
+
+FString UKismetStringLibrary::BuildString_Double(const FString& AppendTo, const FString& Prefix, double InDouble, const FString& Suffix)
+{
+	// despite the name, SanitizeFloat takes a double parameter
+	const FString DoubleStr = FString::SanitizeFloat(InDouble);
 
 	FString StringResult;
-	StringResult.Empty(AppendTo.Len()+Prefix.Len()+FloatStr.Len()+Suffix.Len()+1); // adding one for the string terminator
+	StringResult.Empty(AppendTo.Len() + Prefix.Len() + DoubleStr.Len() + Suffix.Len() + 1);
 	StringResult += AppendTo;
 	StringResult += Prefix;
-	StringResult += FloatStr;
+	StringResult += DoubleStr;
 	StringResult += Suffix;
 
 	return StringResult;
@@ -194,7 +204,7 @@ FString UKismetStringLibrary::BuildString_Float(const FString& AppendTo, const F
 FString UKismetStringLibrary::BuildString_Int(const FString& AppendTo, const FString& Prefix, int32 InInt, const FString& Suffix)
 {
 	// faster, preallocating method
-	FString const IntStr = FString::Printf(TEXT("%d"), InInt);
+	const FString IntStr = FString::Printf(TEXT("%d"), InInt);
 
 	FString StringResult;
 	StringResult.Empty(AppendTo.Len()+Prefix.Len()+IntStr.Len()+Suffix.Len()+1); // adding one for the string terminator
@@ -209,7 +219,7 @@ FString UKismetStringLibrary::BuildString_Int(const FString& AppendTo, const FSt
 FString UKismetStringLibrary::BuildString_Bool(const FString& AppendTo, const FString& Prefix, bool InBool, const FString& Suffix)
 {
 	// faster, preallocating method
-	FString const BoolStr = InBool ? TEXT("true") : TEXT("false");	
+	const FString BoolStr = InBool ? TEXT("true") : TEXT("false");	
 
 	FString StringResult;
 	StringResult.Empty(AppendTo.Len()+Prefix.Len()+BoolStr.Len()+Suffix.Len()+1); // adding one for the string terminator
@@ -224,7 +234,7 @@ FString UKismetStringLibrary::BuildString_Bool(const FString& AppendTo, const FS
 FString UKismetStringLibrary::BuildString_Vector(const FString& AppendTo, const FString& Prefix, FVector InVector, const FString& Suffix)
 {
 	// faster, preallocating method
-	FString const VecStr = InVector.ToString();
+	const FString VecStr = InVector.ToString();
 
 	FString StringResult;
 	StringResult.Empty(AppendTo.Len()+Prefix.Len()+VecStr.Len()+Suffix.Len()+1); // adding one for the string terminator
@@ -239,7 +249,7 @@ FString UKismetStringLibrary::BuildString_Vector(const FString& AppendTo, const
 FString UKismetStringLibrary::BuildString_IntVector(const FString& AppendTo, const FString& Prefix, FIntVector InIntVector, const FString& Suffix)
 {
 	// faster, preallocating method
-	FString const VecStr = InIntVector.ToString();
+	const FString VecStr = InIntVector.ToString();
 
 	FString StringResult;
 	StringResult.Empty(AppendTo.Len() + Prefix.Len() + VecStr.Len() + Suffix.Len() + 1); // adding one for the string terminator
@@ -254,7 +264,7 @@ FString UKismetStringLibrary::BuildString_IntVector(const FString& AppendTo, con
 FString UKismetStringLibrary::BuildString_Vector2d(const FString& AppendTo, const FString& Prefix, FVector2D InVector2d, const FString& Suffix)
 {
 	// faster, preallocating method
-	FString const VecStr = InVector2d.ToString();
+	const FString VecStr = InVector2d.ToString();
 
 	FString StringResult;
 	StringResult.Empty(AppendTo.Len()+Prefix.Len()+VecStr.Len()+Suffix.Len()+1); // adding one for the string terminator
@@ -269,7 +279,7 @@ FString UKismetStringLibrary::BuildString_Vector2d(const FString& AppendTo, cons
 FString UKismetStringLibrary::BuildString_Rotator(const FString& AppendTo, const FString& Prefix, FRotator InRot, const FString& Suffix)
 {
 	// faster, preallocating method
-	FString const RotStr = InRot.ToString();
+	const FString RotStr = InRot.ToString();
 
 	FString StringResult;
 	StringResult.Empty(AppendTo.Len()+Prefix.Len()+RotStr.Len()+Suffix.Len()+1); // adding one for the string terminator
@@ -284,7 +294,7 @@ FString UKismetStringLibrary::BuildString_Rotator(const FString& AppendTo, const
 FString UKismetStringLibrary::BuildString_Object(const FString& AppendTo, const FString& Prefix, class UObject* InObj, const FString& Suffix)
 {
 	// faster, preallocating method
-	FString const ObjStr = (InObj != NULL) ? InObj->GetName() : FString(TEXT("None"));
+	const FString ObjStr = (InObj != NULL) ? InObj->GetName() : FString(TEXT("None"));
 
 	FString StringResult;
 	StringResult.Empty(AppendTo.Len()+Prefix.Len()+ObjStr.Len()+Suffix.Len()+1); // adding one for the string terminator
@@ -299,7 +309,7 @@ FString UKismetStringLibrary::BuildString_Object(const FString& AppendTo, const
 FString UKismetStringLibrary::BuildString_Color(const FString& AppendTo, const FString& Prefix, FLinearColor InColor, const FString& Suffix)
 {
 	// faster, preallocating method
-	FString const ColorStr = InColor.ToString();
+	const FString ColorStr = InColor.ToString();
 
 	FString StringResult;
 	StringResult.Empty(AppendTo.Len()+Prefix.Len()+ColorStr.Len()+Suffix.Len()+1); // adding one for the string terminator
@@ -314,7 +324,7 @@ FString UKismetStringLibrary::BuildString_Color(const FString& AppendTo, const F
 FString UKismetStringLibrary::BuildString_Name(const FString& AppendTo, const FString& Prefix, FName InName, const FString& Suffix)
 {
 	// faster, preallocating method
-	FString const NameStr = InName.ToString();
+	const FString NameStr = InName.ToString();
 
 	FString StringResult;
 	StringResult.Empty(AppendTo.Len()+Prefix.Len()+NameStr.Len()+Suffix.Len()+1); // adding one for the string terminator
