Added mobile shader permutation cvar to support bSupportStaticLighting=true with Force No Precomputed Lighting enabled on a level

diff --git a/Engine/Source/Runtime/Core/Private/HAL/ConsoleManager.cpp b/Engine/Source/Runtime/Core/Private/HAL/ConsoleManager.cpp
--- a/Engine/Source/Runtime/Core/Private/HAL/ConsoleManager.cpp
+++ b/Engine/Source/Runtime/Core/Private/HAL/ConsoleManager.cpp
@@ -2457,6 +2457,13 @@ static TAutoConsoleVariable<int32> CVarMobileEnableMovableLightCSMShaderCulling(
 	TEXT("1: Primitives lit by movable directional light render with the CSM shader when determined to be within CSM range. (default)"),
 	ECVF_RenderThreadSafe | ECVF_ReadOnly);
 
+static TAutoConsoleVariable<int32> CVarMobileEnableNoPrecomputedLightingCSMShader(
+	TEXT("r.Mobile.EnableNoPrecomputedLightingCSMShader"),
+	0,
+	TEXT("0: CSM shaders for scenes without any precomputed lighting are not generated unless r.AllowStaticLighting is 0. (default)\n")
+	TEXT("1: CSM shaders for scenes without any precomputed lighting are always generated."),
+	ECVF_RenderThreadSafe | ECVF_ReadOnly);
+
 static TAutoConsoleVariable<int32> CVarMobileUseCSMShaderBranch(
 	TEXT("r.Mobile.UseCSMShaderBranch"),
 	0,
diff --git a/Engine/Source/Runtime/Engine/Classes/Engine/RendererSettings.h b/Engine/Source/Runtime/Engine/Classes/Engine/RendererSettings.h
--- a/Engine/Source/Runtime/Engine/Classes/Engine/RendererSettings.h
+++ b/Engine/Source/Runtime/Engine/Classes/Engine/RendererSettings.h
@@ -936,6 +936,14 @@ class ENGINE_API URendererSettings : public UDeveloperSettings
 		ConfigRestartRequired = true))
 		uint32 bMobileEnableMovableLightCSMShaderCulling : 1;
 
+	UPROPERTY(config, EditAnywhere, Category = MobileShaderPermutationReduction, meta = (
+		ConsoleVariable = "r.Mobile.EnableNoPrecomputedLightingCSMShader",
+		DisplayName = "Support CSM on levels with Force No Precomputed Lighting enabled",
+		EditCondition = "bAllowStaticLighting",
+		ToolTip = "When Allow Static Lighting is enabled, shaders to support CSM without any precomputed lighting are not normally generated. This setting allows CSM for this case at the cost of extra shader permutations. Changing this setting requires restarting the editor.",
+		ConfigRestartRequired = true))
+		uint32 bMobileEnableNoPrecomputedLightingCSMShader : 1;
+
 	UPROPERTY(config, EditAnywhere, Category = MobileShaderPermutationReduction, meta = (
 		ConsoleVariable = "r.Mobile.AllowDistanceFieldShadows",
 		DisplayName = "Support Pre-baked Distance Field Shadow Maps",
diff --git a/Engine/Source/Runtime/Engine/Private/SceneManagement.cpp b/Engine/Source/Runtime/Engine/Private/SceneManagement.cpp
--- a/Engine/Source/Runtime/Engine/Private/SceneManagement.cpp
+++ b/Engine/Source/Runtime/Engine/Private/SceneManagement.cpp
@@ -1170,6 +1170,7 @@ void FReadOnlyCVARCache::Init()
 	static const auto CVarMobileSkyLightPermutation = IConsoleManager::Get().FindTConsoleVariableDataInt(TEXT("r.Mobile.SkyLightPermutation"));
 	static const auto CVarMobileEnableMovableSpotLights = IConsoleManager::Get().FindTConsoleVariableDataInt(TEXT("r.Mobile.EnableMovableSpotlights"));
 	static const auto CVarMobileEnableMovableSpotLightsShadow = IConsoleManager::Get().FindTConsoleVariableDataInt(TEXT("r.Mobile.EnableMovableSpotlightsShadow"));
+	static const auto CVarMobileEnableNoPrecomputedLightingCSMShader = IConsoleManager::Get().FindTConsoleVariableDataInt(TEXT("r.Mobile.EnableNoPrecomputedLightingCSMShader"));
 
 	const bool bForceAllPermutations = CVarSupportAllShaderPermutations && CVarSupportAllShaderPermutations->GetValueOnAnyThread() != 0;
 
@@ -1187,6 +1188,7 @@ void FReadOnlyCVARCache::Init()
 	MobileSkyLightPermutation = CVarMobileSkyLightPermutation->GetValueOnAnyThread();
 	bMobileEnableMovableSpotlights = CVarMobileEnableMovableSpotLights->GetValueOnAnyThread() != 0;
 	bMobileEnableMovableSpotlightsShadow = bMobileEnableMovableSpotlights && CVarMobileEnableMovableSpotLightsShadow->GetValueOnAnyThread() != 0;
+	bMobileEnableNoPrecomputedLightingCSMShader = CVarMobileEnableNoPrecomputedLightingCSMShader->GetValueOnAnyThread() != 0;
 
 	const bool bShowMissmatchedLowQualityLightmapsWarning = (!bEnableLowQualityLightmaps) && (GEngine->bShouldGenerateLowQualityLightmaps_DEPRECATED);
 	if ( bShowMissmatchedLowQualityLightmapsWarning )
diff --git a/Engine/Source/Runtime/Engine/Public/SceneManagement.h b/Engine/Source/Runtime/Engine/Public/SceneManagement.h
--- a/Engine/Source/Runtime/Engine/Public/SceneManagement.h
+++ b/Engine/Source/Runtime/Engine/Public/SceneManagement.h
@@ -3360,6 +3360,7 @@ struct FReadOnlyCVARCache
 	int32 MobileSkyLightPermutation;
 	bool bMobileEnableMovableSpotlights;
 	bool bMobileEnableMovableSpotlightsShadow;
+	bool bMobileEnableNoPrecomputedLightingCSMShader;
 	
 	bool bInitialized;
 	void Init();
diff --git a/Engine/Source/Runtime/Renderer/Private/LightMapRendering.h b/Engine/Source/Runtime/Renderer/Private/LightMapRendering.h
--- a/Engine/Source/Runtime/Renderer/Private/LightMapRendering.h
+++ b/Engine/Source/Runtime/Renderer/Private/LightMapRendering.h
@@ -492,8 +492,11 @@ struct FMobileDirectionalLightAndCSMPolicy
 		}
 		
 		static auto* CVarAllowStaticLighting = IConsoleManager::Get().FindTConsoleVariableDataInt(TEXT("r.AllowStaticLighting"));
+		static auto* CVarEnableNoPrecomputedLightingCSMShader = IConsoleManager::Get().FindTConsoleVariableDataInt(TEXT("r.Mobile.EnableNoPrecomputedLightingCSMShader"));
 		const bool bAllowStaticLighting = CVarAllowStaticLighting->GetValueOnAnyThread() != 0;
-		return !bAllowStaticLighting && 
+		const bool bEnableNoPrecomputedLightingCSMShader = CVarEnableNoPrecomputedLightingCSMShader && CVarEnableNoPrecomputedLightingCSMShader->GetValueOnAnyThread() != 0;
+
+		return (!bAllowStaticLighting || bEnableNoPrecomputedLightingCSMShader) &&
 			Parameters.MaterialParameters.ShadingModels.IsLit() && 
 			!IsTranslucentBlendMode(Parameters.MaterialParameters.BlendMode);
 	}
diff --git a/Engine/Source/Runtime/Renderer/Private/MobileBasePass.cpp b/Engine/Source/Runtime/Renderer/Private/MobileBasePass.cpp
--- a/Engine/Source/Runtime/Renderer/Private/MobileBasePass.cpp
+++ b/Engine/Source/Runtime/Renderer/Private/MobileBasePass.cpp
@@ -222,7 +222,7 @@ ELightMapPolicyType MobileBasePass::SelectMeshLightmapPolicy(
 	{
 		const FReadOnlyCVARCache& ReadOnlyCVARCache = FReadOnlyCVARCache::Get();
 
-		if (!ReadOnlyCVARCache.bAllowStaticLighting)
+		if (!ReadOnlyCVARCache.bAllowStaticLighting || (ReadOnlyCVARCache.bMobileEnableNoPrecomputedLightingCSMShader && Scene && Scene->GetForceNoPrecomputedLighting()))
 		{
 			if (!IsTranslucentBlendMode(BlendMode))
 			{
diff --git a/Engine/Source/Runtime/Renderer/Private/RendererScene.cpp b/Engine/Source/Runtime/Renderer/Private/RendererScene.cpp
--- a/Engine/Source/Runtime/Renderer/Private/RendererScene.cpp
+++ b/Engine/Source/Runtime/Renderer/Private/RendererScene.cpp
@@ -1190,6 +1190,7 @@ FScene::FScene(UWorld* InWorld, bool bInRequiresHitProxies, bool bInIsEditorScen
 ,	NumVisibleLights_GameThread(0)
 ,	NumEnabledSkylights_GameThread(0)
 ,	SceneFrameNumber(0)
+,	bForceNoPrecomputedLighting(InWorld->GetWorldSettings()->bForceNoPrecomputedLighting)
 {
 	FMemory::Memzero(MobileDirectionalLights);
 	FMemory::Memzero(AtmosphereLights);
diff --git a/Engine/Source/Runtime/Renderer/Private/ScenePrivate.h b/Engine/Source/Runtime/Renderer/Private/ScenePrivate.h
--- a/Engine/Source/Runtime/Renderer/Private/ScenePrivate.h
+++ b/Engine/Source/Runtime/Renderer/Private/ScenePrivate.h
@@ -3379,6 +3379,11 @@ public:
 	 */
 	inline int32 GetMaxPersistentPrimitiveIndex() const { return PersistentPrimitiveIdAllocator.GetMaxSize(); }
 
+	bool GetForceNoPrecomputedLighting() const
+	{
+		return bForceNoPrecomputedLighting;
+	}
+
 protected:
 
 private:
@@ -3534,6 +3539,9 @@ private:
 
 	/** Frame number incremented per-family viewing this scene. */
 	uint32 SceneFrameNumber;
+
+	/** Whether world settings has bForceNoPrecomputedLighting set */
+	bool bForceNoPrecomputedLighting;
 };
 
 inline bool ShouldIncludeDomainInMeshPass(EMaterialDomain Domain)
