Fixed light components fail to propagate certain property changes from Sequencer

diff --git a/Engine/Source/Runtime/Engine/Classes/Components/DirectionalLightComponent.h b/Engine/Source/Runtime/Engine/Classes/Components/DirectionalLightComponent.h
--- a/Engine/Source/Runtime/Engine/Classes/Components/DirectionalLightComponent.h
+++ b/Engine/Source/Runtime/Engine/Classes/Components/DirectionalLightComponent.h
@@ -24,28 +24,28 @@ class ENGINE_API UDirectionalLightComponent : public ULightComponent
 	* A value of 1 scales shadow bias based on each cascade size (Default).
 	* A value of 0 scales shadow bias uniformly accross all cacascade.
 	*/
-	UPROPERTY(EditAnywhere, BlueprintReadOnly, Category = Light, AdvancedDisplay, meta = (UIMin = "0", UIMax = "1"))
+	UPROPERTY(EditAnywhere, BlueprintReadOnly, interp, Category = Light, AdvancedDisplay, meta = (UIMin = "0", UIMax = "1"))
 	float ShadowCascadeBiasDistribution;
 
 	/** Whether to occlude fog and atmosphere inscattering with screenspace blurred occlusion from this light. */
-	UPROPERTY(EditAnywhere, BlueprintReadOnly, Category=LightShafts, meta=(DisplayName = "Light Shaft Occlusion"))
+	UPROPERTY(EditAnywhere, BlueprintReadOnly, interp, Category=LightShafts, meta=(DisplayName = "Light Shaft Occlusion"))
 	uint32 bEnableLightShaftOcclusion:1;
 
 	/** 
 	 * Controls how dark the occlusion masking is, a value of 1 results in no darkening term.
 	 */
-	UPROPERTY(EditAnywhere, BlueprintReadOnly, Category=LightShafts, meta=(UIMin = "0", UIMax = "1"))
+	UPROPERTY(EditAnywhere, BlueprintReadOnly, interp, Category=LightShafts, meta=(UIMin = "0", UIMax = "1"))
 	float OcclusionMaskDarkness;
 
 	/** Everything closer to the camera than this distance will occlude light shafts. */
-	UPROPERTY(EditAnywhere, BlueprintReadOnly, Category=LightShafts, meta=(UIMin = "0", UIMax = "500000"))
+	UPROPERTY(EditAnywhere, BlueprintReadOnly, interp, Category=LightShafts, meta=(UIMin = "0", UIMax = "500000"))
 	float OcclusionDepthRange;
 
 	/** 
 	 * Can be used to make light shafts come from somewhere other than the light's actual direction. 
 	 * This will only be used when non-zero.  It does not have to be normalized.
 	 */
-	UPROPERTY(EditAnywhere, AdvancedDisplay, BlueprintReadOnly, Category=LightShafts)
+	UPROPERTY(EditAnywhere, interp, AdvancedDisplay, BlueprintReadOnly, Category=LightShafts)
 	FVector LightShaftOverrideDirection;
 
 	UPROPERTY()
@@ -55,28 +55,28 @@ class ENGINE_API UDirectionalLightComponent : public ULightComponent
 	 * How far Cascaded Shadow Map dynamic shadows will cover for a movable light, measured from the camera.
 	 * A value of 0 disables the dynamic shadow.
 	 */
-	UPROPERTY(EditAnywhere, BlueprintReadOnly, Category=CascadedShadowMaps, meta=(UIMin = "0", UIMax = "20000", DisplayName = "Dynamic Shadow Distance MovableLight"))
+	UPROPERTY(EditAnywhere, BlueprintReadOnly, interp, Category=CascadedShadowMaps, meta=(UIMin = "0", UIMax = "20000", DisplayName = "Dynamic Shadow Distance MovableLight"))
 	float DynamicShadowDistanceMovableLight;
 
 	/** 
 	 * How far Cascaded Shadow Map dynamic shadows will cover for a stationary light, measured from the camera.
 	 * A value of 0 disables the dynamic shadow.
 	 */
-	UPROPERTY(EditAnywhere, BlueprintReadOnly, Category=CascadedShadowMaps, meta=(UIMin = "0", UIMax = "20000", DisplayName = "Dynamic Shadow Distance StationaryLight"))
+	UPROPERTY(EditAnywhere, BlueprintReadOnly, interp, Category=CascadedShadowMaps, meta=(UIMin = "0", UIMax = "20000", DisplayName = "Dynamic Shadow Distance StationaryLight"))
 	float DynamicShadowDistanceStationaryLight;
 
 	/** 
 	 * Number of cascades to split the view frustum into for the whole scene dynamic shadow.  
 	 * More cascades result in better shadow resolution, but adds significant rendering cost.
 	 */
-	UPROPERTY(EditAnywhere, BlueprintReadOnly, Category=CascadedShadowMaps, meta=(UIMin = "0", UIMax = "4", DisplayName = "Num Dynamic Shadow Cascades"))
+	UPROPERTY(EditAnywhere, BlueprintReadOnly, interp, Category=CascadedShadowMaps, meta=(UIMin = "0", UIMax = "4", DisplayName = "Num Dynamic Shadow Cascades"))
 	int32 DynamicShadowCascades;
 
 	/** 
 	 * Controls whether the cascades are distributed closer to the camera (larger exponent) or further from the camera (smaller exponent).
 	 * An exponent of 1 means that cascade transitions will happen at a distance proportional to their resolution.
 	 */
-	UPROPERTY(EditAnywhere, BlueprintReadOnly, Category=CascadedShadowMaps, meta=(UIMin = "1", UIMax = "4", DisplayName = "Distribution Exponent"))
+	UPROPERTY(EditAnywhere, BlueprintReadOnly, interp, Category=CascadedShadowMaps, meta=(UIMin = "1", UIMax = "4", DisplayName = "Distribution Exponent"))
 	float CascadeDistributionExponent;
 
 	/** 
@@ -88,14 +88,14 @@ class ENGINE_API UDirectionalLightComponent : public ULightComponent
 	 * Ideal values are the smallest possible which still hide the transition.
 	 * An increased fade region size causes an increase in shadow rendering cost.
 	 */
-	UPROPERTY(EditAnywhere, BlueprintReadOnly, Category=CascadedShadowMaps, meta=(UIMin = "0", UIMax = "0.3", DisplayName = "Transition Fraction"))
+	UPROPERTY(EditAnywhere, BlueprintReadOnly, interp, Category=CascadedShadowMaps, meta=(UIMin = "0", UIMax = "0.3", DisplayName = "Transition Fraction"))
 	float CascadeTransitionFraction;
 
 	/** 
 	 * Controls the size of the fade out region at the far extent of the dynamic shadow's influence.  
 	 * This is specified as a fraction of DynamicShadowDistance. 
 	 */
-	UPROPERTY(EditAnywhere, BlueprintReadOnly, Category=CascadedShadowMaps, meta=(UIMin = "0", UIMax = "1.0", DisplayName = "Distance Fadeout Fraction"))
+	UPROPERTY(EditAnywhere, BlueprintReadOnly, interp, Category=CascadedShadowMaps, meta=(UIMin = "0", UIMax = "1.0", DisplayName = "Distance Fadeout Fraction"))
 	float ShadowDistanceFadeoutFraction;
 
 	/** 
@@ -279,15 +279,30 @@ class ENGINE_API UDirectionalLightComponent : public ULightComponent
 	UFUNCTION(BlueprintCallable, Category="Rendering|Lighting")
 	void SetShadowDistanceFadeoutFraction(float NewValue);
 
+	UFUNCTION(BlueprintCallable, Category = "Rendering|Lighting")
+	void SetShadowCascadeBiasDistribution(float NewValue);
+
 	UFUNCTION(BlueprintCallable, Category="Rendering|Lighting")
 	void SetEnableLightShaftOcclusion(bool bNewValue);
 
 	UFUNCTION(BlueprintCallable, Category="Rendering|Lighting")
 	void SetOcclusionMaskDarkness(float NewValue);
 
+	UFUNCTION(BlueprintCallable, Category = "Rendering|Lighting")
+	void SetOcclusionDepthRange(float NewValue);
+
 	UFUNCTION(BlueprintCallable, Category="Rendering|Lighting")
 	void SetLightShaftOverrideDirection(FVector NewValue);
 
+	UFUNCTION(BlueprintCallable, Category = "Rendering|Lighting")
+	void SetLightSourceAngle(float NewValue);
+
+	UFUNCTION(BlueprintCallable, Category = "Rendering|Lighting")
+	void SetLightSourceSoftAngle(float NewValue);
+
+	UFUNCTION(BlueprintCallable, Category = "Rendering|Lighting")
+	void SetShadowSourceAngleFactor(float NewValue);
+
 	UFUNCTION(BlueprintCallable, Category="Rendering|Lighting")
 	void SetShadowAmount(float NewValue);
 
diff --git a/Engine/Source/Runtime/Engine/Classes/Components/PointLightComponent.h b/Engine/Source/Runtime/Engine/Classes/Components/PointLightComponent.h
--- a/Engine/Source/Runtime/Engine/Classes/Components/PointLightComponent.h
+++ b/Engine/Source/Runtime/Engine/Classes/Components/PointLightComponent.h
@@ -40,21 +40,21 @@ class ENGINE_API UPointLightComponent : public ULocalLightComponent
 	 * Radius of light source shape.
 	 * Note that light sources shapes which intersect shadow casting geometry can cause shadowing artifacts.
 	 */
-	UPROPERTY(EditAnywhere, BlueprintReadOnly, Category=Light, meta=(ShouldShowInViewport = true))
+	UPROPERTY(EditAnywhere, BlueprintReadOnly, interp, Category=Light, meta=(ShouldShowInViewport = true))
 	float SourceRadius;
 
 	/**
 	* Soft radius of light source shape.
 	* Note that light sources shapes which intersect shadow casting geometry can cause shadowing artifacts.
 	*/
-	UPROPERTY(EditAnywhere, BlueprintReadOnly, Category = Light)
+	UPROPERTY(EditAnywhere, BlueprintReadOnly, interp, Category = Light)
 	float SoftSourceRadius;
 
 	/** 
 	 * Length of light source shape.
 	 * Note that light sources shapes which intersect shadow casting geometry can cause shadowing artifacts.
 	 */
-	UPROPERTY(EditAnywhere, BlueprintReadOnly, Category=Light)
+	UPROPERTY(EditAnywhere, BlueprintReadOnly, interp, Category=Light)
 	float SourceLength;
 
 	UFUNCTION(BlueprintCallable, Category="Rendering|Lighting")
diff --git a/Engine/Source/Runtime/Engine/Classes/Components/RectLightComponent.h b/Engine/Source/Runtime/Engine/Classes/Components/RectLightComponent.h
--- a/Engine/Source/Runtime/Engine/Classes/Components/RectLightComponent.h
+++ b/Engine/Source/Runtime/Engine/Classes/Components/RectLightComponent.h
@@ -25,26 +25,26 @@ class ENGINE_API URectLightComponent : public ULocalLightComponent
 	 * Width of light source rect.
 	 * Note that light sources shapes which intersect shadow casting geometry can cause shadowing artifacts.
 	 */
-	UPROPERTY(EditAnywhere, BlueprintReadOnly, Category=Light)
+	UPROPERTY(EditAnywhere, BlueprintReadOnly, interp, Category=Light)
 	float SourceWidth;
 
 	/** 
 	 * Height of light source rect.
 	 * Note that light sources shapes which intersect shadow casting geometry can cause shadowing artifacts.
 	 */
-	UPROPERTY(EditAnywhere, BlueprintReadOnly, Category=Light)
+	UPROPERTY(EditAnywhere, BlueprintReadOnly, interp, Category=Light)
 	float SourceHeight;
 
 	/**
 	 * Angle of barn door attached to the light source rect.
 	 */
-	UPROPERTY(EditAnywhere, BlueprintReadOnly, Category = Light, meta = (UIMin = "0.0", UIMax = "90.0"))
+	UPROPERTY(EditAnywhere, BlueprintReadOnly, interp, Category = Light, meta = (UIMin = "0.0", UIMax = "90.0"))
 	float BarnDoorAngle;
 	
 	/**
 	 * Length of barn door attached to the light source rect.
 	 */
-	UPROPERTY(EditAnywhere, BlueprintReadOnly, Category = Light, meta = (UIMin = "0.0"))
+	UPROPERTY(EditAnywhere, BlueprintReadOnly, interp, Category = Light, meta = (UIMin = "0.0"))
 	float BarnDoorLength;
 
 	/** Texture mapped to the light source rectangle */
diff --git a/Engine/Source/Runtime/Engine/Private/Components/DirectionalLightComponent.cpp b/Engine/Source/Runtime/Engine/Private/Components/DirectionalLightComponent.cpp
--- a/Engine/Source/Runtime/Engine/Private/Components/DirectionalLightComponent.cpp
+++ b/Engine/Source/Runtime/Engine/Private/Components/DirectionalLightComponent.cpp
@@ -1231,6 +1231,15 @@ void UDirectionalLightComponent::SetShadowDistanceFadeoutFraction(float NewValue
 	}
 }
 
+void UDirectionalLightComponent::SetShadowCascadeBiasDistribution(float NewValue)
+{
+	if (ShadowCascadeBiasDistribution != NewValue)
+	{
+		ShadowCascadeBiasDistribution = NewValue;
+		MarkRenderStateDirty();
+	}
+}
+
 void UDirectionalLightComponent::SetEnableLightShaftOcclusion(bool bNewValue)
 {
 	if (AreDynamicDataChangesAllowed()
@@ -1251,6 +1260,15 @@ void UDirectionalLightComponent::SetOcclusionMaskDarkness(float NewValue)
 	}
 }
 
+void UDirectionalLightComponent::SetOcclusionDepthRange(float NewValue)
+{
+	if (AreDynamicDataChangesAllowed()
+		&& OcclusionDepthRange != NewValue)
+	{
+		OcclusionDepthRange = NewValue;
+		MarkRenderStateDirty();
+	}
+}
 void UDirectionalLightComponent::SetLightShaftOverrideDirection(FVector NewValue)
 {
 	if (AreDynamicDataChangesAllowed()
@@ -1265,6 +1283,36 @@ void UDirectionalLightComponent::SetLightShaftOverrideDirection(FVector NewValue
 	}
 }
 
+void UDirectionalLightComponent::SetLightSourceAngle(float NewValue)
+{
+	if (AreDynamicDataChangesAllowed()
+		&& LightSourceAngle != NewValue)
+	{
+		LightSourceAngle = NewValue;
+		MarkRenderStateDirty();
+	}
+}
+
+void UDirectionalLightComponent::SetLightSourceSoftAngle(float NewValue)
+{
+	if (AreDynamicDataChangesAllowed()
+		&& LightSourceSoftAngle != NewValue)
+	{
+		LightSourceSoftAngle = NewValue;
+		MarkRenderStateDirty();
+	}
+}
+
+void UDirectionalLightComponent::SetShadowSourceAngleFactor(float NewValue)
+{
+	if (AreDynamicDataChangesAllowed()
+		&& ShadowSourceAngleFactor != NewValue)
+	{
+		ShadowSourceAngleFactor = NewValue;
+		MarkRenderStateDirty();
+	}
+}
+
 void UDirectionalLightComponent::SetShadowAmount(float NewValue)
 {
 	if (AreDynamicDataChangesAllowed()
