Fix or Enable Low Pass Filter bool on audio component being ignored

diff --git a/Engine/Source/Runtime/Engine/Private/Audio.cpp b/Engine/Source/Runtime/Engine/Private/Audio.cpp
--- a/Engine/Source/Runtime/Engine/Private/Audio.cpp
+++ b/Engine/Source/Runtime/Engine/Private/Audio.cpp
@@ -390,7 +390,13 @@ void FSoundSource::SetFilterFrequency()
 			}
 
 			// Set the LPFFrequency to lowest provided value
-			LPFFrequency = FMath::Min(WaveInstance->OcclusionFilterFrequency * OcclusionFilterScale, WaveInstance->LowPassFilterFrequency);
+			LPFFrequency = WaveInstance->OcclusionFilterFrequency * OcclusionFilterScale;
+
+			if (WaveInstance->bEnableLowPassFilter)
+			{
+				LPFFrequency = FMath::Min(LPFFrequency, WaveInstance->LowPassFilterFrequency);
+			}
+
 			LPFFrequency = FMath::Min(LPFFrequency, WaveInstance->AmbientZoneFilterFrequency);
 			LPFFrequency = FMath::Min(LPFFrequency, WaveInstance->AttenuationLowpassFilterFrequency);
 			LPFFrequency = FMath::Min(LPFFrequency, WaveInstance->SoundClassFilterFrequency);
