* Added a scope object (IModularFeatures::FScopedLockModularFeatureList) to more easily lock/unlock the list of modular features when accessing it.
* Fixed render capture provider firing an assert when accessing IModularFeatures from outside the game thread (naive implementation for now, since this is a accessed very sporadically)
* Avoid locking the modular features list when calling GetEngineCrypto(), since the IEngineCrypto* is only accessed once and stored in a static member

diff --git a/Engine/Source/Runtime/Core/Private/IO/IoStore.cpp b/Engine/Source/Runtime/Core/Private/IO/IoStore.cpp
--- a/Engine/Source/Runtime/Core/Private/IO/IoStore.cpp
+++ b/Engine/Source/Runtime/Core/Private/IO/IoStore.cpp
@@ -72,9 +72,18 @@ bool WriteArray(IFileHandle* FileHandle, const ArrayType& Array)
 
 static IEngineCrypto* GetEngineCrypto()
 {
-	IModularFeatures::Get().LockModularFeatureList();
-	static TArray<IEngineCrypto*> Features = IModularFeatures::Get().GetModularFeatureImplementations<IEngineCrypto>(IEngineCrypto::GetFeatureName());
-	IModularFeatures::Get().UnlockModularFeatureList();
+	static bool bFeaturesInitialized = false;
+	static TArray<IEngineCrypto*> Features;
+	if (!bFeaturesInitialized)
+	{
+		IModularFeatures::FScopedLockModularFeatureList ScopedLockModularFeatureList;
+		// bFeaturesInitialized is not atomic so it can have potentially become true since we last checked (except now we're under a lock: IModularFeatures::FScopedLockModularFeatureList, so it will be fine now):
+		if (!bFeaturesInitialized)
+		{
+		    Features = IModularFeatures::Get().GetModularFeatureImplementations<IEngineCrypto>(IEngineCrypto::GetFeatureName());
+		    bFeaturesInitialized = true;
+		}
+	}
 	checkf(Features.Num() > 0, TEXT("RSA functionality was used but no modular feature was registered to provide it. Please make sure your project has the PlatformCrypto plugin enabled!"));
 	return Features[0];
 }
diff --git a/Engine/Source/Runtime/Core/Public/Features/IModularFeatures.h b/Engine/Source/Runtime/Core/Public/Features/IModularFeatures.h
--- a/Engine/Source/Runtime/Core/Public/Features/IModularFeatures.h
+++ b/Engine/Source/Runtime/Core/Public/Features/IModularFeatures.h
@@ -85,6 +85,18 @@ public:
 	/** Unlock modular features when finished accessing from another thread */
 	virtual void UnlockModularFeatureList() = 0;
 
+	struct FScopedLockModularFeatureList
+	{
+		inline FScopedLockModularFeatureList() { IModularFeatures::Get().LockModularFeatureList(); }
+		inline ~FScopedLockModularFeatureList() { IModularFeatures::Get().UnlockModularFeatureList(); }
+
+		// Non-copyable :
+		FScopedLockModularFeatureList(const FScopedLockModularFeatureList&) = delete;
+		FScopedLockModularFeatureList(FScopedLockModularFeatureList&&) = delete;
+		FScopedLockModularFeatureList& operator = (const FScopedLockModularFeatureList&) = delete;
+		FScopedLockModularFeatureList& operator = (FScopedLockModularFeatureList&&) = delete;
+	};
+
 	/**
 	 * Returns the number of registered implementations of the specified feature type.
 	 *
diff --git a/Engine/Source/Runtime/Engine/Private/Animation/AnimRootMotionProvider.cpp b/Engine/Source/Runtime/Engine/Private/Animation/AnimRootMotionProvider.cpp
--- a/Engine/Source/Runtime/Engine/Private/Animation/AnimRootMotionProvider.cpp
+++ b/Engine/Source/Runtime/Engine/Private/Animation/AnimRootMotionProvider.cpp
@@ -14,20 +14,16 @@ const IAnimRootMotionProvider* IAnimRootMotionProvider::Get()
 {
 	if (IsAvailable())
 	{
-		IModularFeatures::Get().LockModularFeatureList();
-		const IAnimRootMotionProvider* AnimRootMotionProvider = &IModularFeatures::Get().GetModularFeature<const IAnimRootMotionProvider>(ModularFeatureName);
-		IModularFeatures::Get().UnlockModularFeatureList();
-		return AnimRootMotionProvider;
+		IModularFeatures::FScopedLockModularFeatureList ScopedLockModularFeatureList;
+		return &IModularFeatures::Get().GetModularFeature<const IAnimRootMotionProvider>(ModularFeatureName);
 	}
 	return nullptr;
 }
 
 bool IAnimRootMotionProvider::IsAvailable()
 {
-	IModularFeatures::Get().LockModularFeatureList();
-	const bool bIsAvailable = IModularFeatures::Get().IsModularFeatureAvailable(ModularFeatureName);
-	IModularFeatures::Get().UnlockModularFeatureList();
-	return bIsAvailable;
+	IModularFeatures::FScopedLockModularFeatureList ScopedLockModularFeatureList;
+	return IModularFeatures::Get().IsModularFeatureAvailable(ModularFeatureName);
 }
 
 }}; // namespace UE::Anim
\ No newline at end of file
diff --git a/Engine/Source/Runtime/RenderCore/Public/IRenderCaptureProvider.h b/Engine/Source/Runtime/RenderCore/Public/IRenderCaptureProvider.h
--- a/Engine/Source/Runtime/RenderCore/Public/IRenderCaptureProvider.h
+++ b/Engine/Source/Runtime/RenderCore/Public/IRenderCaptureProvider.h
@@ -28,6 +28,8 @@ public:
 	 */
 	static inline bool IsAvailable()
 	{
+		// This function is accessible from the render and RHI threads, so lock the list of modular features before calling GetModularFeature
+		IModularFeatures::FScopedLockModularFeatureList ScopedLockModularFeatureList;
 		return IModularFeatures::Get().IsModularFeatureAvailable(GetModularFeatureName());
 	}
 
@@ -39,6 +41,8 @@ public:
 	 */
 	static inline IRenderCaptureProvider& Get()
 	{
+		// This function is accessible from the render and RHI threads, so lock the list of modular features before calling GetModularFeature
+		IModularFeatures::FScopedLockModularFeatureList ScopedLockModularFeatureList;
 		return IModularFeatures::Get().GetModularFeature<IRenderCaptureProvider>(GetModularFeatureName());
 	}
 
