Removing project setting to enable audio stream caching (it is now always enabled.)

diff --git a/Engine/Source/Developer/Windows/WindowsTargetPlatform/Classes/WindowsTargetSettings.h b/Engine/Source/Developer/Windows/WindowsTargetPlatform/Classes/WindowsTargetSettings.h
--- a/Engine/Source/Developer/Windows/WindowsTargetPlatform/Classes/WindowsTargetSettings.h
+++ b/Engine/Source/Developer/Windows/WindowsTargetPlatform/Classes/WindowsTargetSettings.h
@@ -103,10 +103,6 @@ public:
 	UPROPERTY(GlobalConfig, EditAnywhere, Category = "Audio")
 	FPlatformRuntimeAudioCompressionOverrides CompressionOverrides;
 
-	/** When this is enabled, Actual compressed data will be separated from the USoundWave, and loaded into a cache. */
-	UPROPERTY(GlobalConfig, EditAnywhere, Category = "Audio|CookOverrides", meta = (DisplayName = "Use Stream Caching (Experimental)"))
-	bool bUseAudioStreamCaching;
-
 	/** This determines the max amount of memory that should be used for the cache at any given time. If set low (<= 8 MB), it lowers the size of individual chunks of audio during cook. */
 	UPROPERTY(GlobalConfig, EditAnywhere, Category = "Audio|CookOverrides|Stream Caching", meta = (DisplayName = "Max Cache Size (KB)"))
 	int32 CacheSizeKB;
diff --git a/Engine/Source/Runtime/Android/AndroidRuntimeSettings/Classes/AndroidRuntimeSettings.h b/Engine/Source/Runtime/Android/AndroidRuntimeSettings/Classes/AndroidRuntimeSettings.h
--- a/Engine/Source/Runtime/Android/AndroidRuntimeSettings/Classes/AndroidRuntimeSettings.h
+++ b/Engine/Source/Runtime/Android/AndroidRuntimeSettings/Classes/AndroidRuntimeSettings.h
@@ -518,10 +518,6 @@ public:
 	UPROPERTY(config, EditAnywhere, Category = "Audio")
 	FPlatformRuntimeAudioCompressionOverrides CompressionOverrides;
 
-	/** When this is enabled, Actual compressed data will be separated from the USoundWave, and loaded into a cache. */
-	UPROPERTY(GlobalConfig, EditAnywhere, Category = "Audio|CookOverrides", meta = (DisplayName = "Use Stream Caching (Experimental)"))
-	bool bUseAudioStreamCaching;
-
 	/** This determines the max amount of memory that should be used for the cache at any given time. If set low (<= 8 MB), it lowers the size of individual chunks of audio during cook. */
 	UPROPERTY(GlobalConfig, EditAnywhere, Category = "Audio|CookOverrides|Stream Caching", meta = (DisplayName = "Max Cache Size (KB)"))
 	int32 CacheSizeKB;
diff --git a/Engine/Source/Runtime/AudioPlatformConfiguration/Public/AudioCompressionSettings.h b/Engine/Source/Runtime/AudioPlatformConfiguration/Public/AudioCompressionSettings.h
--- a/Engine/Source/Runtime/AudioPlatformConfiguration/Public/AudioCompressionSettings.h
+++ b/Engine/Source/Runtime/AudioPlatformConfiguration/Public/AudioCompressionSettings.h
@@ -74,9 +74,6 @@ struct FPlatformAudioCookOverrides
 	// If StreamCaching is set to true, this will be used 
 	float AutoStreamingThreshold;
 
-	// Whether to use the experimental Load on Demand feature, which uses as little memory at runtime as possible.
-	bool bUseStreamCaching;
-
 	// Whether to put streamed audio chunks inline in the Pak file or not (only matters if bUseStreamCaching is true)
 	bool bInlineStreamedAudioChunks;
 
@@ -87,7 +84,6 @@ struct FPlatformAudioCookOverrides
 		: bResampleForDevice(false)
 		, CompressionQualityModifier(1.0f)
 		, AutoStreamingThreshold(0.0f)
-		, bUseStreamCaching(false)
 		, bInlineStreamedAudioChunks(false)
 	{
 		PlatformSampleRates.Add(ESoundwaveSampleRateSettings::Max, 48000);
@@ -111,23 +107,20 @@ struct FPlatformAudioCookOverrides
 		int32 AutoStreamingThresholdHash = FMath::FloorToInt(InOverrides->AutoStreamingThreshold * 100.0f);
 		OutSuffix.AppendInt(AutoStreamingThresholdHash);
 
-		if (InOverrides->bUseStreamCaching)
+		OutSuffix.Append(TEXT("_StreamCache_Ver"));
+		OutSuffix.AppendInt(GetStreamCachingVersion());
+		OutSuffix.AppendChar('_');
+
+		// cache info:
+		OutSuffix.Append(TEXT("MEM_"));
+		OutSuffix.AppendInt(InOverrides->StreamCachingSettings.CacheSizeKB);
+		OutSuffix.Append(TEXT("MaxChnkSize_"));
+		OutSuffix.AppendInt(InOverrides->StreamCachingSettings.MaxChunkSizeOverrideKB);
+
+		if (InOverrides->StreamCachingSettings.bForceLegacyStreamChunking)
 		{
-			OutSuffix.Append(TEXT("_StreamCache_Ver"));
-			OutSuffix.AppendInt(GetStreamCachingVersion());
-			OutSuffix.AppendChar('_');
-
-			// cache info:
-			OutSuffix.Append(TEXT("MEM_"));
-			OutSuffix.AppendInt(InOverrides->StreamCachingSettings.CacheSizeKB);
-			OutSuffix.Append(TEXT("MaxChnkSize_"));
-			OutSuffix.AppendInt(InOverrides->StreamCachingSettings.MaxChunkSizeOverrideKB);
-
-			if (InOverrides->StreamCachingSettings.bForceLegacyStreamChunking)
-			{
-				OutSuffix.Append(TEXT("_LegacyChunking_"));
-				OutSuffix.AppendInt(InOverrides->StreamCachingSettings.ZerothChunkSizeForLegacyStreamChunkingKB);
-			}
+			OutSuffix.Append(TEXT("_LegacyChunking_"));
+			OutSuffix.AppendInt(InOverrides->StreamCachingSettings.ZerothChunkSizeForLegacyStreamChunkingKB);
 		}
 		
 
diff --git a/Engine/Source/Runtime/Engine/Private/AudioCompressionSettingsUtils.cpp b/Engine/Source/Runtime/Engine/Private/AudioCompressionSettingsUtils.cpp
--- a/Engine/Source/Runtime/Engine/Private/AudioCompressionSettingsUtils.cpp
+++ b/Engine/Source/Runtime/Engine/Private/AudioCompressionSettingsUtils.cpp
@@ -109,39 +109,33 @@ void CacheAudioCookOverrides(FPlatformAudioCookOverrides& OutOverrides, const TC
 		OutOverrides.SoundCueCookQualityIndex = SoundCueQualityIndex;
 	}
 
-//	PlatformFile.GetBool(*CategoryName, TEXT("bUseAudioStreamCaching"), OutOverrides.bUseStreamCaching);
-	OutOverrides.bUseStreamCaching = true;
-
 	GConfig->GetBool(*CategoryName, TEXT("bInlineStreamedAudioChunks"), OutOverrides.bInlineStreamedAudioChunks, GEngineIni);
 
 	/** Memory Load On Demand Settings */
-	if (OutOverrides.bUseStreamCaching)
-	{
-		// Cache size:
-		const int32 DefaultCacheSize = 64 * 1024;
-		int32 RetrievedCacheSize = DefaultCacheSize;
-		int32 RetrievedChunkSizeOverride = INDEX_NONE;
+	// Cache size:
+	const int32 DefaultCacheSize = 64 * 1024;
+	int32 RetrievedCacheSize = DefaultCacheSize;
+	int32 RetrievedChunkSizeOverride = INDEX_NONE;
 
-		PlatformFile.GetInt(*CategoryName, TEXT("CacheSizeKB"), RetrievedCacheSize);
-		if (!RetrievedCacheSize)
-		{
-			PlatformFile.SetInt64(*CategoryName, TEXT("CacheSizeKB"), DefaultCacheSize);
-			RetrievedCacheSize = DefaultCacheSize;
-		}
+	PlatformFile.GetInt(*CategoryName, TEXT("CacheSizeKB"), RetrievedCacheSize);
+	if (!RetrievedCacheSize)
+	{
+		PlatformFile.SetInt64(*CategoryName, TEXT("CacheSizeKB"), DefaultCacheSize);
+		RetrievedCacheSize = DefaultCacheSize;
+	}
 
-		OutOverrides.StreamCachingSettings.CacheSizeKB = RetrievedCacheSize;
+	OutOverrides.StreamCachingSettings.CacheSizeKB = RetrievedCacheSize;
 
-		PlatformFile.GetInt(*CategoryName, TEXT("MaxChunkSizeOverrideKB"), RetrievedChunkSizeOverride);
-		OutOverrides.StreamCachingSettings.MaxChunkSizeOverrideKB = RetrievedChunkSizeOverride;
+	PlatformFile.GetInt(*CategoryName, TEXT("MaxChunkSizeOverrideKB"), RetrievedChunkSizeOverride);
+	OutOverrides.StreamCachingSettings.MaxChunkSizeOverrideKB = RetrievedChunkSizeOverride;
 
-		bool bForceLegacyStreamChunking = false;
-		PlatformFile.GetBool(*CategoryName, TEXT("bForceLegacyStreamChunking"), bForceLegacyStreamChunking);
-		OutOverrides.StreamCachingSettings.bForceLegacyStreamChunking = bForceLegacyStreamChunking;
+	bool bForceLegacyStreamChunking = false;
+	PlatformFile.GetBool(*CategoryName, TEXT("bForceLegacyStreamChunking"), bForceLegacyStreamChunking);
+	OutOverrides.StreamCachingSettings.bForceLegacyStreamChunking = bForceLegacyStreamChunking;
 
-		int32 ZerothChunkSizeForLegacyStreamChunking = 0;
-		PlatformFile.GetInt(*CategoryName, TEXT("ZerothChunkSizeForLegacyStreamChunking"), ZerothChunkSizeForLegacyStreamChunking);
-		OutOverrides.StreamCachingSettings.ZerothChunkSizeForLegacyStreamChunkingKB = ZerothChunkSizeForLegacyStreamChunking;
-	}
+	int32 ZerothChunkSizeForLegacyStreamChunking = 0;
+	PlatformFile.GetInt(*CategoryName, TEXT("ZerothChunkSizeForLegacyStreamChunking"), ZerothChunkSizeForLegacyStreamChunking);
+	OutOverrides.StreamCachingSettings.ZerothChunkSizeForLegacyStreamChunkingKB = ZerothChunkSizeForLegacyStreamChunking;
 
 	PlatformFile.GetBool(*CategoryName, TEXT("bResampleForDevice"), OutOverrides.bResampleForDevice);
 
@@ -359,8 +353,6 @@ const FPlatformAudioCookOverrides* FPlatformCompressionUtilities::GetCookOverrid
 
 bool FPlatformCompressionUtilities::IsCurrentPlatformUsingStreamCaching()
 {
-// 	const FPlatformAudioCookOverrides* Settings = GetCookOverrides();
-// 	return Settings && Settings->bUseStreamCaching;
 	return true;
 }
 
diff --git a/Engine/Source/Runtime/Engine/Private/AudioDerivedData.cpp b/Engine/Source/Runtime/Engine/Private/AudioDerivedData.cpp
--- a/Engine/Source/Runtime/Engine/Private/AudioDerivedData.cpp
+++ b/Engine/Source/Runtime/Engine/Private/AudioDerivedData.cpp
@@ -331,7 +331,6 @@ class FStreamedAudioCacheDerivedDataWorker : public FNonAbandonableTask
 				int32 FirstChunkSize = MaxChunkSizeForCurrentWave;
 
 				const int32 MinimumChunkSize = AudioFormat->GetMinimumSizeForInitialChunk(AudioFormatName, CompressedBuffer);
-				const bool bUseStreamCaching = CompressionOverrides && CompressionOverrides->bUseStreamCaching;
 				const bool bForceLegacyStreamChunking = SoundWave.bStreaming && CompressionOverrides && CompressionOverrides->StreamCachingSettings.bForceLegacyStreamChunking;
 
 				// If the initial chunk  for this sound wave was overridden, use that:
@@ -339,7 +338,7 @@ class FStreamedAudioCacheDerivedDataWorker : public FNonAbandonableTask
 				{
 					FirstChunkSize = FMath::Max(MinimumChunkSize, SoundWave.InitialChunkSize);
 				}
-				else if (bUseStreamCaching)
+				else
 				{
 					// Ensure that the minimum chunk size is nonzero if our compressed buffer is not empty.
 					checkf(CompressedBuffer.Num() == 0 || MinimumChunkSize != 0, TEXT("To use Load On Demand, please override GetMinimumSizeForInitialChunk"));
@@ -362,7 +361,7 @@ class FStreamedAudioCacheDerivedDataWorker : public FNonAbandonableTask
 					}
 				}
 
-				if (bUseStreamCaching && !bForceLegacyStreamChunking)
+				if (!bForceLegacyStreamChunking)
 				{
 					// Use the chunk size for this duration:
 					MaxChunkSizeForCurrentWave = FPlatformCompressionUtilities::GetMaxChunkSizeForCookOverrides(CompressionOverrides);
@@ -420,7 +419,7 @@ class FStreamedAudioCacheDerivedDataWorker : public FNonAbandonableTask
 
 						int32 ZeroPadBytes = 0;
 
-						if (!bUseStreamCaching || bForceLegacyStreamChunking)
+						if (bForceLegacyStreamChunking)
 						{
 							// padding when stream caching is enabled will significantly bloat the amount of space soundwaves take up on disk.
 							ZeroPadBytes = FMath::Max(MaxChunkSizeForCurrentWave - AudioDataSize, 0);
diff --git a/Engine/Source/Runtime/Engine/Private/SoundWave.cpp b/Engine/Source/Runtime/Engine/Private/SoundWave.cpp
--- a/Engine/Source/Runtime/Engine/Private/SoundWave.cpp
+++ b/Engine/Source/Runtime/Engine/Private/SoundWave.cpp
@@ -2904,13 +2904,8 @@ bool USoundWave::IsStreaming(const FPlatformAudioCookOverrides& Overrides) const
 	}
 
 	// For stream caching, the auto streaming threshold is used to force sounds to be inlined:
-	const bool bUsesStreamCache = Overrides.bUseStreamCaching;
-	const bool bOverAutoStreamingThreshold = (Overrides.AutoStreamingThreshold > SMALL_NUMBER && Duration > Overrides.AutoStreamingThreshold);
-
-	const bool Result = bUsesStreamCache || bOverAutoStreamingThreshold;
-
-	SoundWaveDataPtr->bIsStreaming = Result;
-	return Result;
+	SoundWaveDataPtr->bIsStreaming = (Overrides.AutoStreamingThreshold > SMALL_NUMBER && Duration > Overrides.AutoStreamingThreshold);
+	return SoundWaveDataPtr->bIsStreaming;
 }
 
 bool USoundWave::ShouldUseStreamCaching() const
diff --git a/Engine/Source/Runtime/IOS/IOSRuntimeSettings/Classes/IOSRuntimeSettings.h b/Engine/Source/Runtime/IOS/IOSRuntimeSettings/Classes/IOSRuntimeSettings.h
--- a/Engine/Source/Runtime/IOS/IOSRuntimeSettings/Classes/IOSRuntimeSettings.h
+++ b/Engine/Source/Runtime/IOS/IOSRuntimeSettings/Classes/IOSRuntimeSettings.h
@@ -514,10 +514,6 @@ public:
 	UPROPERTY(config, EditAnywhere, Category = "Audio")
 	FPlatformRuntimeAudioCompressionOverrides CompressionOverrides;
 
-	/** When this is enabled, Actual compressed data will be separated from the USoundWave, and loaded into a cache. */
-	UPROPERTY(GlobalConfig, EditAnywhere, Category = "Audio|CookOverrides", meta = (DisplayName = "Use Stream Caching (Experimental)"))
-	bool bUseAudioStreamCaching;
-
 	/** This determines the max amount of memory that should be used for the cache at any given time. If set low (<= 8 MB), it lowers the size of individual chunks of audio during cook. */
 	UPROPERTY(GlobalConfig, EditAnywhere, Category = "Audio|CookOverrides|Stream Caching", meta = (DisplayName = "Max Cache Size (KB)"))
 	int32 CacheSizeKB;
