Texture uprops EditCondition to gray out when not applicable

diff --git a/Engine/Source/Runtime/Engine/Classes/Engine/Texture.h b/Engine/Source/Runtime/Engine/Classes/Engine/Texture.h
--- a/Engine/Source/Runtime/Engine/Classes/Engine/Texture.h
+++ b/Engine/Source/Runtime/Engine/Classes/Engine/Texture.h
@@ -1073,7 +1073,7 @@ public:
 	bool bDoScaleMipsForAlphaCoverage = false;
 	
 	/** Alpha values per channel to compare to when preserving alpha coverage. 0 means disable channel.  Typical good values in 0.5 - 0.9, not 1.0 */
-	UPROPERTY(EditAnywhere, BlueprintReadWrite, Category = Texture, meta=(ClampMin = "0", ClampMax = "1.0"), AdvancedDisplay)
+	UPROPERTY(EditAnywhere, BlueprintReadWrite, Category = Texture, meta=(ClampMin = "0", ClampMax = "1.0", EditCondition="bDoScaleMipsForAlphaCoverage"), AdvancedDisplay)
 	FVector4 AlphaCoverageThresholds = FVector4(0,0,0,0);
 
 	/** When true the texture's border will be preserved during mipmap generation. */
@@ -1189,14 +1189,15 @@ public:
 	uint8 SRGB:1;
 
 #if WITH_EDITORONLY_DATA
-	/** Texture color management settings: source encoding and color space. */
-	UPROPERTY(EditAnywhere, BlueprintReadWrite, Category = Texture, AdvancedDisplay)
-	FTextureSourceColorSettings SourceColorSettings;
 
 	/** A flag for using the simplified legacy gamma space e.g pow(color,1/2.2) for converting from FColor to FLinearColor, if we're doing sRGB. */
-	UPROPERTY(EditAnywhere, BlueprintReadWrite, Category=Texture, AdvancedDisplay)
+	UPROPERTY(EditAnywhere, BlueprintReadWrite, Category=Texture, meta=(DisplayName="sRGB Use Legacy Gamma", EditCondition="SRGB"), AdvancedDisplay)
 	uint8 bUseLegacyGamma:1;
 
+	/** Texture color management settings: source encoding and color space. */
+	UPROPERTY(EditAnywhere, BlueprintReadWrite, Category = Texture, AdvancedDisplay)
+	FTextureSourceColorSettings SourceColorSettings;
+
 	/** Indicates we're currently importing the object (set in PostEditImport, unset in the subsequent PostEditChange) */
 	uint8 bIsImporting : 1;
 	
