Don't allocate shared serialization info until needed
Add csv stats for the shared hit/miss counts

diff --git a/Engine/Source/Runtime/Engine/Private/NetDriver.cpp b/Engine/Source/Runtime/Engine/Private/NetDriver.cpp
--- a/Engine/Source/Runtime/Engine/Private/NetDriver.cpp
+++ b/Engine/Source/Runtime/Engine/Private/NetDriver.cpp
@@ -4991,6 +4991,9 @@ struct FScopedNetDriverStats
 			SET_DWORD_STAT(STAT_SharedSerializationRPCHit, GNumSharedSerializationHit);
 			SET_DWORD_STAT(STAT_SharedSerializationRPCMiss, GNumSharedSerializationMiss);
 
+			CSV_CUSTOM_STAT(Replication, SharedSerializationRPCHit, GNumSharedSerializationHit, ECsvCustomStatOp::Set);
+			CSV_CUSTOM_STAT(Replication, SharedSerializationRPCMiss, GNumSharedSerializationMiss, ECsvCustomStatOp::Set);
+
 			const FNetworkObjectList& NetworkObjectList = NetDriver->GetNetworkObjectList();
 			CSV_CUSTOM_STAT(Replication, NumberOfActiveActors, NetworkObjectList.GetActiveObjects().Num(), ECsvCustomStatOp::Set);
 			CSV_CUSTOM_STAT(Replication, NumberOfFullyDormantActors, NetworkObjectList.GetDormantObjectsOnAllConnections().Num(), ECsvCustomStatOp::Set);
@@ -5030,6 +5033,9 @@ struct FScopedNetDriverStats
 			SET_DWORD_STAT(STAT_SharedSerializationPropertyHit, GNumSharedSerializationHit);
 			SET_DWORD_STAT(STAT_SharedSerializationPropertyMiss, GNumSharedSerializationMiss);
 
+			CSV_CUSTOM_STAT(Replication, SharedSerializationPropertyHit, GNumSharedSerializationHit, ECsvCustomStatOp::Set);
+			CSV_CUSTOM_STAT(Replication, SharedSerializationPropertyMiss, GNumSharedSerializationMiss, ECsvCustomStatOp::Set);
+
 			// Note: we want to reset this at the end of the frame since the RPC stats are incremented at the top (recv)
 			GNumSharedSerializationHit = 0;
 			GNumSharedSerializationMiss = 0;
diff --git a/Engine/Source/Runtime/Engine/Private/RepLayout.cpp b/Engine/Source/Runtime/Engine/Private/RepLayout.cpp
--- a/Engine/Source/Runtime/Engine/Private/RepLayout.cpp
+++ b/Engine/Source/Runtime/Engine/Private/RepLayout.cpp
@@ -2610,7 +2610,6 @@ void FRepSerializationSharedInfo::CountBytes(FArchive& Ar) const
 	GRANULAR_NETWORK_MEMORY_TRACKING_TRACK("SerializedProperties",
 		if (FNetBitWriter const* const LocalSerializedProperties = SerializedProperties.Get())
 		{
-			Ar.CountBytes(sizeof(FNetBitWriter), sizeof(FNetBitWriter));
 			LocalSerializedProperties->CountMemory(Ar);
 		}
 	);
@@ -2647,6 +2646,8 @@ const FRepSerializedPropertyInfo* FRepSerializationSharedInfo::WriteSharedProper
 	}));
 #endif
 
+	check(SerializedProperties.IsValid());
+
 	FRepSerializedPropertyInfo& SharedPropInfo = SharedPropertyInfo.Emplace_GetRef();
 
 	SharedPropInfo.PropertyKey = PropertyKey;
@@ -2763,6 +2764,8 @@ void FRepLayout::SendProperties_r(
 		// Use shared serialization if was found
 		if (SharedPropInfo)
 		{
+			check(SharedInfo->SerializedProperties.IsValid());
+
 			UE_NET_TRACE_DYNAMIC_NAME_SCOPE(Cmd.Property->GetFName(), Writer, GetTraceCollector(Writer), ENetTraceVerbosity::Trace);
 			UE_NET_TRACE_SCOPE(Shared, Writer, GetTraceCollector(Writer), ENetTraceVerbosity::Trace);
 
@@ -6522,6 +6525,8 @@ void FRepLayout::SerializeProperties_r(
 		// Not concerned with unmapped guids because object references can't be shared
 		if (SharedPropInfo)
 		{
+			check(SharedInfo.SerializedProperties.IsValid());
+
 			GNumSharedSerializationHit++;
 #if !(UE_BUILD_SHIPPING || UE_BUILD_TEST)
 			if ((GNetVerifyShareSerializedData != 0) && Ar.IsSaving())
@@ -6654,6 +6659,8 @@ void FRepLayout::BuildSharedSerialization(
 	FChangelistIterator ChangelistIterator(Changed, 0);
 	FRepHandleIterator HandleIterator(Owner, ChangelistIterator, Cmds, BaseHandleToCmdIndex, 0, 1, 0, Cmds.Num() - 1);
 
+	SharedInfo.Init();
+
 	BuildSharedSerialization_r(HandleIterator, Data, bWriteHandle, bDoChecksum, 0, SharedInfo);
 
 	SharedInfo.SetValid();
@@ -6769,6 +6776,7 @@ void FRepLayout::BuildSharedSerializationForRPC(const FConstRepObjectDataBuffer
 {
 	if ((GNetSharedSerializedData != 0) && !SharedInfoRPC.IsValid())
 	{
+		SharedInfoRPC.Init();
 		SharedInfoRPCParentsChanged.Init(false, Parents.Num());
 
 		for (int32 i = 0; i < Parents.Num(); i++)
diff --git a/Engine/Source/Runtime/Engine/Public/Net/RepLayout.h b/Engine/Source/Runtime/Engine/Public/Net/RepLayout.h
--- a/Engine/Source/Runtime/Engine/Public/Net/RepLayout.h
+++ b/Engine/Source/Runtime/Engine/Public/Net/RepLayout.h
@@ -297,7 +297,6 @@ struct FRepSerializedPropertyInfo
 struct FRepSerializationSharedInfo
 {
 	FRepSerializationSharedInfo():
-		SerializedProperties(MakeUnique<FNetBitWriter>(0)),
 		bIsValid(false)
 	{}
 
@@ -306,6 +305,14 @@ struct FRepSerializationSharedInfo
 		bIsValid = true;
 	}
 
+	void Init()
+	{
+		if (!SerializedProperties.IsValid())
+		{
+			SerializedProperties.Reset(new FNetBitWriter(0));
+		}
+	}
+
 	bool IsValid() const
 	{
 		return bIsValid;
