Fix for crash when disabling audio mixer

diff --git a/Engine/Source/Runtime/Engine/Private/AudioThread.cpp b/Engine/Source/Runtime/Engine/Private/AudioThread.cpp
--- a/Engine/Source/Runtime/Engine/Private/AudioThread.cpp
+++ b/Engine/Source/Runtime/Engine/Private/AudioThread.cpp
@@ -636,13 +636,11 @@ void FAudioThread::StartAudioThread()
 
 void FAudioThread::StopAudioThread()
 {
-	if (!bUseThreadedAudio)
+	if (!IsAudioThreadRunning())
 	{
 		return;
 	}
 
-	check(IsAudioThreadRunning());
-
 	GAudioAsyncBatcher.Flush();
 	GAudioAsyncBatcher.LastBatch.Wait();
 	GAudioAsyncBatcher.LastBatch = UE::Tasks::FTask{}; // release the task as it can hold some references
