Added and modified a few stat scopes.

diff --git a/Engine/Source/Runtime/Engine/Private/Animation/AnimInstance.cpp b/Engine/Source/Runtime/Engine/Private/Animation/AnimInstance.cpp
--- a/Engine/Source/Runtime/Engine/Private/Animation/AnimInstance.cpp
+++ b/Engine/Source/Runtime/Engine/Private/Animation/AnimInstance.cpp
@@ -1704,6 +1704,11 @@ FName UAnimInstance::GetCurrentStateName(int32 MachineIndex)
 
 void UAnimInstance::Montage_UpdateWeight(float DeltaSeconds)
 {
+	if (MontageInstances.IsEmpty())
+	{
+		return;
+	}
+
 	SCOPE_CYCLE_COUNTER(STAT_Montage_UpdateWeight);
 
 	// go through all montage instances, and update them
@@ -1719,11 +1724,16 @@ void UAnimInstance::Montage_UpdateWeight(float DeltaSeconds)
 
 void UAnimInstance::Montage_Advance(float DeltaSeconds)
 {
-	SCOPE_CYCLE_COUNTER(STAT_Montage_Advance);
-
 	// We're about to tick montages, queue their events to they're triggered after batched anim notifies.
 	bQueueMontageEvents = true;
 
+	if (MontageInstances.IsEmpty())
+	{
+		return;
+	}
+
+	SCOPE_CYCLE_COUNTER(STAT_Montage_Advance);
+
 	// go through all montage instances, and update them
 	// and make sure their weight is updated properly
 	for (int32 InstanceIndex = 0; InstanceIndex < MontageInstances.Num(); InstanceIndex++)
diff --git a/Engine/Source/Runtime/Engine/Private/PropertyAccess.cpp b/Engine/Source/Runtime/Engine/Private/PropertyAccess.cpp
--- a/Engine/Source/Runtime/Engine/Private/PropertyAccess.cpp
+++ b/Engine/Source/Runtime/Engine/Private/PropertyAccess.cpp
@@ -619,8 +619,6 @@ namespace PropertyAccess
 
 	void ProcessCopies(UObject* InObject, const FPropertyAccessLibrary& InLibrary, const FCopyBatchId& InBatchId)
 	{
-		QUICK_SCOPE_CYCLE_COUNTER(PropertyAccess_ProcessCopies);
-
 		::FPropertyAccessSystem::ProcessCopies(InObject->GetClass(), InObject, InLibrary, InBatchId.Id);
 	}
 
