Start work on splitting material-specific things out of the core HLSLTree system.  Work still remains to decouple preshaders

diff --git a/Engine/Source/Runtime/Engine/Private/HLSLTree/HLSLTree.cpp b/Engine/Source/Runtime/Engine/Private/HLSLTree/HLSLTree.cpp
--- a/Engine/Source/Runtime/Engine/Private/HLSLTree/HLSLTree.cpp
+++ b/Engine/Source/Runtime/Engine/Private/HLSLTree/HLSLTree.cpp
@@ -7,12 +7,7 @@
 #include "Shader/ShaderTypes.h"
 #include "Shader/Preshader.h"
 
-// TODO - M_ForLoop doesn't work yet
-// FPreparedType::GetEvaluation takes scope, checks loop scope automatically
-
-namespace UE
-{
-namespace HLSLTree
+namespace UE::HLSLTree
 {
 
 enum ELocalPHIChainType : uint8
@@ -1317,12 +1312,6 @@ UObject* FOwnerContext::GetCurrentOwner() const
 	return (OwnerStack.Num() > 0) ? OwnerStack.Last() : nullptr;
 }
 
-uint64 Private::GetNextTypeHash()
-{
-	static uint64 Hash = 1;
-	return Hash++;
-}
-
 void FTree::ResetNodes()
 {
 	FNode* Node = Nodes;
@@ -1542,5 +1531,5 @@ FFunction* FTree::NewFunction()
 	return NewFunction;
 }
 
-} // namespace HLSLTree
-} // namespace UE
+} // namespace UE::HLSLTree
+
diff --git a/Engine/Source/Runtime/Engine/Private/HLSLTree/HLSLTreeCommon.cpp b/Engine/Source/Runtime/Engine/Private/HLSLTree/HLSLTreeCommon.cpp
--- a/Engine/Source/Runtime/Engine/Private/HLSLTree/HLSLTreeCommon.cpp
+++ b/Engine/Source/Runtime/Engine/Private/HLSLTree/HLSLTreeCommon.cpp
@@ -2,14 +2,9 @@
 #include "HLSLTree/HLSLTreeCommon.h"
 #include "HLSLTree/HLSLTreeEmit.h"
 #include "Misc/StringBuilder.h"
-#include "MaterialShared.h"
-#include "MaterialSceneTextureId.h"
-#include "Engine/BlendableInterface.h" // BL_AfterTonemapping
 #include "Engine/Texture.h"
 
-namespace UE
-{
-namespace HLSLTree
+namespace UE::HLSLTree
 {
 
 struct FPreshaderLoopScope
@@ -18,131 +13,6 @@ struct FPreshaderLoopScope
 	Shader::FPreshaderLabel BreakLabel;
 };
 
-FExternalInputDescription GetExternalInputDescription(EExternalInput Input)
-{
-	switch (Input)
-	{
-	case EExternalInput::None: return FExternalInputDescription(TEXT("None"), Shader::EValueType::Void);
-
-	case EExternalInput::TexCoord0: return FExternalInputDescription(TEXT("TexCoord0"), Shader::EValueType::Float2, EExternalInput::TexCoord0_Ddx, EExternalInput::TexCoord0_Ddy);
-	case EExternalInput::TexCoord1: return FExternalInputDescription(TEXT("TexCoord1"), Shader::EValueType::Float2, EExternalInput::TexCoord1_Ddx, EExternalInput::TexCoord1_Ddy);
-	case EExternalInput::TexCoord2: return FExternalInputDescription(TEXT("TexCoord2"), Shader::EValueType::Float2, EExternalInput::TexCoord2_Ddx, EExternalInput::TexCoord2_Ddy);
-	case EExternalInput::TexCoord3: return FExternalInputDescription(TEXT("TexCoord3"), Shader::EValueType::Float2, EExternalInput::TexCoord3_Ddx, EExternalInput::TexCoord3_Ddy);
-	case EExternalInput::TexCoord4: return FExternalInputDescription(TEXT("TexCoord4"), Shader::EValueType::Float2, EExternalInput::TexCoord4_Ddx, EExternalInput::TexCoord4_Ddy);
-	case EExternalInput::TexCoord5: return FExternalInputDescription(TEXT("TexCoord5"), Shader::EValueType::Float2, EExternalInput::TexCoord5_Ddx, EExternalInput::TexCoord5_Ddy);
-	case EExternalInput::TexCoord6: return FExternalInputDescription(TEXT("TexCoord6"), Shader::EValueType::Float2, EExternalInput::TexCoord6_Ddx, EExternalInput::TexCoord6_Ddy);
-	case EExternalInput::TexCoord7: return FExternalInputDescription(TEXT("TexCoord7"), Shader::EValueType::Float2, EExternalInput::TexCoord7_Ddx, EExternalInput::TexCoord7_Ddy);
-
-	case EExternalInput::TexCoord0_Ddx: return FExternalInputDescription(TEXT("TexCoord0_Ddx"), Shader::EValueType::Float2);
-	case EExternalInput::TexCoord1_Ddx: return FExternalInputDescription(TEXT("TexCoord1_Ddx"), Shader::EValueType::Float2);
-	case EExternalInput::TexCoord2_Ddx: return FExternalInputDescription(TEXT("TexCoord2_Ddx"), Shader::EValueType::Float2);
-	case EExternalInput::TexCoord3_Ddx: return FExternalInputDescription(TEXT("TexCoord3_Ddx"), Shader::EValueType::Float2);
-	case EExternalInput::TexCoord4_Ddx: return FExternalInputDescription(TEXT("TexCoord4_Ddx"), Shader::EValueType::Float2);
-	case EExternalInput::TexCoord5_Ddx: return FExternalInputDescription(TEXT("TexCoord5_Ddx"), Shader::EValueType::Float2);
-	case EExternalInput::TexCoord6_Ddx: return FExternalInputDescription(TEXT("TexCoord6_Ddx"), Shader::EValueType::Float2);
-	case EExternalInput::TexCoord7_Ddx: return FExternalInputDescription(TEXT("TexCoord7_Ddx"), Shader::EValueType::Float2);
-
-	case EExternalInput::TexCoord0_Ddy: return FExternalInputDescription(TEXT("TexCoord0_Ddy"), Shader::EValueType::Float2);
-	case EExternalInput::TexCoord1_Ddy: return FExternalInputDescription(TEXT("TexCoord1_Ddy"), Shader::EValueType::Float2);
-	case EExternalInput::TexCoord2_Ddy: return FExternalInputDescription(TEXT("TexCoord2_Ddy"), Shader::EValueType::Float2);
-	case EExternalInput::TexCoord3_Ddy: return FExternalInputDescription(TEXT("TexCoord3_Ddy"), Shader::EValueType::Float2);
-	case EExternalInput::TexCoord4_Ddy: return FExternalInputDescription(TEXT("TexCoord4_Ddy"), Shader::EValueType::Float2);
-	case EExternalInput::TexCoord5_Ddy: return FExternalInputDescription(TEXT("TexCoord5_Ddy"), Shader::EValueType::Float2);
-	case EExternalInput::TexCoord6_Ddy: return FExternalInputDescription(TEXT("TexCoord6_Ddy"), Shader::EValueType::Float2);
-	case EExternalInput::TexCoord7_Ddy: return FExternalInputDescription(TEXT("TexCoord7_Ddy"), Shader::EValueType::Float2);
-
-	case EExternalInput::LightmapTexCoord: return FExternalInputDescription(TEXT("LightmapTexCoord"), Shader::EValueType::Float2, EExternalInput::LightmapTexCoord_Ddx, EExternalInput::LightmapTexCoord_Ddy);
-	case EExternalInput::LightmapTexCoord_Ddx: return FExternalInputDescription(TEXT("LightmapTexCoord_Ddx"), Shader::EValueType::Float2);
-	case EExternalInput::LightmapTexCoord_Ddy: return FExternalInputDescription(TEXT("LightmapTexCoord_Ddy"), Shader::EValueType::Float2);
-
-	case EExternalInput::TwoSidedSign: return FExternalInputDescription(TEXT("TwoSidedSign"), Shader::EValueType::Float1);
-	case EExternalInput::VertexColor: return FExternalInputDescription(TEXT("VertexColor"), Shader::EValueType::Float4, EExternalInput::VertexColor_Ddx, EExternalInput::VertexColor_Ddy);
-	case EExternalInput::VertexColor_Ddx: return FExternalInputDescription(TEXT("VertexColor_Ddx"), Shader::EValueType::Float4);
-	case EExternalInput::VertexColor_Ddy: return FExternalInputDescription(TEXT("VertexColor_Ddy"), Shader::EValueType::Float4);
-
-	case EExternalInput::WorldPosition: return FExternalInputDescription(TEXT("WorldPosition"), Shader::EValueType::Double3, EExternalInput::WorldPosition_Ddx, EExternalInput::WorldPosition_Ddy, EExternalInput::PrevWorldPosition);
-	case EExternalInput::WorldPosition_NoOffsets: return FExternalInputDescription(TEXT("WorldPosition_NoOffsets"), Shader::EValueType::Double3, EExternalInput::WorldPosition_Ddx, EExternalInput::WorldPosition_Ddy, EExternalInput::PrevWorldPosition_NoOffsets);
-	case EExternalInput::TranslatedWorldPosition: return FExternalInputDescription(TEXT("TranslatedWorldPosition"), Shader::EValueType::Float3, EExternalInput::WorldPosition_Ddx, EExternalInput::WorldPosition_Ddy, EExternalInput::PrevTranslatedWorldPosition);
-	case EExternalInput::TranslatedWorldPosition_NoOffsets: return FExternalInputDescription(TEXT("TranslatedWorldPosition_NoOffsets"), Shader::EValueType::Float3, EExternalInput::WorldPosition_Ddx, EExternalInput::WorldPosition_Ddy, EExternalInput::PrevTranslatedWorldPosition_NoOffsets);
-	case EExternalInput::ActorWorldPosition: return FExternalInputDescription(TEXT("TranslatedWorldPosition_NoOffsets"), Shader::EValueType::Double3);
-
-	case EExternalInput::PrevWorldPosition: return FExternalInputDescription(TEXT("PrevWorldPosition"), Shader::EValueType::Double3, EExternalInput::WorldPosition_Ddx, EExternalInput::WorldPosition_Ddy);
-	case EExternalInput::PrevWorldPosition_NoOffsets: return FExternalInputDescription(TEXT("PrevWorldPosition_NoOffsets"), Shader::EValueType::Double3, EExternalInput::WorldPosition_Ddx, EExternalInput::WorldPosition_Ddy);
-	case EExternalInput::PrevTranslatedWorldPosition: return FExternalInputDescription(TEXT("PrevTranslatedWorldPosition"), Shader::EValueType::Float3, EExternalInput::WorldPosition_Ddx, EExternalInput::WorldPosition_Ddy);
-	case EExternalInput::PrevTranslatedWorldPosition_NoOffsets: return FExternalInputDescription(TEXT("PrevTranslatedWorldPosition_NoOffsets"), Shader::EValueType::Float3, EExternalInput::WorldPosition_Ddx, EExternalInput::WorldPosition_Ddy);
-
-	case EExternalInput::WorldPosition_Ddx: return FExternalInputDescription(TEXT("WorldPosition_Ddx"), Shader::EValueType::Float3);
-	case EExternalInput::WorldPosition_Ddy: return FExternalInputDescription(TEXT("WorldPosition_Ddx"), Shader::EValueType::Float3);
-
-	case EExternalInput::WorldNormal: return FExternalInputDescription(TEXT("WorldNormal"), Shader::EValueType::Float3);
-	case EExternalInput::WorldReflection: return FExternalInputDescription(TEXT("WorldReflection"), Shader::EValueType::Float3);
-	case EExternalInput::WorldVertexNormal: return FExternalInputDescription(TEXT("WorldVertexNormal"), Shader::EValueType::Float3);
-	case EExternalInput::WorldVertexTangent: return FExternalInputDescription(TEXT("WorldVertexTangent"), Shader::EValueType::Float3);
-
-	case EExternalInput::ViewportUV: return FExternalInputDescription(TEXT("ViewportUV"), Shader::EValueType::Float2);
-	case EExternalInput::PixelPosition: return FExternalInputDescription(TEXT("PixelPosition"), Shader::EValueType::Float2);
-	case EExternalInput::ViewSize: return FExternalInputDescription(TEXT("ViewSize"), Shader::EValueType::Float2);
-	case EExternalInput::RcpViewSize: return FExternalInputDescription(TEXT("RcpViewSize"), Shader::EValueType::Float2);
-	case EExternalInput::FieldOfView: return FExternalInputDescription(TEXT("FieldOfView"), Shader::EValueType::Float1, EExternalInput::None, EExternalInput::None, EExternalInput::PrevFieldOfView);
-	case EExternalInput::TanHalfFieldOfView: return FExternalInputDescription(TEXT("TanHalfFieldOfView"), Shader::EValueType::Float2, EExternalInput::None, EExternalInput::None, EExternalInput::PrevTanHalfFieldOfView);
-	case EExternalInput::CotanHalfFieldOfView: return FExternalInputDescription(TEXT("CotanHalfFieldOfView"), Shader::EValueType::Float2, EExternalInput::None, EExternalInput::None, EExternalInput::PrevCotanHalfFieldOfView);
-	case EExternalInput::TemporalSampleCount: return FExternalInputDescription(TEXT("TemporalSampleCount"), Shader::EValueType::Float1);
-	case EExternalInput::TemporalSampleIndex: return FExternalInputDescription(TEXT("TemporalSampleIndex"), Shader::EValueType::Float1);
-	case EExternalInput::TemporalSampleOffset: return FExternalInputDescription(TEXT("TemporalSampleOffset"), Shader::EValueType::Float2);
-	case EExternalInput::PreExposure: return FExternalInputDescription(TEXT("PreExposure"), Shader::EValueType::Float1);
-	case EExternalInput::RcpPreExposure: return FExternalInputDescription(TEXT("RcpPreExposure"), Shader::EValueType::Float1);
-	case EExternalInput::RuntimeVirtualTextureOutputLevel: return FExternalInputDescription(TEXT("RuntimeVirtualTextureOutputLevel"), Shader::EValueType::Float1);
-	case EExternalInput::RuntimeVirtualTextureOutputDerivative: return FExternalInputDescription(TEXT("RuntimeVirtualTextureOutputDerivative"), Shader::EValueType::Float2);
-	case EExternalInput::RuntimeVirtualTextureMaxLevel: return FExternalInputDescription(TEXT("RuntimeVirtualTextureMaxLevel"), Shader::EValueType::Float1);
-
-	case EExternalInput::CameraVector: return FExternalInputDescription(TEXT("CameraVector"), Shader::EValueType::Float3);
-	case EExternalInput::CameraWorldPosition: return FExternalInputDescription(TEXT("CameraWorldPosition"), Shader::EValueType::Double3, EExternalInput::None, EExternalInput::None, EExternalInput::PrevCameraWorldPosition);
-	case EExternalInput::ViewWorldPosition: return FExternalInputDescription(TEXT("ViewWorldPosition"), Shader::EValueType::Double3, EExternalInput::None, EExternalInput::None, EExternalInput::PrevViewWorldPosition);
-	case EExternalInput::PreViewTranslation: return FExternalInputDescription(TEXT("PreViewTranslation"), Shader::EValueType::Double3, EExternalInput::None, EExternalInput::None, EExternalInput::PrevPreViewTranslation);
-	case EExternalInput::TangentToWorld: return FExternalInputDescription(TEXT("TangentToWorld"), Shader::EValueType::Float4x4);
-	case EExternalInput::LocalToWorld: return FExternalInputDescription(TEXT("LocalToWorld"), Shader::EValueType::Double4x4, EExternalInput::None, EExternalInput::None, EExternalInput::PrevLocalToWorld);
-	case EExternalInput::WorldToLocal: return FExternalInputDescription(TEXT("WorldToLocal"), Shader::EValueType::DoubleInverse4x4, EExternalInput::None, EExternalInput::None, EExternalInput::PrevWorldToLocal);
-	case EExternalInput::TranslatedWorldToCameraView: return FExternalInputDescription(TEXT("TranslatedWorldToCameraView"), Shader::EValueType::Float4x4, EExternalInput::None, EExternalInput::None, EExternalInput::PrevTranslatedWorldToCameraView);
-	case EExternalInput::TranslatedWorldToView: return FExternalInputDescription(TEXT("TranslatedWorldToView"), Shader::EValueType::Float4x4, EExternalInput::None, EExternalInput::None, EExternalInput::PrevTranslatedWorldToView);
-	case EExternalInput::CameraViewToTranslatedWorld: return FExternalInputDescription(TEXT("CameraViewToTranslatedWorld"), Shader::EValueType::Float4x4, EExternalInput::None, EExternalInput::None, EExternalInput::PrevCameraViewToTranslatedWorld);
-	case EExternalInput::ViewToTranslatedWorld: return FExternalInputDescription(TEXT("ViewToTranslatedWorld"), Shader::EValueType::Float4x4, EExternalInput::None, EExternalInput::None, EExternalInput::PrevViewToTranslatedWorld);
-	case EExternalInput::WorldToParticle: return FExternalInputDescription(TEXT("WorldToParticle"), Shader::EValueType::DoubleInverse4x4);
-	case EExternalInput::WorldToInstance: return FExternalInputDescription(TEXT("WorldToInstance"), Shader::EValueType::DoubleInverse4x4);
-	case EExternalInput::ParticleToWorld: return FExternalInputDescription(TEXT("ParticleToWorld"), Shader::EValueType::Double4x4);
-	case EExternalInput::InstanceToWorld: return FExternalInputDescription(TEXT("InstanceToWorld"), Shader::EValueType::Double4x4);
-
-	case EExternalInput::PrevFieldOfView: return FExternalInputDescription(TEXT("PrevFieldOfView"), Shader::EValueType::Float2);
-	case EExternalInput::PrevTanHalfFieldOfView: return FExternalInputDescription(TEXT("PrevTanHalfFieldOfView"), Shader::EValueType::Float2);
-	case EExternalInput::PrevCotanHalfFieldOfView: return FExternalInputDescription(TEXT("PrevCotanHalfFieldOfView"), Shader::EValueType::Float2);
-	case EExternalInput::PrevCameraWorldPosition: return FExternalInputDescription(TEXT("PrevCameraWorldPosition"), Shader::EValueType::Double3);
-	case EExternalInput::PrevViewWorldPosition: return FExternalInputDescription(TEXT("PrevViewWorldPosition"), Shader::EValueType::Double3);
-	case EExternalInput::PrevPreViewTranslation: return FExternalInputDescription(TEXT("PrevPreViewTranslation"), Shader::EValueType::Double3);
-	case EExternalInput::PrevLocalToWorld: return FExternalInputDescription(TEXT("PrevLocalToWorld"), Shader::EValueType::Double4x4);
-	case EExternalInput::PrevWorldToLocal: return FExternalInputDescription(TEXT("PrevWorldToLocal"), Shader::EValueType::DoubleInverse4x4);
-	case EExternalInput::PrevTranslatedWorldToCameraView: return FExternalInputDescription(TEXT("PrevTranslatedWorldToCameraView"), Shader::EValueType::Float4x4);
-	case EExternalInput::PrevTranslatedWorldToView: return FExternalInputDescription(TEXT("PrevTranslatedWorldToView"), Shader::EValueType::Float4x4);
-	case EExternalInput::PrevCameraViewToTranslatedWorld: return FExternalInputDescription(TEXT("PrevCameraViewToTranslatedWorld"), Shader::EValueType::Float4x4);
-	case EExternalInput::PrevViewToTranslatedWorld: return FExternalInputDescription(TEXT("PrevViewToTranslatedWorld"), Shader::EValueType::Float4x4);
-
-	case EExternalInput::PixelDepth: return FExternalInputDescription(TEXT("PixelDepth"), Shader::EValueType::Float1, EExternalInput::PixelDepth_Ddx, EExternalInput::PixelDepth_Ddy);
-	case EExternalInput::PixelDepth_Ddx: return FExternalInputDescription(TEXT("PixelDepth_Ddx"), Shader::EValueType::Float1);
-	case EExternalInput::PixelDepth_Ddy: return FExternalInputDescription(TEXT("PixelDepth_Ddy"), Shader::EValueType::Float1);
-
-	case EExternalInput::GameTime: return FExternalInputDescription(TEXT("GameTime"), Shader::EValueType::Float1, EExternalInput::None, EExternalInput::None, EExternalInput::PrevGameTime);
-	case EExternalInput::RealTime: return FExternalInputDescription(TEXT("RealTime"), Shader::EValueType::Float1, EExternalInput::None, EExternalInput::None, EExternalInput::PrevRealTime);
-	case EExternalInput::DeltaTime: return FExternalInputDescription(TEXT("DeltaTime"), Shader::EValueType::Float1);
-
-	case EExternalInput::PrevGameTime: return FExternalInputDescription(TEXT("PrevGameTime"), Shader::EValueType::Float1);
-	case EExternalInput::PrevRealTime: return FExternalInputDescription(TEXT("PrevRealTime"), Shader::EValueType::Float1);
-
-	case EExternalInput::ParticleColor: return FExternalInputDescription(TEXT("ParticleColor"), Shader::EValueType::Float4);
-	case EExternalInput::ParticleTranslatedWorldPosition: return FExternalInputDescription(TEXT("ParticleTranslatedWorldPosition"), Shader::EValueType::Float3);
-	case EExternalInput::ParticleRadius: return FExternalInputDescription(TEXT("ParticleRadius"), Shader::EValueType::Float1);
-
-	default: checkNoEntry(); return FExternalInputDescription(TEXT("Invalid"), Shader::EValueType::Void);
-	}
-}
-
 FSwizzleParameters::FSwizzleParameters(int8 InR, int8 InG, int8 InB, int8 InA) : NumComponents(0)
 {
 	ComponentIndex[0] = InR;
@@ -239,439 +109,6 @@ void FExpressionConstant::EmitValuePreshader(FEmitContext& Context, FEmitScope&
 	OutResult.Preshader.WriteOpcode(Shader::EPreshaderOpcode::Constant).Write(Value);
 }
 
-void FExpressionMaterialShadingModel::ComputeAnalyticDerivatives(FTree& Tree, FExpressionDerivatives& OutResult) const
-{
-	const Shader::FValue ZeroValue(Shader::EValueType::Float1);
-	OutResult.ExpressionDdx = Tree.NewConstant(ZeroValue);
-	OutResult.ExpressionDdy = OutResult.ExpressionDdx;
-}
-
-bool FExpressionMaterialShadingModel::PrepareValue(FEmitContext& Context, FEmitScope& Scope, const FRequestedType& RequestedType, FPrepareValueResult& OutResult) const
-{
-	return OutResult.SetType(Context, RequestedType, EExpressionEvaluation::Constant, Shader::EValueType::Int1);
-}
-
-void FExpressionMaterialShadingModel::EmitValuePreshader(FEmitContext& Context, FEmitScope& Scope, const FRequestedType& RequestedType, FEmitValuePreshaderResult& OutResult) const
-{
-	Context.ShadingModelsFromCompilation.AddShadingModel(ShadingModel);
-
-	Context.PreshaderStackPosition++;
-	OutResult.Type = Shader::EValueType::Int1;
-	OutResult.Preshader.WriteOpcode(Shader::EPreshaderOpcode::Constant).Write(Shader::FValue((int32)ShadingModel));
-}
-
-void FExpressionMaterialParameter::ComputeAnalyticDerivatives(FTree& Tree, FExpressionDerivatives& OutResult) const
-{
-	const Shader::FType Type = GetShaderValueType(ParameterType);
-	const Shader::FType DerivativeType = Type.GetDerivativeType();
-	if (!DerivativeType.IsVoid())
-	{
-		const Shader::FValue ZeroValue(DerivativeType);
-		OutResult.ExpressionDdx = Tree.NewConstant(ZeroValue);
-		OutResult.ExpressionDdy = OutResult.ExpressionDdx;
-	}
-}
-
-bool FExpressionMaterialParameter::PrepareValue(FEmitContext& Context, FEmitScope& Scope, const FRequestedType& RequestedType, FPrepareValueResult& OutResult) const
-{
-	EExpressionEvaluation Evaluation = EExpressionEvaluation::Shader;
-	if (IsStaticMaterialParameter(ParameterType))
-	{
-		Evaluation = EExpressionEvaluation::Constant;
-	}
-	else if (ParameterType == EMaterialParameterType::Scalar ||
-		ParameterType == EMaterialParameterType::Vector ||
-		ParameterType == EMaterialParameterType::DoubleVector)
-	{
-		Evaluation = EExpressionEvaluation::Preshader;
-	}
-	return OutResult.SetType(Context, RequestedType, Evaluation, DefaultValue.Type);
-}
-
-void FExpressionMaterialParameter::EmitValueShader(FEmitContext& Context, FEmitScope& Scope, const FRequestedType& RequestedType, FEmitValueShaderResult& OutResult) const
-{
-	if (ParameterType == EMaterialParameterType::Texture)
-	{
-		const Shader::FTextureValue* TextureValue = DefaultValue.AsTexture();
-		check(TextureValue);
-
-		const Shader::EValueType TextureType = TextureValue->GetType();
-		const TCHAR* ConstructorName = nullptr;
-		const TCHAR* TextureTypeName = nullptr;
-		int32 ParameterIndex = INDEX_NONE;
-
-		if (TextureType == Shader::EValueType::TextureExternal)
-		{
-			check(TextureValue->SamplerType == SAMPLERTYPE_External);
-
-			ConstructorName = TEXT("MakeTextureExternal");
-			TextureTypeName = TEXT("ExternalTexture");
-
-			FMaterialExternalTextureParameterInfo TextureParameterInfo;
-			TextureParameterInfo.ParameterName = NameToScriptName(ParameterInfo.Name);
-			TextureParameterInfo.ExternalTextureGuid = TextureValue->ExternalTextureGuid;
-			if (TextureValue->Texture)
-			{
-				TextureParameterInfo.SourceTextureIndex = Context.Material->GetReferencedTextures().Find(TextureValue->Texture);
-			}
-			ParameterIndex = Context.MaterialCompilationOutput->UniformExpressionSet.FindOrAddExternalTextureParameter(TextureParameterInfo);
-		}
-		else
-		{
-			EMaterialTextureParameterType TextureParameterType = EMaterialTextureParameterType::Count;
-			switch (TextureType)
-			{
-			case Shader::EValueType::Texture2D:
-				TextureParameterType = EMaterialTextureParameterType::Standard2D;
-				ConstructorName = TEXT("MakeTexture2D");
-				TextureTypeName = TEXT("Texture2D");
-				break;
-			case Shader::EValueType::Texture2DArray:
-				TextureParameterType = EMaterialTextureParameterType::Array2D;
-				ConstructorName = TEXT("MakeTexture2DArray");
-				TextureTypeName = TEXT("Texture2DArray");
-				break;
-			case Shader::EValueType::TextureCube:
-				TextureParameterType = EMaterialTextureParameterType::Cube;
-				ConstructorName = TEXT("MakeTextureCube");
-				TextureTypeName = TEXT("TextureCube");
-				break;
-			case Shader::EValueType::TextureCubeArray:
-				TextureParameterType = EMaterialTextureParameterType::ArrayCube;
-				ConstructorName = TEXT("MakeTextureCubeArray");
-				TextureTypeName = TEXT("TextureCubeArray");
-				break;
-			case Shader::EValueType::Texture3D:
-				TextureParameterType = EMaterialTextureParameterType::Volume;
-				ConstructorName = TEXT("MakeTexture3D");
-				TextureTypeName = TEXT("VolumeTexture");
-				break;
-			default:
-				checkNoEntry();
-				break;
-			}
-
-			FMaterialTextureParameterInfo TextureParameterInfo;
-			TextureParameterInfo.ParameterInfo = ParameterInfo;
-			TextureParameterInfo.TextureIndex = Context.Material->GetReferencedTextures().Find(TextureValue->Texture);
-			TextureParameterInfo.SamplerSource = SSM_FromTextureAsset; // TODO - Is this needed?
-			check(TextureParameterInfo.TextureIndex != INDEX_NONE);
-			ParameterIndex = Context.MaterialCompilationOutput->UniformExpressionSet.FindOrAddTextureParameter(TextureParameterType, TextureParameterInfo);
-		}
-
-		TStringBuilder<256> FormattedCode;
-		FormattedCode.Appendf(TEXT("%s(Material.%s_%d, Material.%s_%dSampler, %d)"), ConstructorName, TextureTypeName, ParameterIndex, TextureTypeName, ParameterIndex, (int32)TextureValue->SamplerType);
-		OutResult.Code = Context.EmitExpression(Scope, TextureType, FormattedCode.ToView());
-	}
-}
-
-void FExpressionMaterialParameter::EmitValuePreshader(FEmitContext& Context, FEmitScope& Scope, const FRequestedType& RequestedType, FEmitValuePreshaderResult& OutResult) const
-{
-	Context.PreshaderStackPosition++;
-	OutResult.Type = GetShaderValueType(ParameterType);
-	if (ParameterType == EMaterialParameterType::StaticSwitch)
-	{
-		Shader::FValue Value = DefaultValue;
-		for (const FStaticSwitchParameter& Parameter : Context.StaticParameters->StaticSwitchParameters)
-		{
-			if (Parameter.ParameterInfo == ParameterInfo)
-			{
-				Value = Parameter.Value;
-				break;
-			}
-		}
-		OutResult.Preshader.WriteOpcode(Shader::EPreshaderOpcode::Constant).Write(Value);
-	}
-	else
-	{
-		const uint32* PrevDefaultOffset = Context.DefaultUniformValues.Find(DefaultValue);
-		uint32 DefaultOffset;
-		if (PrevDefaultOffset)
-		{
-			DefaultOffset = *PrevDefaultOffset;
-		}
-		else
-		{
-			DefaultOffset = Context.MaterialCompilationOutput->UniformExpressionSet.AddDefaultParameterValue(DefaultValue);
-			Context.DefaultUniformValues.Add(DefaultValue, DefaultOffset);
-		}
-		const int32 ParameterIndex = Context.MaterialCompilationOutput->UniformExpressionSet.FindOrAddNumericParameter(ParameterType, ParameterInfo, DefaultOffset);
-		check(ParameterIndex >= 0 && ParameterIndex <= 0xffff);
-		OutResult.Preshader.WriteOpcode(Shader::EPreshaderOpcode::Parameter).Write((uint16)ParameterIndex);
-	}
-}
-
-void FExpressionExternalInput::ComputeAnalyticDerivatives(FTree& Tree, FExpressionDerivatives& OutResult) const
-{
-	const FExternalInputDescription InputDesc = GetExternalInputDescription(InputType);
-	if (InputDesc.Ddx != EExternalInput::None)
-	{
-		check(InputDesc.Ddy != EExternalInput::None);
-		OutResult.ExpressionDdx = Tree.NewExpression<FExpressionExternalInput>(InputDesc.Ddx);
-		OutResult.ExpressionDdy = Tree.NewExpression<FExpressionExternalInput>(InputDesc.Ddy);
-	}
-	else
-	{
-		const Shader::EValueType DerivativeType = Shader::MakeDerivativeType(InputDesc.Type);
-		switch (InputType)
-		{
-		case EExternalInput::ViewportUV:
-		{
-			// Ddx = float2(RcpViewSize.x, 0.0f)
-			// Ddy = float2(0.0f, RcpViewSize.y)
-			FExpression* RcpViewSize = Tree.NewExpression<FExpressionExternalInput>(EExternalInput::RcpViewSize);
-			FExpression* Constant0 = Tree.NewConstant(0.0f);
-			OutResult.ExpressionDdx = Tree.NewExpression<FExpressionAppend>(Tree.NewExpression<FExpressionSwizzle>(MakeSwizzleMask(true, false, false, false), RcpViewSize), Constant0);
-			OutResult.ExpressionDdy = Tree.NewExpression<FExpressionAppend>(Constant0, Tree.NewExpression<FExpressionSwizzle>(MakeSwizzleMask(false, true, false, false), RcpViewSize));
-			break;
-		}
-		default:
-			if (DerivativeType != Shader::EValueType::Void)
-			{
-				OutResult.ExpressionDdx = OutResult.ExpressionDdy = Tree.NewConstant(Shader::FValue(DerivativeType));
-			}
-			break;
-		}
-	}
-}
-
-FExpression* FExpressionExternalInput::ComputePreviousFrame(FTree& Tree, const FRequestedType& RequestedType) const
-{
-	if (InputType == EExternalInput::ActorWorldPosition)
-	{
-		return Tree.NewBinaryOp(EOperation::VecMulMatrix3, Tree.NewBinaryOp(EOperation::VecMulMatrix3,
-			Tree.NewExpression<FExpressionExternalInput>(EExternalInput::ActorWorldPosition),
-			Tree.NewExpression< FExpressionExternalInput>(EExternalInput::WorldToLocal)),
-			Tree.NewExpression< FExpressionExternalInput>(EExternalInput::PrevLocalToWorld));
-	}
-
-	const FExternalInputDescription InputDesc = GetExternalInputDescription(InputType);
-	if (InputDesc.PreviousFrame != EExternalInput::None)
-	{
-		return Tree.NewExpression<FExpressionExternalInput>(InputDesc.PreviousFrame);
-	}
-	return nullptr;
-}
-
-bool FExpressionExternalInput::PrepareValue(FEmitContext& Context, FEmitScope& Scope, const FRequestedType& RequestedType, FPrepareValueResult& OutResult) const
-{
-	const FExternalInputDescription InputDesc = GetExternalInputDescription(InputType);
-
-	switch (InputType)
-	{
-	case EExternalInput::WorldNormal:
-	case EExternalInput::WorldReflection:
-		Context.bReadMaterialNormal = true;
-		break;
-	default:
-		break;
-	}
-
-	return OutResult.SetType(Context, RequestedType, EExpressionEvaluation::Shader, InputDesc.Type);
-}
-
-void FExpressionExternalInput::EmitValueShader(FEmitContext& Context, FEmitScope& Scope, const FRequestedType& RequestedType, FEmitValueShaderResult& OutResult) const
-{
-	const int32 TypeIndex = (int32)InputType;
-	Context.ExternalInputMask[Context.ShaderFrequency][TypeIndex] = true;
-	if (IsTexCoord(InputType))
-	{
-		const int32 TexCoordIndex = TypeIndex - (int32)EExternalInput::TexCoord0;
-		OutResult.Code = Context.EmitInlineExpression(Scope, Shader::EValueType::Float2, TEXT("Parameters.TexCoords[%].xy"), TexCoordIndex);
-	}
-	else if (IsTexCoord_Ddx(InputType))
-	{
-		const int32 TexCoordIndex = TypeIndex - (int32)EExternalInput::TexCoord0_Ddx;
-		OutResult.Code = Context.EmitInlineExpression(Scope, Shader::EValueType::Float2, TEXT("Parameters.TexCoords_DDX[%].xy"), TexCoordIndex);
-	}
-	else if (IsTexCoord_Ddy(InputType))
-	{
-		const int32 TexCoordIndex = TypeIndex - (int32)EExternalInput::TexCoord0_Ddy;
-		OutResult.Code = Context.EmitInlineExpression(Scope, Shader::EValueType::Float2, TEXT("Parameters.TexCoords_DDY[%].xy"), TexCoordIndex);
-	}
-	else
-	{
-		const FExternalInputDescription InputDesc = GetExternalInputDescription(InputType);
-		const TCHAR* Code = nullptr;
-		switch (InputType)
-		{
-		case EExternalInput::LightmapTexCoord: Code = TEXT("GetLightmapUVs(Parameters)"); break;
-		case EExternalInput::LightmapTexCoord_Ddx: Code = TEXT("GetLightmapUVs_DDX(Parameters)"); break;
-		case EExternalInput::LightmapTexCoord_Ddy: Code = TEXT("GetLightmapUVs_DDY(Parameters)"); break;
-		case EExternalInput::TwoSidedSign: Code = TEXT("Parameters.TwoSidedSign"); break;
-		case EExternalInput::VertexColor: Code = TEXT("Parameters.VertexColor"); Context.bUsesVertexColor |= (Context.ShaderFrequency != SF_Vertex); break;
-		case EExternalInput::VertexColor_Ddx: Code = TEXT("Parameters.VertexColor_DDX"); Context.bUsesVertexColor |= (Context.ShaderFrequency != SF_Vertex); break;
-		case EExternalInput::VertexColor_Ddy: Code = TEXT("Parameters.VertexColor_DDY"); Context.bUsesVertexColor |= (Context.ShaderFrequency != SF_Vertex); break;
-		case EExternalInput::WorldPosition: Code = TEXT("GetWorldPosition(Parameters)"); break;
-		case EExternalInput::WorldPosition_NoOffsets: Code = TEXT("GetWorldPosition_NoMaterialOffsets(Parameters)"); break;
-		case EExternalInput::TranslatedWorldPosition: Code = TEXT("GetTranslatedWorldPosition(Parameters)"); break;
-		case EExternalInput::TranslatedWorldPosition_NoOffsets: Code = TEXT("GetTranslatedWorldPosition_NoMaterialOffsets(Parameters)"); break;
-		case EExternalInput::ActorWorldPosition: Code = TEXT("GetActorWorldPosition(Parameters)"); break;
-		case EExternalInput::PrevWorldPosition: Code = TEXT("GetPrevWorldPosition(Parameters)"); break;
-		case EExternalInput::PrevWorldPosition_NoOffsets: Code = TEXT("GetPrevWorldPosition_NoMaterialOffsets(Parameters)"); break;
-		case EExternalInput::PrevTranslatedWorldPosition: Code = TEXT("GetPrevTranslatedWorldPosition(Parameters)"); break;
-		case EExternalInput::PrevTranslatedWorldPosition_NoOffsets: Code = TEXT("GetPrevTranslatedWorldPosition_NoMaterialOffsets(Parameters)"); break;
-		case EExternalInput::WorldPosition_Ddx: Code = TEXT("Parameters.WorldPosition_DDX"); break;
-		case EExternalInput::WorldPosition_Ddy: Code = TEXT("Parameters.WorldPosition_DDY"); break;
-
-		case EExternalInput::WorldNormal: Code = TEXT("Parameters.WorldNormal"); break;
-		case EExternalInput::WorldReflection: Code = TEXT("Parameters.ReflectionVector"); break;
-		case EExternalInput::WorldVertexNormal: Code = TEXT("Parameters.TangentToWorld[2]"); break;
-		case EExternalInput::WorldVertexTangent: Code = TEXT("Parameters.TangentToWorld[0]"); break;
-
-		case EExternalInput::ViewportUV: Code = TEXT("GetViewportUV(Parameters)"); break;
-		case EExternalInput::PixelPosition: Code = TEXT("GetPixelPosition(Parameters)"); break;
-		case EExternalInput::ViewSize: Code = TEXT("View.ViewSizeAndInvSize.xy"); break;
-		case EExternalInput::RcpViewSize: Code = TEXT("View.ViewSizeAndInvSize.zw"); break;
-
-		case EExternalInput::FieldOfView: Code = TEXT("View.FieldOfViewWideAngles"); break;
-		case EExternalInput::TanHalfFieldOfView: Code = TEXT("GetTanHalfFieldOfView()"); break;
-		case EExternalInput::CotanHalfFieldOfView: Code = TEXT("GetCotanHalfFieldOfView()"); break;
-		case EExternalInput::TemporalSampleCount: Code = TEXT("View.TemporalAAParams.y"); break;
-		case EExternalInput::TemporalSampleIndex: Code = TEXT("View.TemporalAAParams.x"); break;
-		case EExternalInput::TemporalSampleOffset: Code = TEXT("View.TemporalAAParams.zw"); break;
-		case EExternalInput::PreExposure: Code = TEXT("View.PreExposure.x"); break;
-		case EExternalInput::RcpPreExposure: Code = TEXT("View.OneOverPreExposure.x"); break;
-		case EExternalInput::RuntimeVirtualTextureOutputLevel:  Code = TEXT("View.RuntimeVirtualTextureMipLevel.x"); break;
-		case EExternalInput::RuntimeVirtualTextureOutputDerivative: Code = TEXT("View.RuntimeVirtualTextureMipLevel.zw"); break;
-		case EExternalInput::RuntimeVirtualTextureMaxLevel:  Code = TEXT("View.RuntimeVirtualTextureMipLevel.y"); break;
-
-		case EExternalInput::CameraVector: Code = TEXT("Parameters.CameraVector"); break;
-		case EExternalInput::CameraWorldPosition: Code = TEXT("ResolvedView.WorldCameraOrigin"); break;
-		case EExternalInput::ViewWorldPosition: Code = TEXT("ResolvedView.WorldViewOrigin"); break;
-		case EExternalInput::PreViewTranslation: Code = TEXT("ResolvedView.PreViewTranslation"); break;
-		case EExternalInput::TangentToWorld: Code = TEXT("Parameters.TangentToWorld"); break;
-		case EExternalInput::LocalToWorld: Code = TEXT("GetLocalToWorld(Parameters)"); break;
-		case EExternalInput::WorldToLocal: Code = TEXT("GetPrimitiveData(Parameters).WorldToLocal"); break;
-		case EExternalInput::TranslatedWorldToCameraView: Code = TEXT("ResolvedView.TranslatedWorldToCameraView"); break;
-		case EExternalInput::TranslatedWorldToView: Code = TEXT("ResolvedView.TranslatedWorldToView"); break;
-		case EExternalInput::CameraViewToTranslatedWorld: Code = TEXT("ResolvedView.CameraViewToTranslatedWorld"); break;
-		case EExternalInput::ViewToTranslatedWorld: Code = TEXT("ResolvedView.ViewToTranslatedWorld"); break;
-		case EExternalInput::WorldToParticle: Code = TEXT("Parameters.Particle.WorldToParticle"); break;
-		case EExternalInput::WorldToInstance: Code = TEXT("GetWorldToInstance(Parameters)"); break;
-		case EExternalInput::ParticleToWorld: Code = TEXT("Parameters.Particle.ParticleToWorld"); break;
-		case EExternalInput::InstanceToWorld: Code = TEXT("GetInstanceToWorld(Parameters)"); break;
-
-		case EExternalInput::PrevFieldOfView: Code = TEXT("View.PrevFieldOfViewWideAngles"); break;
-		case EExternalInput::PrevTanHalfFieldOfView: Code = TEXT("GetPrevTanHalfFieldOfView()"); break;
-		case EExternalInput::PrevCotanHalfFieldOfView: Code = TEXT("GetPrevCotanHalfFieldOfView()"); break;
-		case EExternalInput::PrevCameraWorldPosition: Code = TEXT("ResolvedView.PrevWorldCameraOrigin"); break;
-		case EExternalInput::PrevViewWorldPosition: Code = TEXT("ResolvedView.PrevWorldViewOrigin"); break;
-		case EExternalInput::PrevPreViewTranslation: Code = TEXT("ResolvedView.PrevPreViewTranslation"); break;
-		case EExternalInput::PrevLocalToWorld: Code = TEXT("GetPrevLocalToWorld(Parameters)"); break;
-		case EExternalInput::PrevWorldToLocal: Code = TEXT("GetPrimitiveData(Parameters).PreviousWorldToLocal"); break;
-		case EExternalInput::PrevTranslatedWorldToCameraView: Code = TEXT("ResolvedView.PrevTranslatedWorldToCameraView"); break;
-		case EExternalInput::PrevTranslatedWorldToView: Code = TEXT("ResolvedView.PrevTranslatedWorldToView"); break;
-		case EExternalInput::PrevCameraViewToTranslatedWorld: Code = TEXT("ResolvedView.PrevCameraViewToTranslatedWorld"); break;
-		case EExternalInput::PrevViewToTranslatedWorld: Code = TEXT("ResolvedView.PrevViewToTranslatedWorld"); break;
-
-		case EExternalInput::PixelDepth: Code = TEXT("GetPixelDepth(Parameters)"); break;
-		case EExternalInput::PixelDepth_Ddx: Code = TEXT("Parameters.ScreenPosition_DDX.w"); break;
-		case EExternalInput::PixelDepth_Ddy: Code = TEXT("Parameters.ScreenPosition_DDY.w"); break;
-		case EExternalInput::GameTime: Code = TEXT("View.GameTime"); break;
-		case EExternalInput::RealTime: Code = TEXT("View.RealTime"); break;
-		case EExternalInput::DeltaTime: Code = TEXT("View.DeltaTime"); break;
-		case EExternalInput::PrevGameTime: Code = TEXT("View.PrevFrameGameTime"); break;
-		case EExternalInput::PrevRealTime: Code = TEXT("View.PrevFrameRealTime"); break;
-
-		case EExternalInput::ParticleColor: Code = TEXT("Parameters.Particle.Color"); Context.bUsesParticleColor |= (Context.ShaderFrequency != SF_Vertex); break;
-		case EExternalInput::ParticleTranslatedWorldPosition: Code = TEXT("Parameters.Particle.TranslatedWorldPositionAndSize.xyz"); Context.bNeedsParticlePosition = true; break;
-		case EExternalInput::ParticleRadius: Code = TEXT("Parameters.Particle.TranslatedWorldPositionAndSize.w"); Context.bNeedsParticlePosition = true; break;
-
-		default:
-			checkNoEntry();
-			break;
-		}
-		OutResult.Code = Context.EmitInlineExpression(Scope, InputDesc.Type, Code);
-	}
-}
-
-bool FExpressionMaterialSceneTexture::PrepareValue(FEmitContext& Context, FEmitScope& Scope, const FRequestedType& RequestedType, FPrepareValueResult& OutResult) const
-{
-	Context.PrepareExpression(TexCoordExpression, Scope, ERequestedType::Vector2);
-
-	Context.MaterialCompilationOutput->bNeedsSceneTextures = true;
-	Context.MaterialCompilationOutput->SetIsSceneTextureUsed((ESceneTextureId)SceneTextureId);
-
-	return OutResult.SetType(Context, RequestedType, EExpressionEvaluation::Shader, Shader::EValueType::Float4);
-}
-
-void FExpressionMaterialSceneTexture::EmitValueShader(FEmitContext& Context, FEmitScope& Scope, const FRequestedType& RequestedType, FEmitValueShaderResult& OutResult) const
-{
-	const bool bSupportedOnMobile = SceneTextureId == PPI_PostProcessInput0 ||
-		SceneTextureId == PPI_CustomDepth ||
-		SceneTextureId == PPI_SceneDepth ||
-		SceneTextureId == PPI_CustomStencil;
-
-	FEmitShaderExpression* EmitTexCoord = nullptr;
-	if (TexCoordExpression)
-	{
-		EmitTexCoord = TexCoordExpression->GetValueShader(Context, Scope, Shader::EValueType::Float2);
-		EmitTexCoord = Context.EmitExpression(Scope, Shader::EValueType::Float2, TEXT("ClampSceneTextureUV(ViewportUVToSceneTextureUV(%, %), %)"), EmitTexCoord, (int)SceneTextureId, (int)SceneTextureId);
-	}
-	else
-	{
-		EmitTexCoord = Context.EmitExpression(Scope, Shader::EValueType::Float2, TEXT("GetDefaultSceneTextureUV(Parameters, %)"), (int)SceneTextureId);
-	}
-
-	FEmitShaderExpression* EmitLookup = nullptr;
-	if (Context.Material->GetFeatureLevel() >= ERHIFeatureLevel::SM5)
-	{
-		EmitLookup = Context.EmitExpression(Scope, Shader::EValueType::Float4, TEXT("SceneTextureLookup(%, %, %)"), EmitTexCoord, (int)SceneTextureId, bFiltered);
-	}
-	else
-	{
-		EmitLookup = Context.EmitExpression(Scope, Shader::EValueType::Float4, TEXT("MobileSceneTextureLookup(Parameters, %, %, %)"), (int)SceneTextureId, EmitTexCoord);
-	}
-
-	if (SceneTextureId == PPI_PostProcessInput0 && Context.Material->GetMaterialDomain() == MD_PostProcess && Context.Material->GetBlendableLocation() != BL_AfterTonemapping)
-	{
-		EmitLookup = Context.EmitExpression(Scope, Shader::EValueType::Float4, TEXT("(float4(View.OneOverPreExposure.xxx, 1) * %)"), EmitLookup);
-	}
-	
-	OutResult.Code = EmitLookup;
-}
-
-bool FExpressionMaterialNoise::PrepareValue(FEmitContext& Context, FEmitScope& Scope, const FRequestedType& RequestedType, FPrepareValueResult& OutResult) const
-{
-	const FPreparedType& PositionType = Context.PrepareExpression(PositionExpression, Scope, ERequestedType::Vector3);
-	const FPreparedType& FilterWidthType = Context.PrepareExpression(FilterWidthExpression, Scope, ERequestedType::Scalar);
-	if (PositionType.IsVoid() || FilterWidthType.IsVoid())
-	{
-		return false;
-	}
-	return OutResult.SetType(Context, RequestedType, EExpressionEvaluation::Shader, Shader::EValueType::Float1);
-}
-
-void FExpressionMaterialNoise::EmitValueShader(FEmitContext& Context, FEmitScope& Scope, const FRequestedType& RequestedType, FEmitValueShaderResult& OutResult) const
-{
-	bool bIsLWC = Shader::IsLWCType(PositionExpression->GetPreparedType().ValueComponentType);
-	FEmitShaderExpression* EmitPosition = PositionExpression->GetValueShader(Context, Scope, bIsLWC ? Shader::EValueType::Double3 : Shader::EValueType::Float3);
-	FEmitShaderExpression* EmitFilterWidth = FilterWidthExpression->GetValueShader(Context, Scope, Shader::EValueType::Float1);
-
-	if (bIsLWC)
-	{
-		// If Noise is driven by a LWC position, just take the offset within the current tile
-		// Will generate discontinuity in noise at tile boudaries
-		// Could potentially add noise functions that operate directly on LWC values, but that would be very expensive
-		EmitPosition = Context.EmitExpression(Scope, Shader::EValueType::Float3, TEXT("LWCNormalizeTile(%).Offset"), EmitPosition);
-	}
-
-	OutResult.Code = Context.EmitExpression(Scope, Shader::EValueType::Float1, TEXT("MaterialExpressionNoise(%,%,%,%,%,%,%,%,%,%,%,%)"),
-		EmitPosition,
-		Parameters.Scale,
-		Parameters.Quality,
-		Parameters.NoiseFunction,
-		Parameters.bTurbulence,
-		Parameters.Levels,
-		Parameters.OutputMin,
-		Parameters.OutputMax,
-		Parameters.LevelScale,
-		EmitFilterWidth,
-		Parameters.bTiling,
-		Parameters.RepeatSize);
-}
-
 namespace Private
 {
 Shader::EValueType GetTexCoordType(Shader::EValueType TextureType)
@@ -1572,5 +1009,4 @@ void FStatementLoop::EmitPreshader(FEmitContext& Context, FEmitScope& Scope, con
 	OutPreshader.SetLabel(PreshaderLoopScope.BreakLabel);
 }
 
-} // namespace HLSLTree
-} // namespace UE
+} // namespace UE::HLSLTree
diff --git a/Engine/Source/Runtime/Engine/Private/HLSLTree/HLSLTreeEmit.cpp b/Engine/Source/Runtime/Engine/Private/HLSLTree/HLSLTreeEmit.cpp
--- a/Engine/Source/Runtime/Engine/Private/HLSLTree/HLSLTreeEmit.cpp
+++ b/Engine/Source/Runtime/Engine/Private/HLSLTree/HLSLTreeEmit.cpp
@@ -9,9 +9,7 @@
 #include "MaterialShared.h" // TODO - split preshader out into its own module
 #include "Misc/LargeWorldRenderPosition.h"
 
-namespace UE
-{
-namespace HLSLTree
+namespace UE::HLSLTree
 {
 
 FEmitShaderNode::FEmitShaderNode(FEmitScope& InScope, TArrayView<FEmitShaderNode*> InDependencies)
@@ -199,14 +197,25 @@ FEmitContext::FEmitContext(FMemStackBase& InAllocator, FErrorHandlerInterface& I
 	, Errors(&InErrors)
 	, TypeRegistry(&InTypeRegistry)
 {
-	for (int32 Index = 0; Index < SF_NumFrequencies; ++Index)
+}
+
+FEmitContext::~FEmitContext()
+{
+	for (auto& It : CustomDataMap)
 	{
-		ExternalInputMask[Index].Init(false, 256);
+		delete It.Value;
 	}
 }
 
-FEmitContext::~FEmitContext()
+void FEmitContext::InternalRegisterData(FXxHash64 Hash, FCustomDataWrapper* Data)
 {
+	CustomDataMap.Add(Hash, Data);
+}
+
+FEmitContext::FCustomDataWrapper* FEmitContext::InternalFindData(FXxHash64 Hash) const
+{
+	FCustomDataWrapper* const* PrevData = CustomDataMap.Find(Hash);
+	return PrevData ? *PrevData : nullptr;
 }
 
 namespace Private
@@ -1215,5 +1224,5 @@ void FEmitContext::Finalize()
 	MaterialCompilationOutput->UniformExpressionSet.UniformPreshaderBufferSize = (UniformPreshaderOffset + 3u) / 4u;
 }
 
-} // namespace HLSLTree
-} // namespace UE
+} // namespace UE::HLSLTree
+
diff --git a/Engine/Source/Runtime/Engine/Private/Materials/MaterialExpressionHLSL.cpp b/Engine/Source/Runtime/Engine/Private/Materials/MaterialExpressionHLSL.cpp
--- a/Engine/Source/Runtime/Engine/Private/Materials/MaterialExpressionHLSL.cpp
+++ b/Engine/Source/Runtime/Engine/Private/Materials/MaterialExpressionHLSL.cpp
@@ -3,11 +3,6 @@
 
 #if WITH_EDITOR
 
-#include "MaterialHLSLGenerator.h"
-#include "Misc/MemStackUtility.h"
-#include "HLSLTree/HLSLTree.h"
-#include "HLSLTree/HLSLTreeCommon.h"
-
 #include "Materials/MaterialExpression.h"
 #include "Materials/MaterialExpressionExecBegin.h"
 #include "Materials/MaterialExpressionExecEnd.h"
@@ -136,10 +131,15 @@
 #include "Materials/MaterialExpressionForLoop.h"
 #include "Materials/MaterialExpressionWhileLoop.h"
 #include "Materials/MaterialFunctionInterface.h"
+#include "Materials/MaterialHLSLTree.h"
+#include "HLSLTree/HLSLTree.h"
+#include "HLSLTree/HLSLTreeCommon.h"
 #include "Engine/Engine.h"
 #include "Engine/Texture2D.h"
 #include "Engine/Font.h"
 #include "RenderUtils.h"
+#include "Misc/MemStackUtility.h"
+#include "MaterialHLSLGenerator.h"
 
 bool UMaterialExpression::GenerateHLSLExpression(FMaterialHLSLGenerator& Generator, UE::HLSLTree::FScope& Scope, int32 OutputIndex, UE::HLSLTree::FExpression*& OutExpression)
 {
@@ -207,8 +207,8 @@ bool UMaterialExpressionConstantBiasScale::GenerateHLSLExpression(FMaterialHLSLG
 
 bool UMaterialExpressionShadingModel::GenerateHLSLExpression(FMaterialHLSLGenerator& Generator, UE::HLSLTree::FScope& Scope, int32 OutputIndex, UE::HLSLTree::FExpression*& OutExpression)
 {
-	using namespace UE::HLSLTree;
-	OutExpression = Generator.GetTree().NewExpression<FExpressionMaterialShadingModel>(ShadingModel);
+	using namespace UE::HLSLTree::Material;
+	OutExpression = Generator.GetTree().NewExpression<FExpressionShadingModel>(ShadingModel);
 	return true;
 }
 
@@ -328,14 +328,14 @@ bool UMaterialExpressionStaticSwitchParameter::GenerateHLSLExpression(FMaterialH
 
 bool UMaterialExpressionPixelDepth::GenerateHLSLExpression(FMaterialHLSLGenerator& Generator, UE::HLSLTree::FScope& Scope, int32 OutputIndex, UE::HLSLTree::FExpression*& OutExpression)
 {
-	using namespace UE::HLSLTree;
+	using namespace UE::HLSLTree::Material;
 	OutExpression = Generator.GetTree().NewExpression<FExpressionExternalInput>(EExternalInput::PixelDepth);
 	return true;
 }
 
 bool UMaterialExpressionWorldPosition::GenerateHLSLExpression(FMaterialHLSLGenerator& Generator, UE::HLSLTree::FScope& Scope, int32 OutputIndex, UE::HLSLTree::FExpression*& OutExpression)
 {
-	using namespace UE::HLSLTree;
+	using namespace UE::HLSLTree::Material;
 
 	EExternalInput InputType = EExternalInput::None;
 	switch (WorldPositionShaderOffset)
@@ -353,21 +353,21 @@ bool UMaterialExpressionWorldPosition::GenerateHLSLExpression(FMaterialHLSLGener
 
 bool UMaterialExpressionCameraPositionWS::GenerateHLSLExpression(FMaterialHLSLGenerator& Generator, UE::HLSLTree::FScope& Scope, int32 OutputIndex, UE::HLSLTree::FExpression*& OutExpression)
 {
-	using namespace UE::HLSLTree;
+	using namespace UE::HLSLTree::Material;
 	OutExpression = Generator.GetTree().NewExpression<FExpressionExternalInput>(EExternalInput::CameraWorldPosition);
 	return true;
 }
 
 bool UMaterialExpressionCameraVectorWS::GenerateHLSLExpression(FMaterialHLSLGenerator& Generator, UE::HLSLTree::FScope& Scope, int32 OutputIndex, UE::HLSLTree::FExpression*& OutExpression)
 {
-	using namespace UE::HLSLTree;
+	using namespace UE::HLSLTree::Material;
 	OutExpression = Generator.GetTree().NewExpression<FExpressionExternalInput>(EExternalInput::CameraVector);
 	return true;
 }
 
 bool UMaterialExpressionViewProperty::GenerateHLSLExpression(FMaterialHLSLGenerator& Generator, UE::HLSLTree::FScope& Scope, int32 OutputIndex, UE::HLSLTree::FExpression*& OutExpression)
 {
-	using namespace UE::HLSLTree;
+	using namespace UE::HLSLTree::Material;
 
 	const bool bRequestedRcp = (OutputIndex == 1);
 	bool bApplyRcp = false;
@@ -412,7 +412,7 @@ bool UMaterialExpressionViewProperty::GenerateHLSLExpression(FMaterialHLSLGenera
 
 bool UMaterialExpressionTime::GenerateHLSLExpression(FMaterialHLSLGenerator& Generator, UE::HLSLTree::FScope& Scope, int32 OutputIndex, UE::HLSLTree::FExpression*& OutExpression)
 {
-	using namespace UE::HLSLTree;
+	using namespace UE::HLSLTree::Material;
 
 	if (bOverride_Period && Period == 0.0f)
 	{
@@ -431,13 +431,14 @@ bool UMaterialExpressionTime::GenerateHLSLExpression(FMaterialHLSLGenerator& Gen
 
 bool UMaterialExpressionDeltaTime::GenerateHLSLExpression(FMaterialHLSLGenerator& Generator, UE::HLSLTree::FScope& Scope, int32 OutputIndex, UE::HLSLTree::FExpression*& OutExpression)
 {
-	OutExpression = Generator.GetTree().NewExpression<UE::HLSLTree::FExpressionExternalInput>(UE::HLSLTree::EExternalInput::DeltaTime);
+	using namespace UE::HLSLTree::Material;
+	OutExpression = Generator.GetTree().NewExpression<FExpressionExternalInput>(EExternalInput::DeltaTime);
 	return true;
 }
 
 bool UMaterialExpressionScreenPosition::GenerateHLSLExpression(FMaterialHLSLGenerator& Generator, UE::HLSLTree::FScope& Scope, int32 OutputIndex, UE::HLSLTree::FExpression*& OutExpression)
 {
-	using namespace UE::HLSLTree;
+	using namespace UE::HLSLTree::Material;
 	const EExternalInput InputType = (OutputIndex == 1) ? EExternalInput::PixelPosition : EExternalInput::ViewportUV;
 	OutExpression = Generator.GetTree().NewExpression<FExpressionExternalInput>(InputType);
 	return true;
@@ -445,7 +446,7 @@ bool UMaterialExpressionScreenPosition::GenerateHLSLExpression(FMaterialHLSLGene
 
 bool UMaterialExpressionSceneTexelSize::GenerateHLSLExpression(FMaterialHLSLGenerator& Generator, UE::HLSLTree::FScope& Scope, int32 OutputIndex, UE::HLSLTree::FExpression*& OutExpression)
 {
-	using namespace UE::HLSLTree;
+	using namespace UE::HLSLTree::Material;
 
 	// To make sure any material that were correctly handling BufferUV != ViewportUV, we just lie to material
 	// to make it believe ViewSize == BufferSize, so they are still compatible with SceneTextureLookup().
@@ -455,21 +456,21 @@ bool UMaterialExpressionSceneTexelSize::GenerateHLSLExpression(FMaterialHLSLGene
 
 bool UMaterialExpressionViewSize::GenerateHLSLExpression(FMaterialHLSLGenerator& Generator, UE::HLSLTree::FScope& Scope, int32 OutputIndex, UE::HLSLTree::FExpression*& OutExpression)
 {
-	using namespace UE::HLSLTree;
+	using namespace UE::HLSLTree::Material;
 	OutExpression = Generator.GetTree().NewExpression<FExpressionExternalInput>(EExternalInput::ViewSize);
 	return true;
 }
 
 bool UMaterialExpressionTwoSidedSign::GenerateHLSLExpression(FMaterialHLSLGenerator& Generator, UE::HLSLTree::FScope& Scope, int32 OutputIndex, UE::HLSLTree::FExpression*& OutExpression)
 {
-	using namespace UE::HLSLTree;
+	using namespace UE::HLSLTree::Material;
 	OutExpression = Generator.GetTree().NewExpression<FExpressionExternalInput>(EExternalInput::TwoSidedSign);
 	return true;
 }
 
 bool UMaterialExpressionReflectionVectorWS::GenerateHLSLExpression(FMaterialHLSLGenerator& Generator, UE::HLSLTree::FScope& Scope, int32 OutputIndex, UE::HLSLTree::FExpression*& OutExpression)
 {
-	using namespace UE::HLSLTree;
+	using namespace UE::HLSLTree::Material;
 	check(!CustomWorldNormal.GetTracedInput().Expression); // TODO
 
 	OutExpression = Generator.GetTree().NewExpression<FExpressionExternalInput>(EExternalInput::WorldReflection);
@@ -478,35 +479,35 @@ bool UMaterialExpressionReflectionVectorWS::GenerateHLSLExpression(FMaterialHLSL
 
 bool UMaterialExpressionActorPositionWS::GenerateHLSLExpression(FMaterialHLSLGenerator& Generator, UE::HLSLTree::FScope& Scope, int32 OutputIndex, UE::HLSLTree::FExpression*& OutExpression)
 {
-	using namespace UE::HLSLTree;
+	using namespace UE::HLSLTree::Material;
 	OutExpression = Generator.GetTree().NewExpression<FExpressionExternalInput>(EExternalInput::ActorWorldPosition);
 	return true;
 }
 
 bool UMaterialExpressionPixelNormalWS::GenerateHLSLExpression(FMaterialHLSLGenerator& Generator, UE::HLSLTree::FScope& Scope, int32 OutputIndex, UE::HLSLTree::FExpression*& OutExpression)
 {
-	using namespace UE::HLSLTree;
+	using namespace UE::HLSLTree::Material;
 	OutExpression = Generator.GetTree().NewExpression<FExpressionExternalInput>(EExternalInput::WorldNormal);
 	return true;
 }
 
 bool UMaterialExpressionVertexColor::GenerateHLSLExpression(FMaterialHLSLGenerator& Generator, UE::HLSLTree::FScope& Scope, int32 OutputIndex, UE::HLSLTree::FExpression*& OutExpression)
 {
-	using namespace UE::HLSLTree;
+	using namespace UE::HLSLTree::Material;
 	OutExpression = Generator.GetTree().NewExpression<FExpressionExternalInput>(EExternalInput::VertexColor);
 	return true;
 }
 
 bool UMaterialExpressionVertexNormalWS::GenerateHLSLExpression(FMaterialHLSLGenerator& Generator, UE::HLSLTree::FScope& Scope, int32 OutputIndex, UE::HLSLTree::FExpression*& OutExpression)
 {
-	using namespace UE::HLSLTree;
+	using namespace UE::HLSLTree::Material;
 	OutExpression = Generator.GetTree().NewExpression<FExpressionExternalInput>(EExternalInput::WorldVertexNormal);
 	return true;
 }
 
 bool UMaterialExpressionVertexTangentWS::GenerateHLSLExpression(FMaterialHLSLGenerator& Generator, UE::HLSLTree::FScope& Scope, int32 OutputIndex, UE::HLSLTree::FExpression*& OutExpression)
 {
-	using namespace UE::HLSLTree;
+	using namespace UE::HLSLTree::Material;
 	OutExpression = Generator.GetTree().NewExpression<FExpressionExternalInput>(EExternalInput::WorldVertexTangent);
 	return true;
 }
@@ -536,21 +537,21 @@ bool UMaterialExpressionTextureCoordinate::GenerateHLSLExpression(FMaterialHLSLG
 
 bool UMaterialExpressionLightmapUVs::GenerateHLSLExpression(FMaterialHLSLGenerator& Generator, UE::HLSLTree::FScope& Scope, int32 OutputIndex, UE::HLSLTree::FExpression*& OutExpression)
 {
-	using namespace UE::HLSLTree;
+	using namespace UE::HLSLTree::Material;
 	OutExpression = Generator.GetTree().NewExpression<FExpressionExternalInput>(EExternalInput::LightmapTexCoord);
 	return true;
 }
 
 bool UMaterialExpressionParticleColor::GenerateHLSLExpression(FMaterialHLSLGenerator& Generator, UE::HLSLTree::FScope& Scope, int32 OutputIndex, UE::HLSLTree::FExpression*& OutExpression)
 {
-	using namespace UE::HLSLTree;
+	using namespace UE::HLSLTree::Material;
 	OutExpression = Generator.GetTree().NewExpression<FExpressionExternalInput>(EExternalInput::ParticleColor);
 	return true;
 }
 
 bool UMaterialExpressionParticlePositionWS::GenerateHLSLExpression(FMaterialHLSLGenerator& Generator, UE::HLSLTree::FScope& Scope, int32 OutputIndex, UE::HLSLTree::FExpression*& OutExpression)
 {
-	using namespace UE::HLSLTree;
+	using namespace UE::HLSLTree::Material;
 	OutExpression = Generator.GetTree().NewExpression<FExpressionExternalInput>(EExternalInput::ParticleTranslatedWorldPosition);
 	OutExpression = Generator.GetTree().NewSub(OutExpression, Generator.GetTree().NewExpression<FExpressionExternalInput>(EExternalInput::PreViewTranslation));
 	return true;
@@ -558,7 +559,7 @@ bool UMaterialExpressionParticlePositionWS::GenerateHLSLExpression(FMaterialHLSL
 
 bool UMaterialExpressionParticleRadius::GenerateHLSLExpression(FMaterialHLSLGenerator& Generator, UE::HLSLTree::FScope& Scope, int32 OutputIndex, UE::HLSLTree::FExpression*& OutExpression)
 {
-	using namespace UE::HLSLTree;
+	using namespace UE::HLSLTree::Material;
 	OutExpression = Generator.GetTree().NewExpression<FExpressionExternalInput>(EExternalInput::ParticleRadius);
 	OutExpression = Generator.GetTree().NewMax(OutExpression, Generator.NewConstant(0.0f));
 	return true;
@@ -730,7 +731,7 @@ bool UMaterialExpressionSceneTexture::GenerateHLSLExpression(FMaterialHLSLGenera
 		{
 			ExpressionTexCoord = Coordinates.AcquireHLSLExpression(Generator, Scope);
 		}
-		OutExpression = Generator.GetTree().NewExpression<FExpressionMaterialSceneTexture>(ExpressionTexCoord, SceneTextureId, bFiltered);
+		OutExpression = Generator.GetTree().NewExpression<Material::FExpressionSceneTexture>(ExpressionTexCoord, SceneTextureId, bFiltered);
 		return true;
 	}
 	else if (OutputIndex == 1 || OutputIndex == 2)
@@ -754,11 +755,11 @@ bool UMaterialExpressionNoise::GenerateHLSLExpression(FMaterialHLSLGenerator& Ge
 	FExpression* ExpressionPosition = Position.TryAcquireHLSLExpression(Generator, Scope);
 	if (!ExpressionPosition)
 	{
-		ExpressionPosition = Generator.GetTree().NewExpression<FExpressionExternalInput>(EExternalInput::WorldPosition);
+		ExpressionPosition = Generator.GetTree().NewExpression<Material::FExpressionExternalInput>(Material::EExternalInput::WorldPosition);
 	}
 	FExpression* ExpressionFilterWidth = FilterWidth.AcquireHLSLExpressionOrConstant(Generator, Scope, 0.0f);
 
-	FMaterialNoiseParameters NoiseParameters;
+	Material::FNoiseParameters NoiseParameters;
 	NoiseParameters.Quality = Quality;
 	// to limit performance problems due to values outside reasonable range
 	NoiseParameters.Levels = FMath::Clamp(Levels, 1, 10);
@@ -771,7 +772,7 @@ bool UMaterialExpressionNoise::GenerateHLSLExpression(FMaterialHLSLGenerator& Ge
 	NoiseParameters.bTiling = bTiling;
 	NoiseParameters.bTurbulence = bTurbulence;
 
-	OutExpression = Generator.GetTree().NewExpression<FExpressionMaterialNoise>(NoiseParameters, ExpressionPosition, ExpressionFilterWidth);
+	OutExpression = Generator.GetTree().NewExpression<Material::FExpressionNoise>(NoiseParameters, ExpressionPosition, ExpressionFilterWidth);
 	return true;
 }
 
@@ -1596,6 +1597,8 @@ static UE::HLSLTree::FExpression* TransformBase(UE::HLSLTree::FTree& Tree,
 	bool bWComponent)
 {
 	using namespace UE::HLSLTree;
+	using namespace UE::HLSLTree::Material;
+
 	if (!Input)
 	{
 		// unable to compile
@@ -1873,10 +1876,10 @@ bool UMaterialExpressionFresnel::GenerateHLSLExpression(FMaterialHLSLGenerator&
 	FExpression* ExpressionNormal = Normal.TryAcquireHLSLExpression(Generator, Scope);
 	if (!ExpressionNormal)
 	{
-		ExpressionNormal = Generator.GetTree().NewExpression<FExpressionExternalInput>(EExternalInput::WorldNormal);
+		ExpressionNormal = Generator.GetTree().NewExpression<Material::FExpressionExternalInput>(Material::EExternalInput::WorldNormal);
 	}
 
-	FExpression* ExpressionNdotV = Generator.GetTree().NewDot(ExpressionNormal, Generator.GetTree().NewExpression<FExpressionExternalInput>(EExternalInput::CameraVector));
+	FExpression* ExpressionNdotV = Generator.GetTree().NewDot(ExpressionNormal, Generator.GetTree().NewExpression<Material::FExpressionExternalInput>(Material::EExternalInput::CameraVector));
 	FExpression* ExpressionMax = Generator.GetTree().NewMax(Generator.NewConstant(0.0f), ExpressionNdotV);
 	FExpression* ExpressionMinus = Generator.GetTree().NewSub(Generator.NewConstant(1.0f), ExpressionMax);
 
@@ -1973,7 +1976,7 @@ bool UMaterialExpressionPanner::GenerateHLSLExpression(FMaterialHLSLGenerator& G
 	FExpression* ExpressionTime = Time.TryAcquireHLSLExpression(Generator, Scope);
 	if (!ExpressionTime)
 	{
-		ExpressionTime = Generator.GetTree().NewExpression<FExpressionExternalInput>(EExternalInput::GameTime);
+		ExpressionTime = Generator.GetTree().NewExpression<Material::FExpressionExternalInput>(Material::EExternalInput::GameTime);
 	}
 	FExpression* ExpressionSpeed = Speed.AcquireHLSLExpressionOrConstant(Generator, Scope, FVector2f(SpeedX, SpeedY));
 	FExpression* ExpressionOffset = Generator.GetTree().NewMul(ExpressionSpeed, ExpressionTime);
@@ -1998,7 +2001,7 @@ bool UMaterialExpressionRotator::GenerateHLSLExpression(FMaterialHLSLGenerator&
 	FExpression* ExpressionTime = Time.TryAcquireHLSLExpression(Generator, Scope);
 	if (!ExpressionTime)
 	{
-		ExpressionTime = Generator.GetTree().NewExpression<FExpressionExternalInput>(EExternalInput::GameTime);
+		ExpressionTime = Generator.GetTree().NewExpression<Material::FExpressionExternalInput>(Material::EExternalInput::GameTime);
 	}
 	ExpressionTime = Generator.GetTree().NewMul(ExpressionTime, Generator.NewConstant(Speed));
 
diff --git a/Engine/Source/Runtime/Engine/Private/Materials/MaterialHLSLEmitter.cpp b/Engine/Source/Runtime/Engine/Private/Materials/MaterialHLSLEmitter.cpp
--- a/Engine/Source/Runtime/Engine/Private/Materials/MaterialHLSLEmitter.cpp
+++ b/Engine/Source/Runtime/Engine/Private/Materials/MaterialHLSLEmitter.cpp
@@ -11,6 +11,7 @@
 #include "Materials/MaterialExpressionFunctionOutput.h"
 #include "Materials/MaterialExpressionVolumetricAdvancedMaterialOutput.h"
 #include "Materials/Material.h"
+#include "Materials/MaterialHLSLTree.h"
 #include "MaterialHLSLGenerator.h"
 #include "ShaderCore.h"
 #include "HLSLTree/HLSLTree.h"
@@ -55,6 +56,8 @@ static FString GenerateMaterialTemplateHLSL(EShaderPlatform ShaderPlatform,
 {
 	using namespace UE::HLSLTree;
 
+	const Material::FEmitData& EmitMaterialData = EmitContext.FindData<Material::FEmitData>();
+
 	FString MaterialTemplateSource;
 	LoadShaderSourceFileChecked(TEXT("/Engine/Private/MaterialTemplate.ush"), ShaderPlatform, MaterialTemplateSource);
 
@@ -80,9 +83,9 @@ static FString GenerateMaterialTemplateHLSL(EShaderPlatform ShaderPlatform,
 	FLazyPrintf LazyPrintf(*MaterialTemplateSource);
 
 	uint32 NumTexCoords = 0u;
-	for (int32 TexCoordIndex = UE::HLSLTree::MaxNumTexCoords - 1; TexCoordIndex >= 0; --TexCoordIndex)
+	for (int32 TexCoordIndex = Material::MaxNumTexCoords - 1; TexCoordIndex >= 0; --TexCoordIndex)
 	{
-		if (EmitContext.IsExternalInputUsed(SF_Pixel, MakeInputTexCoord(TexCoordIndex)))
+		if (EmitMaterialData.IsExternalInputUsed(SF_Pixel, Material::MakeInputTexCoord(TexCoordIndex)))
 		{
 			NumTexCoords = TexCoordIndex + 1;
 			break;
@@ -310,9 +313,9 @@ static FString GenerateMaterialTemplateHLSL(EShaderPlatform ShaderPlatform,
 		EvaluateMaterialDeclaration += TEXT("    FMaterialAttributes Result = EvaluateVertexMaterialAttributesInternal(Parameters);" LINE_TERMINATOR);
 		for (uint32 TexCoordIndex = 0; TexCoordIndex < NumTexCoords; ++TexCoordIndex)
 		{
-			const EExternalInput TexCoordInput = MakeInputTexCoord(TexCoordIndex);
+			const Material::EExternalInput TexCoordInput = Material::MakeInputTexCoord(TexCoordIndex);
 
-			if (EmitContext.IsExternalInputUsed(SF_Pixel, TexCoordInput) && !EmitContext.IsExternalInputUsed(SF_Vertex, TexCoordInput))
+			if (EmitMaterialData.IsExternalInputUsed(SF_Pixel, TexCoordInput) && !EmitMaterialData.IsExternalInputUsed(SF_Vertex, TexCoordInput))
 			{
 				const FString AttributeName = FMaterialAttributeDefinitionMap::GetAttributeName((EMaterialProperty)(MP_CustomizedUVs0 + TexCoordIndex));
 				EvaluateMaterialDeclaration += FString::Printf(TEXT("    Result.%s = Parameters.TexCoords[%d];") LINE_TERMINATOR, *AttributeName, TexCoordIndex);
@@ -404,6 +407,10 @@ static void GetMaterialEnvironment(EShaderPlatform InPlatform,
 	const FMaterialCompilationOutput& MaterialCompilationOutput,
 	FShaderCompilerEnvironment& OutEnvironment)
 {
+	using namespace UE::HLSLTree;
+
+	const Material::FEmitData& EmitMaterialData = EmitContext.FindData<Material::FEmitData>();
+
 	bool bMaterialRequestsDualSourceBlending = false;
 
 	OutEnvironment.SetDefine(TEXT("ENABLE_NEW_HLSL_GENERATOR"), 1);
@@ -450,7 +457,7 @@ static void GetMaterialEnvironment(EShaderPlatform InPlatform,
 		OutEnvironment.SetDefine(TEXT("USE_PARTICLE_SUBUVS"), TEXT("1"));
 	}
 
-	if (EmitContext.ExternalInputMask[SF_Pixel][(int32)UE::HLSLTree::EExternalInput::LightmapTexCoord])
+	if (EmitMaterialData.ExternalInputMask[SF_Pixel][(int32)Material::EExternalInput::LightmapTexCoord])
 	{
 		OutEnvironment.SetDefine(TEXT("LIGHTMAP_UV_ACCESS"), TEXT("1"));
 	}
@@ -496,7 +503,7 @@ static void GetMaterialEnvironment(EShaderPlatform InPlatform,
 	// @todo MetalMRT: Remove this hack and implement proper atmospheric-fog solution for Metal MRT...
 	OutEnvironment.SetDefine(TEXT("MATERIAL_ATMOSPHERIC_FOG"), false);// !IsMetalMRTPlatform(InPlatform) ? bUsesAtmosphericFog : 0);
 	OutEnvironment.SetDefine(TEXT("MATERIAL_SKY_ATMOSPHERE"), false);// bUsesSkyAtmosphere);
-	OutEnvironment.SetDefine(TEXT("INTERPOLATE_VERTEX_COLOR"), EmitContext.bUsesVertexColor);
+	OutEnvironment.SetDefine(TEXT("INTERPOLATE_VERTEX_COLOR"), EmitMaterialData.bUsesVertexColor);
 	OutEnvironment.SetDefine(TEXT("NEEDS_PARTICLE_COLOR"), false);// bUsesParticleColor);
 	OutEnvironment.SetDefine(TEXT("NEEDS_PARTICLE_LOCAL_TO_WORLD"), false);// bUsesParticleLocalToWorld);
 	OutEnvironment.SetDefine(TEXT("NEEDS_PARTICLE_WORLD_TO_LOCAL"), false);// bUsesParticleWorldToLocal);
@@ -553,10 +560,10 @@ static void GetMaterialEnvironment(EShaderPlatform InPlatform,
 	// If the material gets its shading model from the material expressions, then we use the result from the compilation (assuming it's valid).
 	// This result will potentially be tighter than what GetShadingModels() returns, because it only picks up the shading models from the expressions that get compiled for a specific feature level and quality level
 	// For example, the material might have shading models behind static switches. GetShadingModels() will return both the true and the false paths from that switch, whereas the shading model field from the compilation will only contain the actual shading model selected 
-	if (InMaterial.IsShadingModelFromMaterialExpression() && EmitContext.ShadingModelsFromCompilation.IsValid())
+	if (InMaterial.IsShadingModelFromMaterialExpression() && EmitMaterialData.ShadingModelsFromCompilation.IsValid())
 	{
 		// Shading models fetched from the compilation of the expression graph
-		ShadingModels = EmitContext.ShadingModelsFromCompilation;
+		ShadingModels = EmitMaterialData.ShadingModelsFromCompilation;
 	}
 	ensure(ShadingModels.IsValid());
 
@@ -810,9 +817,11 @@ bool MaterialEmitHLSL(const FMaterialCompileTargetParameters& InCompilerTarget,
 
 	FEmitContext EmitContext(Allocator, Generator.GetErrors(), TypeRegistry);
 	EmitContext.Material = &InOutMaterial;
-	EmitContext.StaticParameters = &InStaticParameters;
 	EmitContext.MaterialCompilationOutput = &OutCompilationOutput;
 
+	Material::FEmitData& EmitMaterialData = EmitContext.AcquireData<Material::FEmitData>();
+	EmitMaterialData.StaticParameters = &InStaticParameters;
+
 	const FStructField* NormalField = Generator.GetMaterialAttributesType()->FindFieldByName(*FMaterialAttributeDefinitionMap::GetAttributeName(MP_Normal));
 	check(NormalField);
 
@@ -853,7 +862,7 @@ bool MaterialEmitHLSL(const FMaterialCompileTargetParameters& InCompilerTarget,
 			bUsesPixelDepthOffset = true;
 		}
 
-		if (!EmitContext.bReadMaterialNormal)
+		if (!EmitMaterialData.bReadMaterialNormal)
 		{
 			// No access to material normal, can execute everything in phase0
 			RequestedPixelAttributesType.SetFieldRequested(NormalField);
@@ -951,8 +960,8 @@ bool MaterialEmitHLSL(const FMaterialCompileTargetParameters& InCompilerTarget,
 		};
 		const TCHAR* InputPixelCodePhase1[2] =
 		{
-			EmitContext.bReadMaterialNormal ? PixelCodePhase1[0]->ToString() : nullptr,
-			EmitContext.bReadMaterialNormal ? PixelCodePhase1[1]->ToString() : nullptr,
+			EmitMaterialData.bReadMaterialNormal ? PixelCodePhase1[0]->ToString() : nullptr,
+			EmitMaterialData.bReadMaterialNormal ? PixelCodePhase1[1]->ToString() : nullptr,
 		};
 		MaterialTemplateSource = GenerateMaterialTemplateHLSL(InCompilerTarget.ShaderPlatform,
 			InOutMaterial,
diff --git a/Engine/Source/Runtime/Engine/Private/Materials/MaterialHLSLGenerator.cpp b/Engine/Source/Runtime/Engine/Private/Materials/MaterialHLSLGenerator.cpp
--- a/Engine/Source/Runtime/Engine/Private/Materials/MaterialHLSLGenerator.cpp
+++ b/Engine/Source/Runtime/Engine/Private/Materials/MaterialHLSLGenerator.cpp
@@ -12,6 +12,7 @@
 #include "Materials/MaterialExpressionVolumetricAdvancedMaterialOutput.h"
 #include "Materials/MaterialExpressionExecBegin.h"
 #include "Materials/Material.h"
+#include "Materials/MaterialHLSLTree.h"
 #include "ShaderCore.h"
 #include "HLSLTree/HLSLTree.h"
 #include "HLSLTree/HLSLTreeCommon.h"
@@ -458,7 +459,8 @@ UE::HLSLTree::FExpression* FMaterialHLSLGenerator::NewConstant(const UE::Shader:
 
 UE::HLSLTree::FExpression* FMaterialHLSLGenerator::NewTexCoord(int32 Index)
 {
-	return HLSLTree->NewExpression<UE::HLSLTree::FExpressionExternalInput>(UE::HLSLTree::MakeInputTexCoord(Index));
+	using namespace UE::HLSLTree::Material;
+	return HLSLTree->NewExpression<FExpressionExternalInput>(MakeInputTexCoord(Index));
 }
 
 UE::HLSLTree::FExpression* FMaterialHLSLGenerator::NewSwizzle(const UE::HLSLTree::FSwizzleParameters& Params, UE::HLSLTree::FExpression* Input)
@@ -633,7 +635,7 @@ UE::HLSLTree::FExpression* FMaterialHLSLGenerator::GenerateMaterialParameter(EMa
 		}
 	}
 
-	return HLSLTree->NewExpression<UE::HLSLTree::FExpressionMaterialParameter>(InType, GetParameterInfo(InParameterName), DefaultValue);
+	return HLSLTree->NewExpression<UE::HLSLTree::Material::FExpressionParameter>(InType, GetParameterInfo(InParameterName), DefaultValue);
 }
 
 UE::HLSLTree::FExpression* FMaterialHLSLGenerator::GenerateFunctionCall(UE::HLSLTree::FScope& Scope,
diff --git a/Engine/Source/Runtime/Engine/Private/Materials/MaterialHLSLTree.cpp b/Engine/Source/Runtime/Engine/Private/Materials/MaterialHLSLTree.cpp
new file mode 100644
--- /dev/null
+++ b/Engine/Source/Runtime/Engine/Private/Materials/MaterialHLSLTree.cpp
@@ -0,0 +1,577 @@
+// Copyright Epic Games, Inc. All Rights Reserved.
+#include "Materials/MaterialHLSLTree.h"
+#include "HLSLTree/HLSLTreeCommon.h"
+#include "HLSLTree/HLSLTreeEmit.h"
+#include "Misc/StringBuilder.h"
+#include "MaterialShared.h"
+#include "MaterialSceneTextureId.h"
+#include "Engine/BlendableInterface.h" // BL_AfterTonemapping
+#include "Engine/Texture.h"
+
+namespace UE::HLSLTree::Material
+{
+
+FExternalInputDescription GetExternalInputDescription(EExternalInput Input)
+{
+	switch (Input)
+	{
+	case EExternalInput::None: return FExternalInputDescription(TEXT("None"), Shader::EValueType::Void);
+
+	case EExternalInput::TexCoord0: return FExternalInputDescription(TEXT("TexCoord0"), Shader::EValueType::Float2, EExternalInput::TexCoord0_Ddx, EExternalInput::TexCoord0_Ddy);
+	case EExternalInput::TexCoord1: return FExternalInputDescription(TEXT("TexCoord1"), Shader::EValueType::Float2, EExternalInput::TexCoord1_Ddx, EExternalInput::TexCoord1_Ddy);
+	case EExternalInput::TexCoord2: return FExternalInputDescription(TEXT("TexCoord2"), Shader::EValueType::Float2, EExternalInput::TexCoord2_Ddx, EExternalInput::TexCoord2_Ddy);
+	case EExternalInput::TexCoord3: return FExternalInputDescription(TEXT("TexCoord3"), Shader::EValueType::Float2, EExternalInput::TexCoord3_Ddx, EExternalInput::TexCoord3_Ddy);
+	case EExternalInput::TexCoord4: return FExternalInputDescription(TEXT("TexCoord4"), Shader::EValueType::Float2, EExternalInput::TexCoord4_Ddx, EExternalInput::TexCoord4_Ddy);
+	case EExternalInput::TexCoord5: return FExternalInputDescription(TEXT("TexCoord5"), Shader::EValueType::Float2, EExternalInput::TexCoord5_Ddx, EExternalInput::TexCoord5_Ddy);
+	case EExternalInput::TexCoord6: return FExternalInputDescription(TEXT("TexCoord6"), Shader::EValueType::Float2, EExternalInput::TexCoord6_Ddx, EExternalInput::TexCoord6_Ddy);
+	case EExternalInput::TexCoord7: return FExternalInputDescription(TEXT("TexCoord7"), Shader::EValueType::Float2, EExternalInput::TexCoord7_Ddx, EExternalInput::TexCoord7_Ddy);
+
+	case EExternalInput::TexCoord0_Ddx: return FExternalInputDescription(TEXT("TexCoord0_Ddx"), Shader::EValueType::Float2);
+	case EExternalInput::TexCoord1_Ddx: return FExternalInputDescription(TEXT("TexCoord1_Ddx"), Shader::EValueType::Float2);
+	case EExternalInput::TexCoord2_Ddx: return FExternalInputDescription(TEXT("TexCoord2_Ddx"), Shader::EValueType::Float2);
+	case EExternalInput::TexCoord3_Ddx: return FExternalInputDescription(TEXT("TexCoord3_Ddx"), Shader::EValueType::Float2);
+	case EExternalInput::TexCoord4_Ddx: return FExternalInputDescription(TEXT("TexCoord4_Ddx"), Shader::EValueType::Float2);
+	case EExternalInput::TexCoord5_Ddx: return FExternalInputDescription(TEXT("TexCoord5_Ddx"), Shader::EValueType::Float2);
+	case EExternalInput::TexCoord6_Ddx: return FExternalInputDescription(TEXT("TexCoord6_Ddx"), Shader::EValueType::Float2);
+	case EExternalInput::TexCoord7_Ddx: return FExternalInputDescription(TEXT("TexCoord7_Ddx"), Shader::EValueType::Float2);
+
+	case EExternalInput::TexCoord0_Ddy: return FExternalInputDescription(TEXT("TexCoord0_Ddy"), Shader::EValueType::Float2);
+	case EExternalInput::TexCoord1_Ddy: return FExternalInputDescription(TEXT("TexCoord1_Ddy"), Shader::EValueType::Float2);
+	case EExternalInput::TexCoord2_Ddy: return FExternalInputDescription(TEXT("TexCoord2_Ddy"), Shader::EValueType::Float2);
+	case EExternalInput::TexCoord3_Ddy: return FExternalInputDescription(TEXT("TexCoord3_Ddy"), Shader::EValueType::Float2);
+	case EExternalInput::TexCoord4_Ddy: return FExternalInputDescription(TEXT("TexCoord4_Ddy"), Shader::EValueType::Float2);
+	case EExternalInput::TexCoord5_Ddy: return FExternalInputDescription(TEXT("TexCoord5_Ddy"), Shader::EValueType::Float2);
+	case EExternalInput::TexCoord6_Ddy: return FExternalInputDescription(TEXT("TexCoord6_Ddy"), Shader::EValueType::Float2);
+	case EExternalInput::TexCoord7_Ddy: return FExternalInputDescription(TEXT("TexCoord7_Ddy"), Shader::EValueType::Float2);
+
+	case EExternalInput::LightmapTexCoord: return FExternalInputDescription(TEXT("LightmapTexCoord"), Shader::EValueType::Float2, EExternalInput::LightmapTexCoord_Ddx, EExternalInput::LightmapTexCoord_Ddy);
+	case EExternalInput::LightmapTexCoord_Ddx: return FExternalInputDescription(TEXT("LightmapTexCoord_Ddx"), Shader::EValueType::Float2);
+	case EExternalInput::LightmapTexCoord_Ddy: return FExternalInputDescription(TEXT("LightmapTexCoord_Ddy"), Shader::EValueType::Float2);
+
+	case EExternalInput::TwoSidedSign: return FExternalInputDescription(TEXT("TwoSidedSign"), Shader::EValueType::Float1);
+	case EExternalInput::VertexColor: return FExternalInputDescription(TEXT("VertexColor"), Shader::EValueType::Float4, EExternalInput::VertexColor_Ddx, EExternalInput::VertexColor_Ddy);
+	case EExternalInput::VertexColor_Ddx: return FExternalInputDescription(TEXT("VertexColor_Ddx"), Shader::EValueType::Float4);
+	case EExternalInput::VertexColor_Ddy: return FExternalInputDescription(TEXT("VertexColor_Ddy"), Shader::EValueType::Float4);
+
+	case EExternalInput::WorldPosition: return FExternalInputDescription(TEXT("WorldPosition"), Shader::EValueType::Double3, EExternalInput::WorldPosition_Ddx, EExternalInput::WorldPosition_Ddy, EExternalInput::PrevWorldPosition);
+	case EExternalInput::WorldPosition_NoOffsets: return FExternalInputDescription(TEXT("WorldPosition_NoOffsets"), Shader::EValueType::Double3, EExternalInput::WorldPosition_Ddx, EExternalInput::WorldPosition_Ddy, EExternalInput::PrevWorldPosition_NoOffsets);
+	case EExternalInput::TranslatedWorldPosition: return FExternalInputDescription(TEXT("TranslatedWorldPosition"), Shader::EValueType::Float3, EExternalInput::WorldPosition_Ddx, EExternalInput::WorldPosition_Ddy, EExternalInput::PrevTranslatedWorldPosition);
+	case EExternalInput::TranslatedWorldPosition_NoOffsets: return FExternalInputDescription(TEXT("TranslatedWorldPosition_NoOffsets"), Shader::EValueType::Float3, EExternalInput::WorldPosition_Ddx, EExternalInput::WorldPosition_Ddy, EExternalInput::PrevTranslatedWorldPosition_NoOffsets);
+	case EExternalInput::ActorWorldPosition: return FExternalInputDescription(TEXT("TranslatedWorldPosition_NoOffsets"), Shader::EValueType::Double3);
+
+	case EExternalInput::PrevWorldPosition: return FExternalInputDescription(TEXT("PrevWorldPosition"), Shader::EValueType::Double3, EExternalInput::WorldPosition_Ddx, EExternalInput::WorldPosition_Ddy);
+	case EExternalInput::PrevWorldPosition_NoOffsets: return FExternalInputDescription(TEXT("PrevWorldPosition_NoOffsets"), Shader::EValueType::Double3, EExternalInput::WorldPosition_Ddx, EExternalInput::WorldPosition_Ddy);
+	case EExternalInput::PrevTranslatedWorldPosition: return FExternalInputDescription(TEXT("PrevTranslatedWorldPosition"), Shader::EValueType::Float3, EExternalInput::WorldPosition_Ddx, EExternalInput::WorldPosition_Ddy);
+	case EExternalInput::PrevTranslatedWorldPosition_NoOffsets: return FExternalInputDescription(TEXT("PrevTranslatedWorldPosition_NoOffsets"), Shader::EValueType::Float3, EExternalInput::WorldPosition_Ddx, EExternalInput::WorldPosition_Ddy);
+
+	case EExternalInput::WorldPosition_Ddx: return FExternalInputDescription(TEXT("WorldPosition_Ddx"), Shader::EValueType::Float3);
+	case EExternalInput::WorldPosition_Ddy: return FExternalInputDescription(TEXT("WorldPosition_Ddx"), Shader::EValueType::Float3);
+
+	case EExternalInput::WorldNormal: return FExternalInputDescription(TEXT("WorldNormal"), Shader::EValueType::Float3);
+	case EExternalInput::WorldReflection: return FExternalInputDescription(TEXT("WorldReflection"), Shader::EValueType::Float3);
+	case EExternalInput::WorldVertexNormal: return FExternalInputDescription(TEXT("WorldVertexNormal"), Shader::EValueType::Float3);
+	case EExternalInput::WorldVertexTangent: return FExternalInputDescription(TEXT("WorldVertexTangent"), Shader::EValueType::Float3);
+
+	case EExternalInput::ViewportUV: return FExternalInputDescription(TEXT("ViewportUV"), Shader::EValueType::Float2);
+	case EExternalInput::PixelPosition: return FExternalInputDescription(TEXT("PixelPosition"), Shader::EValueType::Float2);
+	case EExternalInput::ViewSize: return FExternalInputDescription(TEXT("ViewSize"), Shader::EValueType::Float2);
+	case EExternalInput::RcpViewSize: return FExternalInputDescription(TEXT("RcpViewSize"), Shader::EValueType::Float2);
+	case EExternalInput::FieldOfView: return FExternalInputDescription(TEXT("FieldOfView"), Shader::EValueType::Float1, EExternalInput::None, EExternalInput::None, EExternalInput::PrevFieldOfView);
+	case EExternalInput::TanHalfFieldOfView: return FExternalInputDescription(TEXT("TanHalfFieldOfView"), Shader::EValueType::Float2, EExternalInput::None, EExternalInput::None, EExternalInput::PrevTanHalfFieldOfView);
+	case EExternalInput::CotanHalfFieldOfView: return FExternalInputDescription(TEXT("CotanHalfFieldOfView"), Shader::EValueType::Float2, EExternalInput::None, EExternalInput::None, EExternalInput::PrevCotanHalfFieldOfView);
+	case EExternalInput::TemporalSampleCount: return FExternalInputDescription(TEXT("TemporalSampleCount"), Shader::EValueType::Float1);
+	case EExternalInput::TemporalSampleIndex: return FExternalInputDescription(TEXT("TemporalSampleIndex"), Shader::EValueType::Float1);
+	case EExternalInput::TemporalSampleOffset: return FExternalInputDescription(TEXT("TemporalSampleOffset"), Shader::EValueType::Float2);
+	case EExternalInput::PreExposure: return FExternalInputDescription(TEXT("PreExposure"), Shader::EValueType::Float1);
+	case EExternalInput::RcpPreExposure: return FExternalInputDescription(TEXT("RcpPreExposure"), Shader::EValueType::Float1);
+	case EExternalInput::RuntimeVirtualTextureOutputLevel: return FExternalInputDescription(TEXT("RuntimeVirtualTextureOutputLevel"), Shader::EValueType::Float1);
+	case EExternalInput::RuntimeVirtualTextureOutputDerivative: return FExternalInputDescription(TEXT("RuntimeVirtualTextureOutputDerivative"), Shader::EValueType::Float2);
+	case EExternalInput::RuntimeVirtualTextureMaxLevel: return FExternalInputDescription(TEXT("RuntimeVirtualTextureMaxLevel"), Shader::EValueType::Float1);
+
+	case EExternalInput::CameraVector: return FExternalInputDescription(TEXT("CameraVector"), Shader::EValueType::Float3);
+	case EExternalInput::CameraWorldPosition: return FExternalInputDescription(TEXT("CameraWorldPosition"), Shader::EValueType::Double3, EExternalInput::None, EExternalInput::None, EExternalInput::PrevCameraWorldPosition);
+	case EExternalInput::ViewWorldPosition: return FExternalInputDescription(TEXT("ViewWorldPosition"), Shader::EValueType::Double3, EExternalInput::None, EExternalInput::None, EExternalInput::PrevViewWorldPosition);
+	case EExternalInput::PreViewTranslation: return FExternalInputDescription(TEXT("PreViewTranslation"), Shader::EValueType::Double3, EExternalInput::None, EExternalInput::None, EExternalInput::PrevPreViewTranslation);
+	case EExternalInput::TangentToWorld: return FExternalInputDescription(TEXT("TangentToWorld"), Shader::EValueType::Float4x4);
+	case EExternalInput::LocalToWorld: return FExternalInputDescription(TEXT("LocalToWorld"), Shader::EValueType::Double4x4, EExternalInput::None, EExternalInput::None, EExternalInput::PrevLocalToWorld);
+	case EExternalInput::WorldToLocal: return FExternalInputDescription(TEXT("WorldToLocal"), Shader::EValueType::DoubleInverse4x4, EExternalInput::None, EExternalInput::None, EExternalInput::PrevWorldToLocal);
+	case EExternalInput::TranslatedWorldToCameraView: return FExternalInputDescription(TEXT("TranslatedWorldToCameraView"), Shader::EValueType::Float4x4, EExternalInput::None, EExternalInput::None, EExternalInput::PrevTranslatedWorldToCameraView);
+	case EExternalInput::TranslatedWorldToView: return FExternalInputDescription(TEXT("TranslatedWorldToView"), Shader::EValueType::Float4x4, EExternalInput::None, EExternalInput::None, EExternalInput::PrevTranslatedWorldToView);
+	case EExternalInput::CameraViewToTranslatedWorld: return FExternalInputDescription(TEXT("CameraViewToTranslatedWorld"), Shader::EValueType::Float4x4, EExternalInput::None, EExternalInput::None, EExternalInput::PrevCameraViewToTranslatedWorld);
+	case EExternalInput::ViewToTranslatedWorld: return FExternalInputDescription(TEXT("ViewToTranslatedWorld"), Shader::EValueType::Float4x4, EExternalInput::None, EExternalInput::None, EExternalInput::PrevViewToTranslatedWorld);
+	case EExternalInput::WorldToParticle: return FExternalInputDescription(TEXT("WorldToParticle"), Shader::EValueType::DoubleInverse4x4);
+	case EExternalInput::WorldToInstance: return FExternalInputDescription(TEXT("WorldToInstance"), Shader::EValueType::DoubleInverse4x4);
+	case EExternalInput::ParticleToWorld: return FExternalInputDescription(TEXT("ParticleToWorld"), Shader::EValueType::Double4x4);
+	case EExternalInput::InstanceToWorld: return FExternalInputDescription(TEXT("InstanceToWorld"), Shader::EValueType::Double4x4);
+
+	case EExternalInput::PrevFieldOfView: return FExternalInputDescription(TEXT("PrevFieldOfView"), Shader::EValueType::Float2);
+	case EExternalInput::PrevTanHalfFieldOfView: return FExternalInputDescription(TEXT("PrevTanHalfFieldOfView"), Shader::EValueType::Float2);
+	case EExternalInput::PrevCotanHalfFieldOfView: return FExternalInputDescription(TEXT("PrevCotanHalfFieldOfView"), Shader::EValueType::Float2);
+	case EExternalInput::PrevCameraWorldPosition: return FExternalInputDescription(TEXT("PrevCameraWorldPosition"), Shader::EValueType::Double3);
+	case EExternalInput::PrevViewWorldPosition: return FExternalInputDescription(TEXT("PrevViewWorldPosition"), Shader::EValueType::Double3);
+	case EExternalInput::PrevPreViewTranslation: return FExternalInputDescription(TEXT("PrevPreViewTranslation"), Shader::EValueType::Double3);
+	case EExternalInput::PrevLocalToWorld: return FExternalInputDescription(TEXT("PrevLocalToWorld"), Shader::EValueType::Double4x4);
+	case EExternalInput::PrevWorldToLocal: return FExternalInputDescription(TEXT("PrevWorldToLocal"), Shader::EValueType::DoubleInverse4x4);
+	case EExternalInput::PrevTranslatedWorldToCameraView: return FExternalInputDescription(TEXT("PrevTranslatedWorldToCameraView"), Shader::EValueType::Float4x4);
+	case EExternalInput::PrevTranslatedWorldToView: return FExternalInputDescription(TEXT("PrevTranslatedWorldToView"), Shader::EValueType::Float4x4);
+	case EExternalInput::PrevCameraViewToTranslatedWorld: return FExternalInputDescription(TEXT("PrevCameraViewToTranslatedWorld"), Shader::EValueType::Float4x4);
+	case EExternalInput::PrevViewToTranslatedWorld: return FExternalInputDescription(TEXT("PrevViewToTranslatedWorld"), Shader::EValueType::Float4x4);
+
+	case EExternalInput::PixelDepth: return FExternalInputDescription(TEXT("PixelDepth"), Shader::EValueType::Float1, EExternalInput::PixelDepth_Ddx, EExternalInput::PixelDepth_Ddy);
+	case EExternalInput::PixelDepth_Ddx: return FExternalInputDescription(TEXT("PixelDepth_Ddx"), Shader::EValueType::Float1);
+	case EExternalInput::PixelDepth_Ddy: return FExternalInputDescription(TEXT("PixelDepth_Ddy"), Shader::EValueType::Float1);
+
+	case EExternalInput::GameTime: return FExternalInputDescription(TEXT("GameTime"), Shader::EValueType::Float1, EExternalInput::None, EExternalInput::None, EExternalInput::PrevGameTime);
+	case EExternalInput::RealTime: return FExternalInputDescription(TEXT("RealTime"), Shader::EValueType::Float1, EExternalInput::None, EExternalInput::None, EExternalInput::PrevRealTime);
+	case EExternalInput::DeltaTime: return FExternalInputDescription(TEXT("DeltaTime"), Shader::EValueType::Float1);
+
+	case EExternalInput::PrevGameTime: return FExternalInputDescription(TEXT("PrevGameTime"), Shader::EValueType::Float1);
+	case EExternalInput::PrevRealTime: return FExternalInputDescription(TEXT("PrevRealTime"), Shader::EValueType::Float1);
+
+	case EExternalInput::ParticleColor: return FExternalInputDescription(TEXT("ParticleColor"), Shader::EValueType::Float4);
+	case EExternalInput::ParticleTranslatedWorldPosition: return FExternalInputDescription(TEXT("ParticleTranslatedWorldPosition"), Shader::EValueType::Float3);
+	case EExternalInput::ParticleRadius: return FExternalInputDescription(TEXT("ParticleRadius"), Shader::EValueType::Float1);
+
+	default: checkNoEntry(); return FExternalInputDescription(TEXT("Invalid"), Shader::EValueType::Void);
+	}
+}
+
+void FExpressionExternalInput::ComputeAnalyticDerivatives(FTree& Tree, FExpressionDerivatives& OutResult) const
+{
+	const FExternalInputDescription InputDesc = GetExternalInputDescription(InputType);
+	if (InputDesc.Ddx != EExternalInput::None)
+	{
+		check(InputDesc.Ddy != EExternalInput::None);
+		OutResult.ExpressionDdx = Tree.NewExpression<FExpressionExternalInput>(InputDesc.Ddx);
+		OutResult.ExpressionDdy = Tree.NewExpression<FExpressionExternalInput>(InputDesc.Ddy);
+	}
+	else
+	{
+		const Shader::EValueType DerivativeType = Shader::MakeDerivativeType(InputDesc.Type);
+		switch (InputType)
+		{
+		case EExternalInput::ViewportUV:
+		{
+			// Ddx = float2(RcpViewSize.x, 0.0f)
+			// Ddy = float2(0.0f, RcpViewSize.y)
+			FExpression* RcpViewSize = Tree.NewExpression<FExpressionExternalInput>(EExternalInput::RcpViewSize);
+			FExpression* Constant0 = Tree.NewConstant(0.0f);
+			OutResult.ExpressionDdx = Tree.NewExpression<FExpressionAppend>(Tree.NewExpression<FExpressionSwizzle>(MakeSwizzleMask(true, false, false, false), RcpViewSize), Constant0);
+			OutResult.ExpressionDdy = Tree.NewExpression<FExpressionAppend>(Constant0, Tree.NewExpression<FExpressionSwizzle>(MakeSwizzleMask(false, true, false, false), RcpViewSize));
+			break;
+		}
+		default:
+			if (DerivativeType != Shader::EValueType::Void)
+			{
+				OutResult.ExpressionDdx = OutResult.ExpressionDdy = Tree.NewConstant(Shader::FValue(DerivativeType));
+			}
+			break;
+		}
+	}
+}
+
+FExpression* FExpressionExternalInput::ComputePreviousFrame(FTree& Tree, const FRequestedType& RequestedType) const
+{
+	if (InputType == EExternalInput::ActorWorldPosition)
+	{
+		return Tree.NewBinaryOp(EOperation::VecMulMatrix3, Tree.NewBinaryOp(EOperation::VecMulMatrix3,
+			Tree.NewExpression<FExpressionExternalInput>(EExternalInput::ActorWorldPosition),
+			Tree.NewExpression< FExpressionExternalInput>(EExternalInput::WorldToLocal)),
+			Tree.NewExpression< FExpressionExternalInput>(EExternalInput::PrevLocalToWorld));
+	}
+
+	const FExternalInputDescription InputDesc = GetExternalInputDescription(InputType);
+	if (InputDesc.PreviousFrame != EExternalInput::None)
+	{
+		return Tree.NewExpression<FExpressionExternalInput>(InputDesc.PreviousFrame);
+	}
+	return nullptr;
+}
+
+bool FExpressionExternalInput::PrepareValue(FEmitContext& Context, FEmitScope& Scope, const FRequestedType& RequestedType, FPrepareValueResult& OutResult) const
+{
+	const FExternalInputDescription InputDesc = GetExternalInputDescription(InputType);
+
+	switch (InputType)
+	{
+	case EExternalInput::WorldNormal:
+	case EExternalInput::WorldReflection:
+		Context.FindData<FEmitData>().bReadMaterialNormal = true;
+		break;
+	default:
+		break;
+	}
+
+	return OutResult.SetType(Context, RequestedType, EExpressionEvaluation::Shader, InputDesc.Type);
+}
+
+void FExpressionExternalInput::EmitValueShader(FEmitContext& Context, FEmitScope& Scope, const FRequestedType& RequestedType, FEmitValueShaderResult& OutResult) const
+{
+	FEmitData& EmitMaterialData = Context.FindData<FEmitData>();
+
+	const int32 TypeIndex = (int32)InputType;
+	EmitMaterialData.ExternalInputMask[Context.ShaderFrequency][TypeIndex] = true;
+	if (IsTexCoord(InputType))
+	{
+		const int32 TexCoordIndex = TypeIndex - (int32)EExternalInput::TexCoord0;
+		OutResult.Code = Context.EmitInlineExpression(Scope, Shader::EValueType::Float2, TEXT("Parameters.TexCoords[%].xy"), TexCoordIndex);
+	}
+	else if (IsTexCoord_Ddx(InputType))
+	{
+		const int32 TexCoordIndex = TypeIndex - (int32)EExternalInput::TexCoord0_Ddx;
+		OutResult.Code = Context.EmitInlineExpression(Scope, Shader::EValueType::Float2, TEXT("Parameters.TexCoords_DDX[%].xy"), TexCoordIndex);
+	}
+	else if (IsTexCoord_Ddy(InputType))
+	{
+		const int32 TexCoordIndex = TypeIndex - (int32)EExternalInput::TexCoord0_Ddy;
+		OutResult.Code = Context.EmitInlineExpression(Scope, Shader::EValueType::Float2, TEXT("Parameters.TexCoords_DDY[%].xy"), TexCoordIndex);
+	}
+	else
+	{
+		const FExternalInputDescription InputDesc = GetExternalInputDescription(InputType);
+		const TCHAR* Code = nullptr;
+		switch (InputType)
+		{
+		case EExternalInput::LightmapTexCoord: Code = TEXT("GetLightmapUVs(Parameters)"); break;
+		case EExternalInput::LightmapTexCoord_Ddx: Code = TEXT("GetLightmapUVs_DDX(Parameters)"); break;
+		case EExternalInput::LightmapTexCoord_Ddy: Code = TEXT("GetLightmapUVs_DDY(Parameters)"); break;
+		case EExternalInput::TwoSidedSign: Code = TEXT("Parameters.TwoSidedSign"); break;
+		case EExternalInput::VertexColor: Code = TEXT("Parameters.VertexColor"); EmitMaterialData.bUsesVertexColor |= (Context.ShaderFrequency != SF_Vertex); break;
+		case EExternalInput::VertexColor_Ddx: Code = TEXT("Parameters.VertexColor_DDX"); EmitMaterialData.bUsesVertexColor |= (Context.ShaderFrequency != SF_Vertex); break;
+		case EExternalInput::VertexColor_Ddy: Code = TEXT("Parameters.VertexColor_DDY"); EmitMaterialData.bUsesVertexColor |= (Context.ShaderFrequency != SF_Vertex); break;
+		case EExternalInput::WorldPosition: Code = TEXT("GetWorldPosition(Parameters)"); break;
+		case EExternalInput::WorldPosition_NoOffsets: Code = TEXT("GetWorldPosition_NoMaterialOffsets(Parameters)"); break;
+		case EExternalInput::TranslatedWorldPosition: Code = TEXT("GetTranslatedWorldPosition(Parameters)"); break;
+		case EExternalInput::TranslatedWorldPosition_NoOffsets: Code = TEXT("GetTranslatedWorldPosition_NoMaterialOffsets(Parameters)"); break;
+		case EExternalInput::ActorWorldPosition: Code = TEXT("GetActorWorldPosition(Parameters)"); break;
+		case EExternalInput::PrevWorldPosition: Code = TEXT("GetPrevWorldPosition(Parameters)"); break;
+		case EExternalInput::PrevWorldPosition_NoOffsets: Code = TEXT("GetPrevWorldPosition_NoMaterialOffsets(Parameters)"); break;
+		case EExternalInput::PrevTranslatedWorldPosition: Code = TEXT("GetPrevTranslatedWorldPosition(Parameters)"); break;
+		case EExternalInput::PrevTranslatedWorldPosition_NoOffsets: Code = TEXT("GetPrevTranslatedWorldPosition_NoMaterialOffsets(Parameters)"); break;
+		case EExternalInput::WorldPosition_Ddx: Code = TEXT("Parameters.WorldPosition_DDX"); break;
+		case EExternalInput::WorldPosition_Ddy: Code = TEXT("Parameters.WorldPosition_DDY"); break;
+
+		case EExternalInput::WorldNormal: Code = TEXT("Parameters.WorldNormal"); break;
+		case EExternalInput::WorldReflection: Code = TEXT("Parameters.ReflectionVector"); break;
+		case EExternalInput::WorldVertexNormal: Code = TEXT("Parameters.TangentToWorld[2]"); break;
+		case EExternalInput::WorldVertexTangent: Code = TEXT("Parameters.TangentToWorld[0]"); break;
+
+		case EExternalInput::ViewportUV: Code = TEXT("GetViewportUV(Parameters)"); break;
+		case EExternalInput::PixelPosition: Code = TEXT("GetPixelPosition(Parameters)"); break;
+		case EExternalInput::ViewSize: Code = TEXT("View.ViewSizeAndInvSize.xy"); break;
+		case EExternalInput::RcpViewSize: Code = TEXT("View.ViewSizeAndInvSize.zw"); break;
+
+		case EExternalInput::FieldOfView: Code = TEXT("View.FieldOfViewWideAngles"); break;
+		case EExternalInput::TanHalfFieldOfView: Code = TEXT("GetTanHalfFieldOfView()"); break;
+		case EExternalInput::CotanHalfFieldOfView: Code = TEXT("GetCotanHalfFieldOfView()"); break;
+		case EExternalInput::TemporalSampleCount: Code = TEXT("View.TemporalAAParams.y"); break;
+		case EExternalInput::TemporalSampleIndex: Code = TEXT("View.TemporalAAParams.x"); break;
+		case EExternalInput::TemporalSampleOffset: Code = TEXT("View.TemporalAAParams.zw"); break;
+		case EExternalInput::PreExposure: Code = TEXT("View.PreExposure.x"); break;
+		case EExternalInput::RcpPreExposure: Code = TEXT("View.OneOverPreExposure.x"); break;
+		case EExternalInput::RuntimeVirtualTextureOutputLevel:  Code = TEXT("View.RuntimeVirtualTextureMipLevel.x"); break;
+		case EExternalInput::RuntimeVirtualTextureOutputDerivative: Code = TEXT("View.RuntimeVirtualTextureMipLevel.zw"); break;
+		case EExternalInput::RuntimeVirtualTextureMaxLevel:  Code = TEXT("View.RuntimeVirtualTextureMipLevel.y"); break;
+
+		case EExternalInput::CameraVector: Code = TEXT("Parameters.CameraVector"); break;
+		case EExternalInput::CameraWorldPosition: Code = TEXT("ResolvedView.WorldCameraOrigin"); break;
+		case EExternalInput::ViewWorldPosition: Code = TEXT("ResolvedView.WorldViewOrigin"); break;
+		case EExternalInput::PreViewTranslation: Code = TEXT("ResolvedView.PreViewTranslation"); break;
+		case EExternalInput::TangentToWorld: Code = TEXT("Parameters.TangentToWorld"); break;
+		case EExternalInput::LocalToWorld: Code = TEXT("GetLocalToWorld(Parameters)"); break;
+		case EExternalInput::WorldToLocal: Code = TEXT("GetPrimitiveData(Parameters).WorldToLocal"); break;
+		case EExternalInput::TranslatedWorldToCameraView: Code = TEXT("ResolvedView.TranslatedWorldToCameraView"); break;
+		case EExternalInput::TranslatedWorldToView: Code = TEXT("ResolvedView.TranslatedWorldToView"); break;
+		case EExternalInput::CameraViewToTranslatedWorld: Code = TEXT("ResolvedView.CameraViewToTranslatedWorld"); break;
+		case EExternalInput::ViewToTranslatedWorld: Code = TEXT("ResolvedView.ViewToTranslatedWorld"); break;
+		case EExternalInput::WorldToParticle: Code = TEXT("Parameters.Particle.WorldToParticle"); break;
+		case EExternalInput::WorldToInstance: Code = TEXT("GetWorldToInstance(Parameters)"); break;
+		case EExternalInput::ParticleToWorld: Code = TEXT("Parameters.Particle.ParticleToWorld"); break;
+		case EExternalInput::InstanceToWorld: Code = TEXT("GetInstanceToWorld(Parameters)"); break;
+
+		case EExternalInput::PrevFieldOfView: Code = TEXT("View.PrevFieldOfViewWideAngles"); break;
+		case EExternalInput::PrevTanHalfFieldOfView: Code = TEXT("GetPrevTanHalfFieldOfView()"); break;
+		case EExternalInput::PrevCotanHalfFieldOfView: Code = TEXT("GetPrevCotanHalfFieldOfView()"); break;
+		case EExternalInput::PrevCameraWorldPosition: Code = TEXT("ResolvedView.PrevWorldCameraOrigin"); break;
+		case EExternalInput::PrevViewWorldPosition: Code = TEXT("ResolvedView.PrevWorldViewOrigin"); break;
+		case EExternalInput::PrevPreViewTranslation: Code = TEXT("ResolvedView.PrevPreViewTranslation"); break;
+		case EExternalInput::PrevLocalToWorld: Code = TEXT("GetPrevLocalToWorld(Parameters)"); break;
+		case EExternalInput::PrevWorldToLocal: Code = TEXT("GetPrimitiveData(Parameters).PreviousWorldToLocal"); break;
+		case EExternalInput::PrevTranslatedWorldToCameraView: Code = TEXT("ResolvedView.PrevTranslatedWorldToCameraView"); break;
+		case EExternalInput::PrevTranslatedWorldToView: Code = TEXT("ResolvedView.PrevTranslatedWorldToView"); break;
+		case EExternalInput::PrevCameraViewToTranslatedWorld: Code = TEXT("ResolvedView.PrevCameraViewToTranslatedWorld"); break;
+		case EExternalInput::PrevViewToTranslatedWorld: Code = TEXT("ResolvedView.PrevViewToTranslatedWorld"); break;
+
+		case EExternalInput::PixelDepth: Code = TEXT("GetPixelDepth(Parameters)"); break;
+		case EExternalInput::PixelDepth_Ddx: Code = TEXT("Parameters.ScreenPosition_DDX.w"); break;
+		case EExternalInput::PixelDepth_Ddy: Code = TEXT("Parameters.ScreenPosition_DDY.w"); break;
+		case EExternalInput::GameTime: Code = TEXT("View.GameTime"); break;
+		case EExternalInput::RealTime: Code = TEXT("View.RealTime"); break;
+		case EExternalInput::DeltaTime: Code = TEXT("View.DeltaTime"); break;
+		case EExternalInput::PrevGameTime: Code = TEXT("View.PrevFrameGameTime"); break;
+		case EExternalInput::PrevRealTime: Code = TEXT("View.PrevFrameRealTime"); break;
+
+		case EExternalInput::ParticleColor: Code = TEXT("Parameters.Particle.Color"); EmitMaterialData.bUsesParticleColor |= (Context.ShaderFrequency != SF_Vertex); break;
+		case EExternalInput::ParticleTranslatedWorldPosition: Code = TEXT("Parameters.Particle.TranslatedWorldPositionAndSize.xyz"); EmitMaterialData.bNeedsParticlePosition = true; break;
+		case EExternalInput::ParticleRadius: Code = TEXT("Parameters.Particle.TranslatedWorldPositionAndSize.w"); EmitMaterialData.bNeedsParticlePosition = true; break;
+
+		default:
+			checkNoEntry();
+			break;
+		}
+		OutResult.Code = Context.EmitInlineExpression(Scope, InputDesc.Type, Code);
+	}
+}
+
+void FExpressionShadingModel::ComputeAnalyticDerivatives(FTree& Tree, FExpressionDerivatives& OutResult) const
+{
+	const Shader::FValue ZeroValue(Shader::EValueType::Float1);
+	OutResult.ExpressionDdx = Tree.NewConstant(ZeroValue);
+	OutResult.ExpressionDdy = OutResult.ExpressionDdx;
+}
+
+bool FExpressionShadingModel::PrepareValue(FEmitContext& Context, FEmitScope& Scope, const FRequestedType& RequestedType, FPrepareValueResult& OutResult) const
+{
+	return OutResult.SetType(Context, RequestedType, EExpressionEvaluation::Constant, Shader::EValueType::Int1);
+}
+
+void FExpressionShadingModel::EmitValuePreshader(FEmitContext& Context, FEmitScope& Scope, const FRequestedType& RequestedType, FEmitValuePreshaderResult& OutResult) const
+{
+	Context.FindData<FEmitData>().ShadingModelsFromCompilation.AddShadingModel(ShadingModel);
+
+	Context.PreshaderStackPosition++;
+	OutResult.Type = Shader::EValueType::Int1;
+	OutResult.Preshader.WriteOpcode(Shader::EPreshaderOpcode::Constant).Write(Shader::FValue((int32)ShadingModel));
+}
+
+void FExpressionParameter::ComputeAnalyticDerivatives(FTree& Tree, FExpressionDerivatives& OutResult) const
+{
+	const Shader::FType Type = GetShaderValueType(ParameterType);
+	const Shader::FType DerivativeType = Type.GetDerivativeType();
+	if (!DerivativeType.IsVoid())
+	{
+		const Shader::FValue ZeroValue(DerivativeType);
+		OutResult.ExpressionDdx = Tree.NewConstant(ZeroValue);
+		OutResult.ExpressionDdy = OutResult.ExpressionDdx;
+	}
+}
+
+bool FExpressionParameter::PrepareValue(FEmitContext& Context, FEmitScope& Scope, const FRequestedType& RequestedType, FPrepareValueResult& OutResult) const
+{
+	EExpressionEvaluation Evaluation = EExpressionEvaluation::Shader;
+	if (IsStaticMaterialParameter(ParameterType))
+	{
+		Evaluation = EExpressionEvaluation::Constant;
+	}
+	else if (ParameterType == EMaterialParameterType::Scalar ||
+		ParameterType == EMaterialParameterType::Vector ||
+		ParameterType == EMaterialParameterType::DoubleVector)
+	{
+		Evaluation = EExpressionEvaluation::Preshader;
+	}
+	return OutResult.SetType(Context, RequestedType, Evaluation, DefaultValue.Type);
+}
+
+void FExpressionParameter::EmitValueShader(FEmitContext& Context, FEmitScope& Scope, const FRequestedType& RequestedType, FEmitValueShaderResult& OutResult) const
+{
+	if (ParameterType == EMaterialParameterType::Texture)
+	{
+		const Shader::FTextureValue* TextureValue = DefaultValue.AsTexture();
+		check(TextureValue);
+
+		const Shader::EValueType TextureType = TextureValue->GetType();
+		const TCHAR* ConstructorName = nullptr;
+		const TCHAR* TextureTypeName = nullptr;
+		int32 ParameterIndex = INDEX_NONE;
+
+		if (TextureType == Shader::EValueType::TextureExternal)
+		{
+			check(TextureValue->SamplerType == SAMPLERTYPE_External);
+
+			ConstructorName = TEXT("MakeTextureExternal");
+			TextureTypeName = TEXT("ExternalTexture");
+
+			FMaterialExternalTextureParameterInfo TextureParameterInfo;
+			TextureParameterInfo.ParameterName = NameToScriptName(ParameterInfo.Name);
+			TextureParameterInfo.ExternalTextureGuid = TextureValue->ExternalTextureGuid;
+			if (TextureValue->Texture)
+			{
+				TextureParameterInfo.SourceTextureIndex = Context.Material->GetReferencedTextures().Find(TextureValue->Texture);
+			}
+			ParameterIndex = Context.MaterialCompilationOutput->UniformExpressionSet.FindOrAddExternalTextureParameter(TextureParameterInfo);
+		}
+		else
+		{
+			EMaterialTextureParameterType TextureParameterType = EMaterialTextureParameterType::Count;
+			switch (TextureType)
+			{
+			case Shader::EValueType::Texture2D:
+				TextureParameterType = EMaterialTextureParameterType::Standard2D;
+				ConstructorName = TEXT("MakeTexture2D");
+				TextureTypeName = TEXT("Texture2D");
+				break;
+			case Shader::EValueType::Texture2DArray:
+				TextureParameterType = EMaterialTextureParameterType::Array2D;
+				ConstructorName = TEXT("MakeTexture2DArray");
+				TextureTypeName = TEXT("Texture2DArray");
+				break;
+			case Shader::EValueType::TextureCube:
+				TextureParameterType = EMaterialTextureParameterType::Cube;
+				ConstructorName = TEXT("MakeTextureCube");
+				TextureTypeName = TEXT("TextureCube");
+				break;
+			case Shader::EValueType::TextureCubeArray:
+				TextureParameterType = EMaterialTextureParameterType::ArrayCube;
+				ConstructorName = TEXT("MakeTextureCubeArray");
+				TextureTypeName = TEXT("TextureCubeArray");
+				break;
+			case Shader::EValueType::Texture3D:
+				TextureParameterType = EMaterialTextureParameterType::Volume;
+				ConstructorName = TEXT("MakeTexture3D");
+				TextureTypeName = TEXT("VolumeTexture");
+				break;
+			default:
+				checkNoEntry();
+				break;
+			}
+
+			FMaterialTextureParameterInfo TextureParameterInfo;
+			TextureParameterInfo.ParameterInfo = ParameterInfo;
+			TextureParameterInfo.TextureIndex = Context.Material->GetReferencedTextures().Find(TextureValue->Texture);
+			TextureParameterInfo.SamplerSource = SSM_FromTextureAsset; // TODO - Is this needed?
+			check(TextureParameterInfo.TextureIndex != INDEX_NONE);
+			ParameterIndex = Context.MaterialCompilationOutput->UniformExpressionSet.FindOrAddTextureParameter(TextureParameterType, TextureParameterInfo);
+		}
+
+		TStringBuilder<256> FormattedCode;
+		FormattedCode.Appendf(TEXT("%s(Material.%s_%d, Material.%s_%dSampler, %d)"), ConstructorName, TextureTypeName, ParameterIndex, TextureTypeName, ParameterIndex, (int32)TextureValue->SamplerType);
+		OutResult.Code = Context.EmitExpression(Scope, TextureType, FormattedCode.ToView());
+	}
+}
+
+void FExpressionParameter::EmitValuePreshader(FEmitContext& Context, FEmitScope& Scope, const FRequestedType& RequestedType, FEmitValuePreshaderResult& OutResult) const
+{
+	FEmitData& EmitMaterialData = Context.FindData<FEmitData>();
+
+	Context.PreshaderStackPosition++;
+	OutResult.Type = GetShaderValueType(ParameterType);
+	if (ParameterType == EMaterialParameterType::StaticSwitch)
+	{
+		Shader::FValue Value = DefaultValue;
+		for (const FStaticSwitchParameter& Parameter : EmitMaterialData.StaticParameters->StaticSwitchParameters)
+		{
+			if (Parameter.ParameterInfo == ParameterInfo)
+			{
+				Value = Parameter.Value;
+				break;
+			}
+		}
+		OutResult.Preshader.WriteOpcode(Shader::EPreshaderOpcode::Constant).Write(Value);
+	}
+	else
+	{
+		const uint32* PrevDefaultOffset = EmitMaterialData.DefaultUniformValues.Find(DefaultValue);
+		uint32 DefaultOffset;
+		if (PrevDefaultOffset)
+		{
+			DefaultOffset = *PrevDefaultOffset;
+		}
+		else
+		{
+			DefaultOffset = Context.MaterialCompilationOutput->UniformExpressionSet.AddDefaultParameterValue(DefaultValue);
+			EmitMaterialData.DefaultUniformValues.Add(DefaultValue, DefaultOffset);
+		}
+		const int32 ParameterIndex = Context.MaterialCompilationOutput->UniformExpressionSet.FindOrAddNumericParameter(ParameterType, ParameterInfo, DefaultOffset);
+		check(ParameterIndex >= 0 && ParameterIndex <= 0xffff);
+		OutResult.Preshader.WriteOpcode(Shader::EPreshaderOpcode::Parameter).Write((uint16)ParameterIndex);
+	}
+}
+
+bool FExpressionSceneTexture::PrepareValue(FEmitContext& Context, FEmitScope& Scope, const FRequestedType& RequestedType, FPrepareValueResult& OutResult) const
+{
+	Context.PrepareExpression(TexCoordExpression, Scope, ERequestedType::Vector2);
+
+	Context.MaterialCompilationOutput->bNeedsSceneTextures = true;
+	Context.MaterialCompilationOutput->SetIsSceneTextureUsed((ESceneTextureId)SceneTextureId);
+
+	return OutResult.SetType(Context, RequestedType, EExpressionEvaluation::Shader, Shader::EValueType::Float4);
+}
+
+void FExpressionSceneTexture::EmitValueShader(FEmitContext& Context, FEmitScope& Scope, const FRequestedType& RequestedType, FEmitValueShaderResult& OutResult) const
+{
+	const bool bSupportedOnMobile = SceneTextureId == PPI_PostProcessInput0 ||
+		SceneTextureId == PPI_CustomDepth ||
+		SceneTextureId == PPI_SceneDepth ||
+		SceneTextureId == PPI_CustomStencil;
+
+	FEmitShaderExpression* EmitTexCoord = nullptr;
+	if (TexCoordExpression)
+	{
+		EmitTexCoord = TexCoordExpression->GetValueShader(Context, Scope, Shader::EValueType::Float2);
+		EmitTexCoord = Context.EmitExpression(Scope, Shader::EValueType::Float2, TEXT("ClampSceneTextureUV(ViewportUVToSceneTextureUV(%, %), %)"), EmitTexCoord, (int)SceneTextureId, (int)SceneTextureId);
+	}
+	else
+	{
+		EmitTexCoord = Context.EmitExpression(Scope, Shader::EValueType::Float2, TEXT("GetDefaultSceneTextureUV(Parameters, %)"), (int)SceneTextureId);
+	}
+
+	FEmitShaderExpression* EmitLookup = nullptr;
+	if (Context.Material->GetFeatureLevel() >= ERHIFeatureLevel::SM5)
+	{
+		EmitLookup = Context.EmitExpression(Scope, Shader::EValueType::Float4, TEXT("SceneTextureLookup(%, %, %)"), EmitTexCoord, (int)SceneTextureId, bFiltered);
+	}
+	else
+	{
+		EmitLookup = Context.EmitExpression(Scope, Shader::EValueType::Float4, TEXT("MobileSceneTextureLookup(Parameters, %, %, %)"), (int)SceneTextureId, EmitTexCoord);
+	}
+
+	if (SceneTextureId == PPI_PostProcessInput0 && Context.Material->GetMaterialDomain() == MD_PostProcess && Context.Material->GetBlendableLocation() != BL_AfterTonemapping)
+	{
+		EmitLookup = Context.EmitExpression(Scope, Shader::EValueType::Float4, TEXT("(float4(View.OneOverPreExposure.xxx, 1) * %)"), EmitLookup);
+	}
+
+	OutResult.Code = EmitLookup;
+}
+
+bool FExpressionNoise::PrepareValue(FEmitContext& Context, FEmitScope& Scope, const FRequestedType& RequestedType, FPrepareValueResult& OutResult) const
+{
+	const FPreparedType& PositionType = Context.PrepareExpression(PositionExpression, Scope, ERequestedType::Vector3);
+	const FPreparedType& FilterWidthType = Context.PrepareExpression(FilterWidthExpression, Scope, ERequestedType::Scalar);
+	if (PositionType.IsVoid() || FilterWidthType.IsVoid())
+	{
+		return false;
+	}
+	return OutResult.SetType(Context, RequestedType, EExpressionEvaluation::Shader, Shader::EValueType::Float1);
+}
+
+void FExpressionNoise::EmitValueShader(FEmitContext& Context, FEmitScope& Scope, const FRequestedType& RequestedType, FEmitValueShaderResult& OutResult) const
+{
+	bool bIsLWC = Shader::IsLWCType(PositionExpression->GetPreparedType().ValueComponentType);
+	FEmitShaderExpression* EmitPosition = PositionExpression->GetValueShader(Context, Scope, bIsLWC ? Shader::EValueType::Double3 : Shader::EValueType::Float3);
+	FEmitShaderExpression* EmitFilterWidth = FilterWidthExpression->GetValueShader(Context, Scope, Shader::EValueType::Float1);
+
+	if (bIsLWC)
+	{
+		// If Noise is driven by a LWC position, just take the offset within the current tile
+		// Will generate discontinuity in noise at tile boudaries
+		// Could potentially add noise functions that operate directly on LWC values, but that would be very expensive
+		EmitPosition = Context.EmitExpression(Scope, Shader::EValueType::Float3, TEXT("LWCNormalizeTile(%).Offset"), EmitPosition);
+	}
+
+	OutResult.Code = Context.EmitExpression(Scope, Shader::EValueType::Float1, TEXT("MaterialExpressionNoise(%,%,%,%,%,%,%,%,%,%,%,%)"),
+		EmitPosition,
+		Parameters.Scale,
+		Parameters.Quality,
+		Parameters.NoiseFunction,
+		Parameters.bTurbulence,
+		Parameters.Levels,
+		Parameters.OutputMin,
+		Parameters.OutputMax,
+		Parameters.LevelScale,
+		EmitFilterWidth,
+		Parameters.bTiling,
+		Parameters.RepeatSize);
+}
+
+} // namespace UE::HLSLTree::Material
+
diff --git a/Engine/Source/Runtime/Engine/Private/Materials/MaterialHLSLTree.h b/Engine/Source/Runtime/Engine/Private/Materials/MaterialHLSLTree.h
new file mode 100644
--- /dev/null
+++ b/Engine/Source/Runtime/Engine/Private/Materials/MaterialHLSLTree.h
@@ -0,0 +1,293 @@
+// Copyright Epic Games, Inc. All Rights Reserved.
+#pragma once
+
+#include "HLSLTree/HLSLTree.h"
+#include "Materials/MaterialLayersFunctions.h"
+#include "RHIDefinitions.h"
+
+enum class EMaterialParameterType : uint8;
+
+namespace UE::HLSLTree::Material
+{
+
+enum class EExternalInput : uint8
+{
+	None,
+
+	TexCoord0,
+	TexCoord1,
+	TexCoord2,
+	TexCoord3,
+	TexCoord4,
+	TexCoord5,
+	TexCoord6,
+	TexCoord7,
+
+	TexCoord0_Ddx,
+	TexCoord1_Ddx,
+	TexCoord2_Ddx,
+	TexCoord3_Ddx,
+	TexCoord4_Ddx,
+	TexCoord5_Ddx,
+	TexCoord6_Ddx,
+	TexCoord7_Ddx,
+
+	TexCoord0_Ddy,
+	TexCoord1_Ddy,
+	TexCoord2_Ddy,
+	TexCoord3_Ddy,
+	TexCoord4_Ddy,
+	TexCoord5_Ddy,
+	TexCoord6_Ddy,
+	TexCoord7_Ddy,
+
+	LightmapTexCoord,
+	LightmapTexCoord_Ddx,
+	LightmapTexCoord_Ddy,
+
+	TwoSidedSign,
+	VertexColor,
+	VertexColor_Ddx,
+	VertexColor_Ddy,
+
+	WorldPosition,
+	WorldPosition_NoOffsets,
+	TranslatedWorldPosition,
+	TranslatedWorldPosition_NoOffsets,
+	ActorWorldPosition,
+
+	PrevWorldPosition,
+	PrevWorldPosition_NoOffsets,
+	PrevTranslatedWorldPosition,
+	PrevTranslatedWorldPosition_NoOffsets,
+
+	WorldPosition_Ddx,
+	WorldPosition_Ddy,
+
+	WorldVertexNormal,
+	WorldVertexTangent,
+	WorldNormal,
+	WorldReflection,
+
+	ViewportUV,
+	PixelPosition,
+	ViewSize,
+	RcpViewSize,
+	FieldOfView,
+	TanHalfFieldOfView,
+	CotanHalfFieldOfView,
+	TemporalSampleCount,
+	TemporalSampleIndex,
+	TemporalSampleOffset,
+	PreExposure,
+	RcpPreExposure,
+	RuntimeVirtualTextureOutputLevel,
+	RuntimeVirtualTextureOutputDerivative,
+	RuntimeVirtualTextureMaxLevel,
+
+	CameraVector,
+	CameraWorldPosition,
+	ViewWorldPosition,
+	PreViewTranslation,
+	TangentToWorld,
+	LocalToWorld,
+	WorldToLocal,
+	TranslatedWorldToCameraView,
+	TranslatedWorldToView,
+	CameraViewToTranslatedWorld,
+	ViewToTranslatedWorld,
+	WorldToParticle,
+	WorldToInstance,
+	ParticleToWorld,
+	InstanceToWorld,
+
+	PrevFieldOfView,
+	PrevTanHalfFieldOfView,
+	PrevCotanHalfFieldOfView,
+
+	PrevCameraWorldPosition,
+	PrevViewWorldPosition,
+	PrevPreViewTranslation,
+	PrevLocalToWorld,
+	PrevWorldToLocal,
+	PrevTranslatedWorldToCameraView,
+	PrevTranslatedWorldToView,
+	PrevCameraViewToTranslatedWorld,
+	PrevViewToTranslatedWorld,
+
+	PixelDepth,
+	PixelDepth_Ddx,
+	PixelDepth_Ddy,
+
+	GameTime,
+	RealTime,
+	DeltaTime,
+
+	PrevGameTime,
+	PrevRealTime,
+
+	ParticleColor,
+	ParticleTranslatedWorldPosition,
+	ParticleRadius,
+};
+static constexpr int32 MaxNumTexCoords = 8;
+
+struct FExternalInputDescription
+{
+	FExternalInputDescription(const TCHAR* InName, Shader::EValueType InType, EExternalInput InDdx = EExternalInput::None, EExternalInput InDdy = EExternalInput::None, EExternalInput InPreviousFrame = EExternalInput::None)
+		: Name(InName), Type(InType), Ddx(InDdx), Ddy(InDdy), PreviousFrame(InPreviousFrame)
+	{}
+
+	const TCHAR* Name;
+	Shader::EValueType Type;
+	EExternalInput Ddx;
+	EExternalInput Ddy;
+	EExternalInput PreviousFrame;
+};
+
+FExternalInputDescription GetExternalInputDescription(EExternalInput Input);
+
+inline bool IsTexCoord(EExternalInput Type)
+{
+	return FMath::IsWithin((int32)Type, (int32)EExternalInput::TexCoord0, (int32)EExternalInput::TexCoord0 + MaxNumTexCoords);
+}
+inline bool IsTexCoord_Ddx(EExternalInput Type)
+{
+	return FMath::IsWithin((int32)Type, (int32)EExternalInput::TexCoord0_Ddx, (int32)EExternalInput::TexCoord0_Ddx + MaxNumTexCoords);
+}
+inline bool IsTexCoord_Ddy(EExternalInput Type)
+{
+	return FMath::IsWithin((int32)Type, (int32)EExternalInput::TexCoord0_Ddy, (int32)EExternalInput::TexCoord0_Ddy + MaxNumTexCoords);
+}
+inline EExternalInput MakeInputTexCoord(int32 Index)
+{
+	check(Index >= 0 && Index < MaxNumTexCoords);
+	return (EExternalInput)((int32)EExternalInput::TexCoord0 + Index);
+}
+
+class FExpressionExternalInput : public FExpression
+{
+public:
+	FExpressionExternalInput(EExternalInput InInputType) : InputType(InInputType) {}
+
+	EExternalInput InputType;
+
+	virtual void ComputeAnalyticDerivatives(FTree& Tree, FExpressionDerivatives& OutResult) const override;
+	virtual FExpression* ComputePreviousFrame(FTree& Tree, const FRequestedType& RequestedType) const override;
+	virtual bool PrepareValue(FEmitContext& Context, FEmitScope& Scope, const FRequestedType& RequestedType, FPrepareValueResult& OutResult) const override;
+	virtual void EmitValueShader(FEmitContext& Context, FEmitScope& Scope, const FRequestedType& RequestedType, FEmitValueShaderResult& OutResult) const override;
+};
+
+class FExpressionShadingModel : public FExpression
+{
+public:
+	explicit FExpressionShadingModel(EMaterialShadingModel InShadingModel)
+		: ShadingModel(InShadingModel)
+	{}
+
+	TEnumAsByte<EMaterialShadingModel> ShadingModel;
+
+	virtual void ComputeAnalyticDerivatives(FTree& Tree, FExpressionDerivatives& OutResult) const override;
+	virtual bool PrepareValue(FEmitContext& Context, FEmitScope& Scope, const FRequestedType& RequestedType, FPrepareValueResult& OutResult) const override;
+	virtual void EmitValuePreshader(FEmitContext& Context, FEmitScope& Scope, const FRequestedType& RequestedType, FEmitValuePreshaderResult& OutResult) const override;
+};
+
+inline void AppendHash(FHasher& Hasher, const FMaterialParameterInfo& Value)
+{
+	AppendHash(Hasher, Value.Name);
+	AppendHash(Hasher, Value.Index);
+	AppendHash(Hasher, Value.Association);
+}
+
+class FExpressionParameter : public FExpression
+{
+public:
+	explicit FExpressionParameter(EMaterialParameterType InType, const FMaterialParameterInfo& InParameterInfo, const Shader::FValue& InDefaultValue)
+		: ParameterInfo(InParameterInfo), DefaultValue(InDefaultValue), ParameterType(InType)
+	{
+	}
+
+	FMaterialParameterInfo ParameterInfo;
+	Shader::FValue DefaultValue;
+	EMaterialParameterType ParameterType;
+
+	virtual void ComputeAnalyticDerivatives(FTree& Tree, FExpressionDerivatives& OutResult) const override;
+	virtual bool PrepareValue(FEmitContext& Context, FEmitScope& Scope, const FRequestedType& RequestedType, FPrepareValueResult& OutResult) const override;
+	virtual void EmitValueShader(FEmitContext& Context, FEmitScope& Scope, const FRequestedType& RequestedType, FEmitValueShaderResult& OutResult) const override;
+	virtual void EmitValuePreshader(FEmitContext& Context, FEmitScope& Scope, const FRequestedType& RequestedType, FEmitValuePreshaderResult& OutResult) const override;
+};
+
+class FExpressionSceneTexture : public FExpression
+{
+public:
+	FExpressionSceneTexture(FExpression* InTexCoordExpression, uint32 InSceneTextureId, bool bInFiltered)
+		: TexCoordExpression(InTexCoordExpression)
+		, SceneTextureId(InSceneTextureId)
+		, bFiltered(bInFiltered)
+	{}
+
+	FExpression* TexCoordExpression;
+	uint32 SceneTextureId;
+	bool bFiltered;
+
+	virtual bool PrepareValue(FEmitContext& Context, FEmitScope& Scope, const FRequestedType& RequestedType, FPrepareValueResult& OutResult) const override;
+	virtual void EmitValueShader(FEmitContext& Context, FEmitScope& Scope, const FRequestedType& RequestedType, FEmitValueShaderResult& OutResult) const override;
+};
+
+struct FNoiseParameters
+{
+	FNoiseParameters() { FMemory::Memzero(*this); }
+
+	int32 Quality;
+	int32 Levels;
+	float Scale;
+	uint32 RepeatSize;
+	float OutputMin;
+	float OutputMax;
+	float LevelScale;
+	uint8 NoiseFunction;
+	bool bTiling;
+	bool bTurbulence;
+};
+
+class FExpressionNoise : public FExpression
+{
+public:
+	FExpressionNoise(const FNoiseParameters& InParams, FExpression* InPositionExpression, FExpression* InFilterWidthExpression)
+		: PositionExpression(InPositionExpression)
+		, FilterWidthExpression(InFilterWidthExpression)
+		, Parameters(InParams)
+	{}
+
+	FExpression* PositionExpression;
+	FExpression* FilterWidthExpression;
+	FNoiseParameters Parameters;
+
+	virtual bool PrepareValue(FEmitContext& Context, FEmitScope& Scope, const FRequestedType& RequestedType, FPrepareValueResult& OutResult) const override;
+	virtual void EmitValueShader(FEmitContext& Context, FEmitScope& Scope, const FRequestedType& RequestedType, FEmitValueShaderResult& OutResult) const override;
+};
+
+class FEmitData
+{
+public:
+	FEmitData()
+	{
+		for (int32 Index = 0; Index < SF_NumFrequencies; ++Index)
+		{
+			ExternalInputMask[Index].Init(false, 256);
+		}
+	}
+
+	const FStaticParameterSet* StaticParameters = nullptr;
+	TMap<Shader::FValue, uint32> DefaultUniformValues;
+	bool bReadMaterialNormal = false;
+	bool bUsesVertexColor = false;
+	bool bNeedsParticlePosition = false;
+	bool bUsesParticleColor = false;
+	TBitArray<> ExternalInputMask[SF_NumFrequencies];
+	FMaterialShadingModelField ShadingModelsFromCompilation;
+
+	bool IsExternalInputUsed(EShaderFrequency Frequency, EExternalInput Input) const { return ExternalInputMask[Frequency][(uint8)Input]; }
+};
+
+} // namespace UE::HLSLTree::Material
diff --git a/Engine/Source/Runtime/Engine/Public/HLSLTree/HLSLTree.h b/Engine/Source/Runtime/Engine/Public/HLSLTree/HLSLTree.h
--- a/Engine/Source/Runtime/Engine/Public/HLSLTree/HLSLTree.h
+++ b/Engine/Source/Runtime/Engine/Public/HLSLTree/HLSLTree.h
@@ -9,6 +9,7 @@
 #include "Containers/List.h"
 #include "Misc/StringBuilder.h"
 #include "Misc/MemStack.h"
+#include "Misc/GeneratedTypeName.h"
 #include "Hash/xxhash.h"
 #include "HLSLTree/HLSLTreeTypes.h"
 #include "HLSLTree/HLSLTreeHash.h"
@@ -17,10 +18,7 @@ class FMaterial;
 class FMaterialCompilationOutput;
 struct FStaticParameterSet;
 
-namespace UE
-{
-
-namespace Shader
+namespace UE::Shader
 {
 class FPreshaderData;
 }
@@ -30,7 +28,7 @@ class FPreshaderData;
  * This allows C++ to procedurally define an HLSL program.  The structure of the tree is designed to be flexible, to facilitate incremental generation from a material node graph
  * Once the tree is complete, HLSL source code may be generated
  */
-namespace HLSLTree
+namespace UE::HLSLTree
 {
 
 class FNode;
@@ -490,20 +488,6 @@ private:
 	int32 NestedLevel = 0;
 };
 
-namespace Private
-{
-uint64 GetNextTypeHash();
-template<typename T>
-struct TTypeHasher
-{
-	static uint64 GetHash()
-	{
-		static uint64 Hash = GetNextTypeHash();
-		return Hash;
-	}
-};
-}
-
 /**
  * The HLSL AST.  Basically a wrapper around the root scope, with some helper methods
  */
@@ -527,7 +511,7 @@ public:
 	inline FExpression* NewExpression(ArgTypes&&... Args)
 	{
 		FHasher Hasher;
-		AppendHash(Hasher, Private::TTypeHasher<T>::GetHash());
+		AppendHash(Hasher, GetGeneratedTypeName<T>());
 		AppendHashes(Hasher, Forward<ArgTypes>(Args)...);
 		const FXxHash64 Hash = Hasher.Finalize();
 		FExpression* Expression = FindExpression(Hash);
@@ -622,5 +606,5 @@ private:
 	friend class FExpressionLocalPHI;
 };
 
-} // namespace HLSLTree
-} // namespace UE
+} // namespace UE::HLSLTree
+
diff --git a/Engine/Source/Runtime/Engine/Public/HLSLTree/HLSLTreeCommon.h b/Engine/Source/Runtime/Engine/Public/HLSLTree/HLSLTreeCommon.h
--- a/Engine/Source/Runtime/Engine/Public/HLSLTree/HLSLTreeCommon.h
+++ b/Engine/Source/Runtime/Engine/Public/HLSLTree/HLSLTreeCommon.h
@@ -2,14 +2,8 @@
 #pragma once
 
 #include "HLSLTree/HLSLTree.h"
-#include "Materials/MaterialLayersFunctions.h" // TODO - split material expressions into separate file/module
 
-enum class EMaterialParameterType : uint8;
-
-namespace UE
-{
-
-namespace HLSLTree
+namespace UE::HLSLTree
 {
 
 class FExpressionConstant : public FExpression
@@ -26,263 +20,6 @@ public:
 	virtual void EmitValuePreshader(FEmitContext& Context, FEmitScope& Scope, const FRequestedType& RequestedType, FEmitValuePreshaderResult& OutResult) const override;
 };
 
-class FExpressionMaterialShadingModel : public FExpression
-{
-public:
-	explicit FExpressionMaterialShadingModel(EMaterialShadingModel InShadingModel)
-		: ShadingModel(InShadingModel)
-	{}
-
-	TEnumAsByte<EMaterialShadingModel> ShadingModel;
-
-	virtual void ComputeAnalyticDerivatives(FTree& Tree, FExpressionDerivatives& OutResult) const override;
-	virtual bool PrepareValue(FEmitContext& Context, FEmitScope& Scope, const FRequestedType& RequestedType, FPrepareValueResult& OutResult) const override;
-	virtual void EmitValuePreshader(FEmitContext& Context, FEmitScope& Scope, const FRequestedType& RequestedType, FEmitValuePreshaderResult& OutResult) const override;
-};
-
-inline void AppendHash(FHasher& Hasher, const FMaterialParameterInfo& Value)
-{
-	AppendHash(Hasher, Value.Name);
-	AppendHash(Hasher, Value.Index);
-	AppendHash(Hasher, Value.Association);
-}
-
-class FExpressionMaterialParameter : public FExpression
-{
-public:
-	explicit FExpressionMaterialParameter(EMaterialParameterType InType, const FMaterialParameterInfo& InParameterInfo, const Shader::FValue& InDefaultValue)
-		: ParameterInfo(InParameterInfo), DefaultValue(InDefaultValue), ParameterType(InType)
-	{
-	}
-
-	FMaterialParameterInfo ParameterInfo;
-	Shader::FValue DefaultValue;
-	EMaterialParameterType ParameterType;
-
-	virtual void ComputeAnalyticDerivatives(FTree& Tree, FExpressionDerivatives& OutResult) const override;
-	virtual bool PrepareValue(FEmitContext& Context, FEmitScope& Scope, const FRequestedType& RequestedType, FPrepareValueResult& OutResult) const override;
-	virtual void EmitValueShader(FEmitContext& Context, FEmitScope& Scope, const FRequestedType& RequestedType, FEmitValueShaderResult& OutResult) const override;
-	virtual void EmitValuePreshader(FEmitContext& Context, FEmitScope& Scope, const FRequestedType& RequestedType, FEmitValuePreshaderResult& OutResult) const override;
-};
-
-enum class EExternalInput : uint8
-{
-	None,
-
-	TexCoord0,
-	TexCoord1,
-	TexCoord2,
-	TexCoord3,
-	TexCoord4,
-	TexCoord5,
-	TexCoord6,
-	TexCoord7,
-
-	TexCoord0_Ddx,
-	TexCoord1_Ddx,
-	TexCoord2_Ddx,
-	TexCoord3_Ddx,
-	TexCoord4_Ddx,
-	TexCoord5_Ddx,
-	TexCoord6_Ddx,
-	TexCoord7_Ddx,
-
-	TexCoord0_Ddy,
-	TexCoord1_Ddy,
-	TexCoord2_Ddy,
-	TexCoord3_Ddy,
-	TexCoord4_Ddy,
-	TexCoord5_Ddy,
-	TexCoord6_Ddy,
-	TexCoord7_Ddy,
-
-	LightmapTexCoord,
-	LightmapTexCoord_Ddx,
-	LightmapTexCoord_Ddy,
-
-	TwoSidedSign,
-	VertexColor,
-	VertexColor_Ddx,
-	VertexColor_Ddy,
-
-	WorldPosition,
-	WorldPosition_NoOffsets,
-	TranslatedWorldPosition,
-	TranslatedWorldPosition_NoOffsets,
-	ActorWorldPosition,
-
-	PrevWorldPosition,
-	PrevWorldPosition_NoOffsets,
-	PrevTranslatedWorldPosition,
-	PrevTranslatedWorldPosition_NoOffsets,
-
-	WorldPosition_Ddx,
-	WorldPosition_Ddy,
-
-	WorldVertexNormal,
-	WorldVertexTangent,
-	WorldNormal,
-	WorldReflection,
-
-	ViewportUV,
-	PixelPosition,
-	ViewSize,
-	RcpViewSize,
-	FieldOfView,
-	TanHalfFieldOfView,
-	CotanHalfFieldOfView,
-	TemporalSampleCount,
-	TemporalSampleIndex,
-	TemporalSampleOffset,
-	PreExposure,
-	RcpPreExposure,
-	RuntimeVirtualTextureOutputLevel,
-	RuntimeVirtualTextureOutputDerivative,
-	RuntimeVirtualTextureMaxLevel,
-
-	CameraVector,
-	CameraWorldPosition,
-	ViewWorldPosition,
-	PreViewTranslation,
-	TangentToWorld,
-	LocalToWorld,
-	WorldToLocal,
-	TranslatedWorldToCameraView,
-	TranslatedWorldToView,
-	CameraViewToTranslatedWorld,
-	ViewToTranslatedWorld,
-	WorldToParticle,
-	WorldToInstance,
-	ParticleToWorld,
-	InstanceToWorld,
-
-	PrevFieldOfView,
-	PrevTanHalfFieldOfView,
-	PrevCotanHalfFieldOfView,
-
-	PrevCameraWorldPosition,
-	PrevViewWorldPosition,
-	PrevPreViewTranslation,
-	PrevLocalToWorld,
-	PrevWorldToLocal,
-	PrevTranslatedWorldToCameraView,
-	PrevTranslatedWorldToView,
-	PrevCameraViewToTranslatedWorld,
-	PrevViewToTranslatedWorld,
-
-	PixelDepth,
-	PixelDepth_Ddx,
-	PixelDepth_Ddy,
-
-	GameTime,
-	RealTime,
-	DeltaTime,
-
-	PrevGameTime,
-	PrevRealTime,
-
-	ParticleColor,
-	ParticleTranslatedWorldPosition,
-	ParticleRadius,
-};
-static constexpr int32 MaxNumTexCoords = 8;
-
-struct FExternalInputDescription
-{
-	FExternalInputDescription(const TCHAR* InName, Shader::EValueType InType, EExternalInput InDdx = EExternalInput::None, EExternalInput InDdy = EExternalInput::None, EExternalInput InPreviousFrame = EExternalInput::None)
-		: Name(InName), Type(InType), Ddx(InDdx), Ddy(InDdy), PreviousFrame(InPreviousFrame)
-	{}
-
-	const TCHAR* Name;
-	Shader::EValueType Type;
-	EExternalInput Ddx;
-	EExternalInput Ddy;
-	EExternalInput PreviousFrame;
-};
-
-FExternalInputDescription GetExternalInputDescription(EExternalInput Input);
-
-inline bool IsTexCoord(EExternalInput Type)
-{
-	return FMath::IsWithin((int32)Type, (int32)EExternalInput::TexCoord0, (int32)EExternalInput::TexCoord0 + MaxNumTexCoords);
-}
-inline bool IsTexCoord_Ddx(EExternalInput Type)
-{
-	return FMath::IsWithin((int32)Type, (int32)EExternalInput::TexCoord0_Ddx, (int32)EExternalInput::TexCoord0_Ddx + MaxNumTexCoords);
-}
-inline bool IsTexCoord_Ddy(EExternalInput Type)
-{
-	return FMath::IsWithin((int32)Type, (int32)EExternalInput::TexCoord0_Ddy, (int32)EExternalInput::TexCoord0_Ddy + MaxNumTexCoords);
-}
-inline EExternalInput MakeInputTexCoord(int32 Index)
-{
-	check(Index >= 0 && Index < MaxNumTexCoords);
-	return (EExternalInput)((int32)EExternalInput::TexCoord0 + Index);
-}
-
-class FExpressionExternalInput : public FExpression
-{
-public:
-	FExpressionExternalInput(EExternalInput InInputType) : InputType(InInputType) {}
-
-	EExternalInput InputType;
-
-	virtual void ComputeAnalyticDerivatives(FTree& Tree, FExpressionDerivatives& OutResult) const override;
-	virtual FExpression* ComputePreviousFrame(FTree& Tree, const FRequestedType& RequestedType) const override;
-	virtual bool PrepareValue(FEmitContext& Context, FEmitScope& Scope, const FRequestedType& RequestedType, FPrepareValueResult& OutResult) const override;
-	virtual void EmitValueShader(FEmitContext& Context, FEmitScope& Scope, const FRequestedType& RequestedType, FEmitValueShaderResult& OutResult) const override;
-};
-
-class FExpressionMaterialSceneTexture : public FExpression
-{
-public:
-	FExpressionMaterialSceneTexture(FExpression* InTexCoordExpression, uint32 InSceneTextureId, bool bInFiltered)
-		: TexCoordExpression(InTexCoordExpression)
-		, SceneTextureId(InSceneTextureId)
-		, bFiltered(bInFiltered)
-	{}
-
-	FExpression* TexCoordExpression;
-	uint32 SceneTextureId;
-	bool bFiltered;
-
-	virtual bool PrepareValue(FEmitContext& Context, FEmitScope& Scope, const FRequestedType& RequestedType, FPrepareValueResult& OutResult) const override;
-	virtual void EmitValueShader(FEmitContext& Context, FEmitScope& Scope, const FRequestedType& RequestedType, FEmitValueShaderResult& OutResult) const override;
-};
-
-struct FMaterialNoiseParameters
-{
-	FMaterialNoiseParameters() { FMemory::Memzero(*this); }
-
-	int32 Quality;
-	int32 Levels;
-	float Scale;
-	uint32 RepeatSize;
-	float OutputMin;
-	float OutputMax;
-	float LevelScale;
-	uint8 NoiseFunction;
-	bool bTiling;
-	bool bTurbulence;
-};
-
-class FExpressionMaterialNoise : public FExpression
-{
-public:
-	FExpressionMaterialNoise(const FMaterialNoiseParameters& InParams, FExpression* InPositionExpression, FExpression* InFilterWidthExpression)
-		: PositionExpression(InPositionExpression)
-		, FilterWidthExpression(InFilterWidthExpression)
-		, Parameters(InParams)
-	{}
-
-	FExpression* PositionExpression;
-	FExpression* FilterWidthExpression;
-	FMaterialNoiseParameters Parameters;
-
-	virtual bool PrepareValue(FEmitContext& Context, FEmitScope& Scope, const FRequestedType& RequestedType, FPrepareValueResult& OutResult) const override;
-	virtual void EmitValueShader(FEmitContext& Context, FEmitScope& Scope, const FRequestedType& RequestedType, FEmitValueShaderResult& OutResult) const override;
-};
-
 class FExpressionTextureSample : public FExpression
 {
 public:
@@ -516,5 +253,4 @@ public:
 	virtual void EmitPreshader(FEmitContext& Context, FEmitScope& Scope, const FRequestedType& RequestedType, TArrayView<const FEmitPreshaderScope> Scopes, Shader::FPreshaderData& OutPreshader) const override;
 };
 
-}
-}
+} // namespace UE::HLSLTree
diff --git a/Engine/Source/Runtime/Engine/Public/HLSLTree/HLSLTreeEmit.h b/Engine/Source/Runtime/Engine/Public/HLSLTree/HLSLTreeEmit.h
--- a/Engine/Source/Runtime/Engine/Public/HLSLTree/HLSLTreeEmit.h
+++ b/Engine/Source/Runtime/Engine/Public/HLSLTree/HLSLTreeEmit.h
@@ -7,23 +7,20 @@
 #include "Containers/StringView.h"
 #include "Misc/StringBuilder.h"
 #include "Misc/MemStack.h"
+#include "Misc/GeneratedTypeName.h"
 #include "RHIDefinitions.h"
 #include "HLSLTree/HLSLTreeTypes.h"
 #include "HLSLTree/HLSLTreeHash.h"
 
 class FMaterial;
 class FMaterialCompilationOutput;
-struct FStaticParameterSet;
 
-namespace UE
-{
-
-namespace Shader
+namespace UE::Shader
 {
 class FPreshaderData;
-} // namespace Shader
+}
 
-namespace HLSLTree
+namespace UE::HLSLTree
 {
 
 class FErrorHandlerInterface;
@@ -37,9 +34,6 @@ class FEmitScope;
 class FEmitShaderExpression;
 class FEmitShaderStatement;
 
-// TODO - remove dependency
-enum class EExternalInput : uint8;
-
 struct FEmitPreshaderScope;
 
 struct FEmitShaderScopeEntry
@@ -256,6 +250,37 @@ public:
 	explicit FEmitContext(FMemStackBase& InAllocator, FErrorHandlerInterface& InErrors, const Shader::FStructTypeRegistry& InTypeRegistry);
 	~FEmitContext();
 
+	template<typename T>
+	T& AcquireData()
+	{
+		const FXxHash64 TypeHash = HashValue(GetGeneratedTypeName<T>());
+		TCustomDataWrapper<T>* Wrapper = static_cast<TCustomDataWrapper<T>*>(InternalFindData(TypeHash));
+		if (!Wrapper)
+		{
+			Wrapper = new TCustomDataWrapper<T>;
+			InternalRegisterData(TypeHash, Wrapper);
+		}
+		return Wrapper->Data;
+	}
+
+	template<typename T>
+	const T& FindData() const
+	{
+		const FXxHash64 TypeHash = HashValue(GetGeneratedTypeName<T>());
+		TCustomDataWrapper<T>* Wrapper = static_cast<TCustomDataWrapper<T>*>(InternalFindData(TypeHash));
+		check(Wrapper);
+		return Wrapper->Data;
+	}
+
+	template<typename T>
+	T& FindData()
+	{
+		const FXxHash64 TypeHash = HashValue(GetGeneratedTypeName<T>());
+		TCustomDataWrapper<T>* Wrapper = static_cast<TCustomDataWrapper<T>*>(InternalFindData(TypeHash));
+		check(Wrapper);
+		return Wrapper->Data;
+	}
+
 	void EmitDeclarationsCode(FStringBuilderBase& OutCode);
 
 	FPreparedType PrepareExpression(FExpression* InExpression, FEmitScope& Scope, const FRequestedType& RequestedType);
@@ -426,6 +451,22 @@ public:
 	/** FStringViews passed to this method are expected to reference persistent memory */
 	FEmitShaderExpression* EmitCustomHLSL(FEmitScope& Scope, FStringView DeclarationCode, FStringView FunctionCode, TConstArrayView<FCustomHLSLInput> Inputs, const Shader::FStructType* OutputType);
 
+	class FCustomDataWrapper
+	{
+	public:
+		virtual ~FCustomDataWrapper() {}
+	};
+
+	template<typename T>
+	class TCustomDataWrapper : public FCustomDataWrapper
+	{
+	public:
+		T Data;
+	};
+
+	void InternalRegisterData(FXxHash64 Hash, FCustomDataWrapper* Data);
+	FCustomDataWrapper* InternalFindData(FXxHash64 Hash) const;
+
 	FMemStackBase* Allocator = nullptr;
 	FErrorHandlerInterface* Errors = nullptr;
 	const Shader::FStructTypeRegistry* TypeRegistry = nullptr;
@@ -442,28 +483,19 @@ public:
 	TMap<FXxHash64, FEmitCustomHLSL> EmitCustomHLSLMap;
 	TArray<struct FPreshaderLoopScope*> PreshaderLoopScopes;
 	TArray<const FPreshaderLocalPHIScope*> PreshaderLocalPHIScopes;
+	TMap<FXxHash64, FCustomDataWrapper*> CustomDataMap;
 	int32 PreshaderStackPosition = 0;
 
-	// TODO - remove preshader material dependency
+	int32 NumExpressionLocals = 0;
+	int32 NumExpressionLocalPHIs = 0;
+
+	// TODO - Material values required for preshaders, need to decouple preshaders from material system
 	const FMaterial* Material = nullptr;
-	const FStaticParameterSet* StaticParameters = nullptr;
 	FMaterialCompilationOutput* MaterialCompilationOutput = nullptr;
-	TMap<Shader::FValue, uint32> DefaultUniformValues;
 	uint32 UniformPreshaderOffset = 0u;
 	uint32 CurrentBoolUniformOffset = ~0u;
 	uint32 CurrentNumBoolComponents = 32u;
-	bool bReadMaterialNormal = false;
-	bool bUsesVertexColor = false;
-	bool bNeedsParticlePosition = false;
-	bool bUsesParticleColor = false;
-	TBitArray<> ExternalInputMask[SF_NumFrequencies];
-	FMaterialShadingModelField ShadingModelsFromCompilation;
-
-	bool IsExternalInputUsed(EShaderFrequency Frequency, EExternalInput Input) const { return ExternalInputMask[Frequency][(uint8)Input]; }
-
-	int32 NumExpressionLocals = 0;
-	int32 NumExpressionLocalPHIs = 0;
 };
 
-} // namespace HLSLTree
-} // namespace UE
+} // namespace UE::HLSLTree
+
diff --git a/Engine/Source/Runtime/Engine/Public/HLSLTree/HLSLTreeHash.h b/Engine/Source/Runtime/Engine/Public/HLSLTree/HLSLTreeHash.h
--- a/Engine/Source/Runtime/Engine/Public/HLSLTree/HLSLTreeHash.h
+++ b/Engine/Source/Runtime/Engine/Public/HLSLTree/HLSLTreeHash.h
@@ -59,6 +59,11 @@ inline void AppendHash(FHasher& Hasher, FStringView Value)
 	Hasher.AppendData(Value.GetData(), Value.Len() * sizeof(TCHAR));
 }
 
+inline void AppendHash(FHasher& Hasher, const TCHAR* Value)
+{
+	Hasher.AppendData(Value, FCString::Strlen(Value) * sizeof(TCHAR));
+}
+
 inline void AppendHash(FHasher& Hasher, const FCustomHLSLInput& Value)
 {
 	AppendHash(Hasher, Value.Name);
@@ -89,5 +94,13 @@ inline void AppendHashes(FHasher& Hasher, const T& Value, ArgTypes&&... Args)
 	AppendHashes(Hasher, Forward<ArgTypes>(Args)...);
 }
 
+template<typename T>
+inline FXxHash64 HashValue(const T& Value)
+{
+	FHasher Hasher;
+	AppendHash(Hasher, Value);
+	return Hasher.Finalize();
+}
+
 } // namespace HLSLTree
 } // namespace UE
diff --git a/Engine/Source/Runtime/Engine/Public/HLSLTree/HLSLTreeTypes.h b/Engine/Source/Runtime/Engine/Public/HLSLTree/HLSLTreeTypes.h
--- a/Engine/Source/Runtime/Engine/Public/HLSLTree/HLSLTreeTypes.h
+++ b/Engine/Source/Runtime/Engine/Public/HLSLTree/HLSLTreeTypes.h
@@ -7,15 +7,12 @@
 
 class UTexture;
 
-namespace UE
-{
-
-namespace Shader
+namespace UE::Shader
 {
 enum class EPreshaderOpcode : uint8;
-} // namespace Shader
+}
 
-namespace HLSLTree
+namespace UE::HLSLTree
 {
 
 class FNode;
@@ -142,5 +139,5 @@ struct FCustomHLSLInput
 	FExpression* Expression = nullptr;
 };
 
-} // namespace HLSLTree
-} // namespace UE
+} // namespace UE::HLSLTree
+
