Fixed allow ray tracing default value.

diff --git a/Engine/Source/Runtime/Engine/Public/SceneView.h b/Engine/Source/Runtime/Engine/Public/SceneView.h
--- a/Engine/Source/Runtime/Engine/Public/SceneView.h
+++ b/Engine/Source/Runtime/Engine/Public/SceneView.h
@@ -1256,7 +1256,7 @@ public:
 	FIESLightProfileResource* IESLightProfileResource;
 
 	/** Use to allow ray tracing on this view. */
-	bool bAllowRayTracing = false;
+	bool bAllowRayTracing = true;
 #endif
 
 protected:
