Fixing issue with streaming sound waves from CL 1904090

diff --git a/Engine/Source/Runtime/Engine/Private/SoundWave.cpp b/Engine/Source/Runtime/Engine/Private/SoundWave.cpp
--- a/Engine/Source/Runtime/Engine/Private/SoundWave.cpp
+++ b/Engine/Source/Runtime/Engine/Private/SoundWave.cpp
@@ -2895,16 +2895,11 @@ bool USoundWave::IsStreaming(const FPlatformAudioCookOverrides& Overrides) const
 	if (bIsForceInline || bProcedural)
 	{
 		SoundWaveDataPtr->bIsStreaming = false;
-		return false;
 	}
-	else if (bStreaming)
+	else
 	{
 		SoundWaveDataPtr->bIsStreaming = true;
-		return true;
 	}
-
-	// For stream caching, the auto streaming threshold is used to force sounds to be inlined:
-	SoundWaveDataPtr->bIsStreaming = (Overrides.AutoStreamingThreshold > SMALL_NUMBER && Duration > Overrides.AutoStreamingThreshold);
 	return SoundWaveDataPtr->bIsStreaming;
 }
 
