'Glider_ID_107_IceMaiden' animation continually pops after deployment in the Frontend on Switch
#fix test for _both_ cooked and editor_filter flags on package when checking whether raw animation data can be evaluated
#misc added checks ensuring raw data is available when forcefully requesting evaluation

diff --git a/Engine/Source/Runtime/Engine/Private/Animation/AnimSequence.cpp b/Engine/Source/Runtime/Engine/Private/Animation/AnimSequence.cpp
--- a/Engine/Source/Runtime/Engine/Private/Animation/AnimSequence.cpp
+++ b/Engine/Source/Runtime/Engine/Private/Animation/AnimSequence.cpp
@@ -1508,6 +1508,7 @@ void UAnimSequence::GetBonePose(FAnimationPoseData& OutAnimationPoseData, const
 	FCompactPose& OutPose = OutAnimationPoseData.GetPose();
 
 	FBoneContainer& RequiredBones = OutPose.GetBoneContainer();
+	check(!bForceUseRawData || CanEvaluateRawAnimationData());
 	const bool bUseRawDataForPoseExtraction = (CanEvaluateRawAnimationData() && bForceUseRawData) || UseRawDataForPoseExtraction(RequiredBones);
 
 	const bool bIsBakedAdditive = !bUseRawDataForPoseExtraction && IsValidAdditive();
@@ -4378,6 +4379,7 @@ void UAnimSequence::EvaluateCurveData(FBlendedCurve& OutCurve, float CurrentTime
 	}
 
 	const bool bEvaluateRawData = bUseRawDataOnly || bForceUseRawData || !IsCurveCompressedDataValid();
+	check(!bForceUseRawData || CanEvaluateRawAnimationData());
 	if (CanEvaluateRawAnimationData() && bEvaluateRawData)
 	{
 		Super::EvaluateCurveData(OutCurve, CurrentTime, bForceUseRawData);
@@ -4394,6 +4396,7 @@ float UAnimSequence::EvaluateCurveData(SmartName::UID_Type CurveUID, float Curre
 	SCOPE_CYCLE_COUNTER(STAT_AnimSeq_EvalCurveData);
 
 	const bool bEvaluateRawData = bUseRawDataOnly || bForceUseRawData || !IsCurveCompressedDataValid();
+	check(!bForceUseRawData || CanEvaluateRawAnimationData());
 	if (CanEvaluateRawAnimationData() && bEvaluateRawData)
 	{
 		return Super::EvaluateCurveData(CurveUID, CurrentTime, bForceUseRawData);
@@ -4409,6 +4412,7 @@ float UAnimSequence::EvaluateCurveData(SmartName::UID_Type CurveUID, float Curre
 bool UAnimSequence::HasCurveData(SmartName::UID_Type CurveUID, bool bForceUseRawData) const
 {
 	const bool bEvaluateRawData = bUseRawDataOnly || bForceUseRawData || !IsCurveCompressedDataValid();
+	check(!bForceUseRawData || CanEvaluateRawAnimationData());
 	if (CanEvaluateRawAnimationData() && bEvaluateRawData)
 	{
 		return Super::HasCurveData(CurveUID, bForceUseRawData);
diff --git a/Engine/Source/Runtime/Engine/Private/Animation/AnimSequenceBase.cpp b/Engine/Source/Runtime/Engine/Private/Animation/AnimSequenceBase.cpp
--- a/Engine/Source/Runtime/Engine/Private/Animation/AnimSequenceBase.cpp
+++ b/Engine/Source/Runtime/Engine/Private/Animation/AnimSequenceBase.cpp
@@ -1315,7 +1315,7 @@ bool UAnimSequenceBase::ShouldDataModelBeValid() const
 {
 	return
 #if WITH_EDITOR
-		!GetOutermost()->HasAnyPackageFlags(PKG_FilterEditorOnly);
+		!GetOutermost()->HasAnyPackageFlags(PKG_Cooked);
 #else
 		false;
 #endif
