- Defaulted world partition runtime grid levels to be aligned on child levels to fix cases where runtime cells could be activated without its overlapping parent level cells being activated.
- Partitioning of spatially loaded actors now always uses actor bounds to find its grid cell (fixes missing actor overlapping another actor bounds)
- Fixed the unaligned version of the runtime spatial hash grid levels.
- Added option to use the unaligned version using the CVar wp.Runtime.RuntimeSpatialHashUseAlignedGridLevels (defaulted to true).
- Adapted the Runtime Grid previewer to respect the value set by wp.Runtime.RuntimeSpatialHashUseAlignedGridLevels.

diff --git a/Engine/Source/Runtime/Engine/Private/WorldPartition/RuntimeSpatialHash/RuntimeSpatialHashGridHelper.cpp b/Engine/Source/Runtime/Engine/Private/WorldPartition/RuntimeSpatialHash/RuntimeSpatialHashGridHelper.cpp
--- a/Engine/Source/Runtime/Engine/Private/WorldPartition/RuntimeSpatialHash/RuntimeSpatialHashGridHelper.cpp
+++ b/Engine/Source/Runtime/Engine/Private/WorldPartition/RuntimeSpatialHash/RuntimeSpatialHashGridHelper.cpp
@@ -4,6 +4,12 @@
 #include "Algo/Transform.h"
 #include "ProfilingDebugging/ScopedTimers.h"
 
+bool GRuntimeSpatialHashUseAlignedGridLevels = true;
+static FAutoConsoleVariableRef CVarUseGrow(
+	TEXT("wp.Runtime.RuntimeSpatialHashUseAlignedGridLevels"),
+	GRuntimeSpatialHashUseAlignedGridLevels,
+	TEXT("Set RuntimeSpatialHashUseAlignedGridLevels to false to help break the pattern caused by world partition promotion of actors to upper grid levels that are always aligned on child levels."));
+
 FSquare2DGridHelper::FSquare2DGridHelper(const FBox& InWorldBounds, const FVector& InOrigin, int32 InCellSize)
 	: WorldBounds(InWorldBounds)
 	, Origin(InOrigin)
@@ -40,12 +46,17 @@ FSquare2DGridHelper::FSquare2DGridHelper(const FBox& InWorldBounds, const FVecto
 	int32 CurrentGridSize = GridSize;
 	for (int32 Level = 0; Level < GridLevelCount; ++Level)
 	{
-		// Except for top level, adding 1 to CurrentGridSize (which is always a power of 2) breaks the pattern of perfectly aligned cell edges between grid level cells.
-		// This will prevent weird artefact during actor promotion when an actor is placed using its bounds and which overlaps multiple cells.
-		// In this situation, the algorithm will try to find a cell that encapsulates completely the actor's bounds by searching in the upper levels, until it finds one.
-		// Also note that, the default origin of each level will always be centered at the middle of the bounds of (level's cellsize * level's grid size).
-		int32 LevelGridSize = (Level == GridLevelCount-1) ? CurrentGridSize : CurrentGridSize + 1;
-		
+		int32 LevelGridSize = CurrentGridSize;
+
+		if (!GRuntimeSpatialHashUseAlignedGridLevels)
+		{
+			// Except for top level, adding 1 to CurrentGridSize (which is always a power of 2) breaks the pattern of perfectly aligned cell edges between grid level cells.
+			// This will prevent weird artefact during actor promotion when an actor is placed using its bounds and which overlaps multiple cells.
+			// In this situation, the algorithm will try to find a cell that encapsulates completely the actor's bounds by searching in the upper levels, until it finds one.
+			// Also note that, the default origin of each level will always be centered at the middle of the bounds of (level's cellsize * level's grid size).
+			LevelGridSize = (Level == GridLevelCount - 1) ? CurrentGridSize : CurrentGridSize + 1;
+		}
+
 		Levels.Emplace(FVector2D(InOrigin), CurrentCellSize, LevelGridSize, Level);
 
 		CurrentCellSize <<= 1;
@@ -71,11 +82,11 @@ void FSquare2DGridHelper::ForEachCells(TFunctionRef<void(const FSquare2DGridHelp
 }
 #endif
 
-int32 FSquare2DGridHelper::ForEachIntersectingCells(const FBox& InBox, TFunctionRef<void(const FIntVector&)> InOperation) const
+int32 FSquare2DGridHelper::ForEachIntersectingCells(const FBox& InBox, TFunctionRef<void(const FIntVector&)> InOperation, int32 InStartLevel) const
 {
 	int32 NumCells = 0;
 
-	for (int32 Level = 0; Level < Levels.Num(); Level++)
+	for (int32 Level = InStartLevel; Level < Levels.Num(); Level++)
 	{
 		NumCells += Levels[Level].ForEachIntersectingCells(InBox, [InBox, Level, InOperation](const FIntVector2& Coord) { InOperation(FIntVector(Coord.X, Coord.Y, Level)); });
 	}
@@ -83,11 +94,11 @@ int32 FSquare2DGridHelper::ForEachIntersectingCells(const FBox& InBox, TFunction
 	return NumCells;
 }
 
-int32 FSquare2DGridHelper::ForEachIntersectingCells(const FSphere& InSphere, TFunctionRef<void(const FIntVector&)> InOperation) const
+int32 FSquare2DGridHelper::ForEachIntersectingCells(const FSphere& InSphere, TFunctionRef<void(const FIntVector&)> InOperation, int32 InStartLevel) const
 {
 	int32 NumCells = 0;
 
-	for (int32 Level = 0; Level < Levels.Num(); Level++)
+	for (int32 Level = InStartLevel; Level < Levels.Num(); Level++)
 	{
 		NumCells += Levels[Level].ForEachIntersectingCells(InSphere, [InSphere, Level, InOperation](const FIntVector2& Coord) { InOperation(FIntVector(Coord.X, Coord.Y, Level)); });
 	}
@@ -95,11 +106,11 @@ int32 FSquare2DGridHelper::ForEachIntersectingCells(const FSphere& InSphere, TFu
 	return NumCells;
 }
 
-int32 FSquare2DGridHelper::ForEachIntersectingCells(const FSphericalSector& InShape, TFunctionRef<void(const FIntVector&)> InOperation) const
+int32 FSquare2DGridHelper::ForEachIntersectingCells(const FSphericalSector& InShape, TFunctionRef<void(const FIntVector&)> InOperation, int32 InStartLevel) const
 {
 	int32 NumCells = 0;
 
-	for (int32 Level = 0; Level < Levels.Num(); Level++)
+	for (int32 Level = InStartLevel; Level < Levels.Num(); Level++)
 	{
 		NumCells += Levels[Level].ForEachIntersectingCells(InShape, [Level, InOperation](const FIntVector2& Coord) { InOperation(FIntVector(Coord.X, Coord.Y, Level)); });
 	}
@@ -160,7 +171,6 @@ FSquare2DGridHelper GetPartitionedActors(const UWorldPartition* WorldPartition,
 		}
 	}
 
-	const float CellArea = PartitionedActors.GetLowestLevel().CellSize * PartitionedActors.GetLowestLevel().CellSize;
 	for (const FActorClusterInstance* ClusterInstance : GridActors)
 	{
 		const FActorCluster* ActorCluster = ClusterInstance->Cluster;
@@ -170,36 +180,24 @@ FSquare2DGridHelper GetPartitionedActors(const UWorldPartition* WorldPartition,
 
 		if (ActorCluster->bIsSpatiallyLoaded)
 		{
-			if (ClusterBounds.GetArea() <= CellArea)
-			{
-				// Find grid level cell that copntains the actor cluster pivot and put actors in it.
-				FIntVector2 CellCoords;
-				if (PartitionedActors.GetLowestLevel().GetCellCoords(ClusterBounds.GetCenter(), CellCoords))
-				{
-					GridCell = &PartitionedActors.GetLowestLevel().GetCell(CellCoords);
-				}
-			}
-			else
+			// Find grid level cell that encompasses the actor cluster bounding box and put actors in it.
+			const FVector2D ClusterSize = ClusterBounds.GetSize();
+			const float MinRequiredCellExtent = FMath::Max(ClusterSize.X, ClusterSize.Y);
+			const int32 FirstPotentialGridLevel = FMath::Max(FMath::CeilToFloat(FMath::Log2(MinRequiredCellExtent / (float)PartitionedActors.CellSize)), 0);
+
+			for (int32 GridLevelIndex = FirstPotentialGridLevel; GridLevelIndex < PartitionedActors.Levels.Num(); GridLevelIndex++)
 			{
-				// Find grid level cell that encompasses the actor cluster bounding box and put actors in it.
-				const FVector2D ClusterSize = ClusterBounds.GetSize();
-				const float MinRequiredCellExtent = FMath::Max(ClusterSize.X, ClusterSize.Y);
-				const int32 FirstPotentialGridLevel = FMath::Max(FMath::CeilToFloat(FMath::Log2(MinRequiredCellExtent / (float)PartitionedActors.CellSize)), 0);
+				FSquare2DGridHelper::FGridLevel& GridLevel = PartitionedActors.Levels[GridLevelIndex];
 
-				for (int32 GridLevelIndex = FirstPotentialGridLevel; GridLevelIndex < PartitionedActors.Levels.Num(); GridLevelIndex++)
+				if (GridLevel.GetNumIntersectingCells(ClusterInstance->Bounds) == 1)
 				{
-					FSquare2DGridHelper::FGridLevel& GridLevel = PartitionedActors.Levels[GridLevelIndex];
-
-					if (GridLevel.GetNumIntersectingCells(ClusterInstance->Bounds) == 1)
+					GridLevel.ForEachIntersectingCells(ClusterInstance->Bounds, [&GridLevel, ActorCluster, ClusterInstance, &GridCell](const FIntVector2& Coords)
 					{
-						GridLevel.ForEachIntersectingCells(ClusterInstance->Bounds, [&GridLevel, ActorCluster, ClusterInstance, &GridCell](const FIntVector2& Coords)
-						{
-							check(!GridCell);
-							GridCell = &GridLevel.GetCell(Coords);
-						});
-
-						break;
-					}
+						check(!GridCell);
+						GridCell = &GridLevel.GetCell(Coords);
+					});
+
+					break;
 				}
 			}
 		}
diff --git a/Engine/Source/Runtime/Engine/Private/WorldPartition/RuntimeSpatialHash/RuntimeSpatialHashGridHelper.h b/Engine/Source/Runtime/Engine/Private/WorldPartition/RuntimeSpatialHash/RuntimeSpatialHashGridHelper.h
--- a/Engine/Source/Runtime/Engine/Private/WorldPartition/RuntimeSpatialHash/RuntimeSpatialHashGridHelper.h
+++ b/Engine/Source/Runtime/Engine/Private/WorldPartition/RuntimeSpatialHash/RuntimeSpatialHashGridHelper.h
@@ -10,6 +10,8 @@
 #include "WorldPartition/WorldPartitionActorCluster.h"
 #include "WorldPartition/DataLayer/DataLayersID.h"
 
+extern ENGINE_API bool GRuntimeSpatialHashUseAlignedGridLevels;
+
 /**
   * Square 2D grid helper
   */
@@ -445,21 +447,21 @@ struct FSquare2DGridHelper
 	 *
 	 * @return the number of intersecting cells
 	 */
-	int32 ForEachIntersectingCells(const FBox& InBox, TFunctionRef<void(const FIntVector&)> InOperation) const;
+	int32 ForEachIntersectingCells(const FBox& InBox, TFunctionRef<void(const FIntVector&)> InOperation, int32 InStartLevel = 0) const;
 
 	/**
 	 * Runs a function on all intersecting cells for the provided sphere
 	 *
 	 * @return the number of intersecting cells
 	 */
-	int32 ForEachIntersectingCells(const FSphere& InSphere, TFunctionRef<void(const FIntVector&)> InOperation) const;
+	int32 ForEachIntersectingCells(const FSphere& InSphere, TFunctionRef<void(const FIntVector&)> InOperation, int32 InStartLevel = 0) const;
 
 	/**
 	 * Runs a function on all intersecting cells for the provided spherical sector
 	 *
 	 * @return the number of intersecting cells
 	 */
-	int32 ForEachIntersectingCells(const FSphericalSector& InShape, TFunctionRef<void(const FIntVector&)> InOperation) const;
+	int32 ForEachIntersectingCells(const FSphericalSector& InShape, TFunctionRef<void(const FIntVector&)> InOperation, int32 InStartLevel = 0) const;
 
 #if WITH_EDITOR
 	// Validates that actor is not referenced by multiple cells
diff --git a/Engine/Source/Runtime/Engine/Private/WorldPartition/WorldPartitionRuntimeSpatialHash.cpp b/Engine/Source/Runtime/Engine/Private/WorldPartition/WorldPartitionRuntimeSpatialHash.cpp
--- a/Engine/Source/Runtime/Engine/Private/WorldPartition/WorldPartitionRuntimeSpatialHash.cpp
+++ b/Engine/Source/Runtime/Engine/Private/WorldPartition/WorldPartitionRuntimeSpatialHash.cpp
@@ -159,10 +159,9 @@ void FSpatialHashStreamingGrid::GetCells(const FWorldPartitionStreamingQuerySour
 		{
 			Helper.ForEachIntersectingCells(Shape, [&](const FIntVector& Coords)
 			{
-				if (const int32* LayerCellIndexPtr = GridLevels[Coords.Z].LayerCellsMapping.Find(Coords.Y * Helper.Levels[Coords.Z].GridSize + Coords.X))
+				if (const FSpatialHashStreamingGridLayerCell* LayerCell = GetLayerCell(Coords))
 				{
-					const FSpatialHashStreamingGridLayerCell& LayerCell = GridLevels[Coords.Z].LayerCells[*LayerCellIndexPtr];
-					for (const UWorldPartitionRuntimeCell* Cell : LayerCell.GridCells)
+					for (const UWorldPartitionRuntimeCell* Cell : LayerCell->GridCells)
 					{
 						if (ShouldAddCell(Cell, QuerySource))
 						{
@@ -192,6 +191,9 @@ void FSpatialHashStreamingGrid::GetCells(const TArray<FWorldPartitionStreamingSo
 {
 	TRACE_CPUPROFILER_EVENT_SCOPE(FSpatialHashStreamingGrid::GetCells);
 	
+	typedef TMap<FIntVector, TArray<UWorldPartitionRuntimeCell::FStreamingSourceInfo>> FIntersectingCells;
+	FIntersectingCells AllActivatedCells;
+
 	const float GridLoadingRange = GetLoadingRange();
 	const FSquare2DGridHelper& Helper = GetGridHelper();
 	for (const FWorldPartitionStreamingSource& Source : Sources)
@@ -201,15 +203,15 @@ void FSpatialHashStreamingGrid::GetCells(const TArray<FWorldPartitionStreamingSo
 			UWorldPartitionRuntimeCell::FStreamingSourceInfo Info(Source, Shape);
 
 			Helper.ForEachIntersectingCells(Shape, [&](const FIntVector& Coords)
-			{ 
+			{
+				bool bAddedActivatedCell = false;
 #if !UE_BUILD_SHIPPING
 				if ((GFilterRuntimeSpatialHashGridLevel == INDEX_NONE) || (GFilterRuntimeSpatialHashGridLevel == Coords.Z))
 #endif
 				{
-					if (const int32* LayerCellIndexPtr = GridLevels[Coords.Z].LayerCellsMapping.Find(Coords.Y * Helper.Levels[Coords.Z].GridSize + Coords.X))
+					if (const FSpatialHashStreamingGridLayerCell* LayerCell = GetLayerCell(Coords))
 					{
-						const FSpatialHashStreamingGridLayerCell& LayerCell = GridLevels[Coords.Z].LayerCells[*LayerCellIndexPtr];
-						for (const UWorldPartitionRuntimeCell* Cell : LayerCell.GridCells)
+						for (const UWorldPartitionRuntimeCell* Cell : LayerCell->GridCells)
 						{
 							if (!Cell->HasDataLayers() || (DataLayerSubsystem && DataLayerSubsystem->IsAnyDataLayerInEffectiveRuntimeState(Cell->GetDataLayers(), EDataLayerRuntimeState::Activated)))
 							{
@@ -221,6 +223,7 @@ void FSpatialHashStreamingGrid::GetCells(const TArray<FWorldPartitionStreamingSo
 								{
 									check(Source.TargetState == EStreamingSourceTargetState::Activated);
 									OutActivateCells.AddCell(Cell, Info);
+									bAddedActivatedCell = (!GRuntimeSpatialHashUseAlignedGridLevels);
 								}
 							}
 							else if (DataLayerSubsystem && DataLayerSubsystem->IsAnyDataLayerInEffectiveRuntimeState(Cell->GetDataLayers(), EDataLayerRuntimeState::Loaded))
@@ -230,11 +233,94 @@ void FSpatialHashStreamingGrid::GetCells(const TArray<FWorldPartitionStreamingSo
 						}
 					}
 				}
+				if (bAddedActivatedCell)
+				{
+					AllActivatedCells.FindOrAdd(Coords).Add(Info);
+				}
 			});
 		});
 	}
 
 	GetAlwaysLoadedCells(DataLayerSubsystem, OutActivateCells.GetCells(), OutLoadCells.GetCells());
+
+	if (!GRuntimeSpatialHashUseAlignedGridLevels)
+	{
+		auto FindIntersectingParents = [&Helper, this](const FIntersectingCells& InAllCells, const FIntersectingCells& InTestCells, FIntersectingCells& OutIntersectingCells)
+		{
+			bool bFound = false;
+			const int32 AlwaysLoadedLevel = Helper.Levels.Num() - 1;
+			for (const auto& InTestCell : InTestCells)
+			{
+				const FIntVector& TestCell = InTestCell.Key;
+				int32 CurrentLevelIndex = TestCell.Z;
+				int32 ParentLevelIndex = CurrentLevelIndex + 1;
+				// Only test with Parent Level if it's below the AlwaysLoaded Level
+				if (ParentLevelIndex < AlwaysLoadedLevel)
+				{
+					FBox2D CurrentLevelCellBounds;
+					Helper.Levels[CurrentLevelIndex].GetCellBounds(FIntVector2(TestCell.X, TestCell.Y), CurrentLevelCellBounds);
+					FBox Box(FVector(CurrentLevelCellBounds.Min, 0), FVector(CurrentLevelCellBounds.Max, 0));
+
+					Helper.ForEachIntersectingCells(Box, [&](const FIntVector& IntersectingCoords)
+					{
+						check(IntersectingCoords.Z >= ParentLevelIndex);
+						if (!InAllCells.Contains(IntersectingCoords))
+						{
+							if (!OutIntersectingCells.Contains(IntersectingCoords))
+							{
+								OutIntersectingCells.Add(IntersectingCoords, InTestCell.Value);
+								bFound = true;
+							}
+						}
+					}, ParentLevelIndex);
+				}
+			}
+			return bFound;
+		};
+	
+		FIntersectingCells AllParentCells;
+		FIntersectingCells TestCells = AllActivatedCells;
+		FIntersectingCells IntersectingCells;
+		bool bFound = false;
+		do
+		{
+			bFound = FindIntersectingParents(AllActivatedCells, TestCells, IntersectingCells);
+			if (bFound)
+			{
+				AllActivatedCells.Append(IntersectingCells);
+				AllParentCells.Append(IntersectingCells);
+				TestCells = MoveTemp(IntersectingCells);
+				check(IntersectingCells.IsEmpty());
+			}
+		} while (bFound);
+
+		for (const auto& ParentCell : AllParentCells)
+		{
+			if (const FSpatialHashStreamingGridLayerCell* LayerCell = GetLayerCell(ParentCell.Key))
+			{
+				for (const UWorldPartitionRuntimeCell* Cell : LayerCell->GridCells)
+				{
+					if (!Cell->HasDataLayers() || (DataLayerSubsystem && DataLayerSubsystem->IsAnyDataLayerInEffectiveRuntimeState(Cell->GetDataLayers(), EDataLayerRuntimeState::Activated)))
+					{
+						for (const auto& Info : ParentCell.Value)
+						{
+							OutActivateCells.AddCell(Cell, Info);
+						}
+					}
+				}
+			}
+		}
+	}
+}
+
+const FSpatialHashStreamingGridLayerCell* FSpatialHashStreamingGrid::GetLayerCell(const FIntVector& Coords) const
+{
+	check(GridLevels.IsValidIndex(Coords.Z));
+	if (const int32* LayerCellIndexPtr = GridLevels[Coords.Z].LayerCellsMapping.Find(Coords.Y * GetGridHelper().Levels[Coords.Z].GridSize + Coords.X))
+	{
+		return &GridLevels[Coords.Z].LayerCells[*LayerCellIndexPtr];
+	}
+	return nullptr;
 }
 
 void FSpatialHashStreamingGrid::GetAlwaysLoadedCells(const UDataLayerSubsystem* DataLayerSubsystem, TSet<const UWorldPartitionRuntimeCell*>& OutActivateCells, TSet<const UWorldPartitionRuntimeCell*>& OutLoadCells) const
diff --git a/Engine/Source/Runtime/Engine/Private/WorldPartition/WorldPartitionRuntimeSpatialHashGridPreviewer.cpp b/Engine/Source/Runtime/Engine/Private/WorldPartition/WorldPartitionRuntimeSpatialHashGridPreviewer.cpp
--- a/Engine/Source/Runtime/Engine/Private/WorldPartition/WorldPartitionRuntimeSpatialHashGridPreviewer.cpp
+++ b/Engine/Source/Runtime/Engine/Private/WorldPartition/WorldPartitionRuntimeSpatialHashGridPreviewer.cpp
@@ -2,6 +2,7 @@
 
 #include "WorldPartition/WorldPartitionRuntimeSpatialHashGridPreviewer.h"
 #include "WorldPartition/WorldPartitionRuntimeSpatialHash.h"
+#include "WorldPartition/RuntimeSpatialHash/RuntimeSpatialHashGridHelper.h"
 #include "Engine/PostProcessVolume.h"
 #include "Materials/MaterialInstanceDynamic.h"
 
@@ -73,6 +74,13 @@ void FWorldPartitionRuntimeSpatialHashGridPreviewer::Draw(UWorld* World, const T
 						MID->SetVectorParameterValue(*FString::Printf(TEXT("Grid%d_Color"), i), Grid.DebugColor);
 						CachedParameters.GridColor = Grid.DebugColor;
 					}
+
+					FVector GridOffset = GRuntimeSpatialHashUseAlignedGridLevels ? FVector(0.5 * Grid.CellSize) : FVector::ZeroVector;
+					if (CachedParameters.GridOffset != GridOffset)
+					{
+						MID->SetVectorParameterValue(*FString::Printf(TEXT("Grid%d_Offset"), i), GridOffset);
+						CachedParameters.GridOffset = GridOffset;
+					}
 				}
 			}
 		}
diff --git a/Engine/Source/Runtime/Engine/Public/WorldPartition/WorldPartitionRuntimeSpatialHash.h b/Engine/Source/Runtime/Engine/Public/WorldPartition/WorldPartitionRuntimeSpatialHash.h
--- a/Engine/Source/Runtime/Engine/Public/WorldPartition/WorldPartitionRuntimeSpatialHash.h
+++ b/Engine/Source/Runtime/Engine/Public/WorldPartition/WorldPartitionRuntimeSpatialHash.h
@@ -142,6 +142,7 @@ struct FSpatialHashStreamingGrid
 	void Draw3D(UWorld* World, const TArray<FWorldPartitionStreamingSource>& Sources, const FTransform& Transform) const;
 
 private:
+	const FSpatialHashStreamingGridLayerCell* GetLayerCell(const FIntVector& Coords) const;
 	void DrawStreamingSource2D(UCanvas* Canvas, const FSphericalSector& Shape, TFunctionRef<FVector2D(const FVector2D&)> WorldToScreen, const FColor& Color) const;
 	void DrawStreamingSource3D(UWorld* World, const FSphericalSector& Shape, const FTransform& Transform, const FColor& Color) const;
 	void GetFilteredCellsForDebugDraw(const FSpatialHashStreamingGridLayerCell* LayerCell, const UDataLayerSubsystem* DataLayerSubsystem, TArray<const UWorldPartitionRuntimeCell*>& FilteredCells) const;
diff --git a/Engine/Source/Runtime/Engine/Public/WorldPartition/WorldPartitionRuntimeSpatialHashGridPreviewer.h b/Engine/Source/Runtime/Engine/Public/WorldPartition/WorldPartitionRuntimeSpatialHashGridPreviewer.h
--- a/Engine/Source/Runtime/Engine/Public/WorldPartition/WorldPartitionRuntimeSpatialHashGridPreviewer.h
+++ b/Engine/Source/Runtime/Engine/Public/WorldPartition/WorldPartitionRuntimeSpatialHashGridPreviewer.h
@@ -37,12 +37,14 @@ struct ENGINE_API FWorldPartitionRuntimeSpatialHashGridPreviewer
 			, CellSize(0)
 			, LoadingRange(0.f)
 			, GridColor(0.f, 0.f, 0.f, 0.f)
+			, GridOffset(FVector::ZeroVector)
 		{}
 
 		bool bEnabled;
 		int32 CellSize;
 		float LoadingRange;
 		FLinearColor GridColor;
+		FVector GridOffset;
 	};
 
 	enum { MAX_PREVIEW_GRIDS = 4 };
