UE-141652 - Remove BP native make/break functionality for LWC types

Previously, vector types were using the single precision variants for their make&break functions. As a result, the split vectors would be floats, which led to precision loss. Going forward, we want to ensure that we use double precision variants of these functions. The single precision functions no longer serve any use, and will be deprecated.

Additionally, FVector_NetQuantize and friends can go back to using the default FVector make&break functions. We don't need to deprecate these functions since they only existed in 5.0. However, some content has already been saved with these functions, so we need to still setup redirects in BaseEngine.ini.

Finally, AssertMacros.h was missing several variants of GET_FUNCTION_NAME_CHECKED_NParams. Without these, GET_FUNCTION_NAME_CHECKED would fail to compile since KismetMathLibrary would have overloads of MakeVector and BreakVector.

diff --git a/Engine/Config/BaseEngine.ini b/Engine/Config/BaseEngine.ini
--- a/Engine/Config/BaseEngine.ini
+++ b/Engine/Config/BaseEngine.ini
@@ -1020,6 +1020,15 @@ NearClipPlane=10.0
 +FunctionRedirects=(OldName="KismetStringLibrary.BuildString_Float", NewName="BuildString_Double")
 +PropertyRedirects=(OldName="KismetStringLibrary.BuildString_Double.InFloat", NewName="BuildString_Double.InDouble")
 
++FunctionRedirects=(OldName="KismetMathLibrary.MakeVector_NetQuantize", NewName="MakeVector")
++FunctionRedirects=(OldName="KismetMathLibrary.MakeVector_NetQuantize10", NewName="MakeVector")
++FunctionRedirects=(OldName="KismetMathLibrary.MakeVector_NetQuantize100", NewName="MakeVector")
++FunctionRedirects=(OldName="KismetMathLibrary.MakeVector_NetQuantizeNormal", NewName="MakeVector")
++FunctionRedirects=(OldName="KismetMathLibrary.BreakVector_NetQuantize", NewName="BreakVector")
++FunctionRedirects=(OldName="KismetMathLibrary.BreakVector_NetQuantize10", NewName="BreakVector")
++FunctionRedirects=(OldName="KismetMathLibrary.BreakVector_NetQuantize100", NewName="BreakVector")
++FunctionRedirects=(OldName="KismetMathLibrary.BreakVector_NetQuantizeNormal", NewName="BreakVector")
+
 [CoreUObject.Metadata]
 ; Note: UnrealHeaderTool should be rerun after making changes to MetadataRedirects to catch any keys specified in class headers
 +MetadataRedirects=(OldKey="K2Protected", NewKey="BlueprintProtected")
diff --git a/Engine/Source/Editor/BlueprintGraph/Private/K2Node_BreakStruct.cpp b/Engine/Source/Editor/BlueprintGraph/Private/K2Node_BreakStruct.cpp
--- a/Engine/Source/Editor/BlueprintGraph/Private/K2Node_BreakStruct.cpp
+++ b/Engine/Source/Editor/BlueprintGraph/Private/K2Node_BreakStruct.cpp
@@ -413,12 +413,12 @@ void UK2Node_BreakStruct::ConvertDeprecatedNode(UEdGraph* Graph, bool bOnlySafeC
 		}
 		else if (StructType == TBaseStructure<FVector>::Get())
 		{
-			BreakNodeFunction = UKismetMathLibrary::StaticClass()->FindFunctionByName(GET_FUNCTION_NAME_CHECKED(UKismetMathLibrary, BreakVector));
+			BreakNodeFunction = UKismetMathLibrary::StaticClass()->FindFunctionByName(GET_FUNCTION_NAME_CHECKED_FourParams(UKismetMathLibrary, BreakVector, FVector, double&, double&, double&));
 			OldPinToNewPinMap.Add(TEXT("Vector"), TEXT("InVec"));
 		}
 		else if (StructType == TBaseStructure<FVector2D>::Get())
 		{
-			BreakNodeFunction = UKismetMathLibrary::StaticClass()->FindFunctionByName(GET_FUNCTION_NAME_CHECKED(UKismetMathLibrary, BreakVector2D));
+			BreakNodeFunction = UKismetMathLibrary::StaticClass()->FindFunctionByName(GET_FUNCTION_NAME_CHECKED_ThreeParams(UKismetMathLibrary, BreakVector2D, FVector2D, double&, double&));
 			OldPinToNewPinMap.Add(TEXT("Vector2D"), TEXT("InVec"));
 		}
 		else
diff --git a/Engine/Source/Editor/BlueprintGraph/Private/K2Node_MakeStruct.cpp b/Engine/Source/Editor/BlueprintGraph/Private/K2Node_MakeStruct.cpp
--- a/Engine/Source/Editor/BlueprintGraph/Private/K2Node_MakeStruct.cpp
+++ b/Engine/Source/Editor/BlueprintGraph/Private/K2Node_MakeStruct.cpp
@@ -418,12 +418,12 @@ void UK2Node_MakeStruct::ConvertDeprecatedNode(UEdGraph* Graph, bool bOnlySafeCh
 		}
 		else if (StructType == TBaseStructure<FVector>::Get())
 		{
-			MakeNodeFunction = UKismetMathLibrary::StaticClass()->FindFunctionByName(GET_FUNCTION_NAME_CHECKED(UKismetMathLibrary, MakeVector));
+			MakeNodeFunction = UKismetMathLibrary::StaticClass()->FindFunctionByName(GET_FUNCTION_NAME_CHECKED_ThreeParams(UKismetMathLibrary, MakeVector, double, double, double));
 			OldPinToNewPinMap.Add(TEXT("Vector"), UEdGraphSchema_K2::PN_ReturnValue);
 		}
 		else if (StructType == TBaseStructure<FVector2D>::Get())
 		{
-			MakeNodeFunction = UKismetMathLibrary::StaticClass()->FindFunctionByName(GET_FUNCTION_NAME_CHECKED(UKismetMathLibrary, MakeVector2D));
+			MakeNodeFunction = UKismetMathLibrary::StaticClass()->FindFunctionByName(GET_FUNCTION_NAME_CHECKED_TwoParams(UKismetMathLibrary, MakeVector2D, double, double));
 			OldPinToNewPinMap.Add(TEXT("Vector2D"), UEdGraphSchema_K2::PN_ReturnValue);
 		}
 		else
diff --git a/Engine/Source/Runtime/Core/Public/Misc/AssertionMacros.h b/Engine/Source/Runtime/Core/Public/Misc/AssertionMacros.h
--- a/Engine/Source/Runtime/Core/Public/Misc/AssertionMacros.h
+++ b/Engine/Source/Runtime/Core/Public/Misc/AssertionMacros.h
@@ -400,9 +400,27 @@ namespace UEAsserts_Private
 #define GET_FUNCTION_NAME_CHECKED_OneParam(ClassName, FunctionName, ArgType) \
 	((void)sizeof((std::declval<ClassName&>().FunctionName(std::declval<PREPROCESSOR_REMOVE_OPTIONAL_PARENS(ArgType)>()), (int)0)), FName(TEXT(#FunctionName)))
 
+#define GET_FUNCTION_NAME_CHECKED_TwoParams(ClassName, FunctionName, ArgType1, ArgType2) \
+	((void)sizeof((std::declval<ClassName&>().FunctionName(std::declval<PREPROCESSOR_REMOVE_OPTIONAL_PARENS(ArgType1)>(), std::declval<PREPROCESSOR_REMOVE_OPTIONAL_PARENS(ArgType2)>()), (int)0)), FName(TEXT(#FunctionName)))
+
+#define GET_FUNCTION_NAME_CHECKED_ThreeParams(ClassName, FunctionName, ArgType1, ArgType2, ArgType3) \
+	((void)sizeof((std::declval<ClassName&>().FunctionName(std::declval<PREPROCESSOR_REMOVE_OPTIONAL_PARENS(ArgType1)>(), std::declval<PREPROCESSOR_REMOVE_OPTIONAL_PARENS(ArgType2)>(), std::declval<PREPROCESSOR_REMOVE_OPTIONAL_PARENS(ArgType3)>()), (int)0)), FName(TEXT(#FunctionName)))
+
+#define GET_FUNCTION_NAME_CHECKED_FourParams(ClassName, FunctionName, ArgType1, ArgType2, ArgType3, ArgType4) \
+	((void)sizeof((std::declval<ClassName&>().FunctionName(std::declval<PREPROCESSOR_REMOVE_OPTIONAL_PARENS(ArgType1)>(), std::declval<PREPROCESSOR_REMOVE_OPTIONAL_PARENS(ArgType2)>(), std::declval<PREPROCESSOR_REMOVE_OPTIONAL_PARENS(ArgType3)>(), std::declval<PREPROCESSOR_REMOVE_OPTIONAL_PARENS(ArgType4)>()), (int)0)), FName(TEXT(#FunctionName)))
+
 #define GET_FUNCTION_NAME_STRING_CHECKED_OneParam(ClassName, FunctionName, ArgType) \
 	((void)sizeof((std::declval<ClassName&>().FunctionName(std::declval<PREPROCESSOR_REMOVE_OPTIONAL_PARENS(ArgType)>()), (int)0)), TEXT(#FunctionName))
 
+#define GET_FUNCTION_NAME_STRING_CHECKED_TwoParams(ClassName, FunctionName, ArgType1, ArgType2) \
+	((void)sizeof((std::declval<ClassName&>().FunctionName(std::declval<PREPROCESSOR_REMOVE_OPTIONAL_PARENS(ArgType1)>(), std::declval<PREPROCESSOR_REMOVE_OPTIONAL_PARENS(ArgType2)>()), (int)0)), TEXT(#FunctionName))
+
+#define GET_FUNCTION_NAME_STRING_CHECKED_ThreeParams(ClassName, FunctionName, ArgType1, ArgType2, ArgType3) \
+	((void)sizeof((std::declval<ClassName&>().FunctionName(std::declval<PREPROCESSOR_REMOVE_OPTIONAL_PARENS(ArgType1)>(), std::declval<PREPROCESSOR_REMOVE_OPTIONAL_PARENS(ArgType2)>(), std::declval<PREPROCESSOR_REMOVE_OPTIONAL_PARENS(ArgType3)>()), (int)0)), TEXT(#FunctionName))
+
+#define GET_FUNCTION_NAME_STRING_CHECKED_FourParams(ClassName, FunctionName, ArgType1, ArgType2, ArgType3, ArgType4) \
+	((void)sizeof((std::declval<ClassName&>().FunctionName(std::declval<PREPROCESSOR_REMOVE_OPTIONAL_PARENS(ArgType1)>(), std::declval<PREPROCESSOR_REMOVE_OPTIONAL_PARENS(ArgType2)>(), std::declval<PREPROCESSOR_REMOVE_OPTIONAL_PARENS(ArgType3)>(), std::declval<PREPROCESSOR_REMOVE_OPTIONAL_PARENS(ArgType4)>()), (int)0)), TEXT(#FunctionName))
+
 /*----------------------------------------------------------------------------
 	Low level error macros
 ----------------------------------------------------------------------------*/
diff --git a/Engine/Source/Runtime/Engine/Classes/Engine/NetSerialization.h b/Engine/Source/Runtime/Engine/Classes/Engine/NetSerialization.h
--- a/Engine/Source/Runtime/Engine/Classes/Engine/NetSerialization.h
+++ b/Engine/Source/Runtime/Engine/Classes/Engine/NetSerialization.h
@@ -401,7 +401,7 @@ bool SerializeFixedVector(FVector3d &Vector, FArchive& Ar)
  *	Note: this is the historical UE format for vector net serialization
  *
  */
-USTRUCT(meta = (HasNativeMake = "Engine.KismetMathLibrary.MakeVector_NetQuantize", HasNativeBreak = "Engine.KismetMathLibrary.BreakVector_NetQuantize"))
+USTRUCT()
 struct FVector_NetQuantize : public FVector
 {
 	GENERATED_USTRUCT_BODY()
@@ -447,7 +447,7 @@ struct TStructOpsTypeTraits< FVector_NetQuantize > : public TStructOpsTypeTraits
  *	Valid range: 2^24 / 10 = +/- 1,677,721.6
  *
  */
-USTRUCT(meta = (HasNativeMake = "Engine.KismetMathLibrary.MakeVector_NetQuantize10", HasNativeBreak = "Engine.KismetMathLibrary.BreakVector_NetQuantize10"))
+USTRUCT()
 struct FVector_NetQuantize10 : public FVector
 {
 	GENERATED_USTRUCT_BODY()
@@ -493,7 +493,7 @@ struct TStructOpsTypeTraits< FVector_NetQuantize10 > : public TStructOpsTypeTrai
  *	Valid range: 2^30 / 100 = +/- 10,737,418.24
  *
  */
-USTRUCT(meta = (HasNativeMake = "Engine.KismetMathLibrary.MakeVector_NetQuantize100", HasNativeBreak = "Engine.KismetMathLibrary.BreakVector_NetQuantize100"))
+USTRUCT()
 struct FVector_NetQuantize100 : public FVector
 {
 	GENERATED_USTRUCT_BODY()
@@ -535,7 +535,7 @@ struct TStructOpsTypeTraits< FVector_NetQuantize100 > : public TStructOpsTypeTra
  *	16 bits per component
  *	Valid range: -1..+1 inclusive
  */
-USTRUCT(meta = (HasNativeMake = "Engine.KismetMathLibrary.MakeVector_NetQuantizeNormal", HasNativeBreak = "Engine.KismetMathLibrary.BreakVector_NetQuantizeNormal"))
+USTRUCT()
 struct FVector_NetQuantizeNormal : public FVector
 {
 	GENERATED_USTRUCT_BODY()
diff --git a/Engine/Source/Runtime/Engine/Classes/Kismet/KismetMathLibrary.h b/Engine/Source/Runtime/Engine/Classes/Kismet/KismetMathLibrary.h
--- a/Engine/Source/Runtime/Engine/Classes/Kismet/KismetMathLibrary.h
+++ b/Engine/Source/Runtime/Engine/Classes/Kismet/KismetMathLibrary.h
@@ -12,7 +12,6 @@
 #include "UObject/ScriptMacros.h"
 #include "Kismet/BlueprintFunctionLibrary.h"
 #include "Misc/QualifiedFrameTime.h"
-#include "Engine/NetSerialization.h"
 
 #include "KismetMathLibrary.generated.h"
 
@@ -1015,11 +1014,14 @@ class ENGINE_API UKismetMathLibrary : public UBlueprintFunctionLibrary
 
 	/** Makes a 2d vector {X, Y} */
 	UFUNCTION(BlueprintPure, Category = "Math|Vector2D", meta = (Keywords = "construct build", NativeMakeFunc))
-	static FVector2D MakeVector2D(float X, float Y);
+	static FVector2D MakeVector2D(double X, double Y);
+
+	UE_DEPRECATED(5.0, "This method has been deprecated and will be removed. Use the double version instead.")
+	static void BreakVector2D(FVector2D InVec, float& X, float& Y);
 
 	/** Breaks a 2D vector apart into X, Y. */
 	UFUNCTION(BlueprintPure, Category = "Math|Vector2D", meta = (NativeBreakFunc))
-	static void BreakVector2D(FVector2D InVec, float& X, float& Y);
+	static void BreakVector2D(FVector2D InVec, double& X, double& Y);
 
 	/** Convert a Vector2D to a Vector */
 	UFUNCTION(BlueprintPure, meta = (DisplayName = "To Vector (Vector2D)", CompactNodeTitle = "->", ScriptMethod = "Vector", Keywords = "cast convert", BlueprintAutocast), Category = "Math|Conversions")
@@ -1310,8 +1312,8 @@ class ENGINE_API UKismetMathLibrary : public UBlueprintFunctionLibrary
 	//
 
 	/** Makes a vector {X, Y, Z} */
-	UFUNCTION(BlueprintPure, Category="Math|Vector", meta=(Keywords="construct build", NativeMakeFunc))
-	static FVector MakeVector(float X, float Y, float Z);
+	UFUNCTION(BlueprintPure, Category = "Math|Vector", meta = (Keywords = "construct build", NativeMakeFunc))
+	static FVector MakeVector(double X, double Y, double Z);
 
 	/** Creates a directional vector from rotation values {Pitch, Yaw} supplied in degrees with specified Length*/	
 	UFUNCTION(BlueprintPure, Category = "Math|Vector", meta = (Keywords = "rotation rotate"))
@@ -1335,9 +1337,12 @@ class ENGINE_API UKismetMathLibrary : public UBlueprintFunctionLibrary
 	UFUNCTION(BlueprintCallable, meta = (ScriptMethod = "Set"), Category = "Math|Vector")
 	static void Vector_Set(UPARAM(ref) FVector& A, float X, float Y, float Z);
 
+	UE_DEPRECATED(5.0, "This method has been deprecated and will be removed. Use the double version instead.")
+	static void BreakVector(FVector InVec, float& X, float& Y, float& Z);
+
 	/** Breaks a vector apart into X, Y, Z */
 	UFUNCTION(BlueprintPure, Category="Math|Vector", meta=(NativeBreakFunc))
-	static void BreakVector(FVector InVec, float& X, float& Y, float& Z);
+	static void BreakVector(FVector InVec, double& X, double& Y, double& Z);
 
 	UFUNCTION(BlueprintPure, Category = "Math|Vector", meta = (NativeBreakFunc))
 	static void BreakVector3f(FVector3f InVec, float& X, float& Y, float& Z);
@@ -2052,11 +2057,14 @@ class ENGINE_API UKismetMathLibrary : public UBlueprintFunctionLibrary
 
 	/** Makes a 4D vector {X, Y, Z, W} */
 	UFUNCTION(BlueprintPure, meta = (Keywords = "construct build", NativeMakeFunc), Category = "Math|Vector4")
-	static FVector4 MakeVector4(float X, float Y, float Z, float W);
+	static FVector4 MakeVector4(double X, double Y, double Z, double W);
+
+	UE_DEPRECATED(5.0, "This method has been deprecated and will be removed. Use the double version instead.")
+	static void BreakVector4(const FVector4& InVec, float& X, float& Y, float& Z, float& W);
 
 	/** Breaks a 4D vector apart into X, Y, Z, W. */
 	UFUNCTION(BlueprintPure, meta = (NativeBreakFunc), Category = "Math|Vector4")
-	static void BreakVector4(const FVector4& InVec, float& X, float& Y, float& Z, float& W);
+	static void BreakVector4(const FVector4& InVec, double& X, double& Y, double& Z, double& W);
 
 	/** Convert a Vector4 to a Vector (dropping the W element) */
 	UFUNCTION(BlueprintPure, meta = (DisplayName = "To Vector (Vector4)", CompactNodeTitle = "->", ScriptMethod = "Vector", Keywords = "cast convert", BlueprintAutocast), Category = "Math|Conversions")
@@ -4268,27 +4276,6 @@ class ENGINE_API UKismetMathLibrary : public UBlueprintFunctionLibrary
 	 */
 	UFUNCTION(BlueprintPure, Category="Math|Smoothing", meta=(DisplayName="Dynamic Weighted Moving Average Rotator"))
 	static FRotator DynamicWeightedMovingAverage_FRotator(FRotator CurrentSample, FRotator PreviousSample, float MaxDistance, float MinWeight, float MaxWeight);
-
-
-	// NetQuantized vector make/breaks
-	// LWC_TODO: These can be removed once blueprints support FReal
-	UFUNCTION(BlueprintPure, Category="Math|Vector", meta=(NativeMakeFunc))
-	static FVector_NetQuantize MakeVector_NetQuantize(float X, float Y, float Z) { return FVector_NetQuantize(X, Y, Z); }
-	UFUNCTION(BlueprintPure, Category="Math|Vector", meta=(NativeMakeFunc))
-	static FVector_NetQuantize10 MakeVector_NetQuantize10(float X, float Y, float Z) { return FVector_NetQuantize10(X, Y, Z); }
-	UFUNCTION(BlueprintPure, Category="Math|Vector", meta=(NativeMakeFunc))
-	static FVector_NetQuantize100 MakeVector_NetQuantize100(float X, float Y, float Z) { return FVector_NetQuantize100(X, Y, Z); }
-	UFUNCTION(BlueprintPure, Category="Math|Vector", meta=(NativeMakeFunc))
-	static FVector_NetQuantizeNormal MakeVector_NetQuantizeNormal(float X, float Y, float Z) { return FVector_NetQuantizeNormal(X, Y, Z); }
-	UFUNCTION(BlueprintPure, Category="Math|Vector", meta=(NativeBreakFunc))
-	static void BreakVector_NetQuantize(FVector_NetQuantize InVec, float& X, float& Y, float& Z) { BreakVector((FVector)InVec, X, Y, Z); }
-	UFUNCTION(BlueprintPure, Category="Math|Vector", meta=(NativeBreakFunc))
-	static void BreakVector_NetQuantize10(FVector_NetQuantize10 InVec, float& X, float& Y, float& Z) { BreakVector((FVector)InVec, X, Y, Z); }
-	UFUNCTION(BlueprintPure, Category="Math|Vector", meta=(NativeBreakFunc))
-	static void BreakVector_NetQuantize100(FVector_NetQuantize100 InVec, float& X, float& Y, float& Z) { BreakVector((FVector)InVec, X, Y, Z); }
-	UFUNCTION(BlueprintPure, Category="Math|Vector", meta=(NativeBreakFunc))
-	static void BreakVector_NetQuantizeNormal(FVector_NetQuantizeNormal InVec, float& X, float& Y, float& Z) { BreakVector((FVector)InVec, X, Y, Z); }
-
 	
 private:
 
diff --git a/Engine/Source/Runtime/Engine/Classes/Kismet/KismetMathLibrary.inl b/Engine/Source/Runtime/Engine/Classes/Kismet/KismetMathLibrary.inl
--- a/Engine/Source/Runtime/Engine/Classes/Kismet/KismetMathLibrary.inl
+++ b/Engine/Source/Runtime/Engine/Classes/Kismet/KismetMathLibrary.inl
@@ -1125,7 +1125,7 @@ FVector2D UKismetMathLibrary::Vector2D_Zero()
 }
 
 KISMET_MATH_FORCEINLINE
-FVector2D UKismetMathLibrary::MakeVector2D(float X, float Y)
+FVector2D UKismetMathLibrary::MakeVector2D(double X, double Y)
 {
 	return FVector2D(X, Y);
 }
@@ -1137,6 +1137,13 @@ void UKismetMathLibrary::BreakVector2D(FVector2D InVec, float& X, float& Y)
 	Y = InVec.Y;
 }
 
+KISMET_MATH_FORCEINLINE
+void UKismetMathLibrary::BreakVector2D(FVector2D InVec, double& X, double& Y)
+{
+	X = InVec.X;
+	Y = InVec.Y;
+}
+
 KISMET_MATH_FORCEINLINE
 FVector UKismetMathLibrary::Conv_Vector2DToVector(FVector2D InVec2D, float Z)
 {
@@ -1438,9 +1445,9 @@ FVector UKismetMathLibrary::Vector_Left()
 }
 
 KISMET_MATH_FORCEINLINE
-FVector UKismetMathLibrary::MakeVector(float X, float Y, float Z)
+FVector UKismetMathLibrary::MakeVector(double X, double Y, double Z)
 {
-	return FVector(X,Y,Z);
+	return FVector(X, Y, Z);
 }
 
 KISMET_MATH_FORCEINLINE
@@ -1463,6 +1470,14 @@ void UKismetMathLibrary::BreakVector(FVector InVec, float& X, float& Y, float& Z
 	Z = InVec.Z;
 }
 
+KISMET_MATH_FORCEINLINE
+void UKismetMathLibrary::BreakVector(FVector InVec, double& X, double& Y, double& Z)
+{
+	X = InVec.X;
+	Y = InVec.Y;
+	Z = InVec.Z;
+}
+
 KISMET_MATH_FORCEINLINE
 void UKismetMathLibrary::BreakVector3f(FVector3f InVec, float& X, float& Y, float& Z)
 {
@@ -2035,9 +2050,8 @@ FVector4 UKismetMathLibrary::Vector4_Zero()
 	return FVector4(ForceInitToZero);
 }
 
-
 KISMET_MATH_FORCEINLINE
-FVector4 UKismetMathLibrary::MakeVector4(float X, float Y, float Z, float W)
+FVector4 UKismetMathLibrary::MakeVector4(double X, double Y, double Z, double W)
 {
 	return FVector4(X, Y, Z, W);
 }
@@ -2051,6 +2065,15 @@ void UKismetMathLibrary::BreakVector4(const FVector4& InVec, float& X, float& Y,
 	W = InVec.W;
 }
 
+KISMET_MATH_FORCEINLINE
+void UKismetMathLibrary::BreakVector4(const FVector4& InVec, double& X, double& Y, double& Z, double& W)
+{
+	X = InVec.X;
+	Y = InVec.Y;
+	Z = InVec.Z;
+	W = InVec.W;
+}
+
 KISMET_MATH_FORCEINLINE
 FVector UKismetMathLibrary::Conv_Vector4ToVector(const FVector4& InVector4)
 {
