Fix for crash from out of bounds access in ISoundParameterControllerInterface

#rnx

diff --git a/Engine/Source/Runtime/Engine/Private/Audio/SoundParameterControllerInterface.cpp b/Engine/Source/Runtime/Engine/Private/Audio/SoundParameterControllerInterface.cpp
--- a/Engine/Source/Runtime/Engine/Private/Audio/SoundParameterControllerInterface.cpp
+++ b/Engine/Source/Runtime/Engine/Private/Audio/SoundParameterControllerInterface.cpp
@@ -52,6 +52,12 @@ void ISoundParameterControllerInterface::SetTriggerParameter(FName InName)
 			{
 				TArray<FAudioParameter> Params = { MoveTemp(ParamToSet) };
 				Sound->InitParameters(Params, SoundParameterControllerInterfacePrivate::ProxyFeatureName);
+				if (Params.Num() == 0)
+				{
+					// USoundBase::InitParameters(...) can remove parameters. 
+					// Exit early if the parameter is removed.
+					return;
+				}
 				ParamToSet = MoveTemp(Params[0]);
 			}
 
@@ -153,26 +159,29 @@ void ISoundParameterControllerInterface::SetParameters(TArray<FAudioParameter>&&
 
 			ParamsToSet = MoveTemp(InValues);
 
-			DECLARE_CYCLE_STAT(TEXT("FAudioThreadTask.SoundParameterControllerInterface.SetParameters"), STAT_AudioSetParameters, STATGROUP_AudioThreadCommands);
-			AudioDevice->SendCommandToActiveSounds(GetInstanceOwnerID(), [AudioDevice, Params = MoveTemp(ParamsToSet)](FActiveSound& ActiveSound)
+			if (ParamsToSet.Num() > 0)
 			{
-				if (Audio::IParameterTransmitter* Transmitter = ActiveSound.GetTransmitter())
+				DECLARE_CYCLE_STAT(TEXT("FAudioThreadTask.SoundParameterControllerInterface.SetParameters"), STAT_AudioSetParameters, STATGROUP_AudioThreadCommands);
+				AudioDevice->SendCommandToActiveSounds(GetInstanceOwnerID(), [AudioDevice, Params = MoveTemp(ParamsToSet)](FActiveSound& ActiveSound)
 				{
-					for (const FAudioParameter& Param : Params)
+					if (Audio::IParameterTransmitter* Transmitter = ActiveSound.GetTransmitter())
 					{
-						const FName ParamName = Param.ParamName;
-						if (!ParamName.IsNone())
+						for (const FAudioParameter& Param : Params)
 						{
-							// Must be copied as original version must be preserved in case command is called on multiple ActiveSounds.
-							FAudioParameter TempParam = Param;
-							if (!Transmitter->SetParameter(MoveTemp(TempParam)))
+							const FName ParamName = Param.ParamName;
+							if (!ParamName.IsNone())
 							{
-								UE_LOG(LogAudio, Warning, TEXT("Failed to set parameter '%s'"), *ParamName.ToString());
+								// Must be copied as original version must be preserved in case command is called on multiple ActiveSounds.
+								FAudioParameter TempParam = Param;
+								if (!Transmitter->SetParameter(MoveTemp(TempParam)))
+								{
+									UE_LOG(LogAudio, Warning, TEXT("Failed to set parameter '%s'"), *ParamName.ToString());
+								}
 							}
 						}
 					}
-				}
-			}, GET_STATID(STAT_AudioSetParameters));
+				}, GET_STATID(STAT_AudioSetParameters));
+			}
 		}
 	}
 }
@@ -199,6 +208,12 @@ void ISoundParameterControllerInterface::SetParameterInternal(FAudioParameter&&
 			{
 				TArray<FAudioParameter> Params = { MoveTemp(InParam) };
 				Sound->InitParameters(Params, SoundParameterControllerInterfacePrivate::ProxyFeatureName);
+				if (Params.Num() == 0)
+				{
+					// USoundBase::InitParameters(...) can remove parameters. 
+					// Exit early if the parameter is removed.
+					return;
+				}
 				ParamToSet = MoveTemp(Params[0]);
 			}
 			else
