- Allow individual primitive components to opt out of editor focus

diff --git a/Engine/Plugins/FX/Niagara/Source/Niagara/Private/NiagaraActor.cpp b/Engine/Plugins/FX/Niagara/Source/Niagara/Private/NiagaraActor.cpp
--- a/Engine/Plugins/FX/Niagara/Source/Niagara/Private/NiagaraActor.cpp
+++ b/Engine/Plugins/FX/Niagara/Source/Niagara/Private/NiagaraActor.cpp
@@ -60,6 +60,9 @@ ANiagaraActor::ANiagaraActor(const FObjectInitializer& ObjectInitializer)
 			ArrowComponent->SpriteInfo.DisplayName = ConstructorStatics.NAME_Effects;
 			ArrowComponent->SetupAttachment(NiagaraComponent);
 			ArrowComponent->SetUsingAbsoluteScale(true);
+		#if WITH_EDITOR
+			ArrowComponent->SetIgnoreBoundsForEditorFocus(true);
+		#endif
 		}
 	}
 #endif // WITH_EDITORONLY_DATA
diff --git a/Engine/Source/Runtime/Engine/Classes/Components/PrimitiveComponent.h b/Engine/Source/Runtime/Engine/Classes/Components/PrimitiveComponent.h
--- a/Engine/Source/Runtime/Engine/Classes/Components/PrimitiveComponent.h
+++ b/Engine/Source/Runtime/Engine/Classes/Components/PrimitiveComponent.h
@@ -617,6 +617,10 @@ protected:
 	UPROPERTY()
 	uint8 bHasNoStreamableTextures : 1;
 
+#if WITH_EDITOR
+	bool bIgnoreBoundsForEditorFocus = false;
+#endif
+
 public:
 	/** If true then DoCustomNavigableGeometryExport will be called to collect navigable geometry of this component. */
 	UPROPERTY()
@@ -1129,7 +1133,12 @@ public:
 	 * Whether or not the bounds of this component should be considered when focusing the editor camera to an actor with this component in it.
 	 * Useful for debug components which need a bounds for rendering but don't contribute to the visible part of the mesh in a meaningful way
 	 */
-	virtual bool IgnoreBoundsForEditorFocus() const { return false; }
+	virtual bool IgnoreBoundsForEditorFocus() const { return bIgnoreBoundsForEditorFocus; }
+	
+	/**
+	 * Set if we should ignore bounds when focusing the editor camera.
+	*/
+	void SetIgnoreBoundsForEditorFocus(bool bIgnore) { bIgnoreBoundsForEditorFocus = bIgnore; }
 #endif
 
 	/** Update current physics volume for this component, if bShouldUpdatePhysicsVolume is true. Overridden to use the overlaps to find the physics volume. */
diff --git a/Engine/Source/Runtime/Engine/Private/Components/SplineComponent.cpp b/Engine/Source/Runtime/Engine/Private/Components/SplineComponent.cpp
--- a/Engine/Source/Runtime/Engine/Private/Components/SplineComponent.cpp
+++ b/Engine/Source/Runtime/Engine/Private/Components/SplineComponent.cpp
@@ -1846,7 +1846,7 @@ FBoxSphereBounds USplineComponent::CalcBounds(const FTransform& LocalToWorld) co
 bool USplineComponent::IgnoreBoundsForEditorFocus() const
 {
 	// Cannot compute proper bounds when there's no point so don't participate to editor focus if that's the case : 
-	return SplineCurves.Position.Points.Num() == 0;
+	return Super::IgnoreBoundsForEditorFocus() || SplineCurves.Position.Points.Num() == 0;
 }
 #endif // WITH_EDITOR
 
