Fix Scale, rotating, or translating Foliage causes the Foliage to flicker and become invisible when placed

#rnx

diff --git a/Engine/Source/Runtime/Engine/Private/HierarchicalInstancedStaticMesh.cpp b/Engine/Source/Runtime/Engine/Private/HierarchicalInstancedStaticMesh.cpp
--- a/Engine/Source/Runtime/Engine/Private/HierarchicalInstancedStaticMesh.cpp
+++ b/Engine/Source/Runtime/Engine/Private/HierarchicalInstancedStaticMesh.cpp
@@ -2344,7 +2344,8 @@ bool UHierarchicalInstancedStaticMeshComponent::UpdateInstanceTransform(int32 In
 		
 		if (!bIsOmittedInstance)
 		{
-			InstanceUpdateCmdBuffer.UpdateInstance(RenderIndex, NewLocalTransform.ToMatrixWithScale());
+
+			InstanceUpdateCmdBuffer.UpdateInstance(RenderIndex, NewLocalTransform.ToMatrixWithScale().ConcatTranslation(-TranslatedInstanceSpaceOrigin));
 			bMarkRenderStateDirty = true;
 		}
 		
@@ -2553,7 +2554,7 @@ TArray<int32> UHierarchicalInstancedStaticMeshComponent::AddInstances(const TArr
 		{
 			InstanceReorderTable.Add(InitialBufferOffset + InstanceIndex);
 
-			InstanceUpdateCmdBuffer.AddInstance(PerInstanceSMData[InstanceIndex].Transform);
+			InstanceUpdateCmdBuffer.AddInstance(PerInstanceSMData[InstanceIndex].Transform.ConcatTranslation(-TranslatedInstanceSpaceOrigin));
 
 			const FBox NewInstanceBounds = GetStaticMesh()->GetBounds().GetBox().TransformBy(PerInstanceSMData[InstanceIndex].Transform);
 			UnbuiltInstanceBounds += NewInstanceBounds;
