Texture.cpp nop fix incorrect comment

diff --git a/Engine/Source/Runtime/Engine/Private/Texture.cpp b/Engine/Source/Runtime/Engine/Private/Texture.cpp
--- a/Engine/Source/Runtime/Engine/Private/Texture.cpp
+++ b/Engine/Source/Runtime/Engine/Private/Texture.cpp
@@ -2243,8 +2243,11 @@ static FName ConditionalGetPrefixedFormat(FName TextureFormatName, const ITarget
 {
 #if WITH_EDITOR
 
-	// Prepend a texture format to allow a module to override the compression (Ex: this allows you to replace TextureFormatDXT with a different compressor)
-
+	// "TextureCompressionFormat" specifies the Oodle Texture plugin to use for textures with OodleTextureSdkVersion == None
+	//		versioned textures always use TFO
+	//		TextureCompressionFormat can specify a pre-TFO plugin if desired
+	// 
+	// if you want Oodle Texture encoding,
 	// TextureCompressionFormat is required, TextureCompressionFormatWithVersion is optional
 
 	FString TextureCompressionFormat;
@@ -2253,10 +2256,9 @@ static FName ConditionalGetPrefixedFormat(FName TextureFormatName, const ITarget
 	
 	if ( bHasFormat )
 	{
-		//	new (optional) pref : TextureCompressionFormatWithVersion
-		//	 if TextureCompressionFormatWithVersion is not set, TextureCompressionFormat is used for both cases (with version & without)
 		if ( ! bOodleTextureSdkVersionIsNone )
 		{
+			//	new (optional) pref : TextureCompressionFormatWithVersion
 			FString TextureCompressionFormatWithVersion;
 			bool bHasFormatWithVersion = TargetPlatform->GetConfigSystem()->GetString(TEXT("AlternateTextureCompression"), TEXT("TextureCompressionFormatWithVersion"), TextureCompressionFormatWithVersion, GEngineIni);
 			bHasFormatWithVersion = bHasFormatWithVersion && ! TextureCompressionFormatWithVersion.IsEmpty();
