fix DoAllocateNewMips using 32 bit mip size

diff --git a/Engine/Source/Runtime/Engine/Private/Streaming/Texture2DStreamIn.cpp b/Engine/Source/Runtime/Engine/Private/Streaming/Texture2DStreamIn.cpp
--- a/Engine/Source/Runtime/Engine/Private/Streaming/Texture2DStreamIn.cpp
+++ b/Engine/Source/Runtime/Engine/Private/Streaming/Texture2DStreamIn.cpp
@@ -34,7 +34,7 @@ void FTexture2DStreamIn::DoAllocateNewMips(const FContext& Context)
 		for (int32 MipIndex = PendingFirstLODIdx; MipIndex < CurrentFirstLODIdx; ++MipIndex)
 		{
 			const FTexture2DMipMap& MipMap = *Context.MipsView[MipIndex];
-			const int32 MipSize = CalcTextureMipMapSize(MipMap.SizeX, MipMap.SizeY, Context.Resource->GetPixelFormat(), 0);
+			const SIZE_T MipSize = CalcTextureMipMapSize(MipMap.SizeX, MipMap.SizeY, Context.Resource->GetPixelFormat(), 0);
 
 			check(!MipData[MipIndex]);
 			MipData[MipIndex] = FMemory::Malloc(MipSize);
