World Partition UE5 default experience:
- Added option 'DisableStreaming' in the World Partition convert commandlet to allow disabling streaming for converted maps.
- Added flag 'bEnableStreaming' for partitioned worlds that can be used to disable streaming but keep key features such as OFPA and Data Layers.
- In this mode, the streaming generation acts as standard level streaming and can be enabled at any time by the user.
- This mode is enabled by default for new partitioned maps created from the empty template or the open world template.
- The World Partition convert commandlet will now duplicate the default HLOD layers setup.
- Show a hint in the World Partition window that streaming is disabled.
- When disabling streaming from the UI, show a dialog to warn the user about the implications of the change.
- When enabling streaming for the first time, show a dialog to warn the user that the world is about to be set up for streaming.
- When saving a map, don't duplicate its HLOD setup unless it's a template map.
- Includes changes from MattK: `Allow external object details customizations to hide the root object node`.
- TODO: Convert all template maps from the Games project to use World Partiton with streaming disabled using the new 'DisableStreaming' commandlet option.

diff --git a/Engine/Config/BaseEngine.ini b/Engine/Config/BaseEngine.ini
--- a/Engine/Config/BaseEngine.ini
+++ b/Engine/Config/BaseEngine.ini
@@ -281,6 +281,8 @@ bEnableOnScreenDebugMessages=true
 ; This shows errors and warnings on screen
 DurationOfErrorsAndWarningsOnHUD=0
 NearClipPlane=10.0
+; Default World Partition HLOD layer
+DefaultWorldPartitionHLODLayer=/Engine/Maps/Templates/HLODs/HLODLayer_Instanced.HLODLayer_Instanced
 
 [CoreRedirects]
 
diff --git a/Engine/Source/Editor/DetailCustomizations/Private/WorldSettingsDetails.cpp b/Engine/Source/Editor/DetailCustomizations/Private/WorldSettingsDetails.cpp
--- a/Engine/Source/Editor/DetailCustomizations/Private/WorldSettingsDetails.cpp
+++ b/Engine/Source/Editor/DetailCustomizations/Private/WorldSettingsDetails.cpp
@@ -200,59 +200,10 @@ void FWorldSettingsDetails::AddWorldCustomization(IDetailLayoutBuilder& DetailBu
 						]
 					]
 				];
-
-			WorldPartitionCategory.AddCustomRow(LOCTEXT("WorldPartitionEditorCellSizeRow", "WorldPartitionEditorCellSize"), true)
-				.NameContent()
-				[
-					SNew(STextBlock)
-					.Text(LOCTEXT("WorldPartitionCellSize", "World Partition Editor Cell Size"))
-					.ToolTipText(LOCTEXT("WorldPartitionEditorCellSize_ToolTip", "Set the world partition editor cell size, will take effect on the next world reload."))
-					.Font(IDetailLayoutBuilder::GetDetailFont())
-					.IsEnabled(bIsPartitionedWorld)
-				]
-				.ValueContent()
-				[
-					SNew(SHorizontalBox)
-					+ SHorizontalBox::Slot()
-					.AutoWidth()
-					.VAlign(VAlign_Center)
-					[
-						SNew(SNumericEntryBox<uint32>)
-						.AllowSpin(false)
-						.MinSliderValue(100)
-						.MaxSliderValue(100000)
-						.OnValueChanged(this, &FWorldSettingsDetails::HandlWorldPartitionEditorCellSizeChanged)
-						.Value(this, &FWorldSettingsDetails::HandleWorldPartitionEditorCellSizeValue)
-					]
-				];
 		}
 	}
 }
 
-void FWorldSettingsDetails::HandlWorldPartitionEditorCellSizeChanged(uint32 NewValue)
-{
-	if (SelectedWorldSettings.IsValid())
-	{
-		if (UWorldPartition* WorldPartition = SelectedWorldSettings->GetWorldPartition())
-		{
-			WorldPartition->SetEditorWantedCellSize(NewValue);
-		}
-	}
-}
-
-TOptional<uint32> FWorldSettingsDetails::HandleWorldPartitionEditorCellSizeValue() const
-{
-	if (SelectedWorldSettings.IsValid())
-	{
-		if (UWorldPartition* WorldPartition = SelectedWorldSettings->GetWorldPartition())
-		{
-			return WorldPartition->GetWantedEditorCellSize();
-		}
-	}
-
-	return TOptional<uint32>();
-}
-
 void FWorldSettingsDetails::OnUseActorFoldersChanged(ECheckBoxState BoxState, ULevel* Level)
 {
 	if (Level && (BoxState == ECheckBoxState::Checked))
diff --git a/Engine/Source/Editor/DetailCustomizations/Private/WorldSettingsDetails.h b/Engine/Source/Editor/DetailCustomizations/Private/WorldSettingsDetails.h
--- a/Engine/Source/Editor/DetailCustomizations/Private/WorldSettingsDetails.h
+++ b/Engine/Source/Editor/DetailCustomizations/Private/WorldSettingsDetails.h
@@ -86,12 +86,6 @@ private:
 	// Handles dropping an asset.
 	void HandleAssetDropped( UObject* Object, TSharedRef<IPropertyHandle> GameInfoProperty );
 
-	// Callback for changes in the world partition editor cell size.
-	void HandlWorldPartitionEditorCellSizeChanged(uint32 NewValue);
-
-	// Callback for getting the world partition editor cell size.
-	TOptional<uint32> HandleWorldPartitionEditorCellSizeValue() const;
-
 	/** Helper class to customizer GameMode property */
 	TSharedPtr<FGameModeInfoCustomizer>	GameInfoModeCustomizer;
 
diff --git a/Engine/Source/Editor/PropertyEditor/Private/DetailCategoryBuilderImpl.cpp b/Engine/Source/Editor/PropertyEditor/Private/DetailCategoryBuilderImpl.cpp
--- a/Engine/Source/Editor/PropertyEditor/Private/DetailCategoryBuilderImpl.cpp
+++ b/Engine/Source/Editor/PropertyEditor/Private/DetailCategoryBuilderImpl.cpp
@@ -377,6 +377,11 @@ IDetailPropertyRow* FDetailCategoryImpl::AddExternalObjects(const TArray<UObject
 
 	FDetailPropertyRow::MakeExternalPropertyRowCustomization(Objects, NAME_None, AsShared(), NewCustomization, AddPropertyParams);
 
+	if (Params.ShouldHideRootObjectNode() && NewCustomization.HasPropertyNode() && NewCustomization.GetPropertyNode()->AsObjectNode())
+	{
+		NewCustomization.PropertyRow->SetForceShowOnlyChildren(true);
+	}
+
 	TSharedPtr<FDetailPropertyRow> NewRow = NewCustomization.PropertyRow;
 
 	if (NewRow.IsValid())
diff --git a/Engine/Source/Editor/PropertyEditor/Private/DetailPropertyRow.cpp b/Engine/Source/Editor/PropertyEditor/Private/DetailPropertyRow.cpp
--- a/Engine/Source/Editor/PropertyEditor/Private/DetailPropertyRow.cpp
+++ b/Engine/Source/Editor/PropertyEditor/Private/DetailPropertyRow.cpp
@@ -171,7 +171,7 @@ bool FDetailPropertyRow::HasColumns() const
 
 bool FDetailPropertyRow::ShowOnlyChildren() const
 {
-	return PropertyTypeLayoutBuilder.IsValid() && CustomPropertyWidget.IsValid() && !CustomPropertyWidget->HasAnyContent();
+	return bForceShowOnlyChildren || (PropertyTypeLayoutBuilder.IsValid() && CustomPropertyWidget.IsValid() && !CustomPropertyWidget->HasAnyContent());
 }
 
 bool FDetailPropertyRow::RequiresTick() const
diff --git a/Engine/Source/Editor/PropertyEditor/Private/DetailPropertyRow.h b/Engine/Source/Editor/PropertyEditor/Private/DetailPropertyRow.h
--- a/Engine/Source/Editor/PropertyEditor/Private/DetailPropertyRow.h
+++ b/Engine/Source/Editor/PropertyEditor/Private/DetailPropertyRow.h
@@ -67,6 +67,12 @@ public:
 	/** @return Gets the custom name that should be used to save and restore this nodes expansion state */
 	FName GetCustomExpansionId() const { return CustomExpansionIdName; }
 
+	/** 
+	 * Sets whether or not we force only children to be shown regardless of what the customization wants.
+	 * Typically used when this row is only for organizational purposes and doesnt have anything valid to show (like external object nodes)
+	 */
+	void SetForceShowOnlyChildren(bool bInForceShowOnlyChildren) { bForceShowOnlyChildren = bInForceShowOnlyChildren; }
+
 	/**
 	 * Called when the owner node is initialized
 	 *
@@ -225,4 +231,6 @@ private:
 	bool bForceAutoExpansion;
 	/** True if we've already attempted to fill out our CustomTypeInterface (no need to rerun the query) */
 	bool bCachedCustomTypeInterface;
+	/** True if this row is only for organizational purposes and doesnt have anything valid to show itself*/
+	bool bForceShowOnlyChildren = false;
 };
diff --git a/Engine/Source/Editor/PropertyEditor/Public/DetailBuilderTypes.h b/Engine/Source/Editor/PropertyEditor/Public/DetailBuilderTypes.h
--- a/Engine/Source/Editor/PropertyEditor/Public/DetailBuilderTypes.h
+++ b/Engine/Source/Editor/PropertyEditor/Public/DetailBuilderTypes.h
@@ -13,6 +13,7 @@ struct FAddPropertyParams
 {
 	FAddPropertyParams()
 		: bForceShowProperty(false)
+		, bHideRootObjectNode(false)
 	{}
 
 
@@ -56,6 +57,15 @@ struct FAddPropertyParams
 	}
 
 
+	/**
+	 * Override whether the root object node should be shown for external object properties. This only applies when adding an entire uobject's properties to the details panel rather than an individual row or struct
+	 */
+	FAddPropertyParams& HideRootObjectNode(bool bInHideRootObjectNode)
+	{
+		bHideRootObjectNode = bInHideRootObjectNode;
+		return *this;
+	}
+
 	/**
 	 * Check whether to forcibly show the property, even if it does not have CPF_Edit
 	 */
@@ -64,6 +74,7 @@ struct FAddPropertyParams
 		return bForceShowProperty;
 	}
 
+	bool ShouldHideRootObjectNode() const { return bHideRootObjectNode; }
 
 	/**
 	 * Conditionally overwrites the specified boolean with a value specifying whether to allow child properties or not.
@@ -104,6 +115,9 @@ private:
 	/** When true, the property will be forcefully shown, even if it does not have CPF_Edit. When false the property will only be created if it has CPF_Edit. */
 	bool bForceShowProperty : 1;
 
+	/** When true any root ObjectPropertyNode generated from an external UObject being added will be hidden from view and only its children are shown */
+	bool bHideRootObjectNode : 1;
+
 	/** Tristate override for allowing children - if value is true: Allow Children, false: Disallow Children, unset: no override */
 	TOptional<bool> bAllowChildrenOverride;
 
diff --git a/Engine/Source/Editor/UnrealEd/Classes/Commandlets/WorldPartitionConvertCommandlet.h b/Engine/Source/Editor/UnrealEd/Classes/Commandlets/WorldPartitionConvertCommandlet.h
--- a/Engine/Source/Editor/UnrealEd/Classes/Commandlets/WorldPartitionConvertCommandlet.h
+++ b/Engine/Source/Editor/UnrealEd/Classes/Commandlets/WorldPartitionConvertCommandlet.h
@@ -56,7 +56,7 @@ protected:
 	virtual void OnWorldLoaded(UWorld* World);
 	virtual void ReadAdditionalTokensAndSwitches(const TArray<FString>& Tokens, const TArray<FString>& Switches) {}
 
-	UWorldPartition* CreateWorldPartition(class AWorldSettings* MainWorldSettings) const;
+	UWorldPartition* CreateWorldPartition(class AWorldSettings* MainWorldSettings);
 	UWorld* LoadWorld(const FString& LevelToLoad);
 	ULevel* InitWorld(UWorld* World);
 
@@ -94,6 +94,7 @@ protected:
 	bool bReportOnly;
 	bool bVerbose;
 	bool bConversionSuffix;
+	bool bDisableStreaming;
 	FString ConversionSuffix;
 
 	UPROPERTY(Config)
diff --git a/Engine/Source/Editor/UnrealEd/Private/Commandlets/WorldPartitionConvertCommandlet.cpp b/Engine/Source/Editor/UnrealEd/Private/Commandlets/WorldPartitionConvertCommandlet.cpp
--- a/Engine/Source/Editor/UnrealEd/Private/Commandlets/WorldPartitionConvertCommandlet.cpp
+++ b/Engine/Source/Editor/UnrealEd/Private/Commandlets/WorldPartitionConvertCommandlet.cpp
@@ -215,11 +215,17 @@ ULevel* UWorldPartitionConvertCommandlet::InitWorld(UWorld* World)
 	return World->PersistentLevel;
 }
 
-UWorldPartition* UWorldPartitionConvertCommandlet::CreateWorldPartition(AWorldSettings* MainWorldSettings) const
+UWorldPartition* UWorldPartitionConvertCommandlet::CreateWorldPartition(AWorldSettings* MainWorldSettings)
 {
 	TRACE_CPUPROFILER_EVENT_SCOPE(UWorldPartitionConvertCommandlet::CreateWorldPartition);
 
 	UWorldPartition* WorldPartition = UWorldPartition::CreateOrRepairWorldPartition(MainWorldSettings, EditorHashClass, RuntimeHashClass);
+
+	if (bDisableStreaming)
+	{
+		WorldPartition->bEnableStreaming = false;
+		WorldPartition->bStreamingWasEnabled = false;
+	}
 		
 	// Read the conversion config file
 	if (FPlatformFileManager::Get().GetPlatformFile().FileExists(*LevelConfigFilename))
@@ -229,6 +235,18 @@ UWorldPartition* UWorldPartitionConvertCommandlet::CreateWorldPartition(AWorldSe
 		WorldPartition->DefaultHLODLayer = HLODLayers.FindRef(DefaultHLODLayerName);
 	}
 
+	if ((WorldPartition->DefaultHLODLayer == UHLODLayer::GetEngineDefaultHLODLayersSetup()) && !bDisableStreaming)
+	{
+		WorldPartition->DefaultHLODLayer = UHLODLayer::DuplicateHLODLayersSetup(UHLODLayer::GetEngineDefaultHLODLayersSetup(), WorldPartition->GetPackage()->GetName(), WorldPartition->GetWorld()->GetName());
+
+		UHLODLayer* CurrentHLODLayer = WorldPartition->DefaultHLODLayer;
+		while (CurrentHLODLayer)
+		{
+			PackagesToSave.Add(CurrentHLODLayer->GetPackage());
+			CurrentHLODLayer = Cast<UHLODLayer>(CurrentHLODLayer->GetParentLayer().Get());
+		}
+	}
+
 	WorldPartition->EditorHash->Initialize();
 	
 	return WorldPartition;
@@ -692,6 +710,7 @@ int32 UWorldPartitionConvertCommandlet::Main(const FString& Params)
 	bGenerateIni = Switches.Contains(TEXT("GenerateIni"));
 	bReportOnly = bGenerateIni || Switches.Contains(TEXT("ReportOnly"));
 	bVerbose = Switches.Contains(TEXT("Verbose"));
+	bDisableStreaming = Switches.Contains(TEXT("DisableStreaming"));
 	ConversionSuffix = GetConversionSuffix(bOnlyMergeSubLevels);
 
 	FString* FoliageTypePathValue = Arguments.Find(TEXT("FoliageTypePath"));
diff --git a/Engine/Source/Editor/UnrealEd/Private/FileHelpers.cpp b/Engine/Source/Editor/UnrealEd/Private/FileHelpers.cpp
--- a/Engine/Source/Editor/UnrealEd/Private/FileHelpers.cpp
+++ b/Engine/Source/Editor/UnrealEd/Private/FileHelpers.cpp
@@ -77,6 +77,7 @@
 #include "WorldPartition/ActorDescContainer.h"
 #include "WorldPartition/WorldPartition.h"
 #include "WorldPartition/WorldPartitionEditorPerProjectUserSettings.h"
+#include "WorldPartition/HLOD/HLODLayer.h"
 #include "PackageSourceControlHelper.h"
 #include "ActorFolder.h"
 #include "InterchangeManager.h"
@@ -291,29 +292,6 @@ static bool ConfirmPackageBranchCheckOutStatus(const TArray<UPackage*>& Packages
 	return true;
 }
 
-static bool SaveExternalPackages(UPackage* Package, bool bCheckDirty)
-{
-	TArray<UPackage*> PackagesToSave;
-	for (UPackage* ExternalPackage : Package->GetExternalPackages())
-	{
-		if (!FPackageName::IsTempPackage(ExternalPackage->GetName()))
-		{
-			PackagesToSave.Add(ExternalPackage);
-		}
-	}
-
-	if (PackagesToSave.Num() > 0)
-	{
-		if (!UEditorLoadingAndSavingUtils::SavePackages(PackagesToSave, bCheckDirty))
-		{
-			FMessageDialog::Open(EAppMsgType::Ok, NSLOCTEXT("UnrealEd", "Error_FailedToSaveExternalActorPackages", "Failed to save external packages"));
-			return false;
-		}
-	}
-
-	return true;
-}
-
 static bool DeleteExistingMapPackages(const FString& ExistingPackageName)
 {
 	// Search for external actor files
@@ -690,6 +668,7 @@ static bool SaveWorld(UWorld* World,
 	const FString OriginalPackageName = Package->GetName();
 	const FString NewWorldAssetName = FPackageName::GetLongPackageAssetName(NewPackageName);
 	const bool bNewPackageExists = FPackageName::DoesPackageExist(NewPackageName);
+	const bool bIsTempPackage = FPackageName::IsTempPackage(World->GetPackage()->GetName());
 	bool bValidWorldName = true;
 	bool bPackageNeedsRename = false;
 	bool bWorldNeedsRename = false;
@@ -753,6 +732,9 @@ static bool SaveWorld(UWorld* World,
 		// Save Loaded cells
 		TArray<FName> LoadedEditorCells;
 
+		// Other packages to save
+		TArray<UPackage*> PackagesToSave;
+
 		// Initialize Physics Scene for save if needed here so that external packages don't get dirtied during the Saving of the map package
 		bool bForceInitializedWorld = false;
 		
@@ -767,11 +749,28 @@ static bool SaveWorld(UWorld* World,
 				}
 
 				RenamedWorldPartition = World->GetWorldPartition();
+
 				// Load all unloaded actors before rename. If this is causing issues (oom or other) map will need to be renamed through a provided builder commandlet
 				if (RenamedWorldPartition)
 				{
 					LoadedEditorCells = RenamedWorldPartition->GetUserLoadedEditorGridCells();
 					RenamedWorldPartition->LoadAllActors(ActorReferences);
+
+					if (bIsTempPackage)
+					{
+						if (UHLODLayer* CurrentHLODLayer = RenamedWorldPartition->DefaultHLODLayer)
+						{
+							CurrentHLODLayer = UHLODLayer::DuplicateHLODLayersSetup(CurrentHLODLayer, NewPackageName, NewWorldAssetName);
+							
+							RenamedWorldPartition->DefaultHLODLayer = CurrentHLODLayer;
+
+							while (CurrentHLODLayer)
+							{
+								PackagesToSave.Add(CurrentHLODLayer->GetPackage());
+								CurrentHLODLayer = Cast<UHLODLayer>(CurrentHLODLayer->GetParentLayer().Get());
+							}
+						}
+					}
 				}
 
 				// If we are doing a SaveAs on a world that already exists on disk, we need to duplicate it:
@@ -847,11 +846,25 @@ static bool SaveWorld(UWorld* World,
 				
 		if(!bAutosaving)
 		{
-			// This will save external actors
 			if (bSuccess)
 			{
-				// Do not check dirty flag if world is newly created
-				bSuccess = SaveExternalPackages(Package, /*bCheckDirty*/!bNewlyCreated);
+				// Gather external actors to save
+				for (UPackage* ExternalPackage : Package->GetExternalPackages())
+				{
+					if (!FPackageName::IsTempPackage(ExternalPackage->GetName()))
+					{
+						PackagesToSave.Add(ExternalPackage);
+					}
+				}
+
+				if (PackagesToSave.Num())
+				{
+					if (!UEditorLoadingAndSavingUtils::SavePackages(PackagesToSave, /*bCheckDirty*/!bNewlyCreated))
+					{
+						FMessageDialog::Open(EAppMsgType::Ok, NSLOCTEXT("UnrealEd", "Error_FailedToSaveHLODLayersPackages", "Failed to save dependant map packages"));
+						bSuccess = false;
+					}
+				}
 			}
 		}
 				
diff --git a/Engine/Source/Editor/WorldPartitionEditor/Private/WorldPartition/Customizations/WorldPartitionDetailsCustomization.cpp b/Engine/Source/Editor/WorldPartitionEditor/Private/WorldPartition/Customizations/WorldPartitionDetailsCustomization.cpp
new file mode 100644
--- /dev/null
+++ b/Engine/Source/Editor/WorldPartitionEditor/Private/WorldPartition/Customizations/WorldPartitionDetailsCustomization.cpp
@@ -0,0 +1,81 @@
+// Copyright Epic Games, Inc. All Rights Reserved.
+
+#include "WorldPartitionDetailsCustomization.h"
+#include "WorldPartition/WorldPartition.h"
+#include "WorldPartition/WorldPartitionRuntimeHash.h"
+#include "DetailLayoutBuilder.h"
+#include "DetailCategoryBuilder.h"
+#include "DetailWidgetRow.h"
+#include "Widgets/Layout/SBorder.h"
+#include "Widgets/Text/STextBlock.h"
+#include "Widgets/Layout/SBox.h"
+#include "Widgets/Input/SButton.h"
+#include "Widgets/Input/SCheckBox.h"
+#include "Widgets/Input/SNumericEntryBox.h"
+
+#define LOCTEXT_NAMESPACE "FWorldPartitionDetails"
+
+TSharedRef<IDetailCustomization> FWorldPartitionDetails::MakeInstance()
+{
+	return MakeShareable(new FWorldPartitionDetails);
+}
+
+void FWorldPartitionDetails::CustomizeDetails(IDetailLayoutBuilder& DetailBuilder)
+{
+	TArray<TWeakObjectPtr<UObject>> ObjectsBeingCustomized;
+	DetailBuilder.GetObjectsBeingCustomized(ObjectsBeingCustomized);
+	check(ObjectsBeingCustomized.Num() == 1);
+
+	WorldPartition = CastChecked<UWorldPartition>(ObjectsBeingCustomized[0].Get());
+
+	IDetailCategoryBuilder& WorldPartitionCategory = DetailBuilder.EditCategory("WorldPartition");
+
+	if (WorldPartition->RuntimeHash)
+	{
+		FAddPropertyParams Params;
+		Params.HideRootObjectNode(true);
+
+		if (IDetailPropertyRow* RuntimeHashRow = WorldPartitionCategory.AddExternalObjects({ WorldPartition->RuntimeHash }, EPropertyLocation::Default, Params))
+		{
+			RuntimeHashRow->ShouldAutoExpand(true)
+				.DisplayName(LOCTEXT("RuntimeHash", "Runtime Hash"))
+				.Visibility(TAttribute<EVisibility>::Create(TAttribute<EVisibility>::FGetter::CreateLambda([this]() { return WorldPartition->bEnableStreaming ? EVisibility::Visible : EVisibility::Hidden; })));
+		}
+	}
+
+	WorldPartitionCategory.AddCustomRow(LOCTEXT("WorldPartitionEditorCellSizeRow", "WorldPartitionEditorCellSize"), true)
+		.NameContent()
+		[
+			SNew(STextBlock)
+			.Text(LOCTEXT("WorldPartitionCellSize", "World Partition Editor Cell Size"))
+			.ToolTipText(LOCTEXT("WorldPartitionEditorCellSize_ToolTip", "Set the world partition editor cell size, will take effect on the next world reload."))
+			.Font(IDetailLayoutBuilder::GetDetailFont())
+		]
+		.ValueContent()
+		[
+			SNew(SHorizontalBox)
+			+ SHorizontalBox::Slot()
+			.AutoWidth()
+			.VAlign(VAlign_Center)
+			[
+				SNew(SNumericEntryBox<uint32>)
+				.AllowSpin(false)
+				.MinSliderValue(100)
+				.MaxSliderValue(100000)
+				.OnValueChanged(this, &FWorldPartitionDetails::HandleWorldPartitionEditorCellSizeChanged)
+				.Value(this, &FWorldPartitionDetails::HandleWorldPartitionEditorCellSizeValue)
+			]
+		];
+}
+
+void FWorldPartitionDetails::HandleWorldPartitionEditorCellSizeChanged(uint32 NewValue)
+{
+	WorldPartition->SetEditorWantedCellSize(NewValue);
+}
+
+TOptional<uint32> FWorldPartitionDetails::HandleWorldPartitionEditorCellSizeValue() const
+{
+	return WorldPartition->GetWantedEditorCellSize();
+}
+
+#undef LOCTEXT_NAMESPACE
diff --git a/Engine/Source/Editor/WorldPartitionEditor/Private/WorldPartition/Customizations/WorldPartitionDetailsCustomization.h b/Engine/Source/Editor/WorldPartitionEditor/Private/WorldPartition/Customizations/WorldPartitionDetailsCustomization.h
new file mode 100644
--- /dev/null
+++ b/Engine/Source/Editor/WorldPartitionEditor/Private/WorldPartition/Customizations/WorldPartitionDetailsCustomization.h
@@ -0,0 +1,33 @@
+// Copyright Epic Games, Inc. All Rights Reserved.
+
+#pragma once
+
+#include "CoreMinimal.h"
+#include "SlateFwd.h"
+#include "IDetailCustomization.h"
+
+class IDetailLayoutBuilder;
+class UWorldPartition;
+
+class FWorldPartitionDetails : public IDetailCustomization
+{
+public:
+	/** Makes a new instance of this detail layout class for a specific detail view requesting it */
+	static TSharedRef<IDetailCustomization> MakeInstance();
+
+protected:
+	FWorldPartitionDetails()
+	{}
+
+private:
+	/** IDetailCustomization interface */
+	virtual void CustomizeDetails(IDetailLayoutBuilder& DetailBuilder) override;
+
+	// Callback for changes in the world partition editor cell size.
+	void HandleWorldPartitionEditorCellSizeChanged(uint32 NewValue);
+
+	// Callback for getting the world partition editor cell size.
+	TOptional<uint32> HandleWorldPartitionEditorCellSizeValue() const;
+
+	UWorldPartition* WorldPartition;
+};
\ No newline at end of file
diff --git a/Engine/Source/Editor/WorldPartitionEditor/Private/WorldPartition/SWorldPartitionEditor.cpp b/Engine/Source/Editor/WorldPartitionEditor/Private/WorldPartition/SWorldPartitionEditor.cpp
--- a/Engine/Source/Editor/WorldPartitionEditor/Private/WorldPartition/SWorldPartitionEditor.cpp
+++ b/Engine/Source/Editor/WorldPartitionEditor/Private/WorldPartition/SWorldPartitionEditor.cpp
@@ -52,13 +52,15 @@ void SWorldPartitionEditor::Refresh()
 	GridView->Refresh();
 }
 
-void SWorldPartitionEditor::OnBrowseWorld(UWorld* InWorld)
+void SWorldPartitionEditor::Reconstruct()
 {
-	World = InWorld;
 	ContentParent->SetContent(ConstructContentWidget());
+}
 
-	// No need to unregister as the previous UWorldPartitionSubsystem is already destroyed
-	WorldPartitionChangedDelegateHandle.Reset();
+void SWorldPartitionEditor::OnBrowseWorld(UWorld* InWorld)
+{
+	World = InWorld;
+	Reconstruct();
 }
 
 TSharedRef<SWidget> SWorldPartitionEditor::ConstructContentWidget()
diff --git a/Engine/Source/Editor/WorldPartitionEditor/Private/WorldPartition/SWorldPartitionEditorGrid.cpp b/Engine/Source/Editor/WorldPartitionEditor/Private/WorldPartition/SWorldPartitionEditorGrid.cpp
--- a/Engine/Source/Editor/WorldPartitionEditor/Private/WorldPartition/SWorldPartitionEditorGrid.cpp
+++ b/Engine/Source/Editor/WorldPartitionEditor/Private/WorldPartition/SWorldPartitionEditorGrid.cpp
@@ -18,14 +18,18 @@ void SWorldPartitionEditorGrid::Construct(const FArguments& InArgs)
 	World = InArgs._InWorld;
 	WorldPartition = World ? World->GetWorldPartition() : nullptr;
 
-	if (!WorldPartition)
+	if (!WorldPartition || !WorldPartition->bEnableStreaming)
 	{
+		const FText Message = WorldPartition ? 
+			LOCTEXT("WorldPartitionHasStreamingDisabled", "World Partition streaming is not enabled for this map") : 
+			LOCTEXT("WorldPartitionMustBeEnabled", "World Partition is not enabled for this map");
+
 		ChildSlot
 		.VAlign(VAlign_Center)
 		.HAlign(HAlign_Center)
 		[
 			SNew(STextBlock)
-			.Text(LOCTEXT("WorldPartitionMustBeEnabled", "World Partition is Not Enabled for this Map"))
+			.Text(Message)
 			.ColorAndOpacity(FSlateColor::UseForeground())
 		];
 	}
@@ -104,4 +108,9 @@ void SWorldPartitionEditorGrid::Refresh()
 	}
 }
 
+bool SWorldPartitionEditorGrid::IsDisabled() const
+{
+	return !WorldPartition || !WorldPartition->bEnableStreaming;
+}
+
 #undef LOCTEXT_NAMESPACE
\ No newline at end of file
diff --git a/Engine/Source/Editor/WorldPartitionEditor/Private/WorldPartitionEditorModule.cpp b/Engine/Source/Editor/WorldPartitionEditor/Private/WorldPartitionEditorModule.cpp
--- a/Engine/Source/Editor/WorldPartitionEditor/Private/WorldPartitionEditorModule.cpp
+++ b/Engine/Source/Editor/WorldPartitionEditor/Private/WorldPartitionEditorModule.cpp
@@ -9,6 +9,8 @@
 #include "WorldPartition/SWorldPartitionEditor.h"
 #include "WorldPartition/SWorldPartitionEditorGridSpatialHash.h"
 
+#include "WorldPartition/Customizations/WorldPartitionDetailsCustomization.h"
+
 #include "WorldPartition/SWorldPartitionConvertDialog.h"
 #include "WorldPartition/WorldPartitionConvertOptions.h"
 #include "WorldPartition/WorldPartitionEditorSettings.h"
@@ -31,6 +33,7 @@
 #include "ContentBrowserModule.h"
 #include "EditorDirectories.h"
 #include "AssetRegistryModule.h"
+#include "PropertyEditorModule.h"
 #include "Editor/WorkspaceMenuStructure/Public/WorkspaceMenuStructure.h"
 #include "Editor/WorkspaceMenuStructure/Public/WorkspaceMenuStructureModule.h"
 #include "Framework/Docking/LayoutExtender.h"
@@ -126,6 +129,9 @@ void FWorldPartitionEditorModule::StartupModule()
 	IAssetTools& AssetTools = FModuleManager::LoadModuleChecked<FAssetToolsModule>("AssetTools").Get();
 	HLODLayerAssetTypeActions = MakeShareable(new FHLODLayerAssetTypeActions);
 	AssetTools.RegisterAssetTypeActions(HLODLayerAssetTypeActions.ToSharedRef());
+
+	FPropertyEditorModule& PropertyEditor = FModuleManager::LoadModuleChecked<FPropertyEditorModule>("PropertyEditor");
+	PropertyEditor.RegisterCustomClassLayout("WorldPartition", FOnGetDetailCustomizationInstance::CreateStatic(&FWorldPartitionDetails::MakeInstance));
 }
 
 void FWorldPartitionEditorModule::ShutdownModule()
@@ -151,21 +157,22 @@ void FWorldPartitionEditorModule::ShutdownModule()
 		UToolMenus::UnregisterOwner(this);
 	}
 
-	IAssetTools* AssetToolsModule = nullptr;
-	if (FModuleManager::Get().IsModuleLoaded("AssetTools"))
-	{
-		AssetToolsModule = &FModuleManager::GetModuleChecked<FAssetToolsModule>("AssetTools").Get();
-	}
-
 	// Unregister the HLODLayer asset type actions
 	if (HLODLayerAssetTypeActions.IsValid())
 	{
-		if (AssetToolsModule)
+		if (FModuleManager::Get().IsModuleLoaded("AssetTools"))
 		{
-			AssetToolsModule->UnregisterAssetTypeActions(HLODLayerAssetTypeActions.ToSharedRef());
+			IAssetTools& AssetTools = FModuleManager::GetModuleChecked<FAssetToolsModule>("AssetTools").Get();
+			AssetTools.UnregisterAssetTypeActions(HLODLayerAssetTypeActions.ToSharedRef());
 		}
 		HLODLayerAssetTypeActions.Reset();
 	}
+
+	if (FModuleManager::Get().IsModuleLoaded("PropertyEditor"))
+	{
+		FPropertyEditorModule& PropertyEditor = FModuleManager::GetModuleChecked<FPropertyEditorModule>("PropertyEditor");
+		PropertyEditor.UnregisterCustomClassLayout("WorldPartition");
+	}
 }
 
 TSharedRef<SWidget> FWorldPartitionEditorModule::CreateWorldPartitionEditor()
diff --git a/Engine/Source/Editor/WorldPartitionEditor/Public/WorldPartition/SWorldPartitionEditor.h b/Engine/Source/Editor/WorldPartitionEditor/Public/WorldPartition/SWorldPartitionEditor.h
--- a/Engine/Source/Editor/WorldPartitionEditor/Public/WorldPartition/SWorldPartitionEditor.h
+++ b/Engine/Source/Editor/WorldPartitionEditor/Public/WorldPartition/SWorldPartitionEditor.h
@@ -22,6 +22,7 @@ public:
 
 	// IWorldPartitionEditor interface
 	virtual void Refresh() override;
+	virtual void Reconstruct() override;
 
 private:
 	void OnBrowseWorld(UWorld* InWorld);
@@ -30,7 +31,6 @@ private:
 
 	FText GetMouseLocationText() const;
 
-	FDelegateHandle WorldPartitionChangedDelegateHandle;
 	TSharedPtr<SBorder>	ContentParent;
 	TSharedPtr<class SWorldPartitionEditorGrid> GridView;
 	TWeakObjectPtr<UWorld> World;
diff --git a/Engine/Source/Editor/WorldPartitionEditor/Public/WorldPartition/SWorldPartitionEditorGrid.h b/Engine/Source/Editor/WorldPartitionEditor/Public/WorldPartition/SWorldPartitionEditorGrid.h
--- a/Engine/Source/Editor/WorldPartitionEditor/Public/WorldPartition/SWorldPartitionEditorGrid.h
+++ b/Engine/Source/Editor/WorldPartitionEditor/Public/WorldPartition/SWorldPartitionEditorGrid.h
@@ -44,7 +44,7 @@ public:
 	// SPanel interface
 	virtual void OnArrangeChildren(const FGeometry& AllottedGeometry, FArrangedChildren& ArrangedChildren) const override
 	{
-		if (!WorldPartition)
+		if (IsDisabled())
 		{
 			SCompoundWidget::OnArrangeChildren(AllottedGeometry, ArrangedChildren);
 		}
@@ -52,12 +52,12 @@ public:
 
 	virtual FVector2D ComputeDesiredSize(float Scale) const override
 	{
-		return !WorldPartition ? FVector2D(100.0f, 100.0f) : SCompoundWidget::ComputeDesiredSize(Scale);
+		return IsDisabled() ? FVector2D(100.0f, 100.0f) : SCompoundWidget::ComputeDesiredSize(Scale);
 	}
 
 	virtual FChildren* GetChildren() override
 	{
-		return !WorldPartition ? &ChildSlot : (FChildren*)&FNoChildren::NoChildrenInstance;
+		return IsDisabled() ? &ChildSlot : (FChildren*)&FNoChildren::NoChildrenInstance;
 	}
 	
 	bool GetPlayerView(FVector& Location, FRotator& Rotation) const;
@@ -70,4 +70,7 @@ public:
 
 protected:
 	static TMap<FName, PartitionEditorGridCreateInstanceFunc> PartitionEditorGridCreateInstanceFactory;
+
+private:
+	bool IsDisabled() const;
 };
\ No newline at end of file
diff --git a/Engine/Source/Runtime/Engine/Classes/GameFramework/WorldSettings.h b/Engine/Source/Runtime/Engine/Classes/GameFramework/WorldSettings.h
--- a/Engine/Source/Runtime/Engine/Classes/GameFramework/WorldSettings.h
+++ b/Engine/Source/Runtime/Engine/Classes/GameFramework/WorldSettings.h
@@ -559,8 +559,8 @@ protected:
 
 public:
 #if WITH_EDITORONLY_DATA
-	/** Size of the grid for instanced foliage actors */
-	UPROPERTY(EditAnywhere, Category = WorldPartition)
+	/** Size of the grid for instanced foliage actors, only used for partitioned worlds */
+	UPROPERTY(EditAnywhere, Category = Foliage)
 	uint32 InstancedFoliageGridSize;
 
 	/** Size of the grid for navigation data chunk actors */
diff --git a/Engine/Source/Runtime/Engine/Private/WorldPartition/HLOD/HLODLayer.cpp b/Engine/Source/Runtime/Engine/Private/WorldPartition/HLOD/HLODLayer.cpp
--- a/Engine/Source/Runtime/Engine/Private/WorldPartition/HLOD/HLODLayer.cpp
+++ b/Engine/Source/Runtime/Engine/Private/WorldPartition/HLOD/HLODLayer.cpp
@@ -11,6 +11,7 @@
 #include "Serialization/ArchiveCrc32.h"
 
 #include "Engine/World.h"
+#include "Misc/ConfigCacheIni.h"
 
 #include "WorldPartition/WorldPartition.h"
 #include "WorldPartition/WorldPartitionActorDesc.h"
@@ -92,6 +93,60 @@ bool UHLODLayer::DoesRequireWarmup() const
 	return false;
 }
 
+UHLODLayer* UHLODLayer::GetEngineDefaultHLODLayersSetup()
+{
+	UHLODLayer* Result = nullptr;
+
+	if (FConfigFile* EngineConfig = GConfig->FindConfigFileWithBaseName(TEXT("Engine")))
+	{
+		FString DefaultHLODLayerName;
+		if (EngineConfig->GetString(TEXT("/Script/Engine.Engine"), TEXT("DefaultWorldPartitionHLODLayer"), DefaultHLODLayerName))
+		{
+			FSoftObjectPath DefaultHLODLayerPath(*DefaultHLODLayerName);
+			TSoftObjectPtr<UHLODLayer> EngineHLODLayerPath(DefaultHLODLayerPath);
+
+			if (UHLODLayer* EngineHLODLayer = EngineHLODLayerPath.LoadSynchronous())
+			{
+				Result = EngineHLODLayer;
+			}
+		}
+	}
+
+	return Result;
+}
+
+UHLODLayer* UHLODLayer::DuplicateHLODLayersSetup(UHLODLayer* HLODLayer, const FString& DestinationPath, const FString& Prefix)
+{
+	UHLODLayer* Result = nullptr;
+
+	UHLODLayer* LastHLODLayer = nullptr;
+	UHLODLayer* CurrentHLODLayer = HLODLayer;
+
+	while (CurrentHLODLayer)
+	{
+		const FString PackageName = DestinationPath + TEXT("_") + CurrentHLODLayer->GetName();
+		UPackage* Package = CreatePackage(*PackageName);
+
+		FString NewHLODLayerName = Prefix + TEXT("_") + CurrentHLODLayer->GetName();
+		UHLODLayer* NewHLODLayer = CastChecked<UHLODLayer>(StaticDuplicateObject(CurrentHLODLayer, Package, *NewHLODLayerName));
+		check(NewHLODLayer);
+
+		if (LastHLODLayer)
+		{
+			LastHLODLayer->SetParentLayer(NewHLODLayer);
+		}
+		else
+		{
+			Result = NewHLODLayer;
+		}
+
+		LastHLODLayer = NewHLODLayer;
+		CurrentHLODLayer = Cast<UHLODLayer>(CurrentHLODLayer->GetParentLayer().LoadSynchronous());
+	}
+
+	return Result;
+}
+
 void UHLODLayer::PostLoad()
 {
 	Super::PostLoad();
@@ -146,4 +201,9 @@ const TSoftObjectPtr<UHLODLayer>& UHLODLayer::GetParentLayer() const
 	return !IsSpatiallyLoaded() ? NullLayer : ParentLayer;
 }
 
-#endif // WITH_EDITOR
+const void UHLODLayer::SetParentLayer(const TSoftObjectPtr<UHLODLayer>& InParentLayer)
+{
+	ParentLayer = InParentLayer;
+}
+
+#endif // WITH_EDITORONLY_DATA
diff --git a/Engine/Source/Runtime/Engine/Private/WorldPartition/WorldPartition.cpp b/Engine/Source/Runtime/Engine/Private/WorldPartition/WorldPartition.cpp
--- a/Engine/Source/Runtime/Engine/Private/WorldPartition/WorldPartition.cpp
+++ b/Engine/Source/Runtime/Engine/Private/WorldPartition/WorldPartition.cpp
@@ -30,6 +30,7 @@
 #include "LevelEditorViewport.h"
 #include "Misc/ScopedSlowTask.h"
 #include "Misc/ScopeExit.h"
+#include "Misc/MessageDialog.h"
 #include "ScopedTransaction.h"
 #include "UnrealEdMisc.h"
 #include "WorldPartition/IWorldPartitionEditorModule.h"
@@ -42,6 +43,7 @@
 #include "WorldPartition/WorldPartitionMiniMapHelper.h"
 #include "WorldPartition/DataLayer/WorldDataLayers.h"
 #include "WorldPartition/WorldPartitionActorDescViewProxy.h"
+#include "WorldPartition/HLOD/HLODLayer.h"
 #include "Modules/ModuleManager.h"
 #include "GameDelegates.h"
 
@@ -255,6 +257,8 @@ UWorldPartition::UWorldPartition(const FObjectInitializer& ObjectInitializer)
 	, EditorHash(nullptr)
 	, bIsEditorLevelAlwaysLoaded(false)
 	, WorldPartitionEditor(nullptr)
+	, bEnableStreaming(true)
+	, bStreamingWasEnabled(true)
 	, bForceGarbageCollection(false)
 	, bForceGarbageCollectionPurge(false)
 	, bIsPIE(false)
@@ -348,7 +352,11 @@ void UWorldPartition::OnEndPlay()
 
 FName UWorldPartition::GetWorldPartitionEditorName() const
 {
-	return EditorHash->GetWorldPartitionEditorName();
+	if (bEnableStreaming)
+	{
+		return EditorHash->GetWorldPartitionEditorName();
+	}
+	return NAME_None;
 }
 
 void UWorldPartition::SetIsEditorLevelAlwaysLoaded(bool bIsAlwaysLoaded)
@@ -722,7 +730,7 @@ UWorldPartition* UWorldPartition::CreateOrRepairWorldPartition(AWorldSettings* W
 
 		WorldSettings->MarkPackageDirty();
 
-		WorldPartition->DefaultHLODLayer = nullptr;
+		WorldPartition->DefaultHLODLayer = UHLODLayer::GetEngineDefaultHLODLayersSetup();
 
 		AWorldDataLayers* WorldDataLayers = World->GetWorldDataLayers();
 		if (!WorldDataLayers)
@@ -1155,6 +1163,86 @@ void UWorldPartition::UnhashActorDesc(FWorldPartitionActorDesc* ActorDesc)
 	FWorldPartitionHandle ActorHandle(this, ActorDesc->GetGuid());
 	EditorHash->UnhashActor(ActorHandle);
 }
+
+void UWorldPartition::PostLoad()
+{
+	FWorldPartitionActorDesc::bForceAlwaysLoaded = !bEnableStreaming;
+
+	if (bEnableStreaming)
+	{
+		bStreamingWasEnabled = true;
+	}
+
+	Super::PostLoad();
+}
+
+void UWorldPartition::PostEditChangeProperty(struct FPropertyChangedEvent& PropertyChangedEvent)
+{
+	if (PropertyChangedEvent.GetPropertyName() == GET_MEMBER_NAME_CHECKED(UWorldPartition, bEnableStreaming))
+	{
+		bool bApplyStreamingChange = true;
+
+		if (bEnableStreaming)
+		{
+			if (!bStreamingWasEnabled)
+			{
+				if (FMessageDialog::Open(EAppMsgType::YesNo, LOCTEXT("WorldPartitionConfirmEnableStreaming", "You are about to enable streaming, the world will be setup to stream. Continue?")) == EAppReturnType::No)
+				{
+					bEnableStreaming = false;
+					bApplyStreamingChange = false;
+				}
+				else
+				{
+					bStreamingWasEnabled = true;
+
+					FMessageDialog::Open(EAppMsgType::Ok, LOCTEXT("WorldPartitionEnableStreamingDialolg", "Please refer to https://docs.unrealengine.com/5.0/en-US/building-virtual-worlds/world-partition for how to set up streaming."));
+				}
+			}
+		}
+		else
+		{
+			if (FMessageDialog::Open(EAppMsgType::YesNo, LOCTEXT("WorldPartitionConfirmDisableStreaming", "You are about to disable streaming, all actors in the world will be always loaded. Continue?")) == EAppReturnType::No)
+			{
+				bEnableStreaming = true;
+				bApplyStreamingChange = false;
+			}
+		}
+
+		if (bApplyStreamingChange)
+		{
+			// Pin the actor handles on the actor to prevent unloading it when unhashing
+			TArray<FWorldPartitionHandlePinRefScope> PinRefScopes;
+			PinRefScopes.Reserve(ActorDescList.Num());
+
+			for (UActorDescContainer::TIterator<> ActorDescIterator(this); ActorDescIterator; ++ActorDescIterator)
+			{
+				FWorldPartitionHandle ExistingActorHandle(this, ActorDescIterator->GetGuid());
+				PinRefScopes.Emplace(ExistingActorHandle);
+
+				UnhashActorDesc(*ActorDescIterator);
+			}
+
+			FWorldPartitionActorDesc::bForceAlwaysLoaded = !bEnableStreaming;
+
+			for (UActorDescContainer::TIterator<> ActorDescIterator(this); ActorDescIterator; ++ActorDescIterator)
+			{
+				HashActorDesc(*ActorDescIterator);
+			}
+
+			if (!bEnableStreaming)
+			{
+				UpdateLoadingEditorCell(EditorHash->GetAlwaysLoadedCell(), true, false);
+			}
+
+			if (WorldPartitionEditor)
+			{
+				WorldPartitionEditor->Reconstruct();
+			}
+		}
+	}
+
+	Super::PostEditChangeProperty(PropertyChangedEvent);
+}
 #endif
 
 bool UWorldPartition::ResolveSubobject(const TCHAR* SubObjectPath, UObject*& OutObject, bool bLoadIfExists)
diff --git a/Engine/Source/Runtime/Engine/Private/WorldPartition/WorldPartitionActorDesc.cpp b/Engine/Source/Runtime/Engine/Private/WorldPartition/WorldPartitionActorDesc.cpp
--- a/Engine/Source/Runtime/Engine/Private/WorldPartition/WorldPartitionActorDesc.cpp
+++ b/Engine/Source/Runtime/Engine/Private/WorldPartition/WorldPartitionActorDesc.cpp
@@ -20,6 +20,7 @@
 
 #if WITH_EDITOR
 uint32 FWorldPartitionActorDesc::GlobalTag = 0;
+bool FWorldPartitionActorDesc::bForceAlwaysLoaded = false;
 
 FWorldPartitionActorDesc::FWorldPartitionActorDesc()
 	: SoftRefCount(0)
diff --git a/Engine/Source/Runtime/Engine/Public/WorldPartition/HLOD/HLODLayer.h b/Engine/Source/Runtime/Engine/Public/WorldPartition/HLOD/HLODLayer.h
--- a/Engine/Source/Runtime/Engine/Public/WorldPartition/HLOD/HLODLayer.h
+++ b/Engine/Source/Runtime/Engine/Public/WorldPartition/HLOD/HLODLayer.h
@@ -41,6 +41,12 @@ public:
 	static UHLODLayer* GetHLODLayer(const AActor* InActor);
 	static UHLODLayer* GetHLODLayer(const FWorldPartitionActorDesc& InActorDesc, const UWorldPartition* InWorldPartition);
 	static UHLODLayer* GetHLODLayer(const FWorldPartitionActorDescView& InActorDescView, const UWorldPartition* InWorldPartition);
+
+	/** Get the default engine HLOD layers setup */
+	static UHLODLayer* GetEngineDefaultHLODLayersSetup();
+
+	/** Duplicate the provided HLOD layers setup */
+	static UHLODLayer* DuplicateHLODLayersSetup(UHLODLayer* HLODLayer, const FString& DestinationPath, const FString& Prefix);
 #endif
 
 #if WITH_EDITORONLY_DATA
@@ -52,6 +58,7 @@ public:
 	int32 GetCellSize() const { return !bIsSpatiallyLoaded ? 0 : CellSize; }
 	float GetLoadingRange() const { return !bIsSpatiallyLoaded ? WORLD_MAX : LoadingRange; }
 	const TSoftObjectPtr<UHLODLayer>& GetParentLayer() const;
+	const void SetParentLayer(const TSoftObjectPtr<UHLODLayer>& InParentLayer);
 	bool IsSpatiallyLoaded() const { return bIsSpatiallyLoaded; }
 	void SetIsSpatiallyLoaded(bool bInIsSpatiallyLoaded) { bIsSpatiallyLoaded = bInIsSpatiallyLoaded; }
 
diff --git a/Engine/Source/Runtime/Engine/Public/WorldPartition/WorldPartition.h b/Engine/Source/Runtime/Engine/Public/WorldPartition/WorldPartition.h
--- a/Engine/Source/Runtime/Engine/Public/WorldPartition/WorldPartition.h
+++ b/Engine/Source/Runtime/Engine/Public/WorldPartition/WorldPartition.h
@@ -52,6 +52,7 @@ enum class EWorldPartitionInitState
 struct ENGINE_API IWorldPartitionEditor
 {
 	virtual void Refresh() {}
+	virtual void Reconstruct() {}
 };
 
 class ENGINE_API ISourceControlHelper
@@ -121,6 +122,10 @@ private:
 
 public:
 	//~ Begin UObject Interface
+#if WITH_EDITOR
+	virtual void PostLoad() override;
+	virtual void PostEditChangeProperty(struct FPropertyChangedEvent& PropertyChangedEvent) override;
+#endif
 	virtual bool ResolveSubobject(const TCHAR* SubObjectPath, UObject*& OutObject, bool bLoadIfExists) override;
 	//~ End UObject Interface
 
@@ -200,11 +205,19 @@ public:
 	IWorldPartitionEditor* WorldPartitionEditor;
 
 	/** Class of WorldPartitionStreamingPolicy to be used to manage world partition streaming. */
-	UPROPERTY(EditAnywhere, Category=WorldPartition, NoClear, meta = (NoResetToDefault))
+	UPROPERTY()
 	TSubclassOf<UWorldPartitionStreamingPolicy> WorldPartitionStreamingPolicyClass;
+
+	/** Enables streaming for this world. */
+	UPROPERTY(EditAnywhere, Category=WorldPartition)
+	bool bEnableStreaming;
+
+	/** Used to know if it's the first time streaming is enabled on this world */
+	UPROPERTY()
+	bool bStreamingWasEnabled;
 #endif
 
-	UPROPERTY(EditAnywhere, Category=WorldPartition, NoClear, meta = (NoResetToDefault), Instanced)
+	UPROPERTY()
 	TObjectPtr<UWorldPartitionRuntimeHash> RuntimeHash;
 
 #if WITH_EDITOR
@@ -215,7 +228,7 @@ public:
 
 #if WITH_EDITORONLY_DATA
 	// Default HLOD layer
-	UPROPERTY(EditAnywhere, Category=WorldPartition, meta = (DisplayName = "Default HLOD Layer"))
+	UPROPERTY(EditAnywhere, Category=WorldPartition, meta = (DisplayName = "Default HLOD Layer", EditCondition="bEnableStreaming", EditConditionHides, HideEditConditionToggle))
 	TObjectPtr<class UHLODLayer> DefaultHLODLayer;
 
 	TArray<FWorldPartitionReference> LoadedSubobjects;
diff --git a/Engine/Source/Runtime/Engine/Public/WorldPartition/WorldPartitionActorDesc.h b/Engine/Source/Runtime/Engine/Public/WorldPartition/WorldPartitionActorDesc.h
--- a/Engine/Source/Runtime/Engine/Public/WorldPartition/WorldPartitionActorDesc.h
+++ b/Engine/Source/Runtime/Engine/Public/WorldPartition/WorldPartitionActorDesc.h
@@ -47,7 +47,7 @@ public:
 	inline UClass* GetActorClass() const { return ActorClass; }
 	inline FVector GetOrigin() const { return GetBounds().GetCenter(); }
 	inline FName GetRuntimeGrid() const { return RuntimeGrid; }
-	inline bool GetIsSpatiallyLoaded() const { return bIsSpatiallyLoaded; }
+	inline bool GetIsSpatiallyLoaded() const { return bForceAlwaysLoaded ? false : bIsSpatiallyLoaded; }
 	inline bool GetActorIsEditorOnly() const { return bActorIsEditorOnly; }
 	inline bool GetLevelBoundsRelevant() const { return bLevelBoundsRelevant; }
 	inline bool GetActorIsHLODRelevant() const { return bActorIsHLODRelevant; }
@@ -196,5 +196,8 @@ public:
 	// Tagging
 	mutable uint32					Tag;
 	static uint32					GlobalTag;
+
+	// Temp
+	static bool						bForceAlwaysLoaded;
 #endif
 };
