Fixed partial out-of-bounds read in morph target decoding.
Load2/3/4 that is partially out of bounds seems to have GPU-specific behavior. Pad input buffer to make sure load4 is always safe.

diff --git a/Engine/Source/Runtime/Engine/Private/Rendering/MorphTargetVertexInfoBuffers.cpp b/Engine/Source/Runtime/Engine/Private/Rendering/MorphTargetVertexInfoBuffers.cpp
--- a/Engine/Source/Runtime/Engine/Private/Rendering/MorphTargetVertexInfoBuffers.cpp
+++ b/Engine/Source/Runtime/Engine/Private/Rendering/MorphTargetVertexInfoBuffers.cpp
@@ -442,6 +442,9 @@ void FMorphTargetVertexInfoBuffers::InitMorphResources(EShaderPlatform ShaderPla
 
 	// Append bitstream data
 	MorphData.Append(BitstreamData);
+	MorphData.Add(0u);	// Pad to make sure it is always safe to access the data with load4s.
+	MorphData.Add(0u);
+	MorphData.Add(0u);
 
 	// UE_LOG(LogStaticMesh, Log, TEXT("Morph compression time: [%.2fs]"), FPlatformTime::ToMilliseconds(FPlatformTime::Cycles() - StartTime) / 1000.0f);
 
diff --git a/Engine/Source/Runtime/Engine/Private/SkeletalMeshRenderData.cpp b/Engine/Source/Runtime/Engine/Private/SkeletalMeshRenderData.cpp
--- a/Engine/Source/Runtime/Engine/Private/SkeletalMeshRenderData.cpp
+++ b/Engine/Source/Runtime/Engine/Private/SkeletalMeshRenderData.cpp
@@ -186,7 +186,7 @@ static void SerializeLODInfoForDDC(USkeletalMesh* SkeletalMesh, FString& KeySuff
 // differences, etc.) replace the version GUID below with a new one.
 // In case of merge conflicts with DDC versions, you MUST generate a new GUID
 // and set this new GUID as the version.
-#define SKELETALMESH_DERIVEDDATA_VER TEXT("B0E9C5D0DA6A4C07815CB825A5C1DE1A")
+#define SKELETALMESH_DERIVEDDATA_VER TEXT("7C7E9A6F35C248188B22C2587772C62B")
 
 const FString& GetSkeletalMeshDerivedDataVersion()
 {
