#ue5 - Apply fixes to checking character perch result from CL 18932018 only when using bUseFlatBaseForFloorChecks.

That setting is not very compatible with perch radius (sweep checks return points on the corners of the box when hitting flat surfaces below, outside perch radius), but this will allow it to work in more situations while not affecting the old capsule path that has been tuned to the old capsule location (usually the same but can differ slightly in some circumstances).

diff --git a/Engine/Source/Runtime/Engine/Private/Components/CharacterMovementComponent.cpp b/Engine/Source/Runtime/Engine/Private/Components/CharacterMovementComponent.cpp
--- a/Engine/Source/Runtime/Engine/Private/Components/CharacterMovementComponent.cpp
+++ b/Engine/Source/Runtime/Engine/Private/Components/CharacterMovementComponent.cpp
@@ -6716,7 +6716,7 @@ bool UCharacterMovementComponent::ComputePerchResult(const float TestRadius, con
 	// Sweep further than actual requested distance, because a reduced capsule radius means we could miss some hits that the normal radius would contact.
 	float PawnRadius, PawnHalfHeight;
 	CharacterOwner->GetCapsuleComponent()->GetScaledCapsuleSize(PawnRadius, PawnHalfHeight);
-	const FVector CapsuleLocation = InHit.TraceStart;
+	const FVector CapsuleLocation = (bUseFlatBaseForFloorChecks ? InHit.TraceStart : InHit.Location);
 
 	const float InHitAboveBase = FMath::Max<float>(0.f, InHit.ImpactPoint.Z - (CapsuleLocation.Z - PawnHalfHeight));
 	const float PerchLineDist = FMath::Max(0.f, InMaxFloorDist - InHitAboveBase);
