SceneOutliner - Added "InitiallyExpanded" option to folders (true by default)

* Allows code to create folders that are initially collapsed
* Used by the HLOD builder to create collapsed HLOD folders

Update default OpenWorld template map

diff --git a/Engine/Source/Editor/UnrealEd/Private/EditorActorFolders.cpp b/Engine/Source/Editor/UnrealEd/Private/EditorActorFolders.cpp
--- a/Engine/Source/Editor/UnrealEd/Private/EditorActorFolders.cpp
+++ b/Engine/Source/Editor/UnrealEd/Private/EditorActorFolders.cpp
@@ -433,6 +433,7 @@ void FActorFolders::OnActorFolderAdded(UActorFolder* InActorFolder)
 	check(Level->IsUsingActorFolders());
 
 	AddFolderToWorld(*Level->GetWorld(), InActorFolder->GetFolder());
+	SetIsFolderExpanded(*Level->GetWorld(), InActorFolder->GetFolder(), InActorFolder->IsInitiallyExpanded());
 }
 
 void FActorFolders::OnFolderRootObjectRemoved(UWorld& InWorld, const FFolder::FRootObject& InFolderRootObject)
diff --git a/Engine/Source/Editor/UnrealEd/Private/WorldPartition/WorldPartitionHLODsBuilder.cpp b/Engine/Source/Editor/UnrealEd/Private/WorldPartition/WorldPartitionHLODsBuilder.cpp
--- a/Engine/Source/Editor/UnrealEd/Private/WorldPartition/WorldPartitionHLODsBuilder.cpp
+++ b/Engine/Source/Editor/UnrealEd/Private/WorldPartition/WorldPartitionHLODsBuilder.cpp
@@ -330,6 +330,8 @@ bool UWorldPartitionHLODsBuilder::SetupHLODActors()
 {
 	auto ActorFolderAddedDelegateHandle = GEngine->OnActorFolderAdded().AddLambda([this](UActorFolder* InActorFolder)
 	{
+		// We don't want the HLOD folders to be expanded by default
+		InActorFolder->SetIsInitiallyExpanded(false);
 		SourceControlHelper->Save(InActorFolder->GetPackage());
 	});
 	
diff --git a/Engine/Source/Runtime/Engine/Private/ActorFolder.cpp b/Engine/Source/Runtime/Engine/Private/ActorFolder.cpp
--- a/Engine/Source/Runtime/Engine/Private/ActorFolder.cpp
+++ b/Engine/Source/Runtime/Engine/Private/ActorFolder.cpp
@@ -38,6 +38,7 @@ UActorFolder* UActorFolder::Create(ULevel* InLevel, const FString& InFolderLabel
 	ActorFolder->FolderGuid = NewFolderGuid;
 	ActorFolder->SetLabel(InFolderLabel);
 	ActorFolder->SetParent(InParent);
+	ActorFolder->SetIsInitiallyExpanded(true);
 
 	FLevelActorFoldersHelper::AddActorFolder(InLevel, ActorFolder, bShouldDirtyLevel);
 	return ActorFolder;
@@ -55,6 +56,7 @@ namespace ActorFolder
 	static const FName NAME_FolderGuid(TEXT("FolderGuid"));
 	static const FName NAME_ParentFolderGuid(TEXT("ParentFolderGuid"));
 	static const FName NAME_FolderLabel(TEXT("FolderLabel"));
+	static const FName NAME_FolderInitiallyExpanded(TEXT("FolderInitiallyExpanded"));
 	static const FName NAME_FolderIsDeleted(TEXT("FolderIsDeleted"));
 	static const FName NAME_OuterPackageName(TEXT("OuterPackageName"));
 };
@@ -64,6 +66,7 @@ void UActorFolder::GetAssetRegistryTags(TArray<FAssetRegistryTag>& OutTags) cons
 	OutTags.Add(FAssetRegistryTag(ActorFolder::NAME_FolderGuid, *FolderGuid.ToString(), FAssetRegistryTag::TT_Hidden));
 	OutTags.Add(FAssetRegistryTag(ActorFolder::NAME_ParentFolderGuid, *ParentFolderGuid.ToString(), FAssetRegistryTag::TT_Hidden));
 	OutTags.Add(FAssetRegistryTag(ActorFolder::NAME_FolderLabel, *FolderLabel, FAssetRegistryTag::TT_Hidden));
+	OutTags.Add(FAssetRegistryTag(ActorFolder::NAME_FolderInitiallyExpanded, bFolderInitiallyExpanded ? TEXT("1") : TEXT("0"), FAssetRegistryTag::TT_Hidden));
 	OutTags.Add(FAssetRegistryTag(ActorFolder::NAME_FolderIsDeleted, bIsDeleted ? TEXT("1") : TEXT("0"), FAssetRegistryTag::TT_Hidden));
 	OutTags.Add(FAssetRegistryTag(ActorFolder::NAME_OuterPackageName, *GetOuterULevel()->GetPackage()->GetName(), FAssetRegistryTag::TT_Hidden));
 }
@@ -94,6 +97,10 @@ FActorFolderDesc UActorFolder::GetAssetRegistryInfoFromPackage(FName ActorFolder
 			{
 				ActorFolderDesc.FolderLabel = Value;
 			}
+			if (Asset.GetTagValue(ActorFolder::NAME_FolderInitiallyExpanded, Value))
+			{
+				ActorFolderDesc.bFolderInitiallyExpanded = (Value == TEXT("1"));
+			}
 			if (Asset.GetTagValue(ActorFolder::NAME_FolderIsDeleted, Value))
 			{
 				ActorFolderDesc.bFolderIsDeleted = (Value == TEXT("1"));
@@ -117,6 +124,16 @@ void UActorFolder::SetLabel(const FString& InFolderLabel)
 	}
 }
 
+void UActorFolder::SetIsInitiallyExpanded(bool bInFolderInitiallyExpanded)
+{
+	check(IsValid());
+	if (bFolderInitiallyExpanded != bInFolderInitiallyExpanded)
+	{
+		Modify();
+		bFolderInitiallyExpanded = bInFolderInitiallyExpanded;
+	}
+}
+
 void UActorFolder::SetParent(UActorFolder* InParent)
 {
 	if ((this != InParent) && (GetParent() != InParent))
diff --git a/Engine/Source/Runtime/Engine/Private/ActorFolderDesc.cpp b/Engine/Source/Runtime/Engine/Private/ActorFolderDesc.cpp
--- a/Engine/Source/Runtime/Engine/Private/ActorFolderDesc.cpp
+++ b/Engine/Source/Runtime/Engine/Private/ActorFolderDesc.cpp
@@ -15,7 +15,8 @@
   */
 
 FActorFolderDesc::FActorFolderDesc()
-	: bFolderIsDeleted(false)
+	: bFolderInitiallyExpanded(true)
+	, bFolderIsDeleted(false)
 {}
 
 FString FActorFolderDesc::GetDisplayName() const
diff --git a/Engine/Source/Runtime/Engine/Public/ActorFolder.h b/Engine/Source/Runtime/Engine/Public/ActorFolder.h
--- a/Engine/Source/Runtime/Engine/Public/ActorFolder.h
+++ b/Engine/Source/Runtime/Engine/Public/ActorFolder.h
@@ -38,6 +38,10 @@ public:
 	void SetLabel(const FString& InFolderLabel);
 	FString GetLabel() const { return FolderLabel; }
 
+	// Control whether the folder is initially expanded or not
+	void SetIsInitiallyExpanded(bool bInFolderInitiallyExpanded);
+	bool IsInitiallyExpanded() const { return bFolderInitiallyExpanded; }
+
 	const FGuid& GetGuid() const { return FolderGuid; }
 	FName GetPath() const;
 	FString GetDisplayName() const;
@@ -74,6 +78,9 @@ private:
 	UPROPERTY()
 	FString FolderLabel;
 
+	UPROPERTY()
+	bool bFolderInitiallyExpanded = true;
+
 	UPROPERTY()
 	bool bIsDeleted;
 #endif
diff --git a/Engine/Source/Runtime/Engine/Public/ActorFolderDesc.h b/Engine/Source/Runtime/Engine/Public/ActorFolderDesc.h
--- a/Engine/Source/Runtime/Engine/Public/ActorFolderDesc.h
+++ b/Engine/Source/Runtime/Engine/Public/ActorFolderDesc.h
@@ -15,6 +15,7 @@ public:
 	const FGuid& GetParentFolderGuid() const { return ParentFolderGuid; }
 	const FGuid& GetFolderGuid() const { return FolderGuid; }
 	const FString& GetFolderLabel() const { return FolderLabel; }
+	bool IsFolderInitiallyExpanded() const { return bFolderInitiallyExpanded; }
 	bool IsFolderDeleted() const { return bFolderIsDeleted; }
 	FString GetOuterPackageName() const { return OuterPackageName; }
 	FString GetDisplayName() const;
@@ -24,6 +25,7 @@ private:
 	FGuid ParentFolderGuid;
 	FGuid FolderGuid;
 	FString FolderLabel;
+	bool bFolderInitiallyExpanded;
 	bool bFolderIsDeleted;
 	FString OuterPackageName;
 
