SceneOutliner - Unloaded LandscapeProxy actors are now placed under their associated Landscape

* FWorldPartitionActorDesc now has a GetSceneOutlinerParent() virtual which is used to build the actor hierarchy of the SceneOutliner
* FLandscapeActorDesc overrides GetSceneOutlinerParent() to return the Landscape guid as it's parent

Resaved all LandscapeProxies in the OpenWorld template map

diff --git a/Engine/Source/Editor/SceneOutliner/Private/ActorHierarchy.cpp b/Engine/Source/Editor/SceneOutliner/Private/ActorHierarchy.cpp
--- a/Engine/Source/Editor/SceneOutliner/Private/ActorHierarchy.cpp
+++ b/Engine/Source/Editor/SceneOutliner/Private/ActorHierarchy.cpp
@@ -278,6 +278,41 @@ FSceneOutlinerTreeItemPtr FActorHierarchy::FindOrCreateParentItem(const ISceneOu
 					return bCreate ? Mode->CreateItemFor<FActorFolderTreeItem>(FActorFolderTreeItem(FFolder(FolderPath), RepresentingWorld.Get()), true) : nullptr;
 				}
 			}
+
+			// Parent Actor (Actor attachement / parenting)
+			const FGuid& ParentActorGuid = ActorDesc->GetSceneOutlinerParent();
+			if (ParentActorGuid.IsValid())
+			{
+				if (UWorldPartition* WorldPartition = RepresentingWorld->GetWorldPartition())
+				{
+					if (const FWorldPartitionActorDesc* ParentActorDesc = WorldPartition->GetActorDesc(ParentActorGuid))
+					{
+						// If parent actor is loaded
+						if (AActor* ParentActor = ParentActorDesc->GetActor())
+						{
+							// Find loaded parent actor node (from the object ptr)
+							if (const FSceneOutlinerTreeItemPtr* ParentItem = Items.Find(ParentActor))
+							{
+								return *ParentItem;
+							}
+							else
+							{
+								return bCreate ? Mode->CreateItemFor<FActorTreeItem>(ParentActor, true) : nullptr;
+							}
+						}
+
+						// Find unloaded parent actor node (from the guid)
+						if (const FSceneOutlinerTreeItemPtr* ParentItem = Items.Find(ParentActorGuid))
+						{
+							return *ParentItem;
+						}
+						else
+						{
+							return bCreate ? Mode->CreateItemFor<FActorDescTreeItem>(FActorDescTreeItem(ParentActorGuid, WorldPartition)) : nullptr;
+						}
+					}
+				}
+			}
 		}
 	}
 
diff --git a/Engine/Source/Runtime/Core/Public/UObject/UE5ReleaseStreamObjectVersion.h b/Engine/Source/Runtime/Core/Public/UObject/UE5ReleaseStreamObjectVersion.h
--- a/Engine/Source/Runtime/Core/Public/UObject/UE5ReleaseStreamObjectVersion.h
+++ b/Engine/Source/Runtime/Core/Public/UObject/UE5ReleaseStreamObjectVersion.h
@@ -108,6 +108,9 @@ struct CORE_API FUE5ReleaseStreamObjectVersion
 		// Add faster damping calculations to the cloth simulation and rename previous Damping parameter to LocalDamping.
 		ChaosClothFasterDamping,
 
+		// Serialize LandscapeActorGuid in FLandscapeActorDesc sub class.
+		WorldPartitionLandscapeActorDescSerializeLandscapeActorGuid,
+
 		// -----<new versions can be added above this line>-------------------------------------------------
 		VersionPlusOne,
 		LatestVersion = VersionPlusOne - 1
diff --git a/Engine/Source/Runtime/Engine/Private/WorldPartition/Landscape/LandscapeActorDesc.cpp b/Engine/Source/Runtime/Engine/Private/WorldPartition/Landscape/LandscapeActorDesc.cpp
--- a/Engine/Source/Runtime/Engine/Private/WorldPartition/Landscape/LandscapeActorDesc.cpp
+++ b/Engine/Source/Runtime/Engine/Private/WorldPartition/Landscape/LandscapeActorDesc.cpp
@@ -3,12 +3,14 @@
 #include "WorldPartition/Landscape/LandscapeActorDesc.h"
 
 #if WITH_EDITOR
+#include "Landscape.h"
+#include "LandscapeInfo.h"
 #include "LandscapeProxy.h"
 #include "LandscapeStreamingProxy.h"
-#include "LandscapeInfo.h"
 #include "WorldPartition/ActorDescContainer.h"
 #include "WorldPartition/WorldPartitionHandle.h"
 #include "UObject/UE5MainStreamObjectVersion.h"
+#include "UObject/UE5ReleaseStreamObjectVersion.h"
 
 void FLandscapeActorDesc::Init(const AActor* InActor)
 {
@@ -19,17 +21,31 @@ void FLandscapeActorDesc::Init(const AActor* InActor)
 	GridIndexX = LandscapeProxy->LandscapeSectionOffset.X / (int32)LandscapeProxy->GridSize;
 	GridIndexY = LandscapeProxy->LandscapeSectionOffset.Y / (int32)LandscapeProxy->GridSize;
 	GridIndexZ = 0;
+
+	const ALandscape* LandscapeActor = LandscapeProxy->GetLandscapeActor();
+	if (LandscapeActor)
+	{
+		LandscapeActorGuid = LandscapeActor->GetActorGuid();
+	}
 }
 
 void FLandscapeActorDesc::Serialize(FArchive& Ar)
 {
 	FPartitionActorDesc::Serialize(Ar);
 
-	if (Ar.IsLoading() && Ar.CustomVer(FUE5MainStreamObjectVersion::GUID) < FUE5MainStreamObjectVersion::FLandscapeActorDescFixupGridIndices)
+	Ar.UsingCustomVersion(FUE5MainStreamObjectVersion::GUID);
+	Ar.UsingCustomVersion(FUE5ReleaseStreamObjectVersion::GUID);
+
+	if (Ar.CustomVer(FUE5MainStreamObjectVersion::GUID) < FUE5MainStreamObjectVersion::FLandscapeActorDescFixupGridIndices)
 	{
 		GridIndexX = (int32)(GridIndexX * GridSize) / (int32)GridSize;
 		GridIndexY = (int32)(GridIndexY * GridSize) / (int32)GridSize;
 	}
+
+	if (Ar.CustomVer(FUE5ReleaseStreamObjectVersion::GUID) >= FUE5ReleaseStreamObjectVersion::WorldPartitionLandscapeActorDescSerializeLandscapeActorGuid)
+	{
+		Ar << LandscapeActorGuid;
+	}
 }
 
 void FLandscapeActorDesc::Unload()
@@ -41,4 +57,21 @@ void FLandscapeActorDesc::Unload()
 
 	FPartitionActorDesc::Unload();
 }
+
+bool FLandscapeActorDesc::Equals(const FWorldPartitionActorDesc* Other) const
+{
+	if (FPartitionActorDesc::Equals(Other))
+	{
+		const FLandscapeActorDesc* LandscapeActorDesc = (FLandscapeActorDesc*)Other;
+		return LandscapeActorGuid == LandscapeActorDesc->LandscapeActorGuid;
+	}
+
+	return false;
+}
+
+const FGuid& FLandscapeActorDesc::GetSceneOutlinerParent() const
+{
+	return LandscapeActorGuid;
+}
+
 #endif
diff --git a/Engine/Source/Runtime/Engine/Public/WorldPartition/Landscape/LandscapeActorDesc.h b/Engine/Source/Runtime/Engine/Public/WorldPartition/Landscape/LandscapeActorDesc.h
--- a/Engine/Source/Runtime/Engine/Public/WorldPartition/Landscape/LandscapeActorDesc.h
+++ b/Engine/Source/Runtime/Engine/Public/WorldPartition/Landscape/LandscapeActorDesc.h
@@ -14,9 +14,14 @@ class ENGINE_API FLandscapeActorDesc : public FPartitionActorDesc
 #if WITH_EDITOR
 public:
 	virtual void Init(const AActor* InActor) override;
+	virtual bool Equals(const FWorldPartitionActorDesc* Other) const override;
 	virtual void Unload() override;
+	virtual const FGuid& GetSceneOutlinerParent() const override;
 
 protected:
 	virtual void Serialize(FArchive& Ar) override;
+
+private:
+	FGuid LandscapeActorGuid;
 #endif
 };
diff --git a/Engine/Source/Runtime/Engine/Public/WorldPartition/WorldPartitionActorDesc.h b/Engine/Source/Runtime/Engine/Public/WorldPartition/WorldPartitionActorDesc.h
--- a/Engine/Source/Runtime/Engine/Public/WorldPartition/WorldPartitionActorDesc.h
+++ b/Engine/Source/Runtime/Engine/Public/WorldPartition/WorldPartitionActorDesc.h
@@ -65,6 +65,7 @@ public:
 	FName GetActorLabelOrName() const;
 
 	virtual bool GetContainerInstance(const UActorDescContainer*& OutLevelContainer, FTransform& OutLevelTransform, EContainerClusterMode& OutClusterMode) const { return false; }
+	virtual const FGuid& GetSceneOutlinerParent() const { return GetParentActor(); }
 
 	bool operator==(const FWorldPartitionActorDesc& Other) const
 	{
