Fix for morph target cooking failing with assert

diff --git a/Engine/Source/Runtime/Engine/Private/Rendering/MorphTargetVertexInfoBuffers.cpp b/Engine/Source/Runtime/Engine/Private/Rendering/MorphTargetVertexInfoBuffers.cpp
--- a/Engine/Source/Runtime/Engine/Private/Rendering/MorphTargetVertexInfoBuffers.cpp
+++ b/Engine/Source/Runtime/Engine/Private/Rendering/MorphTargetVertexInfoBuffers.cpp
@@ -442,9 +442,14 @@ void FMorphTargetVertexInfoBuffers::InitMorphResources(EShaderPlatform ShaderPla
 
 	// Append bitstream data
 	MorphData.Append(BitstreamData);
-	MorphData.Add(0u);	// Pad to make sure it is always safe to access the data with load4s.
-	MorphData.Add(0u);
-	MorphData.Add(0u);
+	
+	if(MorphData.Num() > 0)
+	{
+		// Pad to make sure it is always safe to access the data with load4s.
+		MorphData.Add(0u);
+		MorphData.Add(0u);
+		MorphData.Add(0u);
+	}
 
 	// UE_LOG(LogStaticMesh, Log, TEXT("Morph compression time: [%.2fs]"), FPlatformTime::ToMilliseconds(FPlatformTime::Cycles() - StartTime) / 1000.0f);
 
diff --git a/Engine/Source/Runtime/Engine/Private/SkeletalMeshRenderData.cpp b/Engine/Source/Runtime/Engine/Private/SkeletalMeshRenderData.cpp
--- a/Engine/Source/Runtime/Engine/Private/SkeletalMeshRenderData.cpp
+++ b/Engine/Source/Runtime/Engine/Private/SkeletalMeshRenderData.cpp
@@ -186,7 +186,7 @@ static void SerializeLODInfoForDDC(USkeletalMesh* SkeletalMesh, FString& KeySuff
 // differences, etc.) replace the version GUID below with a new one.
 // In case of merge conflicts with DDC versions, you MUST generate a new GUID
 // and set this new GUID as the version.
-#define SKELETALMESH_DERIVEDDATA_VER TEXT("7C7E9A6F35C248188B22C2587772C62B")
+#define SKELETALMESH_DERIVEDDATA_VER TEXT("724A220873DA44A4B31CBE7004F1F545")
 
 const FString& GetSkeletalMeshDerivedDataVersion()
 {
