ULevel: make GetLoadedExternalObjectPackages() tolerant of memory worlds

This replaces a call to the FPackageName function LongPackageNameToFilename()
with the failure-tolerant TryConvertLongPackageNameToFilename(), which avoids
triggering a fatal error when the level's outer is a memory world.

diff --git a/Engine/Source/Runtime/Engine/Private/Level.cpp b/Engine/Source/Runtime/Engine/Private/Level.cpp
--- a/Engine/Source/Runtime/Engine/Private/Level.cpp
+++ b/Engine/Source/Runtime/Engine/Private/Level.cpp
@@ -3215,7 +3215,12 @@ TArray<UPackage*> ULevel::GetLoadedExternalObjectPackages() const
 	ExternalObjectsPaths.Add(SanitizeExternalPath(ULevel::GetExternalActorsPath(World->GetPackage(), (World->OriginalWorldName == NAME_None) ? World->GetName() : World->OriginalWorldName.ToString())));
 	ExternalObjectsPaths.Add(SanitizeExternalPath(FExternalPackageHelper::GetExternalObjectsPath(World->GetPackage(), (World->OriginalWorldName == NAME_None) ? World->GetName() : World->OriginalWorldName.ToString())));
 	ExternalObjectsPaths = ExternalObjectsPaths.FilterByPredicate([](const FString& Path) {
-			return !Path.IsEmpty() && FPaths::DirectoryExists(FPackageName::LongPackageNameToFilename(Path));
+			FString Filename;
+			if (!Path.IsEmpty() && FPackageName::TryConvertLongPackageNameToFilename(Path, Filename))
+			{
+				return FPaths::DirectoryExists(Filename);
+			}
+			return false;
 		});
 
 	if (!ExternalObjectsPaths.IsEmpty())
