Chaos
- RBAN QPBD and linear joint solver support
- Add 3rd phase to solver loop for projection
- Add settings for QPBD to PhysicsAsset
- Add cvar overrides for all settings
- Hide settings for legacy solver

diff --git a/Engine/Source/Editor/DetailCustomizations/Private/PhysicsConstraintComponentDetails.cpp b/Engine/Source/Editor/DetailCustomizations/Private/PhysicsConstraintComponentDetails.cpp
--- a/Engine/Source/Editor/DetailCustomizations/Private/PhysicsConstraintComponentDetails.cpp
+++ b/Engine/Source/Editor/DetailCustomizations/Private/PhysicsConstraintComponentDetails.cpp
@@ -175,9 +175,10 @@ void FPhysicsConstraintComponentDetails::AddConstraintBehaviorProperties(IDetail
 
 	//Add properties we want in specific order
 	ConstraintCat.AddProperty(ProfilePropertiesProperty->GetChildHandle(GET_MEMBER_NAME_CHECKED(FConstraintProfileProperties, bDisableCollision)));
-	ConstraintCat.AddProperty(ProfilePropertiesProperty->GetChildHandle(GET_MEMBER_NAME_CHECKED(FConstraintProfileProperties, bEnableProjection)));
 #if WITH_CHAOS
-	ConstraintCat.AddProperty(ProfilePropertiesProperty->GetChildHandle(GET_MEMBER_NAME_CHECKED(FConstraintProfileProperties, bEnableSoftProjection)));
+	ConstraintCat.AddProperty(ProfilePropertiesProperty->GetChildHandle(GET_MEMBER_NAME_CHECKED(FConstraintProfileProperties, bEnableLinearProjection)));
+	ConstraintCat.AddProperty(ProfilePropertiesProperty->GetChildHandle(GET_MEMBER_NAME_CHECKED(FConstraintProfileProperties, bEnableAngularProjection)));
+	ConstraintCat.AddProperty(ProfilePropertiesProperty->GetChildHandle(GET_MEMBER_NAME_CHECKED(FConstraintProfileProperties, bEnableShockPropagation)));
 	ConstraintCat.AddProperty(ProfilePropertiesProperty->GetChildHandle(GET_MEMBER_NAME_CHECKED(FConstraintProfileProperties, ProjectionLinearAlpha)));
 	ConstraintCat.AddProperty(ProfilePropertiesProperty->GetChildHandle(GET_MEMBER_NAME_CHECKED(FConstraintProfileProperties, ProjectionAngularAlpha)));
 	ConstraintCat.AddProperty(ProfilePropertiesProperty->GetChildHandle(GET_MEMBER_NAME_CHECKED(FConstraintProfileProperties, ShockPropagationAlpha)));
@@ -185,8 +186,8 @@ void FPhysicsConstraintComponentDetails::AddConstraintBehaviorProperties(IDetail
 
 #if WITH_CHAOS
 	// Hide PhysX-Only settings in Chaos
-	ProfilePropertiesProperty->GetChildHandle(GET_MEMBER_NAME_CHECKED(FConstraintProfileProperties, ProjectionLinearTolerance))->MarkHiddenByCustomization();
-	ProfilePropertiesProperty->GetChildHandle(GET_MEMBER_NAME_CHECKED(FConstraintProfileProperties, ProjectionAngularTolerance))->MarkHiddenByCustomization();
+	//ProfilePropertiesProperty->GetChildHandle(GET_MEMBER_NAME_CHECKED(FConstraintProfileProperties, ProjectionLinearTolerance))->MarkHiddenByCustomization();
+	//ProfilePropertiesProperty->GetChildHandle(GET_MEMBER_NAME_CHECKED(FConstraintProfileProperties, ProjectionAngularTolerance))->MarkHiddenByCustomization();
 #endif
 
 #if !WITH_CHAOS
diff --git a/Engine/Source/Editor/PhysicsAssetEditor/Private/PhysicsAssetEditorAnimInstanceProxy.cpp b/Engine/Source/Editor/PhysicsAssetEditor/Private/PhysicsAssetEditorAnimInstanceProxy.cpp
--- a/Engine/Source/Editor/PhysicsAssetEditor/Private/PhysicsAssetEditorAnimInstanceProxy.cpp
+++ b/Engine/Source/Editor/PhysicsAssetEditor/Private/PhysicsAssetEditorAnimInstanceProxy.cpp
@@ -39,6 +39,7 @@ void FPhysicsAssetEditorAnimInstanceProxy::Initialize(UAnimInstance* InAnimInsta
 	UPhysicsAsset* PhysicsAsset = InAnimInstance->GetSkelMeshComponent()->GetPhysicsAsset();
 	if (PhysicsAsset != nullptr)
 	{
+		SolverSettings = PhysicsAsset->SolverSettings;
 		SolverIterations = PhysicsAsset->SolverIterations;
 	}
 
@@ -83,8 +84,15 @@ bool FPhysicsAssetEditorAnimInstanceProxy::Evaluate_WithRoot(FPoseContext& Outpu
 	ImmediatePhysics::FSimulation* Simulation = RagdollNode.GetSimulation();
 	if (Simulation != nullptr)
 	{
-		Simulation->SetSolverIterations(
-			SolverIterations.FixedTimeStep,
+		Simulation->SetSolverSettings(
+			SolverSettings.FixedTimeStep,
+			SolverSettings.CullDistance,
+			SolverSettings.MaxDepenetrationVelocity,
+			SolverSettings.PositionIterations,
+			SolverSettings.VelocityIterations,
+			SolverSettings.ProjectionIterations);
+
+		Simulation->SetLegacySolverSettings(
 			SolverIterations.SolverIterations,
 			SolverIterations.JointIterations,
 			SolverIterations.CollisionIterations,
diff --git a/Engine/Source/Editor/PhysicsAssetEditor/Private/PhysicsAssetEditorAnimInstanceProxy.h b/Engine/Source/Editor/PhysicsAssetEditor/Private/PhysicsAssetEditorAnimInstanceProxy.h
--- a/Engine/Source/Editor/PhysicsAssetEditor/Private/PhysicsAssetEditorAnimInstanceProxy.h
+++ b/Engine/Source/Editor/PhysicsAssetEditor/Private/PhysicsAssetEditorAnimInstanceProxy.h
@@ -47,6 +47,7 @@ private:
 
 	FAnimNode_RigidBody RagdollNode;
  	FAnimNode_ConvertComponentToLocalSpace ComponentToLocalSpace;
+	FPhysicsAssetSolverSettings SolverSettings;
 	FSolverIterations SolverIterations;
 
 	ImmediatePhysics::FActorHandle* TargetActor;
diff --git a/Engine/Source/Programs/HeadlessChaos/Public/HeadlessChaosCollisionConstraints.h b/Engine/Source/Programs/HeadlessChaos/Public/HeadlessChaosCollisionConstraints.h
--- a/Engine/Source/Programs/HeadlessChaos/Public/HeadlessChaosCollisionConstraints.h
+++ b/Engine/Source/Programs/HeadlessChaos/Public/HeadlessChaosCollisionConstraints.h
@@ -144,6 +144,7 @@ public:
 
 		for (auto Handle : CollisionConstraints.GetConstraintHandles())
 		{
+			Handle->PreGatherInput(SolverData);
 			Handle->GatherInput(Dt, INDEX_NONE, INDEX_NONE, SolverData);
 		}
 	}
diff --git a/Engine/Source/Runtime/AnimGraphRuntime/Private/BoneControllers/AnimNode_RigidBody.cpp b/Engine/Source/Runtime/AnimGraphRuntime/Private/BoneControllers/AnimNode_RigidBody.cpp
--- a/Engine/Source/Runtime/AnimGraphRuntime/Private/BoneControllers/AnimNode_RigidBody.cpp
+++ b/Engine/Source/Runtime/AnimGraphRuntime/Private/BoneControllers/AnimNode_RigidBody.cpp
@@ -158,6 +158,7 @@ FAnimNode_RigidBody::FAnimNode_RigidBody()
 	, bSimulateAnimPhysicsAfterReset(false)
 	, SkelMeshCompWeakPtr()
 	, PhysicsSimulation(nullptr)
+	, SolverSettings()
 	, SolverIterations()
 	, SimulationTask()
 	, OutputBoneData()
@@ -825,8 +826,15 @@ void FAnimNode_RigidBody::EvaluateSkeletalControl_AnyThread(FComponentSpacePoseC
 				UseSimSpaceSettings->MasterAlpha, 
 				UseSimSpaceSettings->ExternalLinearDragV);
 
-			PhysicsSimulation->SetSolverIterations(
-				SolverIterations.FixedTimeStep,
+			PhysicsSimulation->SetSolverSettings(
+				SolverSettings.FixedTimeStep,
+				SolverSettings.CullDistance,
+				SolverSettings.MaxDepenetrationVelocity,
+				SolverSettings.PositionIterations,
+				SolverSettings.VelocityIterations,
+				SolverSettings.ProjectionIterations);
+
+			PhysicsSimulation->SetLegacySolverSettings(
 				SolverIterations.SolverIterations,
 				SolverIterations.JointIterations,
 				SolverIterations.CollisionIterations,
@@ -1248,16 +1256,23 @@ void FAnimNode_RigidBody::InitPhysics(const UAnimInstance* InAnimInstance)
 		PhysicsSimulation->SetIgnoreCollisionActors(IgnoreCollisionActors);
 
 #if WITH_CHAOS
+		SolverSettings = UsePhysicsAsset->SolverSettings;
+		PhysicsSimulation->SetSolverSettings(
+			SolverSettings.FixedTimeStep,
+			SolverSettings.CullDistance,
+			SolverSettings.MaxDepenetrationVelocity,
+			SolverSettings.PositionIterations,
+			SolverSettings.VelocityIterations,
+			SolverSettings.ProjectionIterations);
+
 		SolverIterations = UsePhysicsAsset->SolverIterations;
-		PhysicsSimulation->SetSolverIterations(
-			SolverIterations.FixedTimeStep,
+		PhysicsSimulation->SetLegacySolverSettings(
 			SolverIterations.SolverIterations,
 			SolverIterations.JointIterations,
 			SolverIterations.CollisionIterations,
 			SolverIterations.SolverPushOutIterations,
 			SolverIterations.JointPushOutIterations,
-			SolverIterations.CollisionPushOutIterations
-		);
+			SolverIterations.CollisionPushOutIterations);
 #endif
 	}
 }
diff --git a/Engine/Source/Runtime/AnimGraphRuntime/Public/BoneControllers/AnimNode_RigidBody.h b/Engine/Source/Runtime/AnimGraphRuntime/Public/BoneControllers/AnimNode_RigidBody.h
--- a/Engine/Source/Runtime/AnimGraphRuntime/Public/BoneControllers/AnimNode_RigidBody.h
+++ b/Engine/Source/Runtime/AnimGraphRuntime/Public/BoneControllers/AnimNode_RigidBody.h
@@ -339,7 +339,8 @@ private:
 	TWeakObjectPtr<USkeletalMeshComponent> SkelMeshCompWeakPtr;
 
 	ImmediatePhysics::FSimulation* PhysicsSimulation;
-	FSolverIterations SolverIterations;
+	FPhysicsAssetSolverSettings SolverSettings;
+	FSolverIterations SolverIterations;	// to be deprecated
 
 	friend class FRigidBodyNodeSimulationTask;
 	UE::Tasks::FTask SimulationTask;
diff --git a/Engine/Source/Runtime/Engine/Classes/PhysicsEngine/ConstraintInstance.h b/Engine/Source/Runtime/Engine/Classes/PhysicsEngine/ConstraintInstance.h
--- a/Engine/Source/Runtime/Engine/Classes/PhysicsEngine/ConstraintInstance.h
+++ b/Engine/Source/Runtime/Engine/Classes/PhysicsEngine/ConstraintInstance.h
@@ -27,22 +27,26 @@ struct ENGINE_API FConstraintProfileProperties
 	GENERATED_USTRUCT_BODY()
 
 	/** [PhysX only] Linear tolerance value in world units. If the distance error exceeds this tolerence limit, the body will be projected. */
-	UPROPERTY(EditAnywhere, Category = Projection, meta = (editcondition = "bEnableProjection", ClampMin = "0.0"))
+	UPROPERTY()
 	float ProjectionLinearTolerance;
 
 	/** [PhysX only] Angular tolerance value in world units. If the distance error exceeds this tolerence limit, the body will be projected. */
-	UPROPERTY(EditAnywhere, Category = Projection, meta = (editcondition = "bEnableProjection", ClampMin = "0.0"))
+	UPROPERTY()
 	float ProjectionAngularTolerance;
 
-	/** [Chaos Only] How much linear projection to apply [0-1]. Projection fixes any post-solve position error in the constraint. */
+	/**  How much linear projection to apply [0-1] */
 	UPROPERTY(EditAnywhere, Category = Projection, meta = (ClampMin = "0.0", ClampMax = "1.0"))
 	float ProjectionLinearAlpha;
 
-	/** [Chaos Only] How much angular projection to apply [0-1]. Projection fixes any post-solve angle error in the constraint. */
+	/** How much angular projection to apply [0-1] */
 	UPROPERTY(EditAnywhere, Category = Projection, meta = (ClampMin = "0.0", ClampMax = "1.0"))
 	float ProjectionAngularAlpha;
 
-	/** How much shock propagation to apply [0-1] */
+	/** 
+	 * How much shock propagation to apply [0-1]. Shock propagation increases the mass of the parent body for the last iteration of the
+	 * position and velocity solve phases. This can help stiffen up joint chains, but is also prone to introducing energy down the chain
+	 * especially at high alpha.
+	 */
 	UPROPERTY(EditAnywhere, Category = Projection, meta = (ClampMin = "0.0", ClampMax = "1.0"))
 	float ShockPropagationAlpha;
 
@@ -90,27 +94,46 @@ struct ENGINE_API FConstraintProfileProperties
 	uint8 bParentDominates : 1;
 
 	/**
-	* [PhysX] If distance error between bodies exceeds 0.1 units, or rotation error exceeds 10 degrees, body will be projected to fix this.
-	* For example a chain spinning too fast will have its elements appear detached due to velocity, this will project all bodies so they still appear attached to each other.
+	* NOTE: RigidBody AnimNode only. Projection is not applied to ragdoll physics in the main scene. For Ragdolls, ShockPropagation has a similar effect
+	* and is more compatible with the world solver.
+	* 
+	* Projection is a post-solve position and angular fixup where the parent body in the constraint is treated as having infinite mass and the child body is
+	* translated and rotated to resolve any remaining errors using a semi-physical correction. This can be used to make constraint chains significantly stiffer 
+	* at lower iteration counts. Increasing iterations would have the same effect, but be much more expensive. Projection only works well if the chain is not 
+	* interacting with other objects (e.g., through collisions) because the projection of the bodies in the chain will cause other constraints to be violated. 
+	* Likewise, if a body is influenced by multiple constraints, then enabling projection on more than one constraint may lead to unexpected results - the 
+	* "last" constraint would win but the order in which constraints are solved cannot be directly controlled.
+	* 
+	* Projection is fairly expensive compared to a single position iteration, so if you can get the behaviour you need by adding a couple iterations, that
+	* is probably a better approach.
 	*
-	* [Chaos] Chaos applies a post-solve position and angular fixup where the parent body in the constraint is treated as having infinite mass and the child body is 
-	* translated and rotated to resolve any remaining errors. This can be used to make constraint chains significantly stiffer at lower iteration counts. Increasing
-	* iterations would have the same effect, but be much more expensive. Projection only works well if the chain is not interacting with other objects (e.g.,
-	* through collisions) because the projection of the bodies in the chain will cause other constraints to be violated. Likewise, if a body is influenced by multiple
-	* constraints, then enabling projection on more than one constraint may lead to unexpected results - the "last" constraint would win but the order in which constraints
-	* are solved cannot be directly controlled.
 	*
-	* Note: projection will not be applied to constraints with soft limits.
+	* Note: projection is only applied to hard-limit constraints, anf not applied to constraints with soft limits.
 	*/
 	UPROPERTY(EditAnywhere, Category = Projection)
-	uint8 bEnableProjection : 1;
+	uint8 bEnableLinearProjection : 1;
 
 	/**
-	 * [Chaos Only] Apply projection to constraints with soft limits. This can be used to stiffen up soft joints at low iteration counts, but the projection will
-	 * override a lot of the spring-damper behaviour of the soft limits. E.g., if you have soft projection enabled and ProjectionAngularAlpha = 1.0,
-	 * the joint will act as if it is a hard limit.
+	 * NOTE: RigidBody AnimNode only. See coments on bEnableLinearProjection
+	*/
+	UPROPERTY(EditAnywhere, Category = Projection)
+	uint8 bEnableAngularProjection : 1;
+
+	/**
+	 * Shock propagation increases the mass of the parent body for the last iteration of the position and velocity solve phases. 
+	 * This can help stiffen up joint chains, but is also prone to introducing energy down the chain especially at high alpha.
+	 * 
+	 * NOTE: This is intended to be used for world constraints, not RigidBody AnimNodes which have the Projection system.
 	 */
 	UPROPERTY(EditAnywhere, Category = Projection)
+	uint8 bEnableShockPropagation : 1;
+
+	// HIDDEN - TO BE DEPRECATED
+	UPROPERTY()
+	uint8 bEnableProjection : 1;
+
+	// HIDDEN - TO BE DEPRECATED
+	UPROPERTY()
 	uint8 bEnableSoftProjection : 1;
 
 	/** Whether it is possible to break the joint with angular force. */
@@ -905,15 +928,13 @@ public:
 	void DisableProjection();
 
 	/** Set projection parameters */
-	void SetProjectionParams(bool bEnableProjection, float ProjectionLinearAlphaOrTolerance, float ProjectionAngularAlphaOrTolerance);
+	void SetProjectionParams(bool bEnableLinearProjection, bool bEnableAngularProjection, float ProjectionLinearAlphaOrTolerance, float ProjectionAngularAlphaOrTolerance);
 
-	/** Get projection parameters 
-	* Chaos returns alphas, PhysX returns tolerances
-	*/
+	/** Get projection parameters */
 	void GetProjectionAlphasOrTolerances(float& ProjectionLinearAlphaOrTolerance, float& ProjectionAngularAlphaOrTolerance) const;
 
 	/** Set the shock propagation amount [0, 1] */
-	void SetShockPropagationParams(float ShockPropagationAlpha);
+	void SetShockPropagationParams(bool bEnableShockPropagation, float ShockPropagationAlpha);
 
 	/** Get the shock propagation amount [0, 1] */
 	float GetShockPropagationAlpha() const;
diff --git a/Engine/Source/Runtime/Engine/Classes/PhysicsEngine/PhysicsAsset.h b/Engine/Source/Runtime/Engine/Classes/PhysicsEngine/PhysicsAsset.h
--- a/Engine/Source/Runtime/Engine/Classes/PhysicsEngine/PhysicsAsset.h
+++ b/Engine/Source/Runtime/Engine/Classes/PhysicsEngine/PhysicsAsset.h
@@ -15,47 +15,96 @@
 class FMeshElementCollector;
 class USkeletalBodySetup;
 
+
 /**
- * [Chaos Only]
+ * Solver iterations settings for use by RigidBody AnimNode (RBAN) in the Anim Graph. Each RBAN node runs its own solver with these settings.
+ *
+ * @note These settings have no effect when the Physics Asset is used in a world simulation (i.e., as a ragdoll on a SkeletalMeshComponent).
  */
 USTRUCT(BlueprintType)
-struct ENGINE_API FSolverIterations
+struct ENGINE_API FPhysicsAssetSolverSettings
 {
 	GENERATED_USTRUCT_BODY()
 
-		FSolverIterations();
+	FPhysicsAssetSolverSettings();
+
+	/**
+	 * The number of position iterations to run. The position solve is responsible for depenetration and friction.
+	 * Increasing this will improve simulation stability, but increase the cost.
+	*/
+	UPROPERTY(EditAnywhere, BlueprintReadWrite, Category = SolverSettings)
+	int32 PositionIterations;
 
 	/**
-	 * [Chaos Only]
-	 * The recommended fixed timestep for the solver if supported (e.g., in RigidBody Anim Node). 0 to run with variable timestep.
-	 * NOTE: If this value is non-zero and less than the current frame time, physics will step multiple times.
+	 * The number of velocity iterations to run. The velocity solve is responsible for restitution.
+	*/
+	UPROPERTY(EditAnywhere, BlueprintReadWrite, Category = SolverSettings)
+	int32 VelocityIterations;
+
+	/**
+	 * The number of projection iterations to run. The projection phase is a final pass over the constraints, applying
+	 * a semi-physical correction to any joint errors remaining after the position and velocity solves. It can be
+	 * very helpful to stabilize joint chains, but can cause issues with collision response. The projection magnitude
+	 * can be controlled per-constraint in the constraint settings (assuming ProjectionIteration is not zero).
+	*/
+	UPROPERTY(EditAnywhere, BlueprintReadWrite, Category = SolverSettings)
+	int32 ProjectionIterations;
+
+	/**
+	 * The distance at which collisions are ignored. In general you need this to be a bit larger than the typical relative body
+	 * movement in your simulation, but small enough so that we don't have to speculatively create too many unused collisions.
+	*/
+	UPROPERTY(EditAnywhere, BlueprintReadWrite, Category = SolverSettings, meta = (ClampMin = 0))
+	float CullDistance;
+
+	/**
+	 * When bodies are penetrating, this is the maximum velocity delta that can be applied in one frame.
+	*/
+	UPROPERTY(EditAnywhere, BlueprintReadWrite, Category = SolverSettings, meta = (ClampMin = 0))
+	float MaxDepenetrationVelocity;
+
+	/**
+	 * The recommended fixed timestep for the RBAN solver. Set to 0 to run with variable timestep (default).
+	 * NOTE: If this value is non-zero and less than the current frame time, the simulation will step multiple times
+	 * which increases the cost.
 	 */
 	UPROPERTY(EditAnywhere, BlueprintReadWrite, Category = SolverSettings, meta = (ClampMin = 0))
-		float FixedTimeStep;
+	float FixedTimeStep;
+};
+
+
+/**
+ * Solver settings for use by the Legacy RigidBody AnimNode (RBAN) solver.
+ * Thse settings are no longer used by default and will eventually be deprecated and then removed.
+ * 
+ * @note These settings have no effect when the Physics Asset is used in a world simulation (ragdoll).
+ */
+USTRUCT(BlueprintType)
+struct ENGINE_API FSolverIterations
+{
+	GENERATED_USTRUCT_BODY()
+
+	FSolverIterations();
 
 	/**
-	 * [Chaos Only]
 	 * The recommended number of solver iterations. Increase this if collision and joints are fighting, or joint chains are stretching.
 	 */
 	UPROPERTY(EditAnywhere, BlueprintReadWrite, Category = SolverSettings, meta = (ClampMax = 50))
 		int32 SolverIterations;
 
 	/**
-	 * [Chaos Only]
 	 * The recommended number of joint sub-iterations. Increasing this can help with chains of long-thin bodies.
 	 */
 	UPROPERTY(EditAnywhere, BlueprintReadWrite, Category = SolverSettings, meta = (ClampMax = 50))
 		int32 JointIterations;
 
 	/**
-	 * [Chaos Only]
 	 * The recommended number of collision sub-iterations. Increasing this can help with collision jitter.
 	 */
 	UPROPERTY(EditAnywhere, BlueprintReadWrite, Category = SolverSettings, meta = (ClampMax = 50))
 		int32 CollisionIterations;
 
 	/**
-	 * [Chaos Only]
 	 * The recommended number of solver push-out iterations. Increasing this can help with collision penetration problems.
 	 */
 	 /** Increase this if bodies remain penetrating */
@@ -63,20 +112,21 @@ struct ENGINE_API FSolverIterations
 		int32 SolverPushOutIterations;
 
 	/**
-	 * [Chaos Only]
 	 * The recommended number of joint sub-push-out iterations.
 	 */
 	UPROPERTY(EditAnywhere, BlueprintReadWrite, Category = SolverSettings, meta = (ClampMax = 50))
 		int32 JointPushOutIterations;
 
 	/**
-	 * [Chaos Only]
 	 * The recommended number of joint sub-push-out iterations. Increasing this can help with collision penetration problems.
 	 */
 	UPROPERTY(EditAnywhere, BlueprintReadWrite, Category = SolverSettings, meta = (ClampMax = 50))
 		int32 CollisionPushOutIterations;
 };
 
+
+
+
 UENUM()
 enum class EPhysicsAssetSolverType: uint8
 {
@@ -144,14 +194,29 @@ class UPhysicsAsset : public UObject, public IInterface_PreviewMeshProvider
 
 public:
 
-	/** [Chaos Only] Recommended solver settings. */
-	UPROPERTY(EditAnywhere, AdvancedDisplay, Category = SolverSettings, Config)
+	/** 
+	 * Solver settings when the asset is used with a RigidBody Anim Node (RBAN).
+	 */
+	UPROPERTY(EditAnywhere, Category = SolverSettings, Config)
+	FPhysicsAssetSolverSettings SolverSettings;
+
+	/** 
+	 * Solver settings for RBAN simulations with the legacy RBAN solver.
+	 * These are only used for the legacy solver and are hidden by default.
+	 * They will eventually be deprecated and removed.
+	*/
+	UPROPERTY(Config)	// (EditAnywhere, Category = LegacySolverSettings, Config)
 	FSolverIterations SolverIterations;
 
-	/** [Chaos Only] Solver type used in physics asset editor. */
-	UPROPERTY(EditAnywhere, AdvancedDisplay, Category = SolverSettings)
+	/** 
+	 * Solver type used in physics asset editor. This can be used to make what you see in the asset editror more closely resembles what you
+	 * see in game (though there will be differences owing to framerate variation etc). If your asset will primarily be used as a ragdoll 
+	 * select "World", but if it will be used in the AnimGraph select "RBAN".
+	*/
+	UPROPERTY(EditAnywhere, Category = SolverSettings)
 	EPhysicsAssetSolverType SolverType;
 
+
 	/** If true, we skip instancing bodies for this PhysicsAsset on dedicated servers */
 	UPROPERTY(EditAnywhere, AdvancedDisplay, Category = Physics)
 	uint8 bNotForDedicatedServer:1;
diff --git a/Engine/Source/Runtime/Engine/Private/PhysicsEngine/ConstraintInstance.cpp b/Engine/Source/Runtime/Engine/Private/PhysicsEngine/ConstraintInstance.cpp
--- a/Engine/Source/Runtime/Engine/Private/PhysicsEngine/ConstraintInstance.cpp
+++ b/Engine/Source/Runtime/Engine/Private/PhysicsEngine/ConstraintInstance.cpp
@@ -135,7 +135,7 @@ FConstraintProfileProperties::FConstraintProfileProperties()
 	, ProjectionAngularTolerance(180.f)
 	, ProjectionLinearAlpha(1.0f)
 	, ProjectionAngularAlpha(0.0f)
-	, ShockPropagationAlpha(0.0f)
+	, ShockPropagationAlpha(0.3f)
 	, LinearBreakThreshold(300.f)
 	, LinearPlasticityThreshold(0.1f)
 	, AngularBreakThreshold(500.f)
@@ -143,8 +143,11 @@ FConstraintProfileProperties::FConstraintProfileProperties()
 	, ContactTransferScale(0.f)
 	, bDisableCollision(false)
 	, bParentDominates(false)
-	, bEnableProjection(true)
-	, bEnableSoftProjection(false)
+	, bEnableLinearProjection(false)
+	, bEnableAngularProjection(false)
+	, bEnableShockPropagation(false)
+	, bEnableProjection(true)		// to be deprecated
+	, bEnableSoftProjection(false)	// to be deprecated
 	, bAngularBreakable(false)
 	, bAngularPlasticity(false)
 	, bLinearBreakable(false)
@@ -466,8 +469,8 @@ void FConstraintProfileProperties::UpdateConstraintFlags_AssumesLocked(const FPh
 
 	FPhysicsInterface::SetCollisionEnabled(InConstraintRef, !bDisableCollision);
 #if WITH_CHAOS
-	FPhysicsInterface::SetProjectionEnabled_AssumesLocked(InConstraintRef, bEnableProjection, ProjectionLinearAlpha, ProjectionAngularAlpha);
-	FPhysicsInterface::SetShockPropagationAlpha_AssumesLocked(InConstraintRef, ShockPropagationAlpha);
+	FPhysicsInterface::SetProjectionEnabled_AssumesLocked(InConstraintRef, bEnableLinearProjection || bEnableAngularProjection, ProjectionLinearAlpha, ProjectionAngularAlpha);
+	FPhysicsInterface::SetShockPropagationEnabled_AssumesLocked(InConstraintRef, bEnableShockPropagation, ShockPropagationAlpha);
 #else
 	FPhysicsInterface::SetProjectionEnabled_AssumesLocked(InConstraintRef, bEnableProjection, ProjectionLinearTolerance, ProjectionAngularTolerance);
 #endif
@@ -1192,21 +1195,27 @@ void FConstraintInstance::GetProjectionAlphasOrTolerances(float& ProjectionLinea
 #endif
 }
 
-void FConstraintInstance::SetProjectionParams(bool bEnableProjection, float ProjectionLinearAlphaOrTolerance, float ProjectionAngularAlphaOrTolerance)
+void FConstraintInstance::SetProjectionParams(bool bEnableLinearProjection, bool bEnableAngularProjection, float ProjectionLinearAlphaOrTolerance, float ProjectionAngularAlphaOrTolerance)
 {
-	ProfileInstance.bEnableProjection = bEnableProjection;
+	ProfileInstance.bEnableLinearProjection = bEnableLinearProjection;
+	ProfileInstance.bEnableAngularProjection = bEnableAngularProjection;
 #if WITH_CHAOS
 	ProfileInstance.ProjectionLinearAlpha = ProjectionLinearAlphaOrTolerance;
 	ProfileInstance.ProjectionAngularAlpha = ProjectionAngularAlphaOrTolerance;
 #else
 	ProfileInstance.ProjectionLinearTolerance = ProjectionLinearAlphaOrTolerance;
 	ProfileInstance.ProjectionAngularTolerance = ProjectionAngularAlphaOrTolerance;
-
 #endif
 
 	FPhysicsCommand::ExecuteWrite(ConstraintHandle, [&](const FPhysicsConstraintHandle& Constraint)
 		{
-			if (bEnableProjection)
+#if WITH_CHAOS
+			const bool bEnableProjection = bEnableAngularProjection || ProjectionLinearAlphaOrTolerance;
+			const float LinearAlpha = bEnableLinearProjection ? ProjectionLinearAlphaOrTolerance : 0.0f;
+			const float AngularAlpha = bEnableAngularProjection ? ProjectionAngularAlphaOrTolerance : 0.0f;
+			FPhysicsInterface::SetProjectionEnabled_AssumesLocked(Constraint, bEnableProjection, LinearAlpha, AngularAlpha);
+#else
+			if (bEnableLinearProjection)
 			{
 				FPhysicsInterface::SetProjectionEnabled_AssumesLocked(Constraint, true, ProjectionLinearAlphaOrTolerance, ProjectionAngularAlphaOrTolerance);
 			}
@@ -1214,6 +1223,7 @@ void FConstraintInstance::SetProjectionParams(bool bEnableProjection, float Proj
 			{
 				FPhysicsInterface::SetProjectionEnabled_AssumesLocked(Constraint, false);
 			}
+#endif
 		});
 }
 
@@ -1226,13 +1236,14 @@ float FConstraintInstance::GetShockPropagationAlpha() const
 #endif
 }
 
-void FConstraintInstance::SetShockPropagationParams(float ShockPropagationAlpha)
+void FConstraintInstance::SetShockPropagationParams(bool bEnableShockPropagation, float ShockPropagationAlpha)
 {
 #if WITH_CHAOS
+	ProfileInstance.bEnableShockPropagation = bEnableShockPropagation;
 	ProfileInstance.ShockPropagationAlpha = ShockPropagationAlpha;
 	FPhysicsCommand::ExecuteWrite(ConstraintHandle, [&](const FPhysicsConstraintHandle& Constraint)
 		{
-			FPhysicsInterface::SetShockPropagationAlpha_AssumesLocked(Constraint, ShockPropagationAlpha);
+			FPhysicsInterface::SetShockPropagationEnabled_AssumesLocked(Constraint, bEnableShockPropagation, ShockPropagationAlpha);
 		});
 #endif
 }
@@ -1240,7 +1251,7 @@ void FConstraintInstance::SetShockPropagationParams(float ShockPropagationAlpha)
 void FConstraintInstance::EnableProjection()
 {
 #if WITH_CHAOS
-	SetProjectionParams(true, ProfileInstance.ProjectionLinearAlpha, ProfileInstance.ProjectionAngularAlpha);
+	SetProjectionParams(true, true, ProfileInstance.ProjectionLinearAlpha, ProfileInstance.ProjectionAngularAlpha);
 #else
 	SetProjectionParams(true, ProfileInstance.ProjectionLinearTolerance, ProfileInstance.ProjectionAngularTolerance);
 #endif
@@ -1249,7 +1260,7 @@ void FConstraintInstance::EnableProjection()
 void FConstraintInstance::DisableProjection()
 {
 #if WITH_CHAOS
-	SetProjectionParams(false, ProfileInstance.ProjectionLinearAlpha, ProfileInstance.ProjectionAngularAlpha);
+	SetProjectionParams(false, false, ProfileInstance.ProjectionLinearAlpha, ProfileInstance.ProjectionAngularAlpha);
 #else
 	SetProjectionParams(false, ProfileInstance.ProjectionLinearTolerance, ProfileInstance.ProjectionAngularTolerance);
 #endif
diff --git a/Engine/Source/Runtime/Engine/Private/PhysicsEngine/ConstraintInstanceBlueprintLibrary.cpp b/Engine/Source/Runtime/Engine/Private/PhysicsEngine/ConstraintInstanceBlueprintLibrary.cpp
--- a/Engine/Source/Runtime/Engine/Private/PhysicsEngine/ConstraintInstanceBlueprintLibrary.cpp
+++ b/Engine/Source/Runtime/Engine/Private/PhysicsEngine/ConstraintInstanceBlueprintLibrary.cpp
@@ -56,7 +56,7 @@ void UConstraintInstanceBlueprintLibrary::SetProjectionParams(
 	if (FConstraintInstance* ConstraintInstance = Accessor.Get())
 	{
 		Accessor.Modify();
-		ConstraintInstance->SetProjectionParams(bEnableProjection, ProjectionLinearAlpha, ProjectionAngularAlpha);
+		ConstraintInstance->SetProjectionParams(bEnableProjection, bEnableProjection, ProjectionLinearAlpha, ProjectionAngularAlpha);
 	}
 }
 
diff --git a/Engine/Source/Runtime/Engine/Private/PhysicsEngine/ImmediatePhysics/ImmediatePhysicsChaos/ImmediatePhysicsJointHandle_Chaos.cpp b/Engine/Source/Runtime/Engine/Private/PhysicsEngine/ImmediatePhysics/ImmediatePhysicsChaos/ImmediatePhysicsJointHandle_Chaos.cpp
--- a/Engine/Source/Runtime/Engine/Private/PhysicsEngine/ImmediatePhysics/ImmediatePhysicsChaos/ImmediatePhysicsJointHandle_Chaos.cpp
+++ b/Engine/Source/Runtime/Engine/Private/PhysicsEngine/ImmediatePhysics/ImmediatePhysicsChaos/ImmediatePhysicsJointHandle_Chaos.cpp
@@ -45,11 +45,12 @@ namespace ImmediatePhysics_Chaos
 		ConstraintSettings.AngularLimits[(int32)EJointAngularConstraintIndex::Swing1] = FMath::DegreesToRadians(ConstraintInstance->GetAngularSwing1Limit());
 		ConstraintSettings.AngularLimits[(int32)EJointAngularConstraintIndex::Swing2] = FMath::DegreesToRadians(ConstraintInstance->GetAngularSwing2Limit());
 
-		ConstraintSettings.bProjectionEnabled = Profile.bEnableProjection;
-		ConstraintSettings.bSoftProjectionEnabled = Profile.bEnableSoftProjection;
+		ConstraintSettings.bProjectionEnabled = Profile.bEnableLinearProjection || Profile.bEnableAngularProjection;
+		ConstraintSettings.bShockPropagationEnabled = Profile.bEnableShockPropagation;
 
-		ConstraintSettings.LinearProjection = Profile.ProjectionLinearAlpha;
-		ConstraintSettings.AngularProjection = Profile.ProjectionAngularAlpha;
+		ConstraintSettings.LinearProjection = Profile.bEnableLinearProjection ? Profile.ProjectionLinearAlpha : 0.0f;
+		ConstraintSettings.AngularProjection = Profile.bEnableAngularProjection ? Profile.ProjectionAngularAlpha : 0.0f;
+		ConstraintSettings.ShockPropagation = Profile.bEnableShockPropagation ? Profile.ShockPropagationAlpha : 0.0f;
 		ConstraintSettings.ParentInvMassScale = Profile.bParentDominates ? (FReal)0 : (FReal)1;
 
 		ConstraintSettings.bSoftLinearLimitsEnabled = ConstraintInstance->GetIsSoftLinearLimit();
diff --git a/Engine/Source/Runtime/Engine/Private/PhysicsEngine/ImmediatePhysics/ImmediatePhysicsChaos/ImmediatePhysicsSimulation_Chaos.cpp b/Engine/Source/Runtime/Engine/Private/PhysicsEngine/ImmediatePhysics/ImmediatePhysicsChaos/ImmediatePhysicsSimulation_Chaos.cpp
--- a/Engine/Source/Runtime/Engine/Private/PhysicsEngine/ImmediatePhysics/ImmediatePhysicsChaos/ImmediatePhysicsSimulation_Chaos.cpp
+++ b/Engine/Source/Runtime/Engine/Private/PhysicsEngine/ImmediatePhysics/ImmediatePhysicsChaos/ImmediatePhysicsSimulation_Chaos.cpp
@@ -36,16 +36,22 @@ Chaos::FRealSingle ChaosImmediate_Evolution_StepTime = 0.0f;
 int32 ChaosImmediate_Evolution_NumSteps = 0;
 Chaos::FRealSingle ChaosImmediate_Evolution_InitialStepTime = 0.033f;
 int32 ChaosImmediate_Evolution_DeltaTimeCount = 100;
-int32 ChaosImmediate_Evolution_Iterations = -1;
-int32 ChaosImmediate_Evolution_PushOutIterations = -1;
+int32 ChaosImmediate_Evolution_PositionIterations = -1;
+int32 ChaosImmediate_Evolution_VelocityIterations = -1;
+int32 ChaosImmediate_Evolution_ProjectionIterations = -1;
+int32 ChaosImmediate_Evolution_Iterations = -1;				// Legacy
+int32 ChaosImmediate_Evolution_PushOutIterations = -1;		// Legacy
 Chaos::FRealSingle ChaosImmediate_Evolution_BoundsExtension = 0.0f;
 int32 ChaosImmediate_DisableInactiveByIndex = 1;
 FAutoConsoleVariableRef CVarChaosImmPhysStepTime(TEXT("p.Chaos.ImmPhys.StepTime"), ChaosImmediate_Evolution_StepTime, TEXT("Override step time (if not zero)"));
 FAutoConsoleVariableRef CVarChaosImmPhysNumSteps(TEXT("p.Chaos.ImmPhys.NumSteps"), ChaosImmediate_Evolution_NumSteps, TEXT("Override num steps (if not zero)"));
 FAutoConsoleVariableRef CVarChaosImmPhysInitialStepTime(TEXT("p.Chaos.ImmPhys.InitialStepTime"), ChaosImmediate_Evolution_InitialStepTime, TEXT("Initial step time (then calculated from rolling average)"));
 FAutoConsoleVariableRef CVarChaosImmPhysDeltaTimeCount(TEXT("p.Chaos.ImmPhys.DeltaTimeCount"), ChaosImmediate_Evolution_DeltaTimeCount, TEXT("The number of ticks over which the moving average is calculated"));
-FAutoConsoleVariableRef CVarChaosImmPhysIterations(TEXT("p.Chaos.ImmPhys.Iterations"), ChaosImmediate_Evolution_Iterations, TEXT("Override number of constraint solver loops in immediate physics (if >= 0)"));
-FAutoConsoleVariableRef CVarChaosImmPhysPushOutIterations(TEXT("p.Chaos.ImmPhys.PushOutIterations"), ChaosImmediate_Evolution_PushOutIterations, TEXT("Override number of solver push-out loops (if >= 0)"));
+FAutoConsoleVariableRef CVarChaosImmPhysPositionIterations(TEXT("p.Chaos.ImmPhys.PositionIterations"), ChaosImmediate_Evolution_PositionIterations, TEXT("Override number of position iteration loops in immediate physics (if >= 0)"));
+FAutoConsoleVariableRef CVarChaosImmPhysVelocityIterations(TEXT("p.Chaos.ImmPhys.VelocityIterations"), ChaosImmediate_Evolution_VelocityIterations, TEXT("Override number of velocity iteration loops in immediate physics (if >= 0)"));
+FAutoConsoleVariableRef CVarChaosImmPhysProjectionIterations(TEXT("p.Chaos.ImmPhys.ProjectionIterations"), ChaosImmediate_Evolution_ProjectionIterations, TEXT("Override number of projection iteration loops in immediate physics (if >= 0)"));
+FAutoConsoleVariableRef CVarChaosImmPhysIterations(TEXT("p.Chaos.ImmPhys.Iterations"), ChaosImmediate_Evolution_Iterations, TEXT("[Legacy Solver] Override number of constraint solver loops in immediate physics (if >= 0)"));
+FAutoConsoleVariableRef CVarChaosImmPhysPushOutIterations(TEXT("p.Chaos.ImmPhys.PushOutIterations"), ChaosImmediate_Evolution_PushOutIterations, TEXT("[Legacy Solver] Override number of solver push-out loops (if >= 0)"));
 FAutoConsoleVariableRef CVarChaosImmPhysBoundsExtension(TEXT("p.Chaos.ImmPhys.BoundsExtension"), ChaosImmediate_Evolution_BoundsExtension, TEXT("Bounds are grown by this fraction of their size (should be >= 0.0)"));
 FAutoConsoleVariableRef CVarChaosImmPhysDisableInactiveByIndex(TEXT("p.Chaos.ImmPhys.DisableInactiveByIndex"), ChaosImmediate_DisableInactiveByIndex, TEXT("Disable bodies that are no longer active based on the index, rather than just count."));
 
@@ -62,33 +68,37 @@ Chaos::FRealSingle ChaosImmediate_Evolution_FixedStepTime = -1.0f;
 Chaos::FRealSingle ChaosImmediate_Evolution_FixedStepTolerance = 0.05f;
 FAutoConsoleVariableRef CVarChaosImmPhysMinStepTime(TEXT("p.Chaos.ImmPhys.MinStepTime"), ChaosImmediate_Evolution_MinStepTime, TEXT("If non-zero, then if step time is lower than this, go into fixed step mode with this timestep."));
 FAutoConsoleVariableRef CVarChaosImmPhysFixedStepTime(TEXT("p.Chaos.ImmPhys.FixedStepTime"), ChaosImmediate_Evolution_FixedStepTime, TEXT("Override fixed step time mode: fixed step time (if positive); variable time mode (if zero); asset defined (if negative)"));
-FAutoConsoleVariableRef CVarChaosImmPhysFixedStepTolerance(TEXT("p.Chaos.ImmPhys.FixedStepTolerance"), ChaosImmediate_Evolution_FixedStepTolerance, TEXT("Tiem remainder required to add a new step (fraction of FixedStepTime)"));
+FAutoConsoleVariableRef CVarChaosImmPhysFixedStepTolerance(TEXT("p.Chaos.ImmPhys.FixedStepTolerance"), ChaosImmediate_Evolution_FixedStepTolerance, TEXT("Time remainder required to add a new step (fraction of FixedStepTime)"));
 
 int32 ChaosImmediate_Collision_Enabled = 1;
-int32 ChaosImmediate_Collision_PairIterations = -1;
-int32 ChaosImmediate_Collision_PushOutPairIterations = 0;	// Force disabled - not compatible with EConstraintSolverType::StandardPbd
+int32 ChaosImmediate_Collision_PairIterations = -1;			// Legacy
+int32 ChaosImmediate_Collision_NumPositionFrictionIterations = 0;		// No static friction for RBAN
+int32 ChaosImmediate_Collision_NumVelocityFrictionIterations = 1;		// Dynamic friction for RBAN in velocity solve
+int32 ChaosImmediate_Collision_NumPositionShockPropagationIterations = 0;
+int32 ChaosImmediate_Collision_NumVelocityShockPropagationIterations = 0;
 int32 ChaosImmediate_Collision_Priority = 1;
-Chaos::FRealSingle ChaosImmediate_Collision_CullDistance = 1.0f;
+Chaos::FRealSingle ChaosImmediate_Collision_CullDistance = -1.0f;
+Chaos::FRealSingle ChaosImmediate_Collision_MaxDepenetrationVelocity = -1.0f;
 Chaos::FRealSingle ChaosImmediate_Collision_RestitutionThresholdMultiplier = 3.0f;
 int32 ChaosImmediate_Collision_RestitutionEnabled = true;
 int32 ChaosImmediate_Collision_DeferNarrowPhase = 1;
 int32 ChaosImmediate_Collision_UseManifolds = 0;
 FAutoConsoleVariableRef CVarChaosImmPhysCollisionDisable(TEXT("p.Chaos.ImmPhys.Collision.Enabled"), ChaosImmediate_Collision_Enabled, TEXT("Enable/Disable collisions in Immediate Physics."));
-FAutoConsoleVariableRef CVarChaosImmPhysCollisionPairIterations(TEXT("p.Chaos.ImmPhys.Collision.PairIterations"), ChaosImmediate_Collision_PairIterations, TEXT("Override collision pair iterations (if >= 0)"));
-FAutoConsoleVariableRef CVarChaosImmPhysCollisionPushOutPairIterations(TEXT("p.Chaos.ImmPhys.Collision.PushOutPairIterations"), ChaosImmediate_Collision_PushOutPairIterations, TEXT("Override collision push-out pair iterations (if >= 0)"));
+FAutoConsoleVariableRef CVarChaosImmPhysCollisionPairIterations(TEXT("p.Chaos.ImmPhys.Collision.PairIterations"), ChaosImmediate_Collision_PairIterations, TEXT("[Legacy Solver] Override collision pair iterations (if >= 0)"));
 FAutoConsoleVariableRef CVarChaosImmPhysCollisionPriority(TEXT("p.Chaos.ImmPhys.Collision.Priority"), ChaosImmediate_Collision_Priority, TEXT("Set the Collision constraint sort order (Joints have priority 0)"));
-FAutoConsoleVariableRef CVarChaosImmPhysCollisionCullDistance(TEXT("p.Chaos.ImmPhys.Collision.CullDistance"), ChaosImmediate_Collision_CullDistance, TEXT("CullDistance"));
+FAutoConsoleVariableRef CVarChaosImmPhysCollisionCullDistance(TEXT("p.Chaos.ImmPhys.Collision.CullDistance"), ChaosImmediate_Collision_CullDistance, TEXT("Set the collision CullDistance (if >= 0)"));
+FAutoConsoleVariableRef CVarChaosImmPhysCollisionMaxDepenetrationVelocity(TEXT("p.Chaos.ImmPhys.Collision.MaxDepenetrationVelocity"), ChaosImmediate_Collision_MaxDepenetrationVelocity, TEXT("Set the collision Max Depenetration Velocity (if >= 0)"));
 FAutoConsoleVariableRef CVarChaosImmPhysCollisionRestitutionThresholdMultiplier(TEXT("p.Chaos.ImmPhys.Collision.RestitutionThresholdMultiplier"), ChaosImmediate_Collision_RestitutionThresholdMultiplier, TEXT("Collision Restitution Threshold (Acceleration) = Multiplier * Gravity"));
 FAutoConsoleVariableRef CVarChaosImmPhysCollisionRestitutionEnabled(TEXT("p.Chaos.ImmPhys.Collision.RestitutionEnabled"), ChaosImmediate_Collision_RestitutionEnabled, TEXT("Collision Restitution Enable/Disable"));
-FAutoConsoleVariableRef CVarChaosImmPhysCollisionDeferNarrowPhase(TEXT("p.Chaos.ImmPhys.Collision.DeferNarrowPhase"), ChaosImmediate_Collision_DeferNarrowPhase, TEXT("Create contacts for all broadphase pairs, perform NarrowPhase later."));
-FAutoConsoleVariableRef CVarChaosImmPhysCollisionUseManifolds(TEXT("p.Chaos.ImmPhys.Collision.UseManifolds"), ChaosImmediate_Collision_UseManifolds, TEXT("Enable/Disable use of manifoldes in collision."));
+FAutoConsoleVariableRef CVarChaosImmPhysCollisionDeferNarrowPhase(TEXT("p.Chaos.ImmPhys.Collision.DeferNarrowPhase"), ChaosImmediate_Collision_DeferNarrowPhase, TEXT("[Legacy Solver] Create contacts for all broadphase pairs, perform NarrowPhase later."));
+FAutoConsoleVariableRef CVarChaosImmPhysCollisionUseManifolds(TEXT("p.Chaos.ImmPhys.Collision.UseManifolds"), ChaosImmediate_Collision_UseManifolds, TEXT("[Legacy Solver] Enable/Disable use of manifoldes in collision."));
 
-int32 ChaosImmediate_Joint_PairIterations = -1;
-int32 ChaosImmediate_Joint_PushOutPairIterations = -1;
+int32 ChaosImmediate_Joint_PairIterations = -1;				// Legacy
+int32 ChaosImmediate_Joint_PushOutPairIterations = -1;		// Legacy
 Chaos::FRealSingle ChaosImmediate_Joint_SwingTwistAngleTolerance = 1.0e-6f;
 Chaos::FRealSingle ChaosImmediate_Joint_PositionTolerance = 0.025f;
 Chaos::FRealSingle ChaosImmediate_Joint_AngleTolerance = 0.001f;
-int32 ChaosImmediate_Joint_NumShockPropagationIterations = 0;
+int32 ChaosImmediate_Joint_NumShockPropagationIterations = -1;
 int32 ChaosImmediate_Joint_SolvePositionLast = 1;
 int32 ChaosImmediate_Joint_EnableTwistLimits = 1;
 int32 ChaosImmediate_Joint_EnableSwingLimits = 1;
@@ -108,8 +118,8 @@ Chaos::FRealSingle ChaosImmediate_Joint_AngularDriveStiffness = -1.0f;
 Chaos::FRealSingle ChaosImmediate_Joint_AngularDriveDamping = -1.0f;
 Chaos::FRealSingle ChaosImmediate_Joint_MinParentMassRatio = 0.2f;
 Chaos::FRealSingle ChaosImmediate_Joint_MaxInertiaRatio = 5.0f;
-FAutoConsoleVariableRef CVarChaosImmPhysJointPairIterations(TEXT("p.Chaos.ImmPhys.Joint.PairIterations"), ChaosImmediate_Joint_PairIterations, TEXT("Override joint pair iterations (if >= 0)"));
-FAutoConsoleVariableRef CVarChaosImmPhysJointPushOutPairIterations(TEXT("p.Chaos.ImmPhys.Joint.PushOutPairIterations"), ChaosImmediate_Joint_PushOutPairIterations, TEXT("Override joint push-out pair iterations (if >= 0)"));
+FAutoConsoleVariableRef CVarChaosImmPhysJointPairIterations(TEXT("p.Chaos.ImmPhys.Joint.PairIterations"), ChaosImmediate_Joint_PairIterations, TEXT("[Legacy Solver] Override joint pair iterations (if >= 0)"));
+FAutoConsoleVariableRef CVarChaosImmPhysJointPushOutPairIterations(TEXT("p.Chaos.ImmPhys.Joint.PushOutPairIterations"), ChaosImmediate_Joint_PushOutPairIterations, TEXT("[Legacy Solver] Override joint push-out pair iterations (if >= 0)"));
 FAutoConsoleVariableRef CVarChaosImmPhysJointSwingTwistAngleTolerance(TEXT("p.Chaos.ImmPhys.Joint.SwingTwistAngleTolerance"), ChaosImmediate_Joint_SwingTwistAngleTolerance, TEXT("SwingTwistAngleTolerance."));
 FAutoConsoleVariableRef CVarChaosImmPhysJointPositionTolerance(TEXT("p.Chaos.ImmPhys.Joint.PositionTolerance"), ChaosImmediate_Joint_PositionTolerance, TEXT("PositionTolerance."));
 FAutoConsoleVariableRef CVarChaosImmPhysJointAngleTolerance(TEXT("p.Chaos.ImmPhys.Joint.AngleTolerance"), ChaosImmediate_Joint_AngleTolerance, TEXT("AngleTolerance."));
@@ -138,14 +148,11 @@ FAutoConsoleVariableRef CVarChaosImmPhysJointMaxInertiaRatio(TEXT("p.Chaos.ImmPh
 //
 // Select the solver technique to use until we settle on the final one...
 //
-int32 ChaosImmediate_SolverType = (int32)Chaos::EConstraintSolverType::StandardPbd;
+int32 ChaosImmediate_SolverType = (int32)Chaos::EConstraintSolverType::QuasiPbd;
 FAutoConsoleVariableRef CVarChaosImmPhysSolverType(TEXT("p.Chaos.ImmPhys.SolverType"), ChaosImmediate_SolverType, TEXT("0 = None; 1 = GbfPbd; 2 = Pbd; 3 = QuasiPbd"));
 
-// Whether to use the linear joint solver which is significantly faster than the non-linear one but less accurate. 
-// In princple this choice is independent of the SolverType, but in practice you might want to use the non-Linear Joint Solver when you are using the 
-// Stanbard PBD Solver algorithm. This is because the Linear Joint Solver implementation does not implmenent the Projection phase which is part of the 
-// stabilization of the standard PBD solver. Ymmv.
-bool bChaosImmediate_Joint_UseLinearSolver = true;
+// Whether to use the linear joint solver which is significantly faster than the non-linear one but less accurate. Only applies to the QuasiPBD Solver
+bool bChaosImmediate_Joint_UseLinearSolver = false;
 FAutoConsoleVariableRef CVarChaosImmPhysJointUseCachedSolver(TEXT("p.Chaos.ImmPhys.Joint.UseLinearSolver"), bChaosImmediate_Joint_UseLinearSolver, TEXT("Use linear version of joint solver. (default is true"));
 
 //
@@ -275,13 +282,13 @@ namespace ImmediatePhysics_Chaos
 		FImplementation()
 			: Particles(UniqueIndices)
 			, Joints()
-			, Collisions(Particles, CollidedParticles, ParticleMaterials, PerParticleMaterials, nullptr, 0, 0, ChaosImmediate_Collision_CullDistance)
+			, Collisions(Particles, CollidedParticles, ParticleMaterials, PerParticleMaterials, nullptr)
 			, BroadPhase(&ActivePotentiallyCollidingPairs, nullptr, nullptr)
-			, NarrowPhase(ChaosImmediate_Collision_CullDistance, FReal(0), Collisions.GetConstraintAllocator())
+			, NarrowPhase(FReal(0), FReal(0), Collisions.GetConstraintAllocator())
 			, CollisionDetector(BroadPhase, NarrowPhase, Collisions)
 			, JointsRule(0, Joints)
 			, CollisionsRule(1, Collisions)
-			, Evolution(Particles, ParticlePrevXs, ParticlePrevRs, CollisionDetector, ChaosImmediate_Evolution_BoundsExtension)
+			, Evolution(Particles, ParticlePrevXs, ParticlePrevRs, CollisionDetector, FReal(0))
 			, NumActiveDynamicActorHandles(0)
 			, SimulationSpace()
 			, RollingAverageStepTime(ChaosImmediate_Evolution_InitialStepTime)
@@ -719,13 +726,41 @@ namespace ImmediatePhysics_Chaos
 		SimSpaceSettings.EulerAlpha = ChaosImmediate_Evolution_SimSpaceEulerAlpha;
 	}
 
-	void FSimulation::SetSolverIterations(const FReal InFixedDt, const int32 SolverIts, const int32 JointIts, const int32 CollisionIts, const int32 SolverPushOutIts, const int32 JointPushOutIts, const int32 CollisionPushOutIts)
+	void FSimulation::SetSolverSettings(const FReal FixedDt, const FReal CullDistance, const FReal MaxDepenetrationVelocity, const int32 PositionIts, const int32 VelocityIts, const int32 ProjectionIts)
 	{
-		if (InFixedDt >= 0.0f)
+		if (FixedDt >= FReal(0))
 		{
-			Implementation->FixedStepTime = InFixedDt;
+			Implementation->FixedStepTime = FixedDt;
 		}
 
+		if (CullDistance >= FReal(0))
+		{
+			Implementation->NarrowPhase.SetBoundsExpansion(CullDistance);
+		}
+
+		if (MaxDepenetrationVelocity >= FReal(0))
+		{
+			Implementation->Collisions.SetMaxPushOutVelocity(MaxDepenetrationVelocity);
+		}
+
+		if (PositionIts >= 0)
+		{
+			Implementation->Evolution.SetNumPositionIterations(PositionIts);
+		}
+
+		if (VelocityIts >= 0)
+		{
+			Implementation->Evolution.SetNumVelocityIterations(VelocityIts);
+		}
+
+		if (ProjectionIts >= 0)
+		{
+			Implementation->Evolution.SetNumProjectionIterations(ProjectionIts);
+		}
+	}
+
+	void FSimulation::SetLegacySolverSettings(const int32 SolverIts, const int32 JointIts, const int32 CollisionIts, const int32 SolverPushOutIts, const int32 JointPushOutIts, const int32 CollisionPushOutIts)
+	{
 		if (SolverIts >= 0)
 		{
 			Implementation->Evolution.SetNumIterations(SolverIts);
@@ -788,13 +823,37 @@ namespace ImmediatePhysics_Chaos
 
 		// TEMP: overrides
 		{
+			const EConstraintSolverType SolverType = (EConstraintSolverType)ChaosImmediate_SolverType;
+			Implementation->Evolution.SetSolverType(SolverType);
+			Implementation->Collisions.SetSolverType(SolverType);
+			Implementation->Joints.SetSolverType(SolverType);
+			if (SolverType == EConstraintSolverType::QuasiPbd)
+			{
+				SetSolverSettings(
+					ChaosImmediate_Evolution_FixedStepTime,
+					ChaosImmediate_Collision_CullDistance,
+					ChaosImmediate_Collision_MaxDepenetrationVelocity,
+					ChaosImmediate_Evolution_PositionIterations,
+					ChaosImmediate_Evolution_VelocityIterations,
+					ChaosImmediate_Evolution_ProjectionIterations);
+			}
+			else
+			{
+				SetLegacySolverSettings(
+					ChaosImmediate_Evolution_Iterations,
+					ChaosImmediate_Joint_PairIterations,
+					ChaosImmediate_Collision_PairIterations,
+					ChaosImmediate_Evolution_PushOutIterations,
+					ChaosImmediate_Joint_PushOutPairIterations,
+					0);
+			}
+
 			FPBDJointSolverSettings JointsSettings = Implementation->Joints.GetSettings();
 			JointsSettings.SwingTwistAngleTolerance = ChaosImmediate_Joint_SwingTwistAngleTolerance;
 			JointsSettings.PositionTolerance = ChaosImmediate_Joint_PositionTolerance;
 			JointsSettings.AngleTolerance = ChaosImmediate_Joint_AngleTolerance;
 			JointsSettings.MinParentMassRatio = ChaosImmediate_Joint_MinParentMassRatio;
 			JointsSettings.MaxInertiaRatio = ChaosImmediate_Joint_MaxInertiaRatio;
-			JointsSettings.NumShockPropagationIterations = ChaosImmediate_Joint_NumShockPropagationIterations;
 			JointsSettings.bSolvePositionLast = ChaosImmediate_Joint_SolvePositionLast != 0;
 			JointsSettings.bEnableTwistLimits = ChaosImmediate_Joint_EnableTwistLimits != 0;
 			JointsSettings.bEnableSwingLimits = ChaosImmediate_Joint_EnableSwingLimits != 0;
@@ -814,9 +873,27 @@ namespace ImmediatePhysics_Chaos
 			JointsSettings.LinearDriveDampingOverride = ChaosImmediate_Joint_LinearDriveDamping;
 			JointsSettings.AngularDriveStiffnessOverride = ChaosImmediate_Joint_AngularDriveStiffness;
 			JointsSettings.AngularDriveDampingOverride = ChaosImmediate_Joint_AngularDriveDamping;
-			JointsSettings.bUseLinearSolver = bChaosImmediate_Joint_UseLinearSolver && ((EConstraintSolverType)ChaosImmediate_SolverType == EConstraintSolverType::QuasiPbd);
+			JointsSettings.bUseLinearSolver = false;
+			if (SolverType == EConstraintSolverType::QuasiPbd)
+			{
+				JointsSettings.NumShockPropagationIterations = ChaosImmediate_Joint_NumShockPropagationIterations;
+				JointsSettings.bUseLinearSolver = bChaosImmediate_Joint_UseLinearSolver;
+			}
+			else
+			{
+				JointsSettings.NumShockPropagationIterations = 0;
+				JointsSettings.bUseLinearSolver = false;
+			}
 			Implementation->Joints.SetSettings(JointsSettings);
 
+			if (SolverType == EConstraintSolverType::QuasiPbd)
+			{
+				Implementation->Collisions.SetPositionFrictionIterations(ChaosImmediate_Collision_NumPositionFrictionIterations);
+				Implementation->Collisions.SetVelocityFrictionIterations(ChaosImmediate_Collision_NumVelocityFrictionIterations);
+				Implementation->Collisions.SetPositionShockPropagationIterations(ChaosImmediate_Collision_NumPositionShockPropagationIterations);
+				Implementation->Collisions.SetVelocityShockPropagationIterations(ChaosImmediate_Collision_NumVelocityShockPropagationIterations);
+			}
+
 			Implementation->Collisions.SetRestitutionEnabled(ChaosImmediate_Collision_RestitutionEnabled != 0);
 			Implementation->Collisions.SetRestitutionThreshold(ChaosImmediate_Collision_RestitutionThresholdMultiplier * InGravity.Size());
 			Implementation->Collisions.SetCollisionsEnabled(ChaosImmediate_Collision_Enabled != 0);
@@ -824,14 +901,17 @@ namespace ImmediatePhysics_Chaos
 
 			Implementation->Evolution.SetBoundsExtension(ChaosImmediate_Evolution_BoundsExtension);
 
-			Implementation->NarrowPhase.SetBoundsExpansion(ChaosImmediate_Collision_CullDistance);
-			Implementation->NarrowPhase.GetContext().bDeferUpdate = (ChaosImmediate_Collision_DeferNarrowPhase != 0);
-			Implementation->NarrowPhase.GetContext().bAllowManifolds = (ChaosImmediate_Collision_UseManifolds != 0);
 			Implementation->NarrowPhase.GetContext().bAllowManifoldReuse = false;
-
-			Implementation->Evolution.SetSolverType((EConstraintSolverType)ChaosImmediate_SolverType);
-			Implementation->Collisions.SetSolverType((EConstraintSolverType)ChaosImmediate_SolverType);
-			Implementation->Joints.SetSolverType((EConstraintSolverType)ChaosImmediate_SolverType);
+			if (SolverType == EConstraintSolverType::QuasiPbd)
+			{
+				Implementation->NarrowPhase.GetContext().bDeferUpdate = false;
+				Implementation->NarrowPhase.GetContext().bAllowManifolds = true;
+			}
+			else
+			{
+				Implementation->NarrowPhase.GetContext().bDeferUpdate = (ChaosImmediate_Collision_DeferNarrowPhase != 0);
+				Implementation->NarrowPhase.GetContext().bAllowManifolds = (ChaosImmediate_Collision_UseManifolds != 0);
+			}
 
 			if (ChaosImmediate_Evolution_StepTime > 0)
 			{
@@ -841,15 +921,6 @@ namespace ImmediatePhysics_Chaos
 			{
 				NumSteps = ChaosImmediate_Evolution_NumSteps;
 			}
-
-			SetSolverIterations(
-				ChaosImmediate_Evolution_FixedStepTime,
-				ChaosImmediate_Evolution_Iterations,
-				ChaosImmediate_Joint_PairIterations,
-				ChaosImmediate_Collision_PairIterations,
-				ChaosImmediate_Evolution_PushOutIterations,
-				ChaosImmediate_Joint_PushOutPairIterations,
-				ChaosImmediate_Collision_PushOutPairIterations);
 		}
 
 		// Fixed timestep mode DT (Only used if > 0.0f)
diff --git a/Engine/Source/Runtime/Engine/Private/PhysicsEngine/PhysicsAsset.cpp b/Engine/Source/Runtime/Engine/Private/PhysicsEngine/PhysicsAsset.cpp
--- a/Engine/Source/Runtime/Engine/Private/PhysicsEngine/PhysicsAsset.cpp
+++ b/Engine/Source/Runtime/Engine/Private/PhysicsEngine/PhysicsAsset.cpp
@@ -25,9 +25,18 @@
 bool bUseRBANForDefaultPhysicsAssetSolverType = false;
 FAutoConsoleVariableRef CVarUseRBANForDefaultPhysicsAssetSolverType(TEXT("p.Chaos.UseRBANForDefaultPhysicsAssetSolverType"), bUseRBANForDefaultPhysicsAssetSolverType, TEXT("Boolean to use RBAN for default physics asset solver type (false by default)"));
 
+FPhysicsAssetSolverSettings::FPhysicsAssetSolverSettings()
+	: PositionIterations(6)
+	, VelocityIterations(1)
+	, ProjectionIterations(1)
+	, CullDistance(3.0f)
+	, MaxDepenetrationVelocity(0.0f)
+	, FixedTimeStep(0.0f)
+{
+}
+
 FSolverIterations::FSolverIterations()
-	: FixedTimeStep(0)
-	, SolverIterations(3)
+	: SolverIterations(3)
 	, JointIterations(2)
 	, CollisionIterations(2)
 	, SolverPushOutIterations(1)
diff --git a/Engine/Source/Runtime/Engine/Public/Physics/ImmediatePhysics/ImmediatePhysicsChaos/ImmediatePhysicsSimulation_Chaos.h b/Engine/Source/Runtime/Engine/Public/Physics/ImmediatePhysics/ImmediatePhysicsChaos/ImmediatePhysicsSimulation_Chaos.h
--- a/Engine/Source/Runtime/Engine/Public/Physics/ImmediatePhysics/ImmediatePhysicsChaos/ImmediatePhysicsSimulation_Chaos.h
+++ b/Engine/Source/Runtime/Engine/Public/Physics/ImmediatePhysics/ImmediatePhysicsChaos/ImmediatePhysicsSimulation_Chaos.h
@@ -76,10 +76,17 @@ namespace ImmediatePhysics_Chaos
 			const FReal MasterAlpha, 
 			const FVector& ExternalLinearEtherDrag);
 
-
-		/** Set new iteration counts. A negative value with leave that iteration count unchanged */
-		void SetSolverIterations(
+		/** Set settings. Invalid (negative) values with leave that value unchanged from defaults */
+		void SetSolverSettings(
 			const FReal FixedDt,
+			const FReal CullDistance,
+			const FReal MaxDepenetrationVelocity,
+			const int32 PositionIts,
+			const int32 VelocityIts,
+			const int32 ProjectionIts);
+
+		/** Set iteration counts for the legacy solver. A negative value with leave that iteration count unchanged */
+		void SetLegacySolverSettings(
 			const int32 SolverIts,
 			const int32 JointIts,
 			const int32 CollisionIts,
diff --git a/Engine/Source/Runtime/Experimental/Chaos/Private/Chaos/Collision/PBDCollisionConstraintHandle.cpp b/Engine/Source/Runtime/Experimental/Chaos/Private/Chaos/Collision/PBDCollisionConstraintHandle.cpp
--- a/Engine/Source/Runtime/Experimental/Chaos/Private/Chaos/Collision/PBDCollisionConstraintHandle.cpp
+++ b/Engine/Source/Runtime/Experimental/Chaos/Private/Chaos/Collision/PBDCollisionConstraintHandle.cpp
@@ -62,12 +62,8 @@ namespace Chaos
 		ConcreteContainer()->PreGatherInput(GetContact(), SolverData);
 	}
 
-	void FPBDCollisionConstraintHandle::GatherInput(FReal Dt, const int32 Particle0Level, const int32 Particle1Level, FPBDIslandSolverData& SolverData, const bool bNeedsPreGather)
+	void FPBDCollisionConstraintHandle::GatherInput(FReal Dt, const int32 Particle0Level, const int32 Particle1Level, FPBDIslandSolverData& SolverData)
 	{
-		if (bNeedsPreGather)
-		{
-			PreGatherInput(SolverData);
-		}
 		ConcreteContainer()->GatherInput(Dt, GetContact(), Particle0Level, Particle1Level, SolverData);
 	}
 
diff --git a/Engine/Source/Runtime/Experimental/Chaos/Private/Chaos/Collision/PBDCollisionSolver.cpp b/Engine/Source/Runtime/Experimental/Chaos/Private/Chaos/Collision/PBDCollisionSolver.cpp
--- a/Engine/Source/Runtime/Experimental/Chaos/Private/Chaos/Collision/PBDCollisionSolver.cpp
+++ b/Engine/Source/Runtime/Experimental/Chaos/Private/Chaos/Collision/PBDCollisionSolver.cpp
@@ -27,30 +27,24 @@ namespace Chaos
 		FAutoConsoleVariableRef CVarChaos_PBDCollisionSolver_VectorRegister(TEXT("p.Chaos.PBDCollisionSolver.VectorRegister"), bChaos_PBDCollisionSolver_VectorRegister, TEXT(""));
 
 		bool bChaos_PBDCollisionSolver_Position_SolveEnabled = true;
-		int32 Chaos_PBDCollisionSolver_Position_ShockPropagationIterations = 3;
 		float Chaos_PBDCollisionSolver_Position_MinInvMassScale = 0.3f;
-		int32 Chaos_PBDCollisionSolver_Position_FrictionIterations = 4;
 		float Chaos_PBDCollisionSolver_Position_StaticFrictionStiffness = 0.5f;
 		float Chaos_PBDCollisionSolver_Position_PositionSolverTolerance = 0.001f;		// cms
 		float Chaos_PBDCollisionSolver_Position_RotationSolverTolerance = 0.001f;		// rads
 
 		FAutoConsoleVariableRef CVarChaos_PBDCollisionSolver_Position_SolveEnabled(TEXT("p.Chaos.PBDCollisionSolver.Position.SolveEnabled"), bChaos_PBDCollisionSolver_Position_SolveEnabled, TEXT(""));
-		FAutoConsoleVariableRef CVarChaos_PBDCollisionSolver_Position_UseShockPropagation(TEXT("p.Chaos.PBDCollisionSolver.Position.ShockPropagationIterations"), Chaos_PBDCollisionSolver_Position_ShockPropagationIterations, TEXT(""));
 		FAutoConsoleVariableRef CVarChaos_PBDCollisionSolver_Position_MinInvMassScale(TEXT("p.Chaos.PBDCollisionSolver.Position.MinInvMassScale"), Chaos_PBDCollisionSolver_Position_MinInvMassScale, TEXT(""));
-		FAutoConsoleVariableRef CVarChaos_PBDCollisionSolver_Position_FrictionIterations(TEXT("p.Chaos.PBDCollisionSolver.Position.FrictionIterations"), Chaos_PBDCollisionSolver_Position_FrictionIterations, TEXT(""));
 		FAutoConsoleVariableRef CVarChaos_PBDCollisionSolver_Position_StaticFrictionStiffness(TEXT("p.Chaos.PBDCollisionSolver.Position.StaticFriction.Stiffness"), Chaos_PBDCollisionSolver_Position_StaticFrictionStiffness, TEXT(""));
 		FAutoConsoleVariableRef CVarChaos_PBDCollisionSolver_Position_PositionSolverTolerance(TEXT("p.Chaos.PBDCollisionSolver.Position.PositionTolerance"), Chaos_PBDCollisionSolver_Position_PositionSolverTolerance, TEXT(""));
 		FAutoConsoleVariableRef CVarChaos_PBDCollisionSolver_Position_RotationSolverTolerance(TEXT("p.Chaos.PBDCollisionSolver.Position.RotationTolerance"), Chaos_PBDCollisionSolver_Position_RotationSolverTolerance, TEXT(""));
 
 		bool bChaos_PBDCollisionSolver_Velocity_SolveEnabled = true;
-		int32 Chaos_PBDCollisionSolver_Velocity_ShockPropagationIterations = 1;
 		// If this is the same as Chaos_PBDCollisionSolver_Position_MinInvMassScale and all velocity iterations have shockpropagation, we avoid recalculating constraiunt-space mass
 		float Chaos_PBDCollisionSolver_Velocity_MinInvMassScale = Chaos_PBDCollisionSolver_Position_MinInvMassScale;
 		bool bChaos_PBDCollisionSolver_Velocity_FrictionEnabled = true;
 		bool bChaos_PBDCollisionSolver_Velocity_AveragePointEnabled = true;
 
 		FAutoConsoleVariableRef CVarChaos_PBDCollisionSolver_Velocity_SolveEnabled(TEXT("p.Chaos.PBDCollisionSolver.Velocity.SolveEnabled"), bChaos_PBDCollisionSolver_Velocity_SolveEnabled, TEXT(""));
-		FAutoConsoleVariableRef CVarChaos_PBDCollisionSolver_Velocity_UseShockPropagation(TEXT("p.Chaos.PBDCollisionSolver.Velocity.ShockPropagationIterations"), Chaos_PBDCollisionSolver_Velocity_ShockPropagationIterations, TEXT(""));
 		FAutoConsoleVariableRef CVarChaos_PBDCollisionSolver_Velocity_MinInvMassScale(TEXT("p.Chaos.PBDCollisionSolver.Velocity.MinInvMassScale"), Chaos_PBDCollisionSolver_Velocity_MinInvMassScale, TEXT(""));
 		FAutoConsoleVariableRef CVarChaos_PBDCollisionSolver_Velocity_FrictionEnabled(TEXT("p.Chaos.PBDCollisionSolver.Velocity.FrictionEnabled"), bChaos_PBDCollisionSolver_Velocity_FrictionEnabled, TEXT(""));
 		FAutoConsoleVariableRef CVarChaos_PBDCollisionSolver_Velocity_AveragePointEnabled(TEXT("p.Chaos.PBDCollisionSolver.Velocity.AveragePointEnabled"), bChaos_PBDCollisionSolver_Velocity_AveragePointEnabled, TEXT(""));
diff --git a/Engine/Source/Runtime/Experimental/Chaos/Private/Chaos/Collision/SolverCollisionContainer.cpp b/Engine/Source/Runtime/Experimental/Chaos/Private/Chaos/Collision/SolverCollisionContainer.cpp
--- a/Engine/Source/Runtime/Experimental/Chaos/Private/Chaos/Collision/SolverCollisionContainer.cpp
+++ b/Engine/Source/Runtime/Experimental/Chaos/Private/Chaos/Collision/SolverCollisionContainer.cpp
@@ -21,11 +21,6 @@ namespace Chaos
 	{
 		extern bool bChaos_PBDCollisionSolver_VectorRegister;
 
-		extern int32 Chaos_PBDCollisionSolver_Position_FrictionIterations;
-		extern int32 Chaos_PBDCollisionSolver_Position_ShockPropagationIterations;
-
-		extern int32 Chaos_PBDCollisionSolver_Velocity_ShockPropagationIterations;
-
 		extern bool bChaos_PBDCollisionSolver_Position_SolveEnabled;
 		extern bool bChaos_PBDCollisionSolver_Velocity_SolveEnabled;
 	}
@@ -38,6 +33,22 @@ namespace Chaos
 	//////////////////////////////////////////////////////////////////////////////////////////////////
 
 
+	FPBDCollisionSolverSettings::FPBDCollisionSolverSettings()
+		: MaxPushOutVelocity(0)
+		, NumPositionFrictionIterations(4)
+		, NumVelocityFrictionIterations(1)
+		, NumPositionShockPropagationIterations(3)
+		, NumVelocityShockPropagationIterations(1)
+	{
+	}
+
+
+	//////////////////////////////////////////////////////////////////////////////////////////////////
+	//////////////////////////////////////////////////////////////////////////////////////////////////
+	//////////////////////////////////////////////////////////////////////////////////////////////////
+	//////////////////////////////////////////////////////////////////////////////////////////////////
+
+
 	/**
 	 * @brief A wrapper for FPBDCollisionSolver which binds to a Collision Constraint and adds Gather/Scatter from/to the constraint
 	*/
@@ -72,7 +83,8 @@ namespace Chaos
 			FPBDCollisionConstraint& InConstraint,
 			const int32 Particle0Level,
 			const int32 Particle1Level,
-			FSolverBodyContainer& SolverBodyContainer)
+			FSolverBodyContainer& SolverBodyContainer,
+			const FPBDCollisionSolverSettings& SolverSettings)
 		{
 			bIsIncremental = Constraint->GetUseIncrementalCollisionDetection();
 
@@ -86,19 +98,32 @@ namespace Chaos
 
 			// Friction values. Static and Dynamic friction are applied in the position solve for most shapes.
 			// For quadratic shapes, we run dynamic friction in the velocity solve for better rolling behaviour.
+			// We can also run in a mode without static friction at all. This is faster but stacking is not possible.
 			// @todo(chaos): fix static/dynamic friction for quadratic shapes
 			const FSolverReal StaticFriction = FSolverReal(Constraint->GetStaticFriction());
 			const FSolverReal DynamicFriction = FSolverReal(Constraint->GetDynamicFriction());
-			const bool bIsQuadratic = Constraint->HasQuadraticShape();
-			if (!bIsQuadratic)
+			FSolverReal PositionStaticFriction = FSolverReal(0);
+			FSolverReal PositionDynamicFriction = FSolverReal(0);
+			FSolverReal VelocityDynamicFriction = FSolverReal(0);
+			if (SolverSettings.NumPositionFrictionIterations > 0)
 			{
-				Solver.SetFriction(StaticFriction, DynamicFriction, FSolverReal(0));
+				PositionStaticFriction = StaticFriction;
+				if (!Constraint->HasQuadraticShape())
+				{
+					PositionDynamicFriction = DynamicFriction;
+				}
+				else
+				{
+					VelocityDynamicFriction = DynamicFriction;
+				}
 			}
 			else
 			{
-				Solver.SetFriction(StaticFriction, FReal(0), DynamicFriction);
+				VelocityDynamicFriction = DynamicFriction;
 			}
 
+			Solver.SetFriction(PositionStaticFriction, PositionDynamicFriction, VelocityDynamicFriction);
+
 			Solver.SetStiffness(FSolverReal(Constraint->GetStiffness()));
 
 			Solver.SetSolverBodies(Body0, Body1);
@@ -370,7 +395,6 @@ namespace Chaos
 
 	FPBDCollisionSolverContainer::FPBDCollisionSolverContainer()
 		: FConstraintSolverContainer()
-		, MaxPushOutVelocity(0)
 		, bRequiresIncrementalCollisionDetection(false)
 	{
 	}
@@ -404,7 +428,7 @@ namespace Chaos
 		++ConstraintIndex;
 	}
 
-	void FPBDCollisionSolverContainer::AddConstraintSolver(FReal Dt, FPBDCollisionConstraint& Constraint, const int32 Particle0Level, const int32 Particle1Level, FSolverBodyContainer& SolverBodyContainer)
+	void FPBDCollisionSolverContainer::AddConstraintSolver(FReal Dt, FPBDCollisionConstraint& Constraint, const int32 Particle0Level, const int32 Particle1Level, FSolverBodyContainer& SolverBodyContainer, const FPBDCollisionSolverSettings& SolverSettings)
 	{
 		// This container is required to allocate pointers that are valid for the whole tick,
 		// so we cannot allow the container to resize during the tick. See Reset()
@@ -413,16 +437,15 @@ namespace Chaos
 
 		FPBDCollisionSolverAdapter& CollisionSolver = CollisionSolvers[ConstraintIndex];
 
-		CollisionSolver.GatherInput(Dt, Constraint, Particle0Level, Particle1Level, SolverBodyContainer);
+		CollisionSolver.GatherInput(Dt, Constraint, Particle0Level, Particle1Level, SolverBodyContainer, SolverSettings);
 
 		bRequiresIncrementalCollisionDetection |= CollisionSolver.IsIncrementalManifold();
 	}
 
-	void FPBDCollisionSolverContainer::UpdatePositionShockPropagation(const FReal Dt, const int32 It, const int32 NumIts, const int32 BeginIndex, const int32 EndIndex)
+	void FPBDCollisionSolverContainer::UpdatePositionShockPropagation(const FReal Dt, const int32 It, const int32 NumIts, const int32 BeginIndex, const int32 EndIndex, const FPBDCollisionSolverSettings& SolverSettings)
 	{
 		// If this is the first shock propagation iteration, enable it on each solver
-		const int32 NumShockPropIterations = Chaos_PBDCollisionSolver_Position_ShockPropagationIterations;
-		const bool bEnableShockPropagation = (It == NumIts - NumShockPropIterations);
+		const bool bEnableShockPropagation = (It == NumIts - SolverSettings.NumPositionShockPropagationIterations);
 		if (bEnableShockPropagation)
 		{
 			for (int32 SolverIndex = BeginIndex; SolverIndex < EndIndex; ++SolverIndex)
@@ -432,12 +455,11 @@ namespace Chaos
 		}
 	}
 
-	void FPBDCollisionSolverContainer::UpdateVelocityShockPropagation(const FReal Dt, const int32 It, const int32 NumIts, const int32 BeginIndex, const int32 EndIndex)
+	void FPBDCollisionSolverContainer::UpdateVelocityShockPropagation(const FReal Dt, const int32 It, const int32 NumIts, const int32 BeginIndex, const int32 EndIndex, const FPBDCollisionSolverSettings& SolverSettings)
 	{
 		// Set/reset the shock propagation based on current iteration. The position solve may
 		// have left the bodies with a mass scale and we want to change or reset it.
-		const int32 NumShockPropIterations = Chaos_PBDCollisionSolver_Velocity_ShockPropagationIterations;
-		const bool bEnableShockPropagation = (It == NumIts - NumShockPropIterations);
+		const bool bEnableShockPropagation = (It == NumIts - SolverSettings.NumVelocityShockPropagationIterations);
 		if (bEnableShockPropagation)
 		{
 			for (int32 SolverIndex = BeginIndex; SolverIndex < EndIndex; ++SolverIndex)
@@ -454,28 +476,28 @@ namespace Chaos
 		}
 	}
 
-	bool FPBDCollisionSolverContainer::SolvePositionSerial(const FReal Dt, const int32 It, const int32 NumIts, const int32 BeginIndex, const int32 EndIndex)
+	bool FPBDCollisionSolverContainer::SolvePositionSerial(const FReal Dt, const int32 It, const int32 NumIts, const int32 BeginIndex, const int32 EndIndex, const FPBDCollisionSolverSettings& SolverSettings)
 	{
-		return SolvePositionImpl(Dt, It, NumIts, BeginIndex, EndIndex, false);
+		return SolvePositionImpl(Dt, It, NumIts, BeginIndex, EndIndex, SolverSettings, false);
 	}
 
-	bool FPBDCollisionSolverContainer::SolveVelocitySerial(const FReal Dt, const int32 It, const int32 NumIts, const int32 BeginIndex, const int32 EndIndex)
+	bool FPBDCollisionSolverContainer::SolveVelocitySerial(const FReal Dt, const int32 It, const int32 NumIts, const int32 BeginIndex, const int32 EndIndex, const FPBDCollisionSolverSettings& SolverSettings)
 	{
-		return SolveVelocityImpl(Dt, It, NumIts, BeginIndex, EndIndex, false);
+		return SolveVelocityImpl(Dt, It, NumIts, BeginIndex, EndIndex, SolverSettings, false);
 	}
 
-	bool FPBDCollisionSolverContainer::SolvePositionParallel(const FReal Dt, const int32 It, const int32 NumIts, const int32 BeginIndex, const int32 EndIndex)
+	bool FPBDCollisionSolverContainer::SolvePositionParallel(const FReal Dt, const int32 It, const int32 NumIts, const int32 BeginIndex, const int32 EndIndex, const FPBDCollisionSolverSettings& SolverSettings)
 	{
-		return SolvePositionImpl(Dt, It, NumIts, BeginIndex, EndIndex, true);
+		return SolvePositionImpl(Dt, It, NumIts, BeginIndex, EndIndex, SolverSettings, true);
 	}
 
-	bool FPBDCollisionSolverContainer::SolveVelocityParallel(const FReal Dt, const int32 It, const int32 NumIts, const int32 BeginIndex, const int32 EndIndex)
+	bool FPBDCollisionSolverContainer::SolveVelocityParallel(const FReal Dt, const int32 It, const int32 NumIts, const int32 BeginIndex, const int32 EndIndex, const FPBDCollisionSolverSettings& SolverSettings)
 	{
-		return SolveVelocityImpl(Dt, It, NumIts, BeginIndex, EndIndex, true);
+		return SolveVelocityImpl(Dt, It, NumIts, BeginIndex, EndIndex, SolverSettings, true);
 	}
 
 	// @todo(chaos): parallel version of SolvePosition
-	bool FPBDCollisionSolverContainer::SolvePositionImpl(const FReal Dt, const int32 It, const int32 NumIts, const int32 BeginIndex, const int32 EndIndex, const bool bParallel)
+	bool FPBDCollisionSolverContainer::SolvePositionImpl(const FReal Dt, const int32 It, const int32 NumIts, const int32 BeginIndex, const int32 EndIndex, const FPBDCollisionSolverSettings& SolverSettings, const bool bParallel)
 	{
 		SCOPE_CYCLE_COUNTER(STAT_Collisions_Apply);
 		if (!bChaos_PBDCollisionSolver_Position_SolveEnabled)
@@ -483,13 +505,13 @@ namespace Chaos
 			return false;
 		}
 
-		UpdatePositionShockPropagation(Dt, It, NumIts, BeginIndex, EndIndex);
+		UpdatePositionShockPropagation(Dt, It, NumIts, BeginIndex, EndIndex, SolverSettings);
 
 		// Only apply friction for the last few (tunable) iterations
-		const bool bApplyStaticFriction = (It >= (NumIts - Chaos_PBDCollisionSolver_Position_FrictionIterations));
+		const bool bApplyStaticFriction = (It >= (NumIts - SolverSettings.NumPositionFrictionIterations));
 
 		// Adjust max pushout to attempt to make it iteration count independent
-		const FReal MaxPushOut = (MaxPushOutVelocity > 0) ? (MaxPushOutVelocity * Dt) / FReal(NumIts) : 0;
+		const FReal MaxPushOut = (SolverSettings.MaxPushOutVelocity > 0) ? (SolverSettings.MaxPushOutVelocity * Dt) / FReal(NumIts) : 0;
 
 		// Apply the position correction
 		if (bRequiresIncrementalCollisionDetection)
@@ -562,7 +584,7 @@ namespace Chaos
 		return bNeedsAnotherIteration;
 	}
 
-	bool FPBDCollisionSolverContainer::SolveVelocityImpl(const FReal InDt, const int32 It, const int32 NumIts, const int32 BeginIndex, const int32 EndIndex, const bool bParallel)
+	bool FPBDCollisionSolverContainer::SolveVelocityImpl(const FReal InDt, const int32 It, const int32 NumIts, const int32 BeginIndex, const int32 EndIndex, const FPBDCollisionSolverSettings& SolverSettings, const bool bParallel)
 	{
 		SCOPE_CYCLE_COUNTER(STAT_Collisions_ApplyPushOut);
 		if (!bChaos_PBDCollisionSolver_Velocity_SolveEnabled)
@@ -571,10 +593,9 @@ namespace Chaos
 		}
 		const FSolverReal Dt = FSolverReal(InDt);
 
-		UpdateVelocityShockPropagation(Dt, It, NumIts, BeginIndex, EndIndex);
+		UpdateVelocityShockPropagation(Dt, It, NumIts, BeginIndex, EndIndex, SolverSettings);
 
-		const int32 NumDynamicFrictionIterations = 1;
-		const bool bApplyDynamicFriction = (It >= NumIts - NumDynamicFrictionIterations);
+		const bool bApplyDynamicFriction = (It >= NumIts - SolverSettings.NumVelocityFrictionIterations);
 
 		// Apply the velocity correction
 		// @todo(chaos): parallel version of SolveVelocity
diff --git a/Engine/Source/Runtime/Experimental/Chaos/Private/Chaos/Evolution/PBDMinEvolution.cpp b/Engine/Source/Runtime/Experimental/Chaos/Private/Chaos/Evolution/PBDMinEvolution.cpp
--- a/Engine/Source/Runtime/Experimental/Chaos/Private/Chaos/Evolution/PBDMinEvolution.cpp
+++ b/Engine/Source/Runtime/Experimental/Chaos/Private/Chaos/Evolution/PBDMinEvolution.cpp
@@ -44,8 +44,9 @@ namespace Chaos
 	DECLARE_CYCLE_STAT(TEXT("MinEvolution::ApplyConstraintsPhase1"), STAT_MinEvolution_ApplyConstraintsPhase1, STATGROUP_ChaosMinEvolution);
 	DECLARE_CYCLE_STAT(TEXT("MinEvolution::UpdateVelocities"), STAT_MinEvolution_UpdateVelocites, STATGROUP_ChaosMinEvolution);
 	DECLARE_CYCLE_STAT(TEXT("MinEvolution::ApplyConstraintsPhase2"), STAT_MinEvolution_ApplyConstraintsPhase2, STATGROUP_ChaosMinEvolution);
+	DECLARE_CYCLE_STAT(TEXT("MinEvolution::ApplyCorrections"), STAT_MinEvolution_ApplyCorrections, STATGROUP_ChaosMinEvolution);
+	DECLARE_CYCLE_STAT(TEXT("MinEvolution::ApplyConstraintsPhase3"), STAT_MinEvolution_ApplyConstraintsPhase3, STATGROUP_ChaosMinEvolution);
 	DECLARE_CYCLE_STAT(TEXT("MinEvolution::DetectCollisions"), STAT_MinEvolution_DetectCollisions, STATGROUP_ChaosMinEvolution);
-	DECLARE_CYCLE_STAT(TEXT("MinEvolution::UpdatePositions"), STAT_MinEvolution_UpdatePositions, STATGROUP_ChaosMinEvolution);
 
 	//
 	//
@@ -138,6 +139,9 @@ namespace Chaos
 		, SolverType(EConstraintSolverType::QuasiPbd)
 		, NumApplyIterations(0)
 		, NumApplyPushOutIterations(0)
+		, NumPositionIterations(0)
+		, NumVelocityIterations(0)
+		, NumProjectionIterations(0)
 		, BoundsExtension(InBoundsExtension)
 		, Gravity(FVec3(0))
 		, SimulationSpaceSettings()
@@ -228,9 +232,11 @@ namespace Chaos
 				PostApplyPushOutCallback();
 			}
 
-			ScatterOutput(Dt);
+			ApplyCorrections(Dt);
+
+			ApplyConstraintsPhase3(Dt);
 
-			UpdatePositions(Dt);
+			ScatterOutput(Dt);
 		}
 	}
 
@@ -492,18 +498,28 @@ namespace Chaos
 		}
 
 		SolverData.GetBodyContainer().ScatterOutput();
+	
+		for (auto& Particle : Particles.GetActiveParticlesView())
+		{
+			Particle.Handle()->AuxilaryValue(ParticlePrevXs) = Particle.X();
+			Particle.Handle()->AuxilaryValue(ParticlePrevRs) = Particle.R();
+			Particle.X() = Particle.P();
+			Particle.R() = Particle.Q();
+		}
 	}
 
 	void FPBDMinEvolution::ApplyConstraintsPhase1(FReal Dt)
 	{
 		SCOPE_CYCLE_COUNTER(STAT_MinEvolution_ApplyConstraintsPhase1);
 
-		for (int32 i = 0; i < NumApplyIterations; ++i)
+		const int32 NumIterationsPhase1 = (SolverType == EConstraintSolverType::QuasiPbd) ? NumPositionIterations : NumApplyIterations;
+
+		for (int32 i = 0; i < NumIterationsPhase1; ++i)
 		{
 			bool bNeedsAnotherIteration = Chaos_MinEvolution_ForceMaxConstraintIterations;
 			for (FSimpleConstraintRule* ConstraintRule : PrioritizedConstraintRules)
 			{
-				bNeedsAnotherIteration |= ConstraintRule->ApplyConstraints(Dt, i, NumApplyIterations);
+				bNeedsAnotherIteration |= ConstraintRule->ApplyConstraints(Dt, i, NumIterationsPhase1);
 			}
 
 			if (!bNeedsAnotherIteration)
@@ -539,12 +555,14 @@ namespace Chaos
 	{
 		SCOPE_CYCLE_COUNTER(STAT_MinEvolution_ApplyConstraintsPhase2);
 
-		for (int32 It = 0; It < NumApplyPushOutIterations; ++It)
+		const int32 NumIterationsPhase2 = (SolverType == EConstraintSolverType::QuasiPbd) ? NumVelocityIterations : 0;
+
+		for (int32 It = 0; It < NumIterationsPhase2; ++It)
 		{
 			bool bNeedsAnotherIteration = false;
 			for (FSimpleConstraintRule* ConstraintRule : PrioritizedConstraintRules)
 			{
-				bNeedsAnotherIteration |= ConstraintRule->ApplyPushOut(Dt, It, NumApplyPushOutIterations);
+				bNeedsAnotherIteration |= ConstraintRule->ApplyPushOut(Dt, It, NumIterationsPhase2);
 			}
 
 			if (!bNeedsAnotherIteration)
@@ -554,16 +572,32 @@ namespace Chaos
 		}
 	}
 
-	void FPBDMinEvolution::UpdatePositions(FReal Dt)
+	void FPBDMinEvolution::ApplyCorrections(FReal Dt)
 	{
-		SCOPE_CYCLE_COUNTER(STAT_MinEvolution_UpdatePositions);
-		for (auto& Particle : Particles.GetActiveParticlesView())
+		SCOPE_CYCLE_COUNTER(STAT_MinEvolution_ApplyCorrections);
+
+		SolverData.GetBodyContainer().ApplyCorrections();
+		SolverData.GetBodyContainer().UpdateRotationDependentState();
+	}
+
+	void FPBDMinEvolution::ApplyConstraintsPhase3(FReal Dt)
+	{
+		SCOPE_CYCLE_COUNTER(STAT_MinEvolution_ApplyConstraintsPhase3);
+
+		const int32 NumIterationsPhase3 = (SolverType == EConstraintSolverType::QuasiPbd) ? NumProjectionIterations : NumApplyPushOutIterations;
+
+		for (int32 It = 0; It < NumIterationsPhase3; ++It)
 		{
-			Particle.Handle()->AuxilaryValue(ParticlePrevXs) = Particle.X();
-			Particle.Handle()->AuxilaryValue(ParticlePrevRs) = Particle.R();
-			Particle.X() = Particle.P();
-			Particle.R() = Particle.Q();
+			bool bNeedsAnotherIteration = false;
+			for (FSimpleConstraintRule* ConstraintRule : PrioritizedConstraintRules)
+			{
+				bNeedsAnotherIteration |= ConstraintRule->ApplyProjection(Dt, It, NumIterationsPhase3);
+			}
+
+			if (!bNeedsAnotherIteration)
+			{
+				break;
+			}
 		}
 	}
-
 }
diff --git a/Engine/Source/Runtime/Experimental/Chaos/Private/Chaos/Joint/JointConstraintsCVars.cpp b/Engine/Source/Runtime/Experimental/Chaos/Private/Chaos/Joint/JointConstraintsCVars.cpp
--- a/Engine/Source/Runtime/Experimental/Chaos/Private/Chaos/Joint/JointConstraintsCVars.cpp
+++ b/Engine/Source/Runtime/Experimental/Chaos/Private/Chaos/Joint/JointConstraintsCVars.cpp
@@ -7,10 +7,6 @@ bool bChaos_Joint_ISPC_Enabled = false;
 FAutoConsoleVariableRef CVarChaosJointISPCEnabled(TEXT("p.Chaos.Joint.ISPC"), bChaos_Joint_ISPC_Enabled, TEXT("Whether to use ISPC optimizations in the Joint Solver"));
 #endif
 
-// @todo(chaos): fix joint early out. Joints that are "solved" can be "unsolved" by changes to the bodies they are connected to.
-bool bChaos_Joint_EarlyOut_Enabled = false;
-FAutoConsoleVariableRef CVarChaosJointEarlyOutEnabled(TEXT("p.Chaos.Joint.EarlyOut"), bChaos_Joint_EarlyOut_Enabled, TEXT("Whether to iterating when joints report being solved"));
-
 float Chaos_Joint_DegenerateRotationLimit = -0.998f;	// Cos(176deg)
 FAutoConsoleVariableRef CVarChaosJointDegenerateRotationLimit(TEXT("p.Chaos.Joint.DegenerateRotationLimit"), Chaos_Joint_DegenerateRotationLimit, TEXT("Cosine of the swing angle that is considered degerenerate (default Cos(176deg))"));
 
diff --git a/Engine/Source/Runtime/Experimental/Chaos/Private/Chaos/PBDCollisionConstraints.cpp b/Engine/Source/Runtime/Experimental/Chaos/Private/Chaos/PBDCollisionConstraints.cpp
--- a/Engine/Source/Runtime/Experimental/Chaos/Private/Chaos/PBDCollisionConstraints.cpp
+++ b/Engine/Source/Runtime/Experimental/Chaos/Private/Chaos/PBDCollisionConstraints.cpp
@@ -101,7 +101,7 @@ namespace Chaos
 		const THandleArray<FChaosPhysicsMaterial>* const InSimMaterials,
 		const int32 InApplyPairIterations /*= 1*/,
 		const int32 InApplyPushOutPairIterations /*= 1*/,
-		const FReal InRestitutionThreshold /*= (FReal)0*/)
+		const FReal InRestitutionThreshold /*= (FReal)2000*/)
 		: FPBDConstraintContainer(FConstraintContainerHandle::StaticType())
 		, Particles(InParticles)
 		, NumActivePointConstraints(0)
@@ -120,7 +120,7 @@ namespace Chaos
 		, bCanDisableContacts(true)
 		, GravityDirection(FVec3(0,0,-1))
 		, GravitySize(980)
-		, MaxPushOutVelocity(0)
+		, SolverSettings()
 		, SolverType(EConstraintSolverType::QuasiPbd)
 	{
 	}
@@ -359,10 +359,6 @@ namespace Chaos
 		{
 			FPBDCollisionSolverContainer& SolverContainer = GetConstraintSolverContainer(SolverData);
 			SolverContainer.SetNum(NumIslandConstraints);
-
-			// @todo(chaos): we shouldn't have to pass settings down to the solvers every frame. Ideally the solver containers would be created by the constraint containers and
-			// the settings would be passed on then. Currently the solver containers are created here FPBDIslandSolverData::AddConstraintDatas
-			SolverContainer.SetMaxPushOutVelocity(MaxPushOutVelocity);
 		}
 		else
 		{
@@ -393,7 +389,7 @@ namespace Chaos
 			check(Constraint.IsEnabled());
 
 			FPBDCollisionSolverContainer& SolverContainer = GetConstraintSolverContainer(SolverData);
-			SolverContainer.AddConstraintSolver(Dt, Constraint, Particle0Level, Particle1Level, SolverData.GetBodyContainer());
+			SolverContainer.AddConstraintSolver(Dt, Constraint, Particle0Level, Particle1Level, SolverData.GetBodyContainer(), SolverSettings);
 		}
 		else
 		{
@@ -409,6 +405,7 @@ namespace Chaos
 			{
 				if (Constraint->IsEnabled())
 				{
+					PreGatherInput(*Constraint, SolverData);
 					GatherInput(Dt, *Constraint, INDEX_NONE, INDEX_NONE, SolverData);
 				}
 			}
@@ -458,7 +455,7 @@ namespace Chaos
 		if (SolverType == EConstraintSolverType::QuasiPbd)
 		{
 			FPBDCollisionSolverContainer& SolverContainer = GetConstraintSolverContainer(SolverData);
-			return SolverContainer.SolvePositionSerial(Dt, It, NumIts, 0, SolverContainer.NumSolvers());
+			return SolverContainer.SolvePositionSerial(Dt, It, NumIts, 0, SolverContainer.NumSolvers(), SolverSettings);
 		}
 		else
 		{
@@ -474,7 +471,7 @@ namespace Chaos
 		if (SolverType == EConstraintSolverType::QuasiPbd)
 		{
 			FPBDCollisionSolverContainer& SolverContainer = GetConstraintSolverContainer(SolverData);
-			return SolverContainer.SolvePositionSerial(Dt, It, NumIts, 0, SolverContainer.NumSolvers());
+			return SolverContainer.SolvePositionSerial(Dt, It, NumIts, 0, SolverContainer.NumSolvers(), SolverSettings);
 		}
 		else
 		{
@@ -490,7 +487,7 @@ namespace Chaos
 		if (SolverType == EConstraintSolverType::QuasiPbd)
 		{
 			FPBDCollisionSolverContainer& SolverContainer = GetConstraintSolverContainer(SolverData);
-			return SolverContainer.SolvePositionSerial(Dt, It, NumIts, BeginIndex, EndIndex);
+			return SolverContainer.SolvePositionSerial(Dt, It, NumIts, BeginIndex, EndIndex, SolverSettings);
 		}
 		else
 		{
@@ -506,7 +503,7 @@ namespace Chaos
 		if (SolverType == EConstraintSolverType::QuasiPbd)
 		{
 			FPBDCollisionSolverContainer& SolverContainer = GetConstraintSolverContainer(SolverData);
-			return SolverContainer.SolvePositionParallel(Dt, It, NumIts, BeginIndex, EndIndex);
+			return SolverContainer.SolvePositionParallel(Dt, It, NumIts, BeginIndex, EndIndex, SolverSettings);
 		}
 		else
 		{
@@ -522,7 +519,7 @@ namespace Chaos
 		if (SolverType == EConstraintSolverType::QuasiPbd)
 		{
 			FPBDCollisionSolverContainer& SolverContainer = GetConstraintSolverContainer(SolverData);
-			return SolverContainer.SolveVelocitySerial(Dt, It, NumIts, 0, SolverContainer.NumSolvers());
+			return SolverContainer.SolveVelocitySerial(Dt, It, NumIts, 0, SolverContainer.NumSolvers(), SolverSettings);
 		}
 		else
 		{
@@ -538,7 +535,7 @@ namespace Chaos
 		if (SolverType == EConstraintSolverType::QuasiPbd)
 		{
 			FPBDCollisionSolverContainer& SolverContainer = GetConstraintSolverContainer(SolverData);
-			return SolverContainer.SolveVelocitySerial(Dt, It, NumIts, 0, SolverContainer.NumSolvers());
+			return SolverContainer.SolveVelocitySerial(Dt, It, NumIts, 0, SolverContainer.NumSolvers(), SolverSettings);
 		}
 		else
 		{
@@ -554,7 +551,7 @@ namespace Chaos
 		if (SolverType == EConstraintSolverType::QuasiPbd)
 		{
 			FPBDCollisionSolverContainer& SolverContainer = GetConstraintSolverContainer(SolverData);
-			return SolverContainer.SolveVelocitySerial(Dt, It, NumIts, BeginIndex, EndIndex);
+			return SolverContainer.SolveVelocitySerial(Dt, It, NumIts, BeginIndex, EndIndex, SolverSettings);
 		}
 		else
 		{
@@ -570,7 +567,7 @@ namespace Chaos
 		if (SolverType == EConstraintSolverType::QuasiPbd)
 		{
 			FPBDCollisionSolverContainer& SolverContainer = GetConstraintSolverContainer(SolverData);
-			return SolverContainer.SolveVelocityParallel(Dt, It, NumIts, BeginIndex, EndIndex);
+			return SolverContainer.SolveVelocityParallel(Dt, It, NumIts, BeginIndex, EndIndex, SolverSettings);
 		}
 		else
 		{
diff --git a/Engine/Source/Runtime/Experimental/Chaos/Private/Chaos/PBDConstraintRule.cpp b/Engine/Source/Runtime/Experimental/Chaos/Private/Chaos/PBDConstraintRule.cpp
--- a/Engine/Source/Runtime/Experimental/Chaos/Private/Chaos/PBDConstraintRule.cpp
+++ b/Engine/Source/Runtime/Experimental/Chaos/Private/Chaos/PBDConstraintRule.cpp
@@ -101,7 +101,13 @@ namespace Chaos
 	{
 		return SolverData ? Constraints.ApplyPhase2(Dt, It, NumIts, *SolverData) : false;
 	}
-	
+
+	template<class ConstraintType>
+	bool TSimpleConstraintRule<ConstraintType>::ApplyProjection(const FReal Dt, const int32 It, const int32 NumIts)
+	{
+		return SolverData ? Constraints.ApplyPhase3(Dt, It, NumIts, *SolverData) : false;
+	}
+
 	template<class ConstraintType>
 	TPBDConstraintGraphRuleImpl<ConstraintType>::TPBDConstraintGraphRuleImpl(FConstraints& InConstraints, int32 InPriority)
 		: FPBDConstraintGraphRule(InPriority)
@@ -326,7 +332,7 @@ namespace Chaos
 
 													// Note we are building the SolverBodies as we go, in the order that we visit them. Each constraint
 													// references two bodies, so we won't strictly be accessing only in cache order, but it's about as good as it can be.
-													Constraint->GatherInput(Dt, Particle0Level, Particle1Level, *IslandGroup, false);
+													Constraint->GatherInput(Dt, Particle0Level, Particle1Level, *IslandGroup);
 												}
 											}
 										});
diff --git a/Engine/Source/Runtime/Experimental/Chaos/Private/Chaos/PBDJointConstraintUtilities.cpp b/Engine/Source/Runtime/Experimental/Chaos/Private/Chaos/PBDJointConstraintUtilities.cpp
--- a/Engine/Source/Runtime/Experimental/Chaos/Private/Chaos/PBDJointConstraintUtilities.cpp
+++ b/Engine/Source/Runtime/Experimental/Chaos/Private/Chaos/PBDJointConstraintUtilities.cpp
@@ -605,12 +605,8 @@ namespace Chaos
 		const FPBDJointSolverSettings& SolverSettings,
 		const FPBDJointSettings& JointSettings)
 	{
-		if (JointSettings.bProjectionEnabled)
-		{
-			// ShockProagation setting is a alpha. For an alpha of 0 we want an invmass scale of 1, and vice-versa
-			return (SolverSettings.ShockPropagationOverride >= FReal(0)) ? (FReal(1) - SolverSettings.ShockPropagationOverride) : (FReal(1) - JointSettings.ShockPropagation);
-		}
-		return FReal(1);
+		// ShockProagation setting is a alpha. For an alpha of 0 we want an invmass scale of 1, and vice-versa
+		return (SolverSettings.ShockPropagationOverride >= FReal(0)) ? (FReal(1) - SolverSettings.ShockPropagationOverride) : (FReal(1) - JointSettings.ShockPropagation);
 	}
 
 
diff --git a/Engine/Source/Runtime/Experimental/Chaos/Private/Chaos/PBDJointConstraints.cpp b/Engine/Source/Runtime/Experimental/Chaos/Private/Chaos/PBDJointConstraints.cpp
--- a/Engine/Source/Runtime/Experimental/Chaos/Private/Chaos/PBDJointConstraints.cpp
+++ b/Engine/Source/Runtime/Experimental/Chaos/Private/Chaos/PBDJointConstraints.cpp
@@ -28,6 +28,7 @@ namespace Chaos
 	DECLARE_CYCLE_STAT(TEXT("Joints::Scatter"), STAT_Joints_Scatter, STATGROUP_ChaosJoint);
 	DECLARE_CYCLE_STAT(TEXT("Joints::Apply"), STAT_Joints_Apply, STATGROUP_ChaosJoint);
 	DECLARE_CYCLE_STAT(TEXT("Joints::ApplyPushOut"), STAT_Joints_ApplyPushOut, STATGROUP_ChaosJoint);
+	DECLARE_CYCLE_STAT(TEXT("Joints::ApplyProjection"), STAT_Joints_ApplyProjection, STATGROUP_ChaosJoint);
 
 
 	//
@@ -191,6 +192,7 @@ namespace Chaos
 		, ParentInvMassScale(1)
 		, bCollisionEnabled(true)
 		, bProjectionEnabled(false)
+		, bShockPropagationEnabled(false)
 		, bSoftProjectionEnabled(false)
 		, LinearMotionTypes({ EJointMotionType::Locked, EJointMotionType::Locked, EJointMotionType::Locked })
 		, LinearLimit(FLT_MAX)
@@ -832,6 +834,11 @@ namespace Chaos
 		return ApplyPhase2Serial(Dt, It, NumIts, SolverData);
 	}
 
+	bool FPBDJointConstraints::ApplyPhase3(const FReal Dt, const int32 It, const int32 NumIts, FPBDIslandSolverData& SolverData)
+	{
+		return ApplyPhase3Serial(Dt, It, NumIts, SolverData);
+	}
+
 	//////////////////////////////////////////////////////////////////////////
 	//
 	// End Simple API Solver.
@@ -1011,18 +1018,28 @@ namespace Chaos
 		SCOPE_CYCLE_COUNTER(STAT_Joints_ApplyPushOut);
 
 		int32 NumActive = 0;
-		const int32 NumPairIts = (SolverType == EConstraintSolverType::QuasiPbd) ? 1 : Settings.ApplyPushOutPairIterations;
-		if (NumPairIts > 0)
+		for (int32 ConstraintIndex : SolverData.GetConstraintIndices(ContainerId))
 		{
-			// UpdateConstraintProjection(It, NumIts, SolverData);
-			for (int32 ConstraintIndex : SolverData.GetConstraintIndices(ContainerId))
-			{
-				NumActive += ApplyPhase2Single(Dt, ConstraintIndex, NumPairIts, It, NumIts);
-			}
+			NumActive += ApplyPhase2Single(Dt, ConstraintIndex, It, NumIts);
 		}
 		return (NumActive > 0);
 	}
 
+	bool FPBDJointConstraints::ApplyPhase3Serial(const FReal Dt, const int32 It, const int32 NumIts, FPBDIslandSolverData& SolverData)
+	{
+		CSV_SCOPED_TIMING_STAT(Chaos, ProjectJointConstraints);
+		SCOPE_CYCLE_COUNTER(STAT_Joints_ApplyProjection);
+
+		// UpdateConstraintProjection(It, NumIts, SolverData);
+		for (int32 ConstraintIndex : SolverData.GetConstraintIndices(ContainerId))
+		{
+			ApplyPhase3Single(Dt, ConstraintIndex, It, NumIts);
+		}
+
+		return true;
+	}
+
+
 	void FPBDJointConstraints::UpdateConstraintProjection(const int32 It, const int32 NumIts, const FPBDIslandSolverData& SolverData)
 	{
 		const bool bUseProjection = It == NumIts - 1;
@@ -1120,7 +1137,7 @@ namespace Chaos
 	{
 		// Shock propagation is only enabled for the last iteration, and only for the QPBD solver.
 		// The standard PBD solver runs projection in the second solver phase which is mostly the same thing.
-		if ((It >= (NumIts - Settings.NumShockPropagationIterations)) && (SolverType == EConstraintSolverType::QuasiPbd))
+		if (JointSettings.bShockPropagationEnabled && (It >= (NumIts - Settings.NumShockPropagationIterations)) && (SolverType == EConstraintSolverType::QuasiPbd))
 		{
 			if (Body0.IsDynamic() && Body1.IsDynamic())
 			{
@@ -1140,59 +1157,41 @@ namespace Chaos
 		}
 
 		const TVector<TGeometryParticleHandle<FReal, 3>*, 2>& Constraint = ConstraintParticles[ConstraintIndex];
-		UE_LOG(LogChaosJoint, VeryVerbose, TEXT("Solve Joint Constraint %d %s %s (dt = %f; it = %d / %d)"), ConstraintIndex, *Constraint[0]->ToString(), *Constraint[1]->ToString(), Dt, It, NumIts);
+		UE_LOG(LogChaosJoint, VeryVerbose, TEXT("Solve Joint Position Constraint %d %s %s (dt = %f; it = %d / %d)"), ConstraintIndex, *Constraint[0]->ToString(), *Constraint[1]->ToString(), Dt, It, NumIts);
+
+		// @todo(chaos): store this on the Solver object and don't access the particles here
+		int32 Index0, Index1;
+		GetConstrainedParticleIndices(ConstraintIndex, Index0, Index1);
+		FGenericParticleHandle Particle0 = FGenericParticleHandle(ConstraintParticles[ConstraintIndex][Index0]);
+		FGenericParticleHandle Particle1 = FGenericParticleHandle(ConstraintParticles[ConstraintIndex][Index1]);
+		if ((Particle0->Sleeping() && Particle1->Sleeping())
+			|| (Particle0->IsKinematic() && Particle1->Sleeping())
+			|| (Particle0->Sleeping() && Particle1->IsKinematic()))
+		{
+			return false;
+		}
 
 		const FPBDJointSettings& JointSettings = ConstraintSettings[ConstraintIndex];
 		if (Settings.bUseLinearSolver)
 		{
 			FPBDJointCachedSolver& Solver = CachedConstraintSolvers[ConstraintIndex];
-
-			// @todo(chaos): store this on the Solver object and don't access the particles here
-			int32 Index0, Index1;
-			GetConstrainedParticleIndices(ConstraintIndex, Index0, Index1);
-			FGenericParticleHandle Particle0 = FGenericParticleHandle(ConstraintParticles[ConstraintIndex][Index0]);
-			FGenericParticleHandle Particle1 = FGenericParticleHandle(ConstraintParticles[ConstraintIndex][Index1]);
-			if ((Particle0->Sleeping() && Particle1->Sleeping())
-				|| (Particle0->IsKinematic() && Particle1->Sleeping()) 
-				|| (Particle0->Sleeping() && Particle1->IsKinematic())
-				|| (FMath::IsNearlyZero(Solver.InvM(0)) && FMath::IsNearlyZero(Solver.InvM(1))))
+			if ((FMath::IsNearlyZero(Solver.InvM(0)) && FMath::IsNearlyZero(Solver.InvM(1))))
 			{
 				return false;
 			}
 
-			// If we were solved last iteration and nothing has changed since, we are done
-			const bool bWasActive = Solver.GetIsActive();
 			Solver.Update(Dt, Settings, JointSettings);
 
 			// Set parent inverse mass scale based on current shock propagation state
 			const FReal ShockPropagationInvMassScale = CalculateShockPropagationInvMassScale(Solver.Body0(), Solver.Body1(), JointSettings, It, NumIts);
 			Solver.SetInvMassScales(ShockPropagationInvMassScale, FReal(1), Dt);
 
-			if (!bWasActive && !Solver.GetIsActive() && bChaos_Joint_EarlyOut_Enabled)
-			{
-				return false;
-			}
-
 			const FReal IterationStiffness = CalculateIterationStiffness(It, NumIts);
 			for (int32 PairIt = 0; PairIt < NumPairIts; ++PairIt)
 			{
 				UE_LOG(LogChaosJoint, VeryVerbose, TEXT("  Pair Iteration %d / %d"), PairIt, NumPairIts);
 
-				// This is the same position solver for all SolverTypes (which makes it wrong for the GbfPbd solver which should be a pbd-esque velocity solve)
 				Solver.ApplyConstraints(Dt, IterationStiffness, Settings, JointSettings);
-
-				if (!Solver.GetIsActive() && bChaos_Joint_EarlyOut_Enabled)
-				{
-					break;
-				}
-			}
-
-			// If the solver we are running uses a velocity solver in the Apply step (for other constraint types), 
-			// we need to update the implicit velocity because we did a position solve above. E.g., GbfPbd collisions.
-			if (SolverType == EConstraintSolverType::GbfPbd)
-			{
-				Solver.Body(0).SetImplicitVelocity(Dt);
-				Solver.Body(1).SetImplicitVelocity(Dt);
 			}
 
 			// @todo(ccaulfield): The break limit should really be applied to the impulse in the solver to prevent 1-frame impulses larger than the threshold
@@ -1200,35 +1199,17 @@ namespace Chaos
 			{
 				ApplyBreakThreshold(Dt, ConstraintIndex, Solver.GetNetLinearImpulse(), Solver.GetNetAngularImpulse());
 			}
-
-			return Solver.GetIsActive() || !bChaos_Joint_EarlyOut_Enabled;
 		}
 		else
 		{
 			FPBDJointSolver& Solver = ConstraintSolvers[ConstraintIndex];
-
-			// @todo(chaos): store this on the Solver object and don't access the particles here
-			int32 Index0, Index1;
-			GetConstrainedParticleIndices(ConstraintIndex, Index0, Index1);
-			FGenericParticleHandle Particle0 = FGenericParticleHandle(ConstraintParticles[ConstraintIndex][Index0]);
-			FGenericParticleHandle Particle1 = FGenericParticleHandle(ConstraintParticles[ConstraintIndex][Index1]);
-			if ((Particle0->Sleeping() && Particle1->Sleeping())
-				|| (Particle0->IsKinematic() && Particle1->Sleeping()) 
-				|| (Particle0->Sleeping() && Particle1->IsKinematic())
-				|| (FMath::IsNearlyZero(Solver.InvM(0)) && FMath::IsNearlyZero(Solver.InvM(1))))
+			if ((FMath::IsNearlyZero(Solver.InvM(0)) && FMath::IsNearlyZero(Solver.InvM(1))))
 			{
 				return false;
 			}
 
-			// If we were solved last iteration and nothing has changed since, we are done
-			const bool bWasActive = Solver.GetIsActive();
 			Solver.Update(Dt, Settings, JointSettings);
 
-			if (!bWasActive && !Solver.GetIsActive() && bChaos_Joint_EarlyOut_Enabled)
-			{
-				return false;
-			}
-
 			// Set parent inverse mass scale based on current shock propagation state
 			const FReal ShockPropagationInvMassScale = CalculateShockPropagationInvMassScale(Solver.Body0(), Solver.Body1(), JointSettings, It, NumIts);
 
@@ -1241,10 +1222,11 @@ namespace Chaos
 				{
 					Solver.UpdateMasses(ShockPropagationInvMassScale, FReal(1));
 				}
+				else
+				{
+					Solver.SetInvMassScales(ShockPropagationInvMassScale, FReal(1));
+				}
 
-				Solver.SetInvMassScales(ShockPropagationInvMassScale, FReal(1));
-
-				// This is the same position solver for all SolverTypes (which makes it wrong for the GbfPbd solver which should be a pbd-esque velocity solve)
 				Solver.ApplyConstraints(Dt, IterationStiffness, Settings, JointSettings);
 
 				if (SolverType == EConstraintSolverType::StandardPbd)
@@ -1260,19 +1242,6 @@ namespace Chaos
 						Solver.Body1().UpdateRotationDependentState();
 					}
 				}
-
-				if (!Solver.GetIsActive() && bChaos_Joint_EarlyOut_Enabled)
-				{
-					break;
-				}
-			}
-
-			// If the solver we are running uses a velocity solver in the Apply step (for other constraint types), 
-			// we need to update the implicit velocity because we did a position solve above. E.g., GbfPbd collisions.
-			if (SolverType == EConstraintSolverType::GbfPbd)
-			{
-				Solver.Body(0).SetImplicitVelocity(Dt);
-				Solver.Body(1).SetImplicitVelocity(Dt);
 			}
 
 			// @todo(ccaulfield): The break limit should really be applied to the impulse in the solver to prevent 1-frame impulses larger than the threshold
@@ -1280,139 +1249,165 @@ namespace Chaos
 			{
 				ApplyBreakThreshold(Dt, ConstraintIndex, Solver.GetNetLinearImpulse(), Solver.GetNetAngularImpulse());
 			}
-
-			return Solver.GetIsActive() || !bChaos_Joint_EarlyOut_Enabled;
 		}
+
+		return true;
 	}
 
-	bool FPBDJointConstraints::ApplyPhase2Single(const FReal Dt, const int32 ConstraintIndex, const int32 NumPairIts, const int32 It, const int32 NumIts)
+	// QuasiPBD applies a velocity solve in phase 2
+	// Standard PBD does nothing
+	bool FPBDJointConstraints::ApplyPhase2Single(const FReal Dt, const int32 ConstraintIndex, const int32 It, const int32 NumIts)
 	{
 		if (!CanEvaluate(ConstraintIndex))
 		{
 			return false;
 		}
 
+		if (SolverType == EConstraintSolverType::StandardPbd)
+		{
+			return false;
+		}
+
 		const TVector<TGeometryParticleHandle<FReal, 3>*, 2>& Constraint = ConstraintParticles[ConstraintIndex];
-		UE_LOG(LogChaosJoint, VeryVerbose, TEXT("Project Joint Constraint %d %s %s (dt = %f; it = %d / %d)"), ConstraintIndex, *Constraint[0]->ToString(), *Constraint[1]->ToString(), Dt, It, NumIts);
+		UE_LOG(LogChaosJoint, VeryVerbose, TEXT("Solve Joint Velocity Constraint %d %s %s (dt = %f; it = %d / %d)"), ConstraintIndex, *Constraint[0]->ToString(), *Constraint[1]->ToString(), Dt, It, NumIts);
+
+		// @todo(chaos): store this on the Solver object and don't access the particles here
+		int32 Index0, Index1;
+		GetConstrainedParticleIndices(ConstraintIndex, Index0, Index1);
+		FGenericParticleHandle Particle0 = FGenericParticleHandle(ConstraintParticles[ConstraintIndex][Index0]);
+		FGenericParticleHandle Particle1 = FGenericParticleHandle(ConstraintParticles[ConstraintIndex][Index1]);
+		if ((Particle0->Sleeping() && Particle1->Sleeping())
+			|| (Particle0->IsKinematic() && Particle1->Sleeping())
+			|| (Particle0->Sleeping() && Particle1->IsKinematic()))
+		{
+			return false;
+		}
 
 		const FPBDJointSettings& JointSettings = ConstraintSettings[ConstraintIndex];
 		if (Settings.bUseLinearSolver)
 		{
 			FPBDJointCachedSolver& Solver = CachedConstraintSolvers[ConstraintIndex];
-
-			// @todo(chaos): store this on the Solver object and don't access the particles here
-			int32 Index0, Index1;
-			GetConstrainedParticleIndices(ConstraintIndex, Index0, Index1);
-			FGenericParticleHandle Particle0 = FGenericParticleHandle(ConstraintParticles[ConstraintIndex][Index0]);
-			FGenericParticleHandle Particle1 = FGenericParticleHandle(ConstraintParticles[ConstraintIndex][Index1]);
-			if ((Particle0->Sleeping() && Particle1->Sleeping())
-				|| (Particle0->IsKinematic() && Particle1->Sleeping())
-				|| (Particle0->Sleeping() && Particle1->IsKinematic()) 
-				|| (FMath::IsNearlyZero(Solver.InvM(0)) && FMath::IsNearlyZero(Solver.InvM(1))))
+			if ((FMath::IsNearlyZero(Solver.InvM(0)) && FMath::IsNearlyZero(Solver.InvM(1))))
 			{
 				return false;
 			}
 
-			// If we were solved last iteration and nothing has changed since, we are done
-			const bool bWasActive = Solver.GetIsActive();
 			Solver.Update(Dt, Settings, JointSettings);
 			
 			// Set parent inverse mass scale based on current shock propagation state
 			const FReal ShockPropagationInvMassScale = CalculateShockPropagationInvMassScale(Solver.Body0(), Solver.Body1(), JointSettings, It, NumIts);
 			Solver.SetInvMassScales(ShockPropagationInvMassScale, FReal(1), Dt);
 
-			// For quasipbd, we may still need to stabilize velocities and solve for restitution.
-			if (!bWasActive && !Solver.GetIsActive() && SolverType != EConstraintSolverType::QuasiPbd && bChaos_Joint_EarlyOut_Enabled)
+			const FReal IterationStiffness = CalculateIterationStiffness(It, NumIts);
+			Solver.ApplyVelocityConstraints(Dt, IterationStiffness, Settings, JointSettings);
+
+			// @todo(ccaulfield): should probably add to net impulses in push out too...(for breaking etc)
+		}
+		else
+		{
+			FPBDJointSolver& Solver = ConstraintSolvers[ConstraintIndex];
+			if ((FMath::IsNearlyZero(Solver.InvM(0)) && FMath::IsNearlyZero(Solver.InvM(1))))
 			{
 				return false;
 			}
 
+			Solver.Update(Dt, Settings, JointSettings);
+
+			// Set parent inverse mass scale based on current shock propagation state
+			const FReal ShockPropagationInvMassScale = CalculateShockPropagationInvMassScale(Solver.Body0(), Solver.Body1(), JointSettings, It, NumIts);
+			Solver.SetInvMassScales(ShockPropagationInvMassScale, FReal(1));
+
 			const FReal IterationStiffness = CalculateIterationStiffness(It, NumIts);
-			for (int32 PairIt = 0; PairIt < NumPairIts; ++PairIt)
-			{
-				switch(SolverType)
-				{
-				case EConstraintSolverType::None:
-					break;
-				case EConstraintSolverType::GbfPbd:
-				case EConstraintSolverType::StandardPbd:
-					Solver.ApplyProjections(Dt, IterationStiffness, Settings, JointSettings);
-					break;
-				case EConstraintSolverType::QuasiPbd:
-					Solver.ApplyVelocityConstraints(Dt, IterationStiffness, Settings, JointSettings);
-					break;
-				}
+			Solver.ApplyVelocityConstraints(Dt, IterationStiffness, Settings, JointSettings);
 
-				if (!Solver.GetIsActive() && bChaos_Joint_EarlyOut_Enabled)
-				{
-					break;
-				}
+			// @todo(ccaulfield): should probably add to net impulses in push out too...(for breaking etc)
+		}
+
+		return true;
+	}
+
+	// Projection phase
+	bool FPBDJointConstraints::ApplyPhase3Single(const FReal Dt, const int32 ConstraintIndex, const int32 It, const int32 NumIts)
+	{
+		if (!CanEvaluate(ConstraintIndex))
+		{
+			return false;
+		}
+
+		const TVector<TGeometryParticleHandle<FReal, 3>*, 2>& Constraint = ConstraintParticles[ConstraintIndex];
+		UE_LOG(LogChaosJoint, VeryVerbose, TEXT("Project Joint Constraint %d %s %s (dt = %f; it = %d / %d)"), ConstraintIndex, *Constraint[0]->ToString(), *Constraint[1]->ToString(), Dt, It, NumIts);
+
+		const FPBDJointSettings& JointSettings = ConstraintSettings[ConstraintIndex];
+		if (!JointSettings.bProjectionEnabled && !JointSettings.bSoftProjectionEnabled)
+		{
+			return false;
+		}
+
+		// @todo(chaos): store this on the Solver object and don't access the particles here
+		int32 Index0, Index1;
+		GetConstrainedParticleIndices(ConstraintIndex, Index0, Index1);
+		FGenericParticleHandle Particle0 = FGenericParticleHandle(ConstraintParticles[ConstraintIndex][Index0]);
+		FGenericParticleHandle Particle1 = FGenericParticleHandle(ConstraintParticles[ConstraintIndex][Index1]);
+		if ((Particle0->Sleeping() && Particle1->Sleeping())
+			|| (Particle0->IsKinematic() && Particle1->Sleeping())
+			|| (Particle0->Sleeping() && Particle1->IsKinematic()))
+		{
+			return false;
+		}
+
+		if (Settings.bUseLinearSolver)
+		{
+			FPBDJointCachedSolver& Solver = CachedConstraintSolvers[ConstraintIndex];
+			if ((FMath::IsNearlyZero(Solver.InvM(0)) && FMath::IsNearlyZero(Solver.InvM(1))))
+			{
+				return false;
 			}
 
-			// @todo(ccaulfield): should probably add to net impulses in push out too...(for breaking etc)
+			// @todo(chaos): rebuild the joint state at the latest (post-solve) transforms
+			//Solver.Update(Dt, Settings, JointSettings);
+			//Solver.UpdateMasses(FReal(0), FReal(1));
 
-			return Solver.GetIsActive() || !bChaos_Joint_EarlyOut_Enabled;
+			const FReal IterationStiffness = CalculateIterationStiffness(It, NumIts);
+			Solver.ApplyProjections(Dt, IterationStiffness, Settings, JointSettings);
 		}
 		else
 		{
 			FPBDJointSolver& Solver = ConstraintSolvers[ConstraintIndex];
-
-			// @todo(chaos): store this on the Solver object and don't access the particles here
-			int32 Index0, Index1;
-			GetConstrainedParticleIndices(ConstraintIndex, Index0, Index1);
-			FGenericParticleHandle Particle0 = FGenericParticleHandle(ConstraintParticles[ConstraintIndex][Index0]);
-			FGenericParticleHandle Particle1 = FGenericParticleHandle(ConstraintParticles[ConstraintIndex][Index1]);
-			if ((Particle0->Sleeping() && Particle1->Sleeping())
-				|| (Particle0->IsKinematic() && Particle1->Sleeping())
-				|| (Particle0->Sleeping() && Particle1->IsKinematic()) 
-				|| (FMath::IsNearlyZero(Solver.InvM(0)) && FMath::IsNearlyZero(Solver.InvM(1))))
+			if ((FMath::IsNearlyZero(Solver.InvM(0)) && FMath::IsNearlyZero(Solver.InvM(1))))
 			{
 				return false;
 			}
 
-			// If we were solved last iteration and nothing has changed since, we are done
-			const bool bWasActive = Solver.GetIsActive();
 			Solver.Update(Dt, Settings, JointSettings);
 
-			// For quasipbd, we may still need to stabilize velocities and solve for restitution.
-			if (!bWasActive && !Solver.GetIsActive() && SolverType != EConstraintSolverType::QuasiPbd && bChaos_Joint_EarlyOut_Enabled)
+			if ((SolverType == EConstraintSolverType::StandardPbd) || (It == 0))
 			{
-				return false;
+				// @todo(chaos): support reverse parent/child
+				Solver.UpdateMasses(FReal(0), FReal(1));
 			}
 
-			// Set parent inverse mass scale based on current shock propagation state
-			const FReal ShockPropagationInvMassScale = CalculateShockPropagationInvMassScale(Solver.Body0(), Solver.Body1(), JointSettings, It, NumIts);
-
 			const FReal IterationStiffness = CalculateIterationStiffness(It, NumIts);
-			for (int32 PairIt = 0; PairIt < NumPairIts; ++PairIt)
+			Solver.ApplyProjections(Dt, IterationStiffness, Settings, JointSettings);
+
+			if (SolverType == EConstraintSolverType::StandardPbd)
 			{
-				switch(SolverType)
+				if (Solver.Body0().IsDynamic())
 				{
-				case EConstraintSolverType::None:
-					break;
-				case EConstraintSolverType::GbfPbd:
-				case EConstraintSolverType::StandardPbd:
-					Solver.UpdateMasses(ShockPropagationInvMassScale, FReal(1));
-					Solver.ApplyProjections(Dt, IterationStiffness, Settings, JointSettings);
-					break;
-				case EConstraintSolverType::QuasiPbd:
-					Solver.SetInvMassScales(ShockPropagationInvMassScale, FReal(1));
-					Solver.ApplyVelocityConstraints(Dt, IterationStiffness, Settings, JointSettings);
-					break;
+					Solver.Body0().SolverBody().ApplyCorrections();
+					Solver.Body0().UpdateRotationDependentState();
 				}
-
-				if (!Solver.GetIsActive() && bChaos_Joint_EarlyOut_Enabled)
+				if (Solver.Body1().IsDynamic())
 				{
-					break;
+					Solver.Body1().SolverBody().ApplyCorrections();
+					Solver.Body1().UpdateRotationDependentState();
 				}
 			}
-
-			// @todo(ccaulfield): should probably add to net impulses in push out too...(for breaking etc)
-
-			return Solver.GetIsActive() || !bChaos_Joint_EarlyOut_Enabled;
 		}
+
+		return true;
 	}
 
+
 	void FPBDJointConstraints::ApplyBreakThreshold(const FReal Dt, int32 ConstraintIndex, const FVec3& LinearImpulse, const FVec3& AngularImpulse)
 	{
 		const FPBDJointSettings& JointSettings = ConstraintSettings[ConstraintIndex];
diff --git a/Engine/Source/Runtime/Experimental/Chaos/Private/PBDRigidsSolver.cpp b/Engine/Source/Runtime/Experimental/Chaos/Private/PBDRigidsSolver.cpp
--- a/Engine/Source/Runtime/Experimental/Chaos/Private/PBDRigidsSolver.cpp
+++ b/Engine/Source/Runtime/Experimental/Chaos/Private/PBDRigidsSolver.cpp
@@ -186,6 +186,15 @@ namespace Chaos
 		// Iteration count cvars
 		// These override the engine config if >= 0
 
+		int32 ChaosSolverCollisionPositionFrictionIterations = -1;
+		int32 ChaosSolverCollisionVelocityFrictionIterations = -1;
+		int32 ChaosSolverCollisionPositionShockPropagationIterations = -1;
+		int32 ChaosSolverCollisionVelocityShockPropagationIterations = -1;
+		FAutoConsoleVariableRef CVarChaosSolverCollisionPositionFrictionIterations(TEXT("p.Chaos.Solver.Collision.PositionFrictionIterations"), ChaosSolverCollisionPositionFrictionIterations, TEXT("Override number of position iterations where friction is applied (if >= 0)"));
+		FAutoConsoleVariableRef CVarChaosSolverCollisionVelocityFrictionIterations(TEXT("p.Chaos.Solver.Collision.VelocityFrictionIterations"), ChaosSolverCollisionVelocityFrictionIterations, TEXT("Override number of velocity iterations where friction is applied (if >= 0)"));
+		FAutoConsoleVariableRef CVarChaosSolverCollisionPositionShockPropagationIterations(TEXT("p.Chaos.Solver.Collision.PositionShockPropagationIterations"), ChaosSolverCollisionPositionShockPropagationIterations, TEXT("Override number of position iterations where shock propagation is applied (if >= 0)"));
+		FAutoConsoleVariableRef CVarChaosSolverCollisionVelocityShockPropagationIterations(TEXT("p.Chaos.Solver.Collision.VelocityShockPropagationIterations"), ChaosSolverCollisionVelocityShockPropagationIterations, TEXT("Override number of velocity iterations where shock propagation is applied (if >= 0)"));
+
 		int32 ChaosSolverIterations = -1;
 		FAutoConsoleVariableRef CVarChaosSolverIterations(TEXT("p.Chaos.Solver.Iterations"), ChaosSolverIterations, TEXT("Override number of solver iterations (-1 to use config)"));
 
@@ -832,6 +841,22 @@ namespace Chaos
 			GetEvolution()->GetCollisionConstraints().SetSolverType((EConstraintSolverType)ChaosSolver_SolverType);
 			GetJointConstraints().SetSolverType((EConstraintSolverType)ChaosSolver_SolverType);
 
+			if (ChaosSolverCollisionPositionFrictionIterations >= 0)
+			{
+				MEvolution->GetCollisionConstraints().SetPositionFrictionIterations(ChaosSolverCollisionPositionFrictionIterations);
+			}
+			if (ChaosSolverCollisionVelocityFrictionIterations >= 0)
+			{
+				MEvolution->GetCollisionConstraints().SetVelocityFrictionIterations(ChaosSolverCollisionVelocityFrictionIterations);
+			}
+			if (ChaosSolverCollisionPositionShockPropagationIterations >= 0)
+			{
+				MEvolution->GetCollisionConstraints().SetPositionShockPropagationIterations(ChaosSolverCollisionPositionShockPropagationIterations);
+			}
+			if (ChaosSolverCollisionVelocityShockPropagationIterations >= 0)
+			{
+				MEvolution->GetCollisionConstraints().SetVelocityShockPropagationIterations(ChaosSolverCollisionVelocityShockPropagationIterations);
+			}
 			if (ChaosSolverIterations >= 0)
 			{
 				SetIterations(ChaosSolverIterations);
diff --git a/Engine/Source/Runtime/Experimental/Chaos/Public/Chaos/Collision/PBDCollisionConstraintHandle.h b/Engine/Source/Runtime/Experimental/Chaos/Public/Chaos/Collision/PBDCollisionConstraintHandle.h
--- a/Engine/Source/Runtime/Experimental/Chaos/Public/Chaos/Collision/PBDCollisionConstraintHandle.h
+++ b/Engine/Source/Runtime/Experimental/Chaos/Public/Chaos/Collision/PBDCollisionConstraintHandle.h
@@ -88,7 +88,7 @@ namespace Chaos
 		TVector<TGeometryParticleHandle<FReal, 3>*, 2> GetConstrainedParticles();
 
 		void PreGatherInput(FPBDIslandSolverData& SolverData);
-		void GatherInput(FReal Dt, const int32 Particle0Level, const int32 Particle1Level, FPBDIslandSolverData& SolverData, const bool bNeedsPreGather = true);
+		void GatherInput(FReal Dt, const int32 Particle0Level, const int32 Particle1Level, FPBDIslandSolverData& SolverData);
 
 		FSolverBody* GetSolverBody0();
 		FSolverBody* GetSolverBody1();
diff --git a/Engine/Source/Runtime/Experimental/Chaos/Public/Chaos/Collision/SolverCollisionContainer.h b/Engine/Source/Runtime/Experimental/Chaos/Public/Chaos/Collision/SolverCollisionContainer.h
--- a/Engine/Source/Runtime/Experimental/Chaos/Public/Chaos/Collision/SolverCollisionContainer.h
+++ b/Engine/Source/Runtime/Experimental/Chaos/Public/Chaos/Collision/SolverCollisionContainer.h
@@ -20,6 +20,33 @@ namespace Chaos
 	class FPBDCollisionSolverManifoldPoint;
 	class FSolverBodyContainer;
 
+	/**
+	 * @brief Settings to control the low-level collision solver behaviour
+	*/
+	class FPBDCollisionSolverSettings
+	{
+	public:
+		FPBDCollisionSolverSettings();
+
+		// Maximum speed at which two objects can depenetrate (actually, how much relative velocity can be added
+		// to a contact per frame when depentrating. Stacks and deep penetrations can lead to larger velocities)
+		FReal MaxPushOutVelocity;
+
+		// How many of the position iterations should run static/dynamic friction
+		int32 NumPositionFrictionIterations;
+
+		// How many of the velocity iterations should run dynamic friction
+		// @todo(chaos): if NumVelocityFrictionIterations > 1, then dynamic friction in the velocity phase will be iteration 
+		// count dependent (velocity-solve friction is currentlyused by quadratic shapes and RBAN)
+		int32 NumVelocityFrictionIterations;
+
+		// How many position iterations should have shock propagation enabled
+		int32 NumPositionShockPropagationIterations;
+
+		// How many velocity iterations should have shock propagation enabled
+		int32 NumVelocityShockPropagationIterations;
+	};
+
 	/**
 	 * @brief A container of low-level data used to solve collision constraints
 	*/
@@ -31,8 +58,6 @@ namespace Chaos
 
 		int32 NumSolvers() const { return CollisionSolvers.Num(); }
 
-		void SetMaxPushOutVelocity(const FReal InMaxPushOutVelocity) { MaxPushOutVelocity = InMaxPushOutVelocity; }
-
 		virtual void Reset(const int32 InMaxCollisions) override;
 
 		/** Set the number of collision solvers in the container
@@ -42,28 +67,27 @@ namespace Chaos
 
 		// Add a solver constraint for the specified constraint and gather the required solver data
 		void PreAddConstraintSolver(FPBDCollisionConstraint& Constraint, FSolverBodyContainer& SolverBodyContainer, int32& ConstraintIndex);
-		void AddConstraintSolver(FReal Dt, FPBDCollisionConstraint& Constraint, const int32 Particle0Level, const int32 Particle1Level, FSolverBodyContainer& SolverBodyContainer);
+		void AddConstraintSolver(FReal Dt, FPBDCollisionConstraint& Constraint, const int32 Particle0Level, const int32 Particle1Level, FSolverBodyContainer& SolverBodyContainer, const FPBDCollisionSolverSettings& SolverSettings);
 
-		bool SolvePositionSerial(const FReal Dt, const int32 It, const int32 NumIts, const int32 BeginIndex, const int32 EndIndex);
-		bool SolvePositionParallel(const FReal Dt, const int32 It, const int32 NumIts, const int32 BeginIndex, const int32 EndIndex);
+		bool SolvePositionSerial(const FReal Dt, const int32 It, const int32 NumIts, const int32 BeginIndex, const int32 EndIndex, const FPBDCollisionSolverSettings& SolverSettings);
+		bool SolvePositionParallel(const FReal Dt, const int32 It, const int32 NumIts, const int32 BeginIndex, const int32 EndIndex, const FPBDCollisionSolverSettings& SolverSettings);
 
-		bool SolveVelocitySerial(const FReal Dt, const int32 It, const int32 NumIts, const int32 BeginIndex, const int32 EndIndex);
-		bool SolveVelocityParallel(const FReal Dt, const int32 It, const int32 NumIts, const int32 BeginIndex, const int32 EndIndex);
+		bool SolveVelocitySerial(const FReal Dt, const int32 It, const int32 NumIts, const int32 BeginIndex, const int32 EndIndex, const FPBDCollisionSolverSettings& SolverSettings);
+		bool SolveVelocityParallel(const FReal Dt, const int32 It, const int32 NumIts, const int32 BeginIndex, const int32 EndIndex, const FPBDCollisionSolverSettings& SolverSettings);
 
 		void ScatterOutput(const FReal Dt, const int32 BeginIndex, const int32 EndIndex);
 
 	private:
-		void UpdatePositionShockPropagation(const FReal Dt, const int32 It, const int32 NumIts, const int32 BeginIndex, const int32 EndIndex);
-		void UpdateVelocityShockPropagation(const FReal Dt, const int32 It, const int32 NumIts, const int32 BeginIndex, const int32 EndIndex);
-		bool SolvePositionImpl(const FReal Dt, const int32 It, const int32 NumIts, const int32 BeginIndex, const int32 EndIndex, const bool bParallel);
+		void UpdatePositionShockPropagation(const FReal Dt, const int32 It, const int32 NumIts, const int32 BeginIndex, const int32 EndIndex, const FPBDCollisionSolverSettings& SolverSettings);
+		void UpdateVelocityShockPropagation(const FReal Dt, const int32 It, const int32 NumIts, const int32 BeginIndex, const int32 EndIndex, const FPBDCollisionSolverSettings& SolverSettings);
+		bool SolvePositionImpl(const FReal Dt, const int32 It, const int32 NumIts, const int32 BeginIndex, const int32 EndIndex, const FPBDCollisionSolverSettings& SolverSettings, const bool bParallel);
 		bool SolvePositionIncrementalImpl(const FReal Dt, const int32 BeginIndex, const int32 EndIndex, const FReal MaxPushOut, const bool bApplyStaticFriction);
 		bool SolvePositionWithFrictionImpl(const FReal Dt, const int32 BeginIndex, const int32 EndIndex, const FReal MaxPushOut);
 		bool SolvePositionNoFrictionImpl(const FReal Dt, const int32 BeginIndex, const int32 EndIndex, const FReal MaxPushOut);
-		bool SolveVelocityImpl(const FReal Dt, const int32 It, const int32 NumIts, const int32 BeginIndex, const int32 EndIndex, const bool bParallel);
+		bool SolveVelocityImpl(const FReal Dt, const int32 It, const int32 NumIts, const int32 BeginIndex, const int32 EndIndex, const FPBDCollisionSolverSettings& SolverSettings, const bool bParallel);
 		void ScatterOutputImpl(const FReal Dt, const int32 BeginIndex, const int32 EndIndex, const bool bParallel);
 
 		TArray<FPBDCollisionSolverAdapter> CollisionSolvers;
-		FReal MaxPushOutVelocity;
 		bool bRequiresIncrementalCollisionDetection;
 	};
 }
diff --git a/Engine/Source/Runtime/Experimental/Chaos/Public/Chaos/Evolution/PBDMinEvolution.h b/Engine/Source/Runtime/Experimental/Chaos/Public/Chaos/Evolution/PBDMinEvolution.h
--- a/Engine/Source/Runtime/Experimental/Chaos/Public/Chaos/Evolution/PBDMinEvolution.h
+++ b/Engine/Source/Runtime/Experimental/Chaos/Public/Chaos/Evolution/PBDMinEvolution.h
@@ -45,11 +45,28 @@ namespace Chaos
 			SolverType = InSolverType;
 		}
 
+		void SetNumPositionIterations(const int32 NumIts)
+		{
+			NumPositionIterations = NumIts;
+		}
+
+		void SetNumVelocityIterations(const int32 NumIts)
+		{
+			NumVelocityIterations = NumIts;
+		}
+
+		void SetNumProjectionIterations(const int32 NumIts)
+		{
+			NumProjectionIterations = NumIts;
+		}
+
+		// Legacy
 		void SetNumIterations(const int32 NumIts)
 		{
 			NumApplyIterations = NumIts;
 		}
 
+		// Legacy
 		void SetNumPushOutIterations(const int32 NumIts)
 		{
 			NumApplyPushOutIterations = NumIts;
@@ -117,7 +134,8 @@ namespace Chaos
 		void ApplyConstraintsPhase1(FReal Dt);
 		void UpdateVelocities(FReal Dt);
 		void ApplyConstraintsPhase2(FReal Dt);
-		void UpdatePositions(FReal Dt);
+		void ApplyCorrections(FReal Dt);
+		void ApplyConstraintsPhase3(FReal Dt);
 
 		FRigidParticleSOAs& Particles;
 		FCollisionDetector& CollisionDetector;
@@ -132,6 +150,9 @@ namespace Chaos
 		EConstraintSolverType SolverType;
 		int32 NumApplyIterations;
 		int32 NumApplyPushOutIterations;
+		int32 NumPositionIterations;
+		int32 NumVelocityIterations;
+		int32 NumProjectionIterations;
 		FReal BoundsExtension;
 		FVec3 Gravity;
 		FSimulationSpaceSettings SimulationSpaceSettings;
diff --git a/Engine/Source/Runtime/Experimental/Chaos/Public/Chaos/JointProperties.inl b/Engine/Source/Runtime/Experimental/Chaos/Public/Chaos/JointProperties.inl
--- a/Engine/Source/Runtime/Experimental/Chaos/Public/Chaos/JointProperties.inl
+++ b/Engine/Source/Runtime/Experimental/Chaos/Public/Chaos/JointProperties.inl
@@ -4,6 +4,7 @@
 CHAOS_INNER_JOINT_PROPERTY(JointSettings, JointTransforms, ConnectorTransforms, FTransformPair)
 CHAOS_INNER_JOINT_PROPERTY(JointSettings, CollisionEnabled, bCollisionEnabled, bool)
 CHAOS_INNER_JOINT_PROPERTY(JointSettings, ProjectionEnabled, bProjectionEnabled, bool)
+CHAOS_INNER_JOINT_PROPERTY(JointSettings, ShockPropagationEnabled, bShockPropagationEnabled, bool)
 CHAOS_INNER_JOINT_PROPERTY(JointSettings, ProjectionLinearAlpha, LinearProjection, FReal)
 CHAOS_INNER_JOINT_PROPERTY(JointSettings, ProjectionAngularAlpha, AngularProjection, FReal)
 CHAOS_INNER_JOINT_PROPERTY(JointSettings, ShockPropagationAlpha, ShockPropagation, FReal)
diff --git a/Engine/Source/Runtime/Experimental/Chaos/Public/Chaos/PBDCollisionConstraints.h b/Engine/Source/Runtime/Experimental/Chaos/Public/Chaos/PBDCollisionConstraints.h
--- a/Engine/Source/Runtime/Experimental/Chaos/Public/Chaos/PBDCollisionConstraints.h
+++ b/Engine/Source/Runtime/Experimental/Chaos/Public/Chaos/PBDCollisionConstraints.h
@@ -159,6 +159,7 @@ public:
 
 	bool ApplyPhase1(const FReal Dt, const int32 It, const int32 NumIts, FPBDIslandSolverData& SolverData);
 	bool ApplyPhase2(const FReal Dt, const int32 It, const int32 NumIts, FPBDIslandSolverData& SolverData);
+	bool ApplyPhase3(const FReal Dt, const int32 It, const int32 NumIts, FPBDIslandSolverData& SolverData) { return false; }
 
 	//
 	// Island Rule API
@@ -265,7 +266,27 @@ public:
 
 	void SetMaxPushOutVelocity(const FReal InMaxPushOutVelocity)
 	{
-		MaxPushOutVelocity = InMaxPushOutVelocity;
+		SolverSettings.MaxPushOutVelocity = InMaxPushOutVelocity;
+	}
+
+	void SetPositionFrictionIterations(const int32 InNumIterations)
+	{
+		SolverSettings.NumPositionFrictionIterations = InNumIterations;
+	}
+
+	void SetVelocityFrictionIterations(const int32 InNumIterations)
+	{
+		SolverSettings.NumVelocityFrictionIterations = InNumIterations;
+	}
+
+	void SetPositionShockPropagationIterations(const int32 InNumIterations)
+	{
+		SolverSettings.NumPositionShockPropagationIterations = InNumIterations;
+	}
+
+	void SetVelocityShockPropagationIterations(const int32 InNumIterations)
+	{
+		SolverSettings.NumVelocityShockPropagationIterations = InNumIterations;
 	}
 
 	int32 NumConstraints() const
@@ -338,7 +359,8 @@ private:
 	FVec3 GravityDirection;
 	FReal GravitySize;
 
-	FReal MaxPushOutVelocity;
+	// Settings for the low-level collision solvers
+	FPBDCollisionSolverSettings SolverSettings;
 
 	EConstraintSolverType SolverType;
 };
diff --git a/Engine/Source/Runtime/Experimental/Chaos/Public/Chaos/PBDConstraintRule.h b/Engine/Source/Runtime/Experimental/Chaos/Public/Chaos/PBDConstraintRule.h
--- a/Engine/Source/Runtime/Experimental/Chaos/Public/Chaos/PBDConstraintRule.h
+++ b/Engine/Source/Runtime/Experimental/Chaos/Public/Chaos/PBDConstraintRule.h
@@ -62,6 +62,7 @@ namespace Chaos
 		virtual void ScatterSolverOutput(const FReal Dt) {}
 		virtual bool ApplyConstraints(const FReal Dt, const int32 It, const int32 NumIts) { return false; }
 		virtual bool ApplyPushOut(const FReal Dt, const int32 It, const int32 NumIts) { return false; }
+		virtual bool ApplyProjection(const FReal Dt, const int32 It, const int32 NumIts) { return false; }
 	protected:
 		/** Solver datas that are coming from the evolution */
 		FPBDIslandSolverData* SolverData = nullptr;
@@ -94,6 +95,8 @@ namespace Chaos
 
 		virtual bool ApplyPushOut(const FReal Dt, const int32 It, const int32 NumIts) override;
 
+		virtual bool ApplyProjection(const FReal Dt, const int32 It, const int32 NumIts) override;
+
 	private:
 		FConstraints& Constraints;
 	};
diff --git a/Engine/Source/Runtime/Experimental/Chaos/Public/Chaos/PBDJointConstraintTypes.h b/Engine/Source/Runtime/Experimental/Chaos/Public/Chaos/PBDJointConstraintTypes.h
--- a/Engine/Source/Runtime/Experimental/Chaos/Public/Chaos/PBDJointConstraintTypes.h
+++ b/Engine/Source/Runtime/Experimental/Chaos/Public/Chaos/PBDJointConstraintTypes.h
@@ -129,13 +129,14 @@ namespace Chaos
 		FTransformPair ConnectorTransforms;
 
 		FReal Stiffness;
-		FReal LinearProjection;			// @chaos(todo): remove - old PBD solver only
-		FReal AngularProjection;		// @chaos(todo): remove - old PBD solver only
+		FReal LinearProjection;
+		FReal AngularProjection;
 		FReal ShockPropagation;
 		FReal ParentInvMassScale;
 
 		bool bCollisionEnabled;
-		bool bProjectionEnabled;		// @chaos(todo): remove - old PBD solver only
+		bool bProjectionEnabled;		// @chaos(todo): remove - implied by alpha
+		bool bShockPropagationEnabled;	// @chaos(todo): remove - implied by alpha
 		bool bSoftProjectionEnabled;	// @chaos(todo): remove - old PBD solver only
 
 		TVector<EJointMotionType, 3> LinearMotionTypes;
diff --git a/Engine/Source/Runtime/Experimental/Chaos/Public/Chaos/PBDJointConstraints.h b/Engine/Source/Runtime/Experimental/Chaos/Public/Chaos/PBDJointConstraints.h
--- a/Engine/Source/Runtime/Experimental/Chaos/Public/Chaos/PBDJointConstraints.h
+++ b/Engine/Source/Runtime/Experimental/Chaos/Public/Chaos/PBDJointConstraints.h
@@ -268,6 +268,7 @@ namespace Chaos
 
 		bool ApplyPhase1(const FReal Dt, const int32 It, const int32 NumIts, FPBDIslandSolverData& SolverData);
 		bool ApplyPhase2(const FReal Dt, const int32 It, const int32 NumIts, FPBDIslandSolverData& SolverData);
+		bool ApplyPhase3(const FReal Dt, const int32 It, const int32 NumIts, FPBDIslandSolverData& SolverData);
 
 		//
 		// Island Rule API
@@ -277,6 +278,7 @@ namespace Chaos
 		void GatherInput(const FReal Dt, const int32 ConstraintIndex, const int32 Particle0Level, const int32 Particle1Level, FPBDIslandSolverData& SolverData);
 		bool ApplyPhase1Serial(const FReal Dt, const int32 It, const int32 NumIts, FPBDIslandSolverData& SolverData);
 		bool ApplyPhase2Serial(const FReal Dt, const int32 It, const int32 NumIts, FPBDIslandSolverData& SolverData);
+		bool ApplyPhase3Serial(const FReal Dt, const int32 It, const int32 NumIts, FPBDIslandSolverData& SolverData);
 
 		/**
 		 * Set the solver method to use
@@ -305,7 +307,8 @@ namespace Chaos
 		bool CanEvaluate(const int32 ConstraintIndex) const;
 
 		bool ApplyPhase1Single(const FReal Dt, const int32 ConstraintIndex, const int32 NumPairIts, const int32 It, const int32 NumIts);
-		bool ApplyPhase2Single(const FReal Dt, const int32 ConstraintIndex, const int32 NumPairIts, const int32 It, const int32 NumIts);
+		bool ApplyPhase2Single(const FReal Dt, const int32 ConstraintIndex, const int32 It, const int32 NumIts);
+		bool ApplyPhase3Single(const FReal Dt, const int32 ConstraintIndex, const int32 It, const int32 NumIts);
 		void ApplyBreakThreshold(const FReal Dt, int32 ConstraintIndex, const FVec3& LinearImpulse, const FVec3& AngularImpulse);
 		void ApplyPlasticityLimits(const int32 ConstraintIndex);
 
diff --git a/Engine/Source/Runtime/Experimental/Chaos/Public/Chaos/PBDRigidSpringConstraints.h b/Engine/Source/Runtime/Experimental/Chaos/Public/Chaos/PBDRigidSpringConstraints.h
--- a/Engine/Source/Runtime/Experimental/Chaos/Public/Chaos/PBDRigidSpringConstraints.h
+++ b/Engine/Source/Runtime/Experimental/Chaos/Public/Chaos/PBDRigidSpringConstraints.h
@@ -167,6 +167,7 @@ namespace Chaos
 
 		bool ApplyPhase1(const FReal Dt, const int32 It, const int32 NumIts, FPBDIslandSolverData& SolverData);
 		bool ApplyPhase2(const FReal Dt, const int32 It, const int32 NumIts, FPBDIslandSolverData& SolverData) { return false; }
+		bool ApplyPhase3(const FReal Dt, const int32 It, const int32 NumIts, FPBDIslandSolverData& SolverData) { return false; }
 
 		//
 		// Island Rule API
diff --git a/Engine/Source/Runtime/PhysicsCore/Private/ChaosEngineInterface.cpp b/Engine/Source/Runtime/PhysicsCore/Private/ChaosEngineInterface.cpp
--- a/Engine/Source/Runtime/PhysicsCore/Private/ChaosEngineInterface.cpp
+++ b/Engine/Source/Runtime/PhysicsCore/Private/ChaosEngineInterface.cpp
@@ -1289,12 +1289,13 @@ void FChaosEngineInterface::SetProjectionEnabled_AssumesLocked(const FPhysicsCon
 	}
 }
 
-void FChaosEngineInterface::SetShockPropagationAlpha_AssumesLocked(const FPhysicsConstraintHandle& InConstraintRef, float InShockPropagationAlpha)
+void FChaosEngineInterface::SetShockPropagationEnabled_AssumesLocked(const FPhysicsConstraintHandle& InConstraintRef, bool bInShockPropagationEnabled, float InShockPropagationAlpha)
 {
 	if (InConstraintRef.IsValid() && InConstraintRef.Constraint->IsType(Chaos::EConstraintType::JointConstraintType))
 	{
 		if (Chaos::FJointConstraint* Constraint = static_cast<Chaos::FJointConstraint*>(InConstraintRef.Constraint))
 		{
+			Constraint->SetShockPropagationEnabled(bInShockPropagationEnabled);
 			Constraint->SetShockPropagationAlpha(InShockPropagationAlpha);
 		}
 	}
diff --git a/Engine/Source/Runtime/PhysicsCore/Public/Chaos/ChaosEngineInterface.h b/Engine/Source/Runtime/PhysicsCore/Public/Chaos/ChaosEngineInterface.h
--- a/Engine/Source/Runtime/PhysicsCore/Public/Chaos/ChaosEngineInterface.h
+++ b/Engine/Source/Runtime/PhysicsCore/Public/Chaos/ChaosEngineInterface.h
@@ -498,7 +498,7 @@ public:
 	static void SetCanVisualize(const FPhysicsConstraintHandle& InConstraintRef,bool bInCanVisualize);
 	static void SetCollisionEnabled(const FPhysicsConstraintHandle& InConstraintRef,bool bInCollisionEnabled);
 	static void SetProjectionEnabled_AssumesLocked(const FPhysicsConstraintHandle& InConstraintRef,bool bInProjectionEnabled,float InLinearAlpha = 1.0f,float InAngularAlpha = 0.0f);
-	static void SetShockPropagationAlpha_AssumesLocked(const FPhysicsConstraintHandle& InConstraintRef, float InShockPropagationAlpha);
+	static void SetShockPropagationEnabled_AssumesLocked(const FPhysicsConstraintHandle& InConstraintRef, bool bInShockPropagationEnabled, float InShockPropagationAlpha);
 	static void SetParentDominates_AssumesLocked(const FPhysicsConstraintHandle& InConstraintRef,bool bInParentDominates);
 	static void SetBreakForces_AssumesLocked(const FPhysicsConstraintHandle& InConstraintRef,float InLinearBreakForce,float InAngularBreakForce);
 	static void SetPlasticityLimits_AssumesLocked(const FPhysicsConstraintHandle& InConstraintRef, float InLinearPlasticityLimit, float InAngularPlasticityLimit, EConstraintPlasticityType InLinearPlasticityType);
