UE-143850 - //UE5/Release-5.0 - Cook Windows Client - Can't connect pins

Build blocker fix.

19122243 removed the custom make/break functions for the quantize vectors, which should've been safe to do. However, those temporary functions actually permitted conversions between FVector and its quantize counterparts, which was previously illegal in UE4. Some BP content had already started to use this new conversion, so reverting the change would cause BP compilation errors.

For now, just re-add the custom make/break functions. Arguably, this should be allowed anyway since the quantize vectors inherit from FVector, and don't change its layout at all.

diff --git a/Engine/Source/Runtime/Engine/Classes/Engine/NetSerialization.h b/Engine/Source/Runtime/Engine/Classes/Engine/NetSerialization.h
--- a/Engine/Source/Runtime/Engine/Classes/Engine/NetSerialization.h
+++ b/Engine/Source/Runtime/Engine/Classes/Engine/NetSerialization.h
@@ -401,7 +401,7 @@ bool SerializeFixedVector(FVector3d &Vector, FArchive& Ar)
  *	Note: this is the historical UE format for vector net serialization
  *
  */
-USTRUCT()
+USTRUCT(meta = (HasNativeMake = "Engine.KismetMathLibrary.MakeVector_NetQuantize", HasNativeBreak = "Engine.KismetMathLibrary.BreakVector_NetQuantize"))
 struct FVector_NetQuantize : public FVector
 {
 	GENERATED_USTRUCT_BODY()
@@ -413,7 +413,7 @@ struct FVector_NetQuantize : public FVector
 	: FVector(E)
 	{}
 
-	FORCEINLINE FVector_NetQuantize( float InX, float InY, float InZ )
+	FORCEINLINE FVector_NetQuantize(double InX, double InY, double InZ)
 	: FVector(InX, InY, InZ)
 	{}
 
@@ -447,7 +447,7 @@ struct TStructOpsTypeTraits< FVector_NetQuantize > : public TStructOpsTypeTraits
  *	Valid range: 2^24 / 10 = +/- 1,677,721.6
  *
  */
-USTRUCT()
+USTRUCT(meta = (HasNativeMake = "Engine.KismetMathLibrary.MakeVector_NetQuantize10", HasNativeBreak = "Engine.KismetMathLibrary.BreakVector_NetQuantize10"))
 struct FVector_NetQuantize10 : public FVector
 {
 	GENERATED_USTRUCT_BODY()
@@ -459,7 +459,7 @@ struct FVector_NetQuantize10 : public FVector
 	: FVector(E)
 	{}
 
-	FORCEINLINE FVector_NetQuantize10( float InX, float InY, float InZ )
+	FORCEINLINE FVector_NetQuantize10(double InX, double InY, double InZ)
 	: FVector(InX, InY, InZ)
 	{}
 
@@ -493,7 +493,7 @@ struct TStructOpsTypeTraits< FVector_NetQuantize10 > : public TStructOpsTypeTrai
  *	Valid range: 2^30 / 100 = +/- 10,737,418.24
  *
  */
-USTRUCT()
+USTRUCT(meta = (HasNativeMake = "Engine.KismetMathLibrary.MakeVector_NetQuantize100", HasNativeBreak = "Engine.KismetMathLibrary.BreakVector_NetQuantize100"))
 struct FVector_NetQuantize100 : public FVector
 {
 	GENERATED_USTRUCT_BODY()
@@ -505,7 +505,7 @@ struct FVector_NetQuantize100 : public FVector
 	: FVector(E)
 	{}
 
-	FORCEINLINE FVector_NetQuantize100( float InX, float InY, float InZ )
+	FORCEINLINE FVector_NetQuantize100(double InX, double InY, double InZ)
 	: FVector(InX, InY, InZ)
 	{}
 
@@ -535,7 +535,7 @@ struct TStructOpsTypeTraits< FVector_NetQuantize100 > : public TStructOpsTypeTra
  *	16 bits per component
  *	Valid range: -1..+1 inclusive
  */
-USTRUCT()
+USTRUCT(meta = (HasNativeMake = "Engine.KismetMathLibrary.MakeVector_NetQuantizeNormal", HasNativeBreak = "Engine.KismetMathLibrary.BreakVector_NetQuantizeNormal"))
 struct FVector_NetQuantizeNormal : public FVector
 {
 	GENERATED_USTRUCT_BODY()
@@ -547,7 +547,7 @@ struct FVector_NetQuantizeNormal : public FVector
 	: FVector(E)
 	{}
 
-	FORCEINLINE FVector_NetQuantizeNormal( float InX, float InY, float InZ )
+	FORCEINLINE FVector_NetQuantizeNormal(double InX, double InY, double InZ)
 	: FVector(InX, InY, InZ)
 	{}
 
diff --git a/Engine/Source/Runtime/Engine/Classes/Kismet/KismetMathLibrary.h b/Engine/Source/Runtime/Engine/Classes/Kismet/KismetMathLibrary.h
--- a/Engine/Source/Runtime/Engine/Classes/Kismet/KismetMathLibrary.h
+++ b/Engine/Source/Runtime/Engine/Classes/Kismet/KismetMathLibrary.h
@@ -4277,6 +4277,24 @@ class ENGINE_API UKismetMathLibrary : public UBlueprintFunctionLibrary
 	UFUNCTION(BlueprintPure, Category="Math|Smoothing", meta=(DisplayName="Dynamic Weighted Moving Average Rotator"))
 	static FRotator DynamicWeightedMovingAverage_FRotator(FRotator CurrentSample, FRotator PreviousSample, float MaxDistance, float MinWeight, float MaxWeight);
 	
+	// NetQuantized vector make/breaks
+	UFUNCTION(BlueprintPure, Category = "Math|Vector", meta = (NativeMakeFunc))
+	static FVector_NetQuantize MakeVector_NetQuantize(double X, double Y, double Z) { return FVector_NetQuantize(X, Y, Z); }
+	UFUNCTION(BlueprintPure, Category = "Math|Vector", meta = (NativeMakeFunc))
+	static FVector_NetQuantize10 MakeVector_NetQuantize10(double X, double Y, double Z) { return FVector_NetQuantize10(X, Y, Z); }
+	UFUNCTION(BlueprintPure, Category = "Math|Vector", meta = (NativeMakeFunc))
+	static FVector_NetQuantize100 MakeVector_NetQuantize100(double X, double Y, double Z) { return FVector_NetQuantize100(X, Y, Z); }
+	UFUNCTION(BlueprintPure, Category = "Math|Vector", meta = (NativeMakeFunc))
+	static FVector_NetQuantizeNormal MakeVector_NetQuantizeNormal(double X, double Y, double Z) { return FVector_NetQuantizeNormal(X, Y, Z); }
+	UFUNCTION(BlueprintPure, Category = "Math|Vector", meta = (NativeBreakFunc))
+	static void BreakVector_NetQuantize(FVector_NetQuantize InVec, double& X, double& Y, double& Z) { BreakVector((FVector)InVec, X, Y, Z); }
+	UFUNCTION(BlueprintPure, Category = "Math|Vector", meta = (NativeBreakFunc))
+	static void BreakVector_NetQuantize10(FVector_NetQuantize10 InVec, double& X, double& Y, double& Z) { BreakVector((FVector)InVec, X, Y, Z); }
+	UFUNCTION(BlueprintPure, Category = "Math|Vector", meta = (NativeBreakFunc))
+	static void BreakVector_NetQuantize100(FVector_NetQuantize100 InVec, double& X, double& Y, double& Z) { BreakVector((FVector)InVec, X, Y, Z); }
+	UFUNCTION(BlueprintPure, Category = "Math|Vector", meta = (NativeBreakFunc))
+	static void BreakVector_NetQuantizeNormal(FVector_NetQuantizeNormal InVec, double& X, double& Y, double& Z) { BreakVector((FVector)InVec, X, Y, Z); }
+
 private:
 
 	static void ReportError_Divide_ByteByte();
