Material transform expression is supposed to accept a scalar input value when input is tangent space.  Change this to accept both MCT_Float and MCT_Float1, rather than just MCT_Float1.  Most places that generate scalar values use the MCT_Float type.

diff --git a/Engine/Source/Runtime/Engine/Private/Materials/HLSLMaterialTranslator.cpp b/Engine/Source/Runtime/Engine/Private/Materials/HLSLMaterialTranslator.cpp
--- a/Engine/Source/Runtime/Engine/Private/Materials/HLSLMaterialTranslator.cpp
+++ b/Engine/Source/Runtime/Engine/Private/Materials/HLSLMaterialTranslator.cpp
@@ -8415,11 +8415,12 @@ int32 FHLSLMaterialTranslator::TransformBase(EMaterialCommonBasis SourceCoordBas
 		
 		// Construct float3(0,0,x) out of the input if it is a scalar
 		// This way artists can plug in a scalar and it will be treated as height, or a vector displacement
-		if (SourceType == MCT_Float1 && SourceCoordBasis == MCB_Tangent)
+		const uint32 NumInputComponents = GetNumComponents(SourceType);
+		if (NumInputComponents == 1u && SourceCoordBasis == MCB_Tangent)
 		{
 			A = AppendVector(Constant2(0, 0), A);
 		}
-		else if (GetNumComponents(SourceType) < 3)
+		else if (NumInputComponents < 3u)
 		{
 			return Errorf(TEXT("input must be a vector (%s: %s) or a scalar (if source is Tangent)"), *GetParameterCode(A), DescribeType(SourceType));
 		}
