Fix access after deletion of chunk pointers.
ActiveChunks needs to be cleared on WaitFlushRequests_AnyThread() for this.
Unfortunately that means we need to widen the lock to include accessing it.
#rnx

diff --git a/Engine/Source/Runtime/Engine/Private/VT/VirtualTextureChunkDDCCache.cpp b/Engine/Source/Runtime/Engine/Private/VT/VirtualTextureChunkDDCCache.cpp
--- a/Engine/Source/Runtime/Engine/Private/VT/VirtualTextureChunkDDCCache.cpp
+++ b/Engine/Source/Runtime/Engine/Private/VT/VirtualTextureChunkDDCCache.cpp
@@ -232,34 +232,28 @@ void FVirtualTextureChunkDDCCache::Initialize()
 void FVirtualTextureChunkDDCCache::ShutDown()
 {
 	WaitFlushRequests_AnyThread();
-
-	ENQUEUE_RENDER_COMMAND(FVirtualTextureChunkDDCCache_Shutdown)(
-		[this](FRHICommandListImmediate& RHICmdList)
-		{
-			UpdateRequests();
-			ensure(ActiveChunks.Num() == 0);
-		});
 }
 
 void FVirtualTextureChunkDDCCache::UpdateRequests()
 {
-	// Remove completed chunks from array.
 	check(IsInRenderingThread());
+
+	// Needs a lock because these arrays are accessed by rare calls to WaitFlushRequests_AnyThread().
+	// All other access is render thread.
+	FScopeLock Lock(&ActiveTasksLock);
+
+	// Remove completed chunks from array.
 	ActiveChunks.RemoveAllSwap([](auto Chunk) -> bool {return Chunk->bFileAvailableInVTDDCDache == true; }, false);
 
 	// Remove completed tasks from array.
-	// ActiveTasks needs a lock because it is accessed by WaitFlushRequests_AnyThread()
+	for (int32 TaskIndex = 0; TaskIndex < ActiveTasks.Num(); ++TaskIndex)
 	{
-		FScopeLock Lock(&ActiveTasksLock);
-		for (int32 TaskIndex = 0; TaskIndex < ActiveTasks.Num(); ++TaskIndex)
+		FAsyncTask<FAsyncFillCacheWorker>* Task = ActiveTasks[TaskIndex];
+		if (Task->IsWorkDone())
 		{
-			FAsyncTask<FAsyncFillCacheWorker>* Task = ActiveTasks[TaskIndex];
-			if (Task->IsWorkDone())
-			{
-				Task->EnsureCompletion();
-				delete Task;
-				ActiveTasks.RemoveAtSwap(TaskIndex--, 1, false);
-			}
+			Task->EnsureCompletion();
+			delete Task;
+			ActiveTasks.RemoveAtSwap(TaskIndex--, 1, false);
 		}
 	}
 }
@@ -269,6 +263,7 @@ void FVirtualTextureChunkDDCCache::WaitFlushRequests_AnyThread()
 	TArray<FAsyncTask<FAsyncFillCacheWorker>*> ActiveTasksCopy;
 	{
 		FScopeLock Lock(&ActiveTasksLock);
+		ActiveChunks.Empty();
 		ActiveTasksCopy = MoveTemp(ActiveTasks);
 	}
 	for (int32 TaskIndex = 0; TaskIndex < ActiveTasksCopy.Num(); ++TaskIndex)
@@ -304,11 +299,14 @@ bool FVirtualTextureChunkDDCCache::MakeChunkAvailable(struct FVirtualTextureData
 		return true;
 	}
 
-	// Are we already processing this chunk ?
-	const int32 ChunkInProgressIdx = ActiveChunks.Find(Chunk);
-	if (ChunkInProgressIdx != -1)
 	{
-		return false;
+		// Are we already processing this chunk ?
+		FScopeLock Lock(&ActiveTasksLock);
+		const int32 ChunkInProgressIdx = ActiveChunks.Find(Chunk);
+		if (ChunkInProgressIdx != -1)
+		{
+			return false;
+		}
 	}
 	
 	// start filling it to the cache
@@ -319,9 +317,9 @@ bool FVirtualTextureChunkDDCCache::MakeChunkAvailable(struct FVirtualTextureData
 		{
 			FScopeLock Lock(&ActiveTasksLock);
 			ActiveTasks.Add(Task);
+			ActiveChunks.Add(Chunk);
 		}
 		
-		ActiveChunks.Add(Chunk);
 		Task->StartBackgroundTask();
 	}
 	else
