#ue5 - Kismet Math Library float->double conversion. Where a float function already had a matching double function (eg Add_FloatFloat / Add_DoubleDouble), the old one has been deprecated and redirected to the new one. Where functions only had a float version, the arguments and return types have been changed to double.

Not all functions have been modified to reduce scope, but the intent is to generally improve what is needed to support larger double range without touching some more complex functions at this time.

diff --git a/Engine/Config/BaseEngine.ini b/Engine/Config/BaseEngine.ini
--- a/Engine/Config/BaseEngine.ini
+++ b/Engine/Config/BaseEngine.ini
@@ -1031,6 +1031,17 @@ DefaultWorldPartitionHLODLayer=/Engine/Maps/Templates/HLODs/HLODLayer_Instanced.
 +FunctionRedirects=(OldName="KismetMathLibrary.BreakVector_NetQuantize100", NewName="BreakVector")
 +FunctionRedirects=(OldName="KismetMathLibrary.BreakVector_NetQuantizeNormal", NewName="BreakVector")
 
++FunctionRedirects=(OldName="KismetMathLibrary.Multiply_FloatFloat", NewName="Multiply_DoubleDouble")
++FunctionRedirects=(OldName="KismetMathLibrary.Divide_FloatFloat", NewName="Divide_DoubleDouble")
++FunctionRedirects=(OldName="KismetMathLibrary.Add_FloatFloat", NewName="Add_DoubleDouble")
++FunctionRedirects=(OldName="KismetMathLibrary.Subtract_FloatFloat", NewName="Subtract_DoubleDouble")
++FunctionRedirects=(OldName="KismetMathLibrary.Less_FloatFloat", NewName="Less_DoubleDouble")
++FunctionRedirects=(OldName="KismetMathLibrary.Greater_FloatFloat", NewName="Greater_DoubleDouble")
++FunctionRedirects=(OldName="KismetMathLibrary.LessEqual_FloatFloat", NewName="LessEqual_DoubleDouble")
++FunctionRedirects=(OldName="KismetMathLibrary.GreaterEqual_FloatFloat", NewName="GreaterEqual_DoubleDouble")
++FunctionRedirects=(OldName="KismetMathLibrary.EqualEqual_FloatFloat", NewName="EqualEqual_DoubleDouble")
++FunctionRedirects=(OldName="KismetMathLibrary.NotEqual_FloatFloat", NewName="NotEqual_DoubleDouble")
+
 [CoreUObject.Metadata]
 ; Note: UnrealHeaderTool should be rerun after making changes to MetadataRedirects to catch any keys specified in class headers
 +MetadataRedirects=(OldKey="K2Protected", NewKey="BlueprintProtected")
diff --git a/Engine/Source/Runtime/Engine/Classes/Kismet/KismetMathLibrary.h b/Engine/Source/Runtime/Engine/Classes/Kismet/KismetMathLibrary.h
--- a/Engine/Source/Runtime/Engine/Classes/Kismet/KismetMathLibrary.h
+++ b/Engine/Source/Runtime/Engine/Classes/Kismet/KismetMathLibrary.h
@@ -494,155 +494,124 @@ class ENGINE_API UKismetMathLibrary : public UBlueprintFunctionLibrary
 	//
 
 	/** Power (Base to the Exp-th power) */
-	UFUNCTION(BlueprintPure, meta=(DisplayName = "Power" ), Category="Math|Float")
-	static float MultiplyMultiply_FloatFloat(float Base, float Exp);
+	UFUNCTION(BlueprintPure, meta = (DisplayName = "Power"), Category = "Math|Float")
+	static double MultiplyMultiply_FloatFloat(double Base, double Exp);
 
-	/** Multiplication (A * B) */
-	UFUNCTION(BlueprintPure, meta=(DisplayName = "float * float", CompactNodeTitle = "*", Keywords = "* multiply", CommutativeAssociativeBinaryOperator = "true"), Category="Math|Float")
+	UE_DEPRECATED(5.0, "This method has been deprecated and will be removed. Use the double version instead.")
 	static float Multiply_FloatFloat(float A, float B);
 
 	/** Multiplication (A * B) */
-	UFUNCTION(BlueprintPure, meta=(DisplayName = "int * float", CompactNodeTitle = "*", Keywords = "* multiply"), Category="Math|Float")
-	static float Multiply_IntFloat(int32 A, float B);
+	UFUNCTION(BlueprintPure, meta = (DisplayName = "int * float", CompactNodeTitle = "*", Keywords = "* multiply"), Category = "Math|Float")
+	static double Multiply_IntFloat(int32 A, double B);
 
-	/** Division (A / B) */
-	UFUNCTION(BlueprintPure, CustomThunk, meta=(DisplayName = "float / float", CompactNodeTitle = "/", Keywords = "/ divide division"), Category="Math|Float")
+	UE_DEPRECATED(5.0, "This method has been deprecated and will be removed. Use the double version instead.")
 	static float Divide_FloatFloat(float A, float B = 1.f);
-	
-	static float GenericDivide_FloatFloat(float A, float B);
-
-	/** Custom thunk to allow script stack trace in case of divide by zero */
-	DECLARE_FUNCTION(execDivide_FloatFloat)
-	{
-		P_GET_PROPERTY(FFloatProperty, A);
-		P_GET_PROPERTY(FFloatProperty, B);
-
-		P_FINISH;
-
-		if (B == 0.f)
-		{
-			FFrame::KismetExecutionMessage(*FString::Printf(TEXT("Divide by zero detected: %f / 0\n%s"), A, *Stack.GetStackTrace()), ELogVerbosity::Warning);
-			*(float*)RESULT_PARAM = 0;
-			return;
-		}
 
-		*(float*)RESULT_PARAM = GenericDivide_FloatFloat(A, B);
-	}
-
-	/** Modulo (A % B) */
 	UFUNCTION(BlueprintPure, CustomThunk, meta = (DisplayName = "% (Float)", CompactNodeTitle = "%", Keywords = "% modulus"), Category = "Math|Float")
-	static float Percent_FloatFloat(float A, float B = 1.f);
+	static double Percent_FloatFloat(double A, double B = 1.f);
 
-	static float GenericPercent_FloatFloat(float A, float B);
+	static double GenericPercent_FloatFloat(double A, double B);
 
 	/** Custom thunk to allow script stack trace in case of modulo by zero */
 	DECLARE_FUNCTION(execPercent_FloatFloat)
 	{
-		P_GET_PROPERTY(FFloatProperty, A);
-		P_GET_PROPERTY(FFloatProperty, B);
+		P_GET_PROPERTY(FDoubleProperty, A);
+		P_GET_PROPERTY(FDoubleProperty, B);
 
 		P_FINISH;
 
 		if (B == 0.f)
 		{
 			FFrame::KismetExecutionMessage(*FString::Printf(TEXT("Modulo by zero detected: %f %% 0\n%s"), A, *Stack.GetStackTrace()), ELogVerbosity::Warning);
-			*(float*)RESULT_PARAM = 0;
+			*(double*)RESULT_PARAM = 0;
 			return;
 		}
 
-		*(float*)RESULT_PARAM = GenericPercent_FloatFloat(A, B);
+		*(double*)RESULT_PARAM = GenericPercent_FloatFloat(A, B);
 	}
 
-	/** Returns the fractional part of a float. */
+	/** Returns the fractional part of a double. */
 	UFUNCTION(BlueprintPure, Category="Math|Float")
-	static float Fraction(float A);
+	static double Fraction(double A);
 
 	/** Addition (A + B) */
-	UFUNCTION(BlueprintPure, meta = (DisplayName = "double + double", CompactNodeTitle = "+", Keywords = "+ add plus", CommutativeAssociativeBinaryOperator = "true"), Category = "Math|Double")
+	UFUNCTION(BlueprintPure, meta = (DisplayName = "float + float", CompactNodeTitle = "+", Keywords = "+ add plus", CommutativeAssociativeBinaryOperator = "true"), Category = "Math|Float")
 	static double Add_DoubleDouble(double A, double B = 1.0);
 
 	/** Subtraction (A - B) */
-	UFUNCTION(BlueprintPure, meta = (DisplayName = "double - double", CompactNodeTitle = "-", Keywords = "- subtract minus"), Category = "Math|Double")
+	UFUNCTION(BlueprintPure, meta = (DisplayName = "float - float", CompactNodeTitle = "-", Keywords = "- subtract minus"), Category = "Math|Float")
 	static double Subtract_DoubleDouble(double A, double B = 1.0);
 
 	/** Multiplication (A * B) */
-	UFUNCTION(BlueprintPure, meta = (DisplayName = "double * double", CompactNodeTitle = "*", Keywords = "* multiply", CommutativeAssociativeBinaryOperator = "true"), Category = "Math|Double")
+	UFUNCTION(BlueprintPure, meta = (DisplayName = "float * float", CompactNodeTitle = "*", Keywords = "* multiply", CommutativeAssociativeBinaryOperator = "true"), Category = "Math|Float")
 	static double Multiply_DoubleDouble(double A, double B);
 
 	/** Division (A / B) */
-	UFUNCTION(BlueprintPure, meta = (DisplayName = "double / double", CompactNodeTitle = "/", Keywords = "/ divide division"), Category = "Math|Double")
+	UFUNCTION(BlueprintPure, meta = (DisplayName = "float / float", CompactNodeTitle = "/", Keywords = "/ divide division"), Category = "Math|Float")
 	static double Divide_DoubleDouble(double A, double B = 1.0);
 
 	/** Returns true if A is Less than B (A < B) */
-	UFUNCTION(BlueprintPure, meta = (DisplayName = "double < double", CompactNodeTitle = "<", Keywords = "< less"), Category = "Math|double")
+	UFUNCTION(BlueprintPure, meta = (DisplayName = "float < float", CompactNodeTitle = "<", Keywords = "< less"), Category = "Math|Float")
 	static bool Less_DoubleDouble(double A, double B);
 
 	/** Returns true if A is greater than B (A > B) */
-	UFUNCTION(BlueprintPure, meta = (DisplayName = "double > double", CompactNodeTitle = ">", Keywords = "> greater"), Category = "Math|double")
+	UFUNCTION(BlueprintPure, meta = (DisplayName = "float > float", CompactNodeTitle = ">", Keywords = "> greater"), Category = "Math|Float")
 	static bool Greater_DoubleDouble(double A, double B);
 
 	/** Returns true if A is Less than or equal to B (A <= B) */
-	UFUNCTION(BlueprintPure, meta = (DisplayName = "double <= double", CompactNodeTitle = "<=", Keywords = "<= less"), Category = "Math|double")
+	UFUNCTION(BlueprintPure, meta = (DisplayName = "float <= float", CompactNodeTitle = "<=", Keywords = "<= less"), Category = "Math|Float")
 	static bool LessEqual_DoubleDouble(double A, double B);
 
 	/** Returns true if A is greater than or equal to B (A >= B) */
-	UFUNCTION(BlueprintPure, meta = (DisplayName = "double >= double", CompactNodeTitle = ">=", Keywords = ">= greater"), Category = "Math|double")
+	UFUNCTION(BlueprintPure, meta = (DisplayName = "float >= float", CompactNodeTitle = ">=", Keywords = ">= greater"), Category = "Math|Float")
 	static bool GreaterEqual_DoubleDouble(double A, double B);
 
 	/** Returns true if A is exactly equal to B (A == B)*/
-	UFUNCTION(BlueprintPure, meta = (DisplayName = "Equal (Double)", CompactNodeTitle = "==", Keywords = "== equal"), Category = "Math|double")
+	UFUNCTION(BlueprintPure, meta = (DisplayName = "Equal (Float)", CompactNodeTitle = "==", Keywords = "== equal"), Category = "Math|Float")
 	static bool EqualEqual_DoubleDouble(double A, double B);
 
-	/** Addition (A + B) */
-	UFUNCTION(BlueprintPure, meta=(DisplayName = "float + float", CompactNodeTitle = "+", Keywords = "+ add plus", CommutativeAssociativeBinaryOperator = "true"), Category="Math|Float")
+	UE_DEPRECATED(5.0, "This method has been deprecated and will be removed. Use the double version instead.")
 	static float Add_FloatFloat(float A, float B = 1.f);
 
-	/** Subtraction (A - B) */
-	UFUNCTION(BlueprintPure, meta=(DisplayName = "float - float", CompactNodeTitle = "-", Keywords = "- subtract minus"), Category="Math|Float")
+	UE_DEPRECATED(5.0, "This method has been deprecated and will be removed. Use the double version instead.")
 	static float Subtract_FloatFloat(float A, float B = 1.f);
 
-	/** Returns true if A is Less than B (A < B) */
-	UFUNCTION(BlueprintPure, meta=(DisplayName = "float < float", CompactNodeTitle = "<", Keywords = "< less"), Category="Math|Float")
+	UE_DEPRECATED(5.0, "This method has been deprecated and will be removed. Use the double version instead.")
 	static bool Less_FloatFloat(float A, float B);
 
-	/** Returns true if A is greater than B (A > B) */
-	UFUNCTION(BlueprintPure, meta=(DisplayName = "float > float", CompactNodeTitle = ">", Keywords = "> greater"), Category="Math|Float")
+	UE_DEPRECATED(5.0, "This method has been deprecated and will be removed. Use the double version instead.")
 	static bool Greater_FloatFloat(float A, float B);
 
-	/** Returns true if A is Less than or equal to B (A <= B) */
-	UFUNCTION(BlueprintPure, meta=(DisplayName = "float <= float", CompactNodeTitle = "<=", Keywords = "<= less"), Category="Math|Float")
+	UE_DEPRECATED(5.0, "This method has been deprecated and will be removed. Use the double version instead.")
 	static bool LessEqual_FloatFloat(float A, float B);
 
-	/** Returns true if A is greater than or equal to B (A >= B) */
-	UFUNCTION(BlueprintPure, meta=(DisplayName = "float >= float", CompactNodeTitle = ">=", Keywords = ">= greater"), Category="Math|Float")
+	UE_DEPRECATED(5.0, "This method has been deprecated and will be removed. Use the double version instead.")
 	static bool GreaterEqual_FloatFloat(float A, float B);
 
-	/** Returns true if A is exactly equal to B (A == B)*/
-	UFUNCTION(BlueprintPure, meta=(DisplayName = "Equal (Float)", CompactNodeTitle = "==", Keywords = "== equal"), Category="Math|Float")
+	UE_DEPRECATED(5.0, "This method has been deprecated and will be removed. Use the double version instead.")
 	static bool EqualEqual_FloatFloat(float A, float B);
 
 	/** Returns true if A is nearly equal to B (|A - B| < ErrorTolerance) */
-	UFUNCTION(BlueprintPure, meta=(DisplayName = "Nearly Equal (Float)", Keywords = "== equal"), Category="Math|Float")
-	static bool NearlyEqual_FloatFloat(float A, float B, float ErrorTolerance = 1.e-6f);
+	UFUNCTION(BlueprintPure, meta = (DisplayName = "Nearly Equal (Float)", Keywords = "== equal"), Category = "Math|Float")
+	static bool NearlyEqual_FloatFloat(double A, double B, double ErrorTolerance = 1.e-6);
 
 	/** Returns true if A does not equal B (A != B)*/
-	UFUNCTION(BlueprintPure, meta=(DisplayName = "Not Equal (Float)", CompactNodeTitle = "!=", Keywords = "!= not equal"), Category="Math|Float")
-	static bool NotEqual_FloatFloat(float A, float B);
-
-	/** Returns true if A does not equal B (A != B)*/
-	UFUNCTION(BlueprintPure, meta = (DisplayName = "Not Equal (Double)", CompactNodeTitle = "!=", Keywords = "!= not equal"), Category = "Math|Double")
+	UFUNCTION(BlueprintPure, meta = (DisplayName = "Not Equal (Float)", CompactNodeTitle = "!=", Keywords = "!= not equal"), Category = "Math|Float")
 	static bool NotEqual_DoubleDouble(double A, double B);
 
+	UE_DEPRECATED(5.0, "This method has been deprecated and will be removed. Use the double version instead.")
+	static bool NotEqual_FloatFloat(float A, float B);
+
 	/** Returns true if value is between Min and Max (V >= Min && V <= Max)
 	 * If InclusiveMin is true, value needs to be equal or larger than Min, else it needs to be larger
 	 * If InclusiveMax is true, value needs to be smaller or equal than Max, else it needs to be smaller
 	 */
-	UFUNCTION(BlueprintPure, meta=(DisplayName = "In Range (Float)", Min="0.0", Max="1.0"), Category="Math|Float")
-	static bool InRange_FloatFloat(float Value, float Min, float Max, bool InclusiveMin = true, bool InclusiveMax = true);
+	UFUNCTION(BlueprintPure, meta = (DisplayName = "In Range (Float)", Min = "0.0", Max = "1.0"), Category = "Math|Float")
+	static bool InRange_FloatFloat(double Value, double Min, double Max, bool InclusiveMin = true, bool InclusiveMax = true);
 
 	/** Returns the hypotenuse of a right-angled triangle given the width and height. */
 	UFUNCTION(BlueprintPure, meta=(Keywords = "pythagorean theorem"), Category = "Math|Float")
-	static float Hypotenuse(float Width, float Height);
+	static double Hypotenuse(double Width, double Height);
 	
 	/** Snaps a value to the nearest grid multiple. E.g.,
 	 *		Location = 5.1, GridSize = 10.0 : return value = 10.0
@@ -650,111 +619,111 @@ class ENGINE_API UKismetMathLibrary : public UBlueprintFunctionLibrary
 	 * if GridSize is very small precision issues may occur.
 	 */
 	UFUNCTION(BlueprintPure, meta = (DisplayName = "Snap To Grid (Float)"), Category = "Math|Float")
-	static float GridSnap_Float(float Location, float GridSize);
+	static double GridSnap_Float(double Location, double GridSize);
 
 	/** Returns the absolute (positive) value of A */
 	UFUNCTION(BlueprintPure, meta=(DisplayName = "Absolute (Float)", CompactNodeTitle = "ABS"), Category="Math|Float")
-	static float Abs(float A);
+	static double Abs(double A);
 
 	/** Returns the sine of A (expects Radians)*/
 	UFUNCTION(BlueprintPure, meta=(DisplayName = "Sin (Radians)", CompactNodeTitle = "SIN", Keywords = "sine"), Category="Math|Trig")
-	static float Sin(float A);
+	static double Sin(double A);
 
 	/** Returns the inverse sine (arcsin) of A (result is in Radians) */
 	UFUNCTION(BlueprintPure, meta=(DisplayName = "Asin (Radians)", CompactNodeTitle = "ASIN", Keywords = "sine"), Category="Math|Trig")
-	static float Asin(float A);
+	static double Asin(double A);
 
 	/** Returns the cosine of A (expects Radians)*/
 	UFUNCTION(BlueprintPure, meta=(DisplayName = "Cos (Radians)", CompactNodeTitle = "COS"), Category="Math|Trig")
-	static float Cos(float A);
+	static double Cos(double A);
 
 	/** Returns the inverse cosine (arccos) of A (result is in Radians) */
 	UFUNCTION(BlueprintPure, meta=(DisplayName = "Acos (Radians)", CompactNodeTitle = "ACOS"), Category="Math|Trig")
-	static float Acos(float A);
+	static double Acos(double A);
 
 	/** Returns the tan of A (expects Radians)*/
 	UFUNCTION(BlueprintPure, meta=(DisplayName = "Tan (Radians)", CompactNodeTitle = "TAN"), Category="Math|Trig")
-	static float Tan(float A);
+	static double Tan(double A);
 
 	/** Returns the inverse tan (atan) (result is in Radians)*/
 	UFUNCTION(BlueprintPure, meta=(DisplayName = "Atan (Radians)"), Category="Math|Trig")
-	static float Atan(float A);
+	static double Atan(double A);
 
 	/** Returns the inverse tan (atan2) of A/B (result is in Radians)*/
 	UFUNCTION(BlueprintPure, meta=(DisplayName = "Atan2 (Radians)"), Category="Math|Trig")
-	static float Atan2(float Y, float X);
+	static double Atan2(double Y, double X);
 
 	/** Returns exponential(e) to the power A (e^A)*/
 	UFUNCTION(BlueprintPure, Category="Math|Float", meta=(CompactNodeTitle = "e"))
-	static float Exp(float A);
+	static double Exp(double A);
 
 	/** Returns log of A base B (if B^R == A, returns R)*/
 	UFUNCTION(BlueprintPure, Category = "Math|Float")
-	static float Log(float A, float Base = 1.f);
+	static double Log(double A, double Base = 1.0);
 
 	/** Returns natural log of A (if e^R == A, returns R)*/
 	UFUNCTION(BlueprintPure, Category="Math|Float")
-	static float Loge(float A);
+	static double Loge(double A);
 
 	/** Returns square root of A*/
 	UFUNCTION(BlueprintPure, Category="Math|Float", meta=(Keywords = "square root", CompactNodeTitle = "SQRT"))
-	static float Sqrt(float A);
+	static double Sqrt(double A);
 
 	/** Returns square of A (A*A)*/
 	UFUNCTION(BlueprintPure, Category="Math|Float", meta=(CompactNodeTitle = "^2"))
-	static float Square(float A);
+	static double Square(double A);
 
 	/** Returns a random float between 0 and 1 */
 	UFUNCTION(BlueprintPure, Category="Math|Random", meta=(NotBlueprintThreadSafe))
-	static float RandomFloat();
+	static double RandomFloat();
 
 	/** Generate a random number between Min and Max */
 	UFUNCTION(BlueprintPure, Category="Math|Random", meta=(NotBlueprintThreadSafe))
-	static float RandomFloatInRange(float Min, float Max);
+	static double RandomFloatInRange(double Min, double Max);
 
 	/** Returns the value of PI */
 	UFUNCTION(BlueprintPure, meta=(DisplayName = "Get PI", CompactNodeTitle = "PI"), Category="Math|Trig")
-	static float GetPI();
+	static double GetPI();
 
 	/** Returns the value of TAU (= 2 * PI) */
 	UFUNCTION(BlueprintPure, meta=(DisplayName = "Get TAU", CompactNodeTitle = "TAU"), Category="Math|Trig")
-	static float GetTAU();
+	static double GetTAU();
 
 	/** Returns radians value based on the input degrees */
 	UFUNCTION(BlueprintPure, meta=(DisplayName = "Degrees To Radians", CompactNodeTitle = "D2R"), Category="Math|Trig")
-	static float DegreesToRadians(float A);
+	static double DegreesToRadians(double A);
 
 	/** Returns degrees value based on the input radians */
 	UFUNCTION(BlueprintPure, meta=(DisplayName = "Radians To Degrees", CompactNodeTitle = "R2D"), Category="Math|Trig")
-	static float RadiansToDegrees(float A);
+	static double RadiansToDegrees(double A);
 
 	/** Returns the sin of A (expects Degrees)*/
 	UFUNCTION(BlueprintPure, meta=(DisplayName = "Sin (Degrees)", CompactNodeTitle = "SINd", Keywords = "sine"), Category="Math|Trig")
-	static float DegSin(float A);
+	static double DegSin(double A);
 
 	/** Returns the inverse sin (arcsin) of A (result is in Degrees) */
 	UFUNCTION(BlueprintPure, meta=(DisplayName = "Asin (Degrees)", CompactNodeTitle = "ASINd", Keywords = "sine"), Category="Math|Trig")
-	static float DegAsin(float A);
+	static double DegAsin(double A);
 
 	/** Returns the cos of A (expects Degrees)*/
 	UFUNCTION(BlueprintPure, meta=(DisplayName = "Cos (Degrees)", CompactNodeTitle = "COSd"), Category="Math|Trig")
-	static float DegCos(float A);
+	static double DegCos(double A);
 
 	/** Returns the inverse cos (arccos) of A (result is in Degrees) */
 	UFUNCTION(BlueprintPure, meta=(DisplayName = "Acos (Degrees)", CompactNodeTitle = "ACOSd"), Category="Math|Trig")
-	static float DegAcos(float A);
+	static double DegAcos(double A);
 
 	/** Returns the tan of A (expects Degrees)*/
 	UFUNCTION(BlueprintPure, meta=(DisplayName = "Tan (Degrees)", CompactNodeTitle = "TANd"), Category="Math|Trig")
-	static float DegTan(float A);
+	static double DegTan(double A);
 
 	/** Returns the inverse tan (atan) (result is in Degrees)*/
 	UFUNCTION(BlueprintPure, meta=(DisplayName = "Atan (Degrees)"), Category="Math|Trig")
-	static float DegAtan(float A);
+	static double DegAtan(double A);
 
 	/** Returns the inverse tan (atan2) of A/B (result is in Degrees)*/
 	UFUNCTION(BlueprintPure, meta=(DisplayName = "Atan2 (Degrees)"), Category="Math|Trig")
-	static float DegAtan2(float Y, float X);
+	static double DegAtan2(double Y, double X);
 
 	/** 
 	 * Clamps an arbitrary angle to be between the given angles.  Will clamp to nearest boundary.
@@ -764,27 +733,27 @@ class ENGINE_API UKismetMathLibrary : public UBlueprintFunctionLibrary
 	 * @return Returns clamped angle in the range -180..180.
 	 */
 	UFUNCTION(BlueprintPure, meta=(DisplayName = "Clamp Angle"), Category="Math|Float")
-	static float ClampAngle(float AngleDegrees, float MinAngleDegrees, float MaxAngleDegrees);
+	static double ClampAngle(double AngleDegrees, double MinAngleDegrees, double MaxAngleDegrees);
 
 	/** Returns the minimum value of A and B */
 	UFUNCTION(BlueprintPure, meta=(DisplayName = "Min (Float)", CompactNodeTitle = "MIN", CommutativeAssociativeBinaryOperator = "true"), Category="Math|Float")
-	static float FMin(float A, float B);
+	static double FMin(double A, double B);
 
 	/** Returns the maximum value of A and B */
 	UFUNCTION(BlueprintPure, meta=(DisplayName = "Max (Float)", CompactNodeTitle = "MAX", CommutativeAssociativeBinaryOperator = "true"), Category="Math|Float")
-	static float FMax(float A, float B);
+	static double FMax(double A, double B);
 
 	/** Returns Value clamped between A and B (inclusive) */
 	UFUNCTION(BlueprintPure, meta=(DisplayName = "Clamp (Float)", Min="0.0", Max="1.0"), Category="Math|Float")
-	static float FClamp(float Value, float Min, float Max);
+	static double FClamp(double Value, double Min, double Max);
 
 	/** Returns Value wrapped from A and B (inclusive) */
 	UFUNCTION(BlueprintPure, meta = (DisplayName = "Wrap (Float)", Min = "0.0", Max = "1.0"), Category = "Math|Float")
-	static float FWrap(float Value, float Min, float Max);
+	static double FWrap(double Value, double Min, double Max);
 
 	/** This functions returns 0 if B (the denominator) is zero */
 	UFUNCTION(BlueprintPure, Category = "Math|Float", meta = (Keywords = "percent"))
-	static float SafeDivide(float A, float B);
+	static double SafeDivide(double A, double B);
 
 	/** Returns max of all array entries and the index at which it was found. Returns value of 0 and index of -1 if the supplied array is empty. */
 	UFUNCTION(BlueprintPure, Category="Math|Integer")
@@ -820,74 +789,81 @@ class ENGINE_API UKismetMathLibrary : public UBlueprintFunctionLibrary
 
 	/** Linearly interpolates between A and B based on Alpha (100% of A when Alpha=0 and 100% of B when Alpha=1) */
 	UFUNCTION(BlueprintPure, Category="Math|Float")
-	static float Lerp(float A, float B, float Alpha);
+	static double Lerp(double A, double B, double Alpha);
 	
 	UE_DEPRECATED(4.19, "Use NormalizeToRange instead")
 	static float InverseLerp(float A, float B, float Value);
 
 	/** Easeing  between A and B using a specified easing function */
 	UFUNCTION(BlueprintPure, meta = (DisplayName = "Ease", BlueprintInternalUseOnly = "true"), Category = "Math|Interpolation")
-	static float Ease(float A, float B, float Alpha, TEnumAsByte<EEasingFunc::Type> EasingFunc, float BlendExp = 2, int32 Steps = 2);
+	static double Ease(double A, double B, double Alpha, TEnumAsByte<EEasingFunc::Type> EasingFunc, double BlendExp = 2, int32 Steps = 2);
 
 	/** Rounds A to the nearest integer (e.g., -1.6 becomes -2 and 1.6 becomes 2) */
 	UFUNCTION(BlueprintPure, Category="Math|Float")
-	static int32 Round(float A);
+	static int32 Round(double A);
 
 	/** Rounds A down towards negative infinity / down to the previous integer (e.g., -1.6 becomes -2 and 1.6 becomes 1) */
 	UFUNCTION(BlueprintPure, meta=(DisplayName = "Floor"), Category="Math|Float")
-	static int32 FFloor(float A);
+	static int32 FFloor(double A);
 	
 	/** Rounds A towards zero, truncating the fractional part (e.g., -1.6 becomes -1 and 1.6 becomes 1) */
 	UFUNCTION(BlueprintPure, meta=(DisplayName = "Truncate", BlueprintAutocast), Category="Math|Float")
-	static int32 FTrunc(float A);
+	static int32 FTrunc(double A); // TODO: ok to change with BP Autocast?
 
 	/** Rounds A up towards positive infinity / up to the next integer (e.g., -1.6 becomes -1 and 1.6 becomes 2) */
 	UFUNCTION(BlueprintPure, Category = "Math|Float", meta=(DisplayName="Ceil"))
-	static int32 FCeil(float A);
+	static int32 FCeil(double A);
 
 	/** Rounds A to the nearest integer (e.g., -1.6 becomes -2 and 1.6 becomes 2) */
 	UFUNCTION(BlueprintPure, meta = (DisplayName = "Round to Integer64"), Category = "Math|Float")
-	static int64 Round64(float A);
+	static int64 Round64(double A);
 
 	/** Rounds A down towards negative infinity / down to the previous integer (e.g., -1.6 becomes -2 and 1.6 becomes 1) */
 	UFUNCTION(BlueprintPure, meta = (DisplayName = "Floor to Integer64"), Category = "Math|Float")
-	static int64 FFloor64(float A);
+	static int64 FFloor64(double A);
 
 	/** Rounds A towards zero, truncating the fractional part (e.g., -1.6 becomes -1 and 1.6 becomes 1) */
 	UFUNCTION(BlueprintPure, meta = (DisplayName = "Truncate to Integer64", BlueprintAutocast), Category = "Math|Float")
-	static int64 FTrunc64(float A);
+	static int64 FTrunc64(double A);
 
 	/** Rounds A up towards positive infinity / up to the next integer (e.g., -1.6 becomes -1 and 1.6 becomes 2) */
 	UFUNCTION(BlueprintPure, meta = (DisplayName = "Ceil to Integer64"), Category = "Math|Float")
-	static int64 FCeil64(float A);
+	static int64 FCeil64(double A);
 
 	/** Returns the number of times Divisor will go into Dividend (i.e., Dividend divided by Divisor), as well as the remainder */
 	UFUNCTION(BlueprintPure, meta=(DisplayName = "Division (Whole and Remainder)"), Category="Math|Float")
+	static int32 FMod(double Dividend, double Divisor, double& Remainder);
+
+	/** Returns the number of times Divisor will go into Dividend (i.e., Dividend divided by Divisor), as well as the remainder */
+	UFUNCTION(BlueprintPure, meta = (DisplayName = "Division (Whole and Remainder) to Integer64"), Category = "Math|Float")
+	static int64 FMod64(double Dividend, double Divisor, double& Remainder);
+
+	UE_DEPRECATED(5.0, "This method has been deprecated and will be removed. Use the double version instead.")
 	static int32 FMod(float Dividend, float Divisor, float& Remainder);
 
 	/** Sign (float, returns -1 if A < 0, 0 if A is zero, and +1 if A > 0) */
 	UFUNCTION(BlueprintPure, meta=(DisplayName = "Sign (Float)"), Category="Math|Float")
-	static float SignOfFloat(float A);
+	static double SignOfFloat(double A);
 
 	/** Returns Value normalized to the given range.  (e.g. 20 normalized to the range 10->50 would result in 0.25) */
 	UFUNCTION(BlueprintPure, Category="Math|Float")
-	static float NormalizeToRange(float Value, float RangeMin, float RangeMax);
+	static double NormalizeToRange(double Value, double RangeMin, double RangeMax);
 
 	/** Returns Value mapped from one range into another.  (e.g. 20 normalized from the range 10->50 to 20->40 would result in 25) */
 	UFUNCTION(BlueprintPure, Category="Math|Float", meta=(Keywords = "get mapped value"))
-	static float MapRangeUnclamped(float Value, float InRangeA, float InRangeB, float OutRangeA, float OutRangeB);
+	static double MapRangeUnclamped(double Value, double InRangeA, double InRangeB, double OutRangeA, double OutRangeB);
 
 	/** Returns Value mapped from one range into another where the Value is clamped to the Input Range.  (e.g. 0.5 normalized from the range 0->1 to 0->50 would result in 25) */
 	UFUNCTION(BlueprintPure, Category="Math|Float", meta=(Keywords = "get mapped value"))
-	static float MapRangeClamped(float Value, float InRangeA, float InRangeB, float OutRangeA, float OutRangeB);
+	static double MapRangeClamped(double Value, double InRangeA, double InRangeB, double OutRangeA, double OutRangeB);
 	
 	/** Multiplies the input value by pi. */
 	UFUNCTION(BlueprintPure, meta=(Keywords = "* multiply"), Category="Math|Float")
-	static float MultiplyByPi(float Value);
+	static double MultiplyByPi(double Value);
 
 	/** Interpolate between A and B, applying an ease in/out function.  Exp controls the degree of the curve. */
 	UFUNCTION(BlueprintPure, Category = "Math|Float")
-	static float FInterpEaseInOut(float A, float B, float Alpha, float Exponent);
+	static double FInterpEaseInOut(double A, double B, double Alpha, double Exponent);
 
 	/**
 	* Simple function to create a pulsating scalar value
@@ -1037,7 +1013,7 @@ class ENGINE_API UKismetMathLibrary : public UBlueprintFunctionLibrary
 
 	/** Returns Vector A added by B */
 	UFUNCTION(BlueprintPure, meta = (DisplayName = "vector2d + float", CompactNodeTitle = "+", ScriptMethod = "AddFloat", ScriptOperator = "+;+=", Keywords = "+ add plus"), Category = "Math|Vector2D")
-	static FVector2D Add_Vector2DFloat(FVector2D A, float B);
+	static FVector2D Add_Vector2DFloat(FVector2D A, double B);
 
 	/** Returns subtraction of Vector B from Vector A (A - B) */
 	UFUNCTION(BlueprintPure, meta = (DisplayName = "vector2d - vector2d", CompactNodeTitle = "-", ScriptMethod = "Subtract", ScriptOperator = "-;-=", Keywords = "- subtract minus"), Category = "Math|Vector2D")
@@ -1045,7 +1021,7 @@ class ENGINE_API UKismetMathLibrary : public UBlueprintFunctionLibrary
 
 	/** Returns Vector A subtracted by B */
 	UFUNCTION(BlueprintPure, meta = (DisplayName = "vector2d - float", CompactNodeTitle = "-", ScriptMethod = "SubtractFloat", ScriptOperator = "-;-=", Keywords = "- subtract minus"), Category = "Math|Vector2D")
-	static FVector2D Subtract_Vector2DFloat(FVector2D A, float B);
+	static FVector2D Subtract_Vector2DFloat(FVector2D A, double B);
 
 	/** Element-wise Vector multiplication (Result = {A.x*B.x, A.y*B.y}) */
 	UFUNCTION(BlueprintPure, meta = (DisplayName = "vector2d * vector2d", CompactNodeTitle = "*", ScriptMethod = "Multiply", ScriptOperator = "*;*=", Keywords = "* multiply", CommutativeAssociativeBinaryOperator = "true"), Category = "Math|Vector2D")
@@ -1053,7 +1029,7 @@ class ENGINE_API UKismetMathLibrary : public UBlueprintFunctionLibrary
 
 	/** Returns Vector A scaled by B */
 	UFUNCTION(BlueprintPure, meta = (DisplayName = "vector2d * float", CompactNodeTitle = "*", ScriptMethod = "MultiplyFloat", ScriptOperator = "*;*=", Keywords = "* multiply"), Category = "Math|Vector2D")
-	static FVector2D Multiply_Vector2DFloat(FVector2D A, float B);
+	static FVector2D Multiply_Vector2DFloat(FVector2D A, double B);
 
 	/** Element-wise Vector divide (Result = {A.x/B.x, A.y/B.y}) */
 	UFUNCTION(BlueprintPure, meta = (DisplayName = "vector2d / vector2d", CompactNodeTitle = "/", ScriptMethod = "Divide", ScriptOperator = "/;/=", Keywords = "/ divide division"), Category = "Math|Vector2D")
@@ -1061,7 +1037,7 @@ class ENGINE_API UKismetMathLibrary : public UBlueprintFunctionLibrary
 
 	/** Returns Vector A divided by B */
 	UFUNCTION(BlueprintPure, meta = (DisplayName = "vector2d / float", CompactNodeTitle = "/", ScriptMethod = "DivideFloat", ScriptOperator = "/;/=", Keywords = "/ divide division"), Category = "Math|Vector2D")
-	static FVector2D Divide_Vector2DFloat(FVector2D A, float B = 1.f);
+	static FVector2D Divide_Vector2DFloat(FVector2D A, double B = 1.f);
 
 	/** Returns true if vector A is equal to vector B (A == B) */
 	UFUNCTION(BlueprintPure, meta=(DisplayName = "Equal Exactly (Vector2D)", CompactNodeTitle = "===", ScriptMethod = "Equals", ScriptOperator = "==", Keywords = "== equal"), Category="Math|Vector2D")
@@ -1090,18 +1066,18 @@ class ENGINE_API UKismetMathLibrary : public UBlueprintFunctionLibrary
 	 * @param InY New Y coordinate.
 	 */
 	UFUNCTION(BlueprintCallable, meta = (ScriptMethod = "Set"), Category = "Math|Vector2D")
-	static void Set2D(UPARAM(ref) FVector2D& A, float X, float Y);
+	static void Set2D(UPARAM(ref) FVector2D& A, double X, double Y);
 
 	/**
 	 * Creates a copy of this vector with both axes clamped to the given range.
 	 * @return New vector with clamped axes.
 	 */
 	UFUNCTION(BlueprintPure, meta = (ScriptMethod = "ClampedAxes"), Category = "Math|Vector2D")
-	static FVector2D ClampAxes2D(FVector2D A, float MinAxisVal, float MaxAxisVal);
+	static FVector2D ClampAxes2D(FVector2D A, double MinAxisVal, double MaxAxisVal);
 
 	/** Returns the cross product of two 2d vectors - see  http://mathworld.wolfram.com/CrossProduct.html */
 	UFUNCTION(BlueprintPure, meta = (DisplayName = "Cross Product (2D)", CompactNodeTitle = "cross", ScriptMethod = "Cross", ScriptOperator = "^"), Category = "Math|Vector2D")
-	static float CrossProduct2D(FVector2D A, FVector2D B);
+	static double CrossProduct2D(FVector2D A, FVector2D B);
 
 	/**
 	 * Distance between two 2D points.
@@ -1111,7 +1087,7 @@ class ENGINE_API UKismetMathLibrary : public UBlueprintFunctionLibrary
 	 * @return The distance between two 2D points.
 	 */
 	UFUNCTION(BlueprintPure, meta = (Keywords = "magnitude", ScriptMethod = "Distance"), Category = "Math|Vector2D")
-	static float Distance2D(FVector2D V1, FVector2D V2);
+	static double Distance2D(FVector2D V1, FVector2D V2);
 
 	/**
 	 * Squared distance between two 2D points.
@@ -1121,11 +1097,11 @@ class ENGINE_API UKismetMathLibrary : public UBlueprintFunctionLibrary
 	 * @return The squared distance between two 2D points.
 	 */
 	UFUNCTION(BlueprintPure, meta = (Keywords = "magnitude", ScriptMethod = "DistanceSquared"), Category = "Math|Vector2D")
-	static float DistanceSquared2D(FVector2D V1, FVector2D V2);
+	static double DistanceSquared2D(FVector2D V1, FVector2D V2);
 
 	/** Returns the dot product of two 2d vectors - see http://mathworld.wolfram.com/DotProduct.html */
 	UFUNCTION(BlueprintPure, meta = (DisplayName = "Dot Product (2D)", CompactNodeTitle = "dot", ScriptMethod = "Dot", ScriptOperator = "|"), Category = "Math|Vector2D")
-	static float DotProduct2D(FVector2D A, FVector2D B);
+	static double DotProduct2D(FVector2D A, FVector2D B);
 
 	/**
 	* Get a copy of this vector with absolute value of each component.
@@ -1141,7 +1117,7 @@ class ENGINE_API UKismetMathLibrary : public UBlueprintFunctionLibrary
 	 * @return The maximum absolute value of the vector's components.
 	 */
 	UFUNCTION(BlueprintPure, meta = (ScriptMethod = "GetAbsMax"), Category = "Math|Vector2D")
-	static float GetAbsMax2D(FVector2D A);
+	static double GetAbsMax2D(FVector2D A);
 
 	/**
 	 * Get the maximum value of the vector's components.
@@ -1149,7 +1125,7 @@ class ENGINE_API UKismetMathLibrary : public UBlueprintFunctionLibrary
 	 * @return The maximum value of the vector's components.
 	 */
 	UFUNCTION(BlueprintPure, meta = (ScriptMethod = "GetMax"), Category = "Math|Vector2D")
-	static float GetMax2D(FVector2D A);
+	static double GetMax2D(FVector2D A);
 
 	/**
 	 * Get the minimum value of the vector's components.
@@ -1157,7 +1133,7 @@ class ENGINE_API UKismetMathLibrary : public UBlueprintFunctionLibrary
 	 * @return The minimum value of the vector's components.
 	 */
 	UFUNCTION(BlueprintPure, meta = (ScriptMethod = "GetMin"), Category = "Math|Vector2D")
-	static float GetMin2D(FVector2D A);
+	static double GetMin2D(FVector2D A);
 
 	/**
 	 * Rotates around axis (0,0,1)
@@ -1243,7 +1219,10 @@ class ENGINE_API UKismetMathLibrary : public UBlueprintFunctionLibrary
 	 * @param OutLength Reference passed in to store length of the vector.
 	 */
 	UFUNCTION(BlueprintPure, meta = (DisplayName = "To Direction And Length", ScriptMethod = "ToDirectionAndLength"), Category = "Math|Vector2D")
-	static void ToDirectionAndLength2D(FVector2D A, FVector2D &OutDir, float &OutLength);
+	static void ToDirectionAndLength2D(FVector2D A, FVector2D &OutDir, double &OutLength);
+
+	UE_DEPRECATED(5.0, "This method has been deprecated and will be removed. Use the double version instead.")
+	static void ToDirectionAndLength2D(FVector2D A, FVector2D& OutDir, float& OutLength);
 
 	/**
 	 * Get this vector as a vector where each component has been rounded to the nearest int.
@@ -1264,11 +1243,11 @@ class ENGINE_API UKismetMathLibrary : public UBlueprintFunctionLibrary
 
 	/** Returns the length of a 2D Vector. */
 	UFUNCTION(BlueprintPure, meta = (DisplayName = "Vector2D Length", Keywords = "magnitude", ScriptMethod = "Length"), Category = "Math|Vector2D")
-	static float VSize2D(FVector2D A);
+	static double VSize2D(FVector2D A);
 
 	/** Returns the squared length of a 2D Vector. */
 	UFUNCTION(BlueprintPure, meta = (DisplayName = "Vector2D Length Squared", Keywords = "magnitude", ScriptMethod = "LengthSquared"), Category = "Math|Vector2D")
-	static float VSize2DSquared(FVector2D A);
+	static double VSize2DSquared(FVector2D A);
 
 
 	//
@@ -1335,7 +1314,7 @@ class ENGINE_API UKismetMathLibrary : public UBlueprintFunctionLibrary
 	 * @param InZ New Z coordinate.
 	 */
 	UFUNCTION(BlueprintCallable, meta = (ScriptMethod = "Set"), Category = "Math|Vector")
-	static void Vector_Set(UPARAM(ref) FVector& A, float X, float Y, float Z);
+	static void Vector_Set(UPARAM(ref) FVector& A, double X, double Y, double Z);
 
 	UE_DEPRECATED(5.0, "This method has been deprecated and will be removed. Use the double version instead.")
 	static void BreakVector(FVector InVec, float& X, float& Y, float& Z);
@@ -1393,7 +1372,7 @@ class ENGINE_API UKismetMathLibrary : public UBlueprintFunctionLibrary
 
 	/** Adds a float to each component of a vector */
 	UFUNCTION(BlueprintPure, meta=(DisplayName = "vector + float", CompactNodeTitle = "+", ScriptMethod = "AddFloat", Keywords = "+ add plus"), Category="Math|Vector")
-	static FVector Add_VectorFloat(FVector A, float B);
+	static FVector Add_VectorFloat(FVector A, double B);
 	
 	/** Adds an integer to each component of a vector */
 	UFUNCTION(BlueprintPure, meta=(DisplayName = "vector + integer", CompactNodeTitle = "+", ScriptMethod = "AddInt", Keywords = "+ add plus"), Category="Math|Vector")
@@ -1405,7 +1384,7 @@ class ENGINE_API UKismetMathLibrary : public UBlueprintFunctionLibrary
 
 	/** Subtracts a float from each component of a vector */
 	UFUNCTION(BlueprintPure, meta=(DisplayName = "vector - float", CompactNodeTitle = "-", ScriptMethod = "SubtractFloat", Keywords = "- subtract minus"), Category="Math|Vector")
-	static FVector Subtract_VectorFloat(FVector A, float B);
+	static FVector Subtract_VectorFloat(FVector A, double B);
 
 	/** Subtracts an integer from each component of a vector */
 	UFUNCTION(BlueprintPure, meta=(DisplayName = "vector - integer", CompactNodeTitle = "-", ScriptMethod = "SubtractInt", Keywords = "- subtract minus"), Category="Math|Vector")
@@ -1417,7 +1396,7 @@ class ENGINE_API UKismetMathLibrary : public UBlueprintFunctionLibrary
 
 	/** Scales Vector A by B */
 	UFUNCTION(BlueprintPure, meta=(DisplayName = "vector * float", CompactNodeTitle = "*", ScriptMethod = "MultiplyFloat", ScriptOperator = "*;*=", Keywords = "* multiply"), Category="Math|Vector")
-	static FVector Multiply_VectorFloat(FVector A, float B);
+	static FVector Multiply_VectorFloat(FVector A, double B);
 	
 	/** Scales Vector A by B */
 	UFUNCTION(BlueprintPure, meta=(DisplayName = "vector * integer", CompactNodeTitle = "*", ScriptMethod = "MultiplyInt", Keywords = "* multiply"), Category="Math|Vector")
@@ -1429,7 +1408,7 @@ class ENGINE_API UKismetMathLibrary : public UBlueprintFunctionLibrary
 
 	/** Vector divide by a float */
 	UFUNCTION(BlueprintPure, meta=(DisplayName = "vector / float", CompactNodeTitle = "/", ScriptMethod = "DivideFloat", ScriptOperator = "/;/=", Keywords = "/ divide division"), Category="Math|Vector")
-	static FVector Divide_VectorFloat(FVector A, float B = 1.f);
+	static FVector Divide_VectorFloat(FVector A, double B = 1.f);
 
 	/** Vector divide by an integer */
 	UFUNCTION(BlueprintPure, meta=(DisplayName = "vector / integer", CompactNodeTitle = "/", ScriptMethod = "DivideInt", Keywords = "/ divide division"), Category="Math|Vector")
@@ -1457,7 +1436,7 @@ class ENGINE_API UKismetMathLibrary : public UBlueprintFunctionLibrary
 
 	/** Returns the dot product of two 3d vectors - see http://mathworld.wolfram.com/DotProduct.html */
 	UFUNCTION(BlueprintPure, meta=(DisplayName = "Dot Product", CompactNodeTitle = "dot", ScriptMethod = "Dot", ScriptOperator = "|"), Category="Math|Vector" )
-	static float Dot_VectorVector(FVector A, FVector B);
+	static double Dot_VectorVector(FVector A, FVector B);
 
 	/** Returns the cross product of two 3d vectors - see http://mathworld.wolfram.com/CrossProduct.html */
 	UFUNCTION(BlueprintPure, meta=(DisplayName = "Cross Product", CompactNodeTitle = "cross", ScriptMethod = "Cross", ScriptOperator = "^"), Category="Math|Vector" )
@@ -1481,35 +1460,35 @@ class ENGINE_API UKismetMathLibrary : public UBlueprintFunctionLibrary
 
 	/** Create a copy of this vector, with its magnitude/size/length clamped between Min and Max. */
 	UFUNCTION(BlueprintPure, meta=(ScriptMethod = "ClampedSize"), Category="Math|Vector")
-	static FVector ClampVectorSize(FVector A, float Min, float Max);
+	static FVector ClampVectorSize(FVector A, double Min, double Max);
 
 	/** Create a copy of this vector, with the 2D magnitude/size/length clamped between Min and Max. Z is unchanged. */
 	UFUNCTION(BlueprintPure, meta=(ScriptMethod = "ClampedSize2D"), Category="Math|Vector")
-	static FVector Vector_ClampSize2D(FVector A, float Min, float Max);
+	static FVector Vector_ClampSize2D(FVector A, double Min, double Max);
 
 	/** Create a copy of this vector, with its maximum magnitude/size/length clamped to MaxSize. */
 	UFUNCTION(BlueprintPure, meta=(ScriptMethod = "ClampedSizeMax"), Category="Math|Vector")
-	static FVector Vector_ClampSizeMax(FVector A, float Max);
+	static FVector Vector_ClampSizeMax(FVector A, double Max);
 
 	/** Create a copy of this vector, with the maximum 2D magnitude/size/length clamped to MaxSize. Z is unchanged. */
 	UFUNCTION(BlueprintPure, meta=(ScriptMethod = "ClampedSizeMax2D"), Category="Math|Vector")
-	static FVector Vector_ClampSizeMax2D(FVector A, float Max);
+	static FVector Vector_ClampSizeMax2D(FVector A, double Max);
 
 	/** Find the minimum element (X, Y or Z) of a vector */
 	UFUNCTION(BlueprintPure, meta=(ScriptMethod = "GetMinElement"), Category="Math|Vector")
-	static float GetMinElement(FVector A);
+	static double GetMinElement(FVector A);
 
 	/** Find the maximum element (X, Y or Z) of a vector */
 	UFUNCTION(BlueprintPure, meta=(ScriptMethod = "GetMaxElement"), Category="Math|Vector")
-	static float GetMaxElement(FVector A);
+	static double GetMaxElement(FVector A);
 
 	/** Find the maximum absolute element (abs(X), abs(Y) or abs(Z)) of a vector */
 	UFUNCTION(BlueprintPure, meta=(ScriptMethod = "GetAbsMax"), Category="Math|Vector")
-	static float Vector_GetAbsMax(FVector A);
+	static double Vector_GetAbsMax(FVector A);
 
 	/** Find the minimum absolute element (abs(X), abs(Y) or abs(Z)) of a vector */
 	UFUNCTION(BlueprintPure, meta=(ScriptMethod = "GetAbsMin"), Category="Math|Vector")
-	static float Vector_GetAbsMin(FVector A);
+	static double Vector_GetAbsMin(FVector A);
 
 	/**
 	 * Get a copy of this vector with absolute value of each component.
@@ -1550,7 +1529,7 @@ class ENGINE_API UKismetMathLibrary : public UBlueprintFunctionLibrary
 	 * @return 'Heading' angle between +/-PI radians. 0 is pointing down +X.
 	 */
 	UFUNCTION(BlueprintPure, meta=(ScriptMethod = "HeadingAngle"), Category="Math|Vector")
-	static float Vector_HeadingAngle(FVector A);
+	static double Vector_HeadingAngle(FVector A);
 
 	/**
 	 * Returns the cosine of the angle between this vector and another projected onto the XY plane (no Z).
@@ -1559,7 +1538,7 @@ class ENGINE_API UKismetMathLibrary : public UBlueprintFunctionLibrary
 	 * @return The cosine.
 	 */
 	UFUNCTION(BlueprintPure, meta=(ScriptMethod = "CosineAngle2D"), Category="Math|Vector")
-	static float Vector_CosineAngle2D(FVector A, FVector B);
+	static double Vector_CosineAngle2D(FVector A, FVector B);
 
 	/**
 	 * Converts a vector containing degree values to a vector containing radian values.
@@ -1613,7 +1592,7 @@ class ENGINE_API UKismetMathLibrary : public UBlueprintFunctionLibrary
 	 * @return The distance between two points.
 	 */
 	UFUNCTION(BlueprintPure, meta = (DisplayName = "Distance (Vector)", ScriptMethod = "Distance", Keywords = "magnitude"), Category = "Math|Vector")
-	static float Vector_Distance(FVector V1, FVector V2);
+	static double Vector_Distance(FVector V1, FVector V2);
 
 	/**
 	 * Squared distance between two points.
@@ -1623,7 +1602,7 @@ class ENGINE_API UKismetMathLibrary : public UBlueprintFunctionLibrary
 	 * @return The squared distance between two points.
 	 */
 	UFUNCTION(BlueprintPure, meta = (DisplayName = "Distance Squared (Vector)", ScriptMethod = "DistanceSquared", Keywords = "magnitude"), Category = "Math|Vector")
-	static float Vector_DistanceSquared(FVector V1, FVector V2);
+	static double Vector_DistanceSquared(FVector V1, FVector V2);
 
 	/**
 	* Euclidean distance between two points in the XY plane (ignoring Z).
@@ -1633,7 +1612,7 @@ class ENGINE_API UKismetMathLibrary : public UBlueprintFunctionLibrary
 	* @return The distance between two points in the XY plane.
 	*/
 	UFUNCTION(BlueprintPure, meta = (DisplayName = "Distance2D (Vector)", ScriptMethod = "Distance2D", Keywords = "magnitude"), Category = "Math|Vector")
-	static float Vector_Distance2D(FVector V1, FVector V2);
+	static double Vector_Distance2D(FVector V1, FVector V2);
 
 	/**
 	* Squared euclidean distance between two points in the XY plane (ignoring Z).
@@ -1643,23 +1622,23 @@ class ENGINE_API UKismetMathLibrary : public UBlueprintFunctionLibrary
 	* @return The distance between two points in the XY plane.
 	*/
 	UFUNCTION(BlueprintPure, meta = (DisplayName = "Distance2D Squared (Vector)", ScriptMethod = "Distance2DSquared", Keywords = "magnitude"), Category = "Math|Vector")
-	static float Vector_Distance2DSquared(FVector V1, FVector V2);
+	static double Vector_Distance2DSquared(FVector V1, FVector V2);
 
 	/** Returns the length of the vector */
 	UFUNCTION(BlueprintPure, meta=(DisplayName = "Vector Length", ScriptMethod = "Length", Keywords="magnitude"), Category="Math|Vector")
-	static float VSize(FVector A);
+	static double VSize(FVector A);
 
 	/** Returns the squared length of the vector */
 	UFUNCTION(BlueprintPure, meta=(DisplayName = "Vector Length Squared", ScriptMethod = "LengthSquared", Keywords="magnitude"), Category="Math|Vector")
-	static float VSizeSquared(FVector A);
+	static double VSizeSquared(FVector A);
 
 	/** Returns the length of the vector's XY components. */
 	UFUNCTION(BlueprintPure, meta=(DisplayName = "Vector Length XY", ScriptMethod = "Length2D", Keywords="magnitude"), Category="Math|Vector")
-	static float VSizeXY(FVector A);
+	static double VSizeXY(FVector A);
 
 	/** Returns the squared length of the vector's XY components. */
 	UFUNCTION(BlueprintPure, meta=(DisplayName = "Vector Length XY Squared", ScriptMethod = "Length2DSquared", Keywords="magnitude"), Category="Math|Vector")
-	static float VSizeXYSquared(FVector A);
+	static double VSizeXYSquared(FVector A);
 
 	/**
 	 * Checks whether vector is near to zero within a specified tolerance.
@@ -2147,7 +2126,7 @@ class ENGINE_API UKismetMathLibrary : public UBlueprintFunctionLibrary
 	 * @param InW New W coordinate.
 	 */
 	UFUNCTION(BlueprintCallable, meta = (ScriptMethod = "Set"), Category = "Math|Vector4")
-	static void Vector4_Set(UPARAM(ref) FVector4& A, float X, float Y, float Z, float W);
+	static void Vector4_Set(UPARAM(ref) FVector4& A, double X, double Y, double Z, double W);
 
 	/** Returns the cross product of two vectors - see  http://mathworld.wolfram.com/CrossProduct.html */
 	UFUNCTION(BlueprintPure, meta = (DisplayName = "Cross Product XYZ (Vector4)", CompactNodeTitle = "cross3", ScriptMethod = "Cross3"), Category = "Math|Vector4")
@@ -2155,11 +2134,11 @@ class ENGINE_API UKismetMathLibrary : public UBlueprintFunctionLibrary
 
 	/** Returns the dot product of two vectors - see http://mathworld.wolfram.com/DotProduct.html */
 	UFUNCTION(BlueprintPure, meta = (DisplayName = "Dot Product (Vector4)", CompactNodeTitle = "dot", ScriptMethod = "Dot", ScriptOperator = "|"), Category = "Math|Vector4")
-	static float Vector4_DotProduct(const FVector4& A, const FVector4& B);
+	static double Vector4_DotProduct(const FVector4& A, const FVector4& B);
 
 	/** Returns the dot product of two vectors - see http://mathworld.wolfram.com/DotProduct.html The W element is ignored.*/
 	UFUNCTION(BlueprintPure, meta = (DisplayName = "Dot Product XYZ (Vector4)", CompactNodeTitle = "dot3", ScriptMethod = "Dot3"), Category = "Math|Vector4")
-	static float Vector4_DotProduct3(const FVector4& A, const FVector4& B);
+	static double Vector4_DotProduct3(const FVector4& A, const FVector4& B);
 
 	/**
 	 * Determines if any component is not a number (NAN)
@@ -2188,19 +2167,19 @@ class ENGINE_API UKismetMathLibrary : public UBlueprintFunctionLibrary
 
 	/** Returns the length of the vector. */
 	UFUNCTION(BlueprintPure, meta = (DisplayName = "Length (Vector4)", ScriptMethod = "Length", Keywords = "magnitude"), Category = "Math|Vector4")
-	static float Vector4_Size(const FVector4& A);
+	static double Vector4_Size(const FVector4& A);
 
 	/** Returns the squared length of the vector. */
 	UFUNCTION(BlueprintPure, meta = (DisplayName = "Length Squared (Vector4)", ScriptMethod = "LengthSquared", Keywords = "magnitude"), Category = "Math|Vector4")
-	static float Vector4_SizeSquared(const FVector4& A);
+	static double Vector4_SizeSquared(const FVector4& A);
 
 	/** Returns the length of the vector. The W element is ignored. */
 	UFUNCTION(BlueprintPure, meta = (DisplayName = "Length XYZ (Vector4)", ScriptMethod = "Length3", Keywords = "magnitude"), Category = "Math|Vector4")
-	static float Vector4_Size3(const FVector4& A);
+	static double Vector4_Size3(const FVector4& A);
 
 	/** Returns the squared length of the vector. The W element is ignored. */
 	UFUNCTION(BlueprintPure, meta = (DisplayName = "Length XYZ Squared (Vector4)", ScriptMethod = "LengthSquared3", Keywords = "magnitude"), Category = "Math|Vector4")
-	static float Vector4_SizeSquared3(const FVector4& A);
+	static double Vector4_SizeSquared3(const FVector4& A);
 
 	/**
 	 * Determines if vector is normalized / unit (length 1) within specified squared tolerance. The W element is ignored.
@@ -2474,7 +2453,7 @@ class ENGINE_API UKismetMathLibrary : public UBlueprintFunctionLibrary
 	  * If your Matrix represents a Transform that you wish to scale you should use Apply Scale instead
 	  */
 	UFUNCTION(BlueprintPure, meta = (DisplayName = "Matrix * Float", CompactNodeTitle = "*", ScriptMethod = "MultiplyFloat", ScriptOperator = "*;*=", Keywords = "* multiply"), Category = "Math|Matrix")
-	static FMatrix Multiply_MatrixFloat (const FMatrix& A, float B);
+	static FMatrix Multiply_MatrixFloat (const FMatrix& A, double B);
 
 	/**
 	 * Checks whether another Matrix is equal to this, within specified tolerance.
@@ -3676,7 +3655,7 @@ class ENGINE_API UKismetMathLibrary : public UBlueprintFunctionLibrary
 
 	/** If bPickA is true, A is returned, otherwise B is */
 	UFUNCTION(BlueprintPure, Category="Math|Float")
-	static float SelectFloat(float A, float B, bool bPickA);
+	static double SelectFloat(double A, double B, bool bPickA);
 
 	/** If bPickA is true, A is returned, otherwise B is */
 	UFUNCTION(BlueprintPure, Category="Math|Vector")
@@ -3902,7 +3881,7 @@ class ENGINE_API UKismetMathLibrary : public UBlueprintFunctionLibrary
 	 * @return		New interpolated position
 	 */
 	UFUNCTION(BlueprintPure, Category="Math|Interpolation")
-	static float FInterpTo(float Current, float Target, float DeltaTime, float InterpSpeed);
+	static double FInterpTo(double Current, double Target, double DeltaTime, double InterpSpeed);
 
 	/**
 	 * Tries to reach Target at a constant rate.
@@ -3914,7 +3893,7 @@ class ENGINE_API UKismetMathLibrary : public UBlueprintFunctionLibrary
 	 * @return		New interpolated position
 	 */
 	UFUNCTION(BlueprintPure, Category="Math|Interpolation")
-	static float FInterpTo_Constant(float Current, float Target, float DeltaTime, float InterpSpeed);
+	static double FInterpTo_Constant(double Current, double Target, double DeltaTime, double InterpSpeed);
 
 	/**
 	 * Tries to reach Target rotation based on Current rotation, giving a nice smooth feeling when rotating to Target rotation.
@@ -4300,6 +4279,7 @@ private:
 	static void ReportError_Divide_ByteByte();
 	static void ReportError_Percent_ByteByte();
 	static void ReportError_Divide_IntInt();
+	static void ReportError_Divide_FloatFloat();
 	static void ReportError_Divide_DoubleDouble();
 	static void ReportError_Divide_Int64Int64();
 	static void ReportError_Percent_IntInt();
diff --git a/Engine/Source/Runtime/Engine/Classes/Kismet/KismetMathLibrary.inl b/Engine/Source/Runtime/Engine/Classes/Kismet/KismetMathLibrary.inl
--- a/Engine/Source/Runtime/Engine/Classes/Kismet/KismetMathLibrary.inl
+++ b/Engine/Source/Runtime/Engine/Classes/Kismet/KismetMathLibrary.inl
@@ -403,10 +403,10 @@ int64 UKismetMathLibrary::Abs_Int64(int64 A)
 }
 
 KISMET_MATH_FORCEINLINE
-float UKismetMathLibrary::MultiplyMultiply_FloatFloat(float Base, float Exp)
+double UKismetMathLibrary::MultiplyMultiply_FloatFloat(double Base, double Exp)
 {
 	return FMath::Pow(Base, Exp);
-}	
+}
 
 KISMET_MATH_FORCEINLINE
 float UKismetMathLibrary::Multiply_FloatFloat(float A, float B)
@@ -415,7 +415,7 @@ float UKismetMathLibrary::Multiply_FloatFloat(float A, float B)
 }	
 
 KISMET_MATH_FORCEINLINE
-float UKismetMathLibrary::Multiply_IntFloat(int32 A, float B)
+double UKismetMathLibrary::Multiply_IntFloat(int32 A, double B)
 {
 	return A * B;
 }	
@@ -458,13 +458,7 @@ int32 UKismetMathLibrary::Percent_IntInt(int32 A, int32 B)
 }
 
 KISMET_MATH_FORCEINLINE
-float UKismetMathLibrary::GenericDivide_FloatFloat(float A, float B)
-{
-	return A / B;
-}	
-
-KISMET_MATH_FORCEINLINE
-float UKismetMathLibrary::Fraction(float A)
+double UKismetMathLibrary::Fraction(double A)
 {
 	return FMath::Fractional(A);
 }
@@ -487,12 +481,24 @@ double UKismetMathLibrary::Multiply_DoubleDouble(double A, double B)
 	return A * B;
 }
 
+KISMET_MATH_FORCEINLINE
+float UKismetMathLibrary::Divide_FloatFloat(float A, float B)
+{
+	if (B == 0.f)
+	{
+		ReportError_Divide_FloatFloat();
+		return 0.f;
+	}
+	return A / B;
+}
+
 KISMET_MATH_FORCEINLINE
 double UKismetMathLibrary::Divide_DoubleDouble(double A, double B)
 {
-	if(B == 0.0)
+	if (B == 0.0)
 	{
 		ReportError_Divide_DoubleDouble();
+		return 0.0;
 	}
 	return A / B;
 }
@@ -570,7 +576,7 @@ bool UKismetMathLibrary::EqualEqual_FloatFloat(float A, float B)
 }	
 
 KISMET_MATH_FORCEINLINE
-bool UKismetMathLibrary::NearlyEqual_FloatFloat(float A, float B, float ErrorTolerance)
+bool UKismetMathLibrary::NearlyEqual_FloatFloat(double A, double B, double ErrorTolerance)
 {
 	return FMath::IsNearlyEqual(A, B, ErrorTolerance);
 }
@@ -588,221 +594,221 @@ bool UKismetMathLibrary::NotEqual_DoubleDouble(double A, double B)
 }
 
 KISMET_MATH_FORCEINLINE
-float UKismetMathLibrary::GridSnap_Float(float Location, float GridSize)
+double UKismetMathLibrary::GridSnap_Float(double Location, double GridSize)
 {
 	return FMath::GridSnap(Location, GridSize);
 }
 
 KISMET_MATH_FORCEINLINE
-float UKismetMathLibrary::GetPI()
+double UKismetMathLibrary::GetPI()
 {
-	return PI;
+	return DOUBLE_PI;
 }
 
 KISMET_MATH_FORCEINLINE
-float UKismetMathLibrary::GetTAU()
+double UKismetMathLibrary::GetTAU()
 {
-	return 2.f * PI;
+	return 2.0 * DOUBLE_PI;
 }
 
 KISMET_MATH_FORCEINLINE
-float UKismetMathLibrary::DegreesToRadians(float A)
+double UKismetMathLibrary::DegreesToRadians(double A)
 {
 	return FMath::DegreesToRadians(A);
 }
 
 KISMET_MATH_FORCEINLINE
-float UKismetMathLibrary::RadiansToDegrees(float A)
+double UKismetMathLibrary::RadiansToDegrees(double A)
 {
 	return FMath::RadiansToDegrees(A);
 }
 
 KISMET_MATH_FORCEINLINE
-float UKismetMathLibrary::Abs(float A)
+double UKismetMathLibrary::Abs(double A)
 {
 	return FMath::Abs(A);
 }	
 
 KISMET_MATH_FORCEINLINE
-float UKismetMathLibrary::Sin(float A)
+double UKismetMathLibrary::Sin(double A)
 {
 	return FMath::Sin(A);
 }	
 
 KISMET_MATH_FORCEINLINE
-float UKismetMathLibrary::Asin(float A)
+double UKismetMathLibrary::Asin(double A)
 {
 	return FMath::Asin(A);
 }	
 
 KISMET_MATH_FORCEINLINE
-float UKismetMathLibrary::Cos(float A)
+double UKismetMathLibrary::Cos(double A)
 {
 	return FMath::Cos(A);
 }	
 
 KISMET_MATH_FORCEINLINE
-float UKismetMathLibrary::Acos(float A)
+double UKismetMathLibrary::Acos(double A)
 {
 	return FMath::Acos(A);
 }	
 
 KISMET_MATH_FORCEINLINE
-float UKismetMathLibrary::Tan(float A)
+double UKismetMathLibrary::Tan(double A)
 {
 	return FMath::Tan(A);
 }	
 
 KISMET_MATH_FORCEINLINE
-float UKismetMathLibrary::Atan(float A)
+double UKismetMathLibrary::Atan(double A)
 {
 	return FMath::Atan(A);
 }
 
 KISMET_MATH_FORCEINLINE
-float UKismetMathLibrary::Atan2(float Y, float X)
+double UKismetMathLibrary::Atan2(double Y, double X)
 {
 	return FMath::Atan2(Y, X);
 }	
 
 KISMET_MATH_FORCEINLINE
-float UKismetMathLibrary::DegSin(float A)
+double UKismetMathLibrary::DegSin(double A)
 {
-	return FMath::Sin(PI/(180.f) * A);
+	return FMath::Sin(DOUBLE_PI/(180.0) * A);
 }
 
 KISMET_MATH_FORCEINLINE
-float UKismetMathLibrary::DegAsin(float A)
+double UKismetMathLibrary::DegAsin(double A)
 {
-	return (180.f)/PI * FMath::Asin(A);
+	return (180.0)/DOUBLE_PI * FMath::Asin(A);
 }
 
 KISMET_MATH_FORCEINLINE
-float UKismetMathLibrary::DegCos(float A)
+double UKismetMathLibrary::DegCos(double A)
 {
-	return FMath::Cos(PI/(180.f) * A);
+	return FMath::Cos(DOUBLE_PI/(180.0) * A);
 }
 
 KISMET_MATH_FORCEINLINE
-float UKismetMathLibrary::DegAcos(float A)
+double UKismetMathLibrary::DegAcos(double A)
 {
-	return (180.f)/PI * FMath::Acos(A);
+	return (180.0)/DOUBLE_PI * FMath::Acos(A);
 }
 
 KISMET_MATH_FORCEINLINE
-float UKismetMathLibrary::DegTan(float A)
+double UKismetMathLibrary::DegTan(double A)
 {
-	return FMath::Tan(PI/(180.f) * A);
+	return FMath::Tan(DOUBLE_PI/(180.0) * A);
 }
 
 KISMET_MATH_FORCEINLINE
-float UKismetMathLibrary::DegAtan(float A)
+double UKismetMathLibrary::DegAtan(double A)
 {
-	return (180.f)/PI * FMath::Atan(A);
+	return (180.0)/DOUBLE_PI * FMath::Atan(A);
 }
 
 KISMET_MATH_FORCEINLINE
-float UKismetMathLibrary::DegAtan2(float Y, float X)
+double UKismetMathLibrary::DegAtan2(double Y, double X)
 {
-	return (180.f)/PI * FMath::Atan2(Y, X);
+	return (180.0)/DOUBLE_PI * FMath::Atan2(Y, X);
 }
 
 KISMET_MATH_FORCEINLINE
-float UKismetMathLibrary::ClampAngle(float AngleDegrees, float MinAngleDegrees, float MaxAngleDegrees)
+double UKismetMathLibrary::ClampAngle(double AngleDegrees, double MinAngleDegrees, double MaxAngleDegrees)
 {
 	return FMath::ClampAngle(AngleDegrees, MinAngleDegrees, MaxAngleDegrees);
 }
 
 KISMET_MATH_FORCEINLINE
-float UKismetMathLibrary::Exp(float A)
+double UKismetMathLibrary::Exp(double A)
 {
 	return FMath::Exp(A);
 }	
 
 KISMET_MATH_FORCEINLINE
-float UKismetMathLibrary::Loge(float A)
+double UKismetMathLibrary::Loge(double A)
 {
 	return FMath::Loge(A);
 }	
 
 KISMET_MATH_FORCEINLINE
-float UKismetMathLibrary::Sqrt(float A)
+double UKismetMathLibrary::Sqrt(double A)
 {
-	if (A >= 0.f)
+	if (A >= 0.0)
 	{
 		return FMath::Sqrt(A);
 	}
 	else
 	{
 		ReportError_Sqrt();
-		return 0.f;
+		return 0.0;
 	}
 }
 
 KISMET_MATH_FORCEINLINE
-float UKismetMathLibrary::Square(float A)
+double UKismetMathLibrary::Square(double A)
 {
 	return FMath::Square(A);
 }	
 
 KISMET_MATH_FORCEINLINE
-int32 UKismetMathLibrary::Round(float A)
+int32 UKismetMathLibrary::Round(double A)
 {
 	return FMath::RoundToInt(A);
 }	
 
 KISMET_MATH_FORCEINLINE
-int32 UKismetMathLibrary::FFloor(float A)
+int32 UKismetMathLibrary::FFloor(double A)
 {
 	return FMath::FloorToInt(A);
 }	
 
 KISMET_MATH_FORCEINLINE
-int32 UKismetMathLibrary::FTrunc(float A)
+int32 UKismetMathLibrary::FTrunc(double A)
 {
 	return FMath::TruncToInt(A);
 }	
 
 KISMET_MATH_FORCEINLINE
-int64 UKismetMathLibrary::Round64(float A)
+int64 UKismetMathLibrary::Round64(double A)
 {
 	return (int64)FMath::RoundToDouble(A);
 }
 
 KISMET_MATH_FORCEINLINE
-int64 UKismetMathLibrary::FFloor64(float A)
+int64 UKismetMathLibrary::FFloor64(double A)
 {
 	return (int64)FMath::FloorToDouble(A);
 }
 
 KISMET_MATH_FORCEINLINE
-int64 UKismetMathLibrary::FTrunc64(float A)
+int64 UKismetMathLibrary::FTrunc64(double A)
 {
 	return (int64)A;
 }
 
 KISMET_MATH_FORCEINLINE
-int64 UKismetMathLibrary::FCeil64(float A)
+int64 UKismetMathLibrary::FCeil64(double A)
 {
 	return (int64)FMath::CeilToDouble(A);
 }
 
 KISMET_MATH_FORCEINLINE
-int32 UKismetMathLibrary::FCeil(float A)
+int32 UKismetMathLibrary::FCeil(double A)
 {
 	return FMath::CeilToInt(A);
 }	
 
 KISMET_MATH_FORCEINLINE
-float UKismetMathLibrary::SignOfFloat(float A)
+double UKismetMathLibrary::SignOfFloat(double A)
 {
-	return FMath::Sign<float>(A);
+	return FMath::Sign<double>(A);
 }
 
 KISMET_MATH_FORCEINLINE
-float UKismetMathLibrary::MultiplyByPi(float Value)
+double UKismetMathLibrary::MultiplyByPi(double Value)
 {
-	return Value * PI;
+	return Value * DOUBLE_PI;
 }
 
 KISMET_MATH_FORCEINLINE
@@ -812,55 +818,55 @@ float UKismetMathLibrary::FixedTurn(float InCurrent, float InDesired, float InDe
 }
 
 KISMET_MATH_FORCEINLINE
-float UKismetMathLibrary::RandomFloat()
+double UKismetMathLibrary::RandomFloat()
 {
-	return FMath::FRand();
+	return FMath::FRand(); // LWC_todo: add and switch to double version.
 }	
 
 KISMET_MATH_FORCEINLINE
-float UKismetMathLibrary::RandomFloatInRange(float Min, float Max)
+double UKismetMathLibrary::RandomFloatInRange(double Min, double Max)
 {
 	return FMath::FRandRange(Min, Max);
 }	
 
 KISMET_MATH_FORCEINLINE
-float UKismetMathLibrary::FMin(float A, float B)
+double UKismetMathLibrary::FMin(double A, double B)
 {
 	return FMath::Min(A, B);
 }	
 
 KISMET_MATH_FORCEINLINE
-float UKismetMathLibrary::FMax(float A, float B)
+double UKismetMathLibrary::FMax(double A, double B)
 {
 	return FMath::Max(A, B);
 }	
 
 KISMET_MATH_FORCEINLINE
-float UKismetMathLibrary::FClamp(float V, float A, float B)
+double UKismetMathLibrary::FClamp(double V, double A, double B)
 {
 	return FMath::Clamp(V, A, B);
 }	
 
 KISMET_MATH_FORCEINLINE
-float UKismetMathLibrary::FWrap(float Value, float Min, float Max)
+double UKismetMathLibrary::FWrap(double Value, double Min, double Max)
 {
 	return FMath::Wrap(Value, Min, Max);
 }
 
 KISMET_MATH_FORCEINLINE
-float UKismetMathLibrary::Lerp(float A, float B, float V)
+double UKismetMathLibrary::Lerp(double A, double B, double V)
 {
 	return A + V*(B-A);
 }	
 
 KISMET_MATH_FORCEINLINE
-float UKismetMathLibrary::FInterpTo(float Current, float Target, float DeltaTime, float InterpSpeed)
+double UKismetMathLibrary::FInterpTo(double Current, double Target, double DeltaTime, double InterpSpeed)
 {
 	return FMath::FInterpTo(Current, Target, DeltaTime, InterpSpeed);
 }
 
 KISMET_MATH_FORCEINLINE
-float UKismetMathLibrary::FInterpTo_Constant(float Current, float Target, float DeltaTime, float InterpSpeed)
+double UKismetMathLibrary::FInterpTo_Constant(double Current, double Target, double DeltaTime, double InterpSpeed)
 {
 	return FMath::FInterpConstantTo(Current, Target, DeltaTime, InterpSpeed);
 }
@@ -1169,7 +1175,7 @@ FVector2D UKismetMathLibrary::Subtract_Vector2DVector2D(FVector2D A, FVector2D B
 }
 
 KISMET_MATH_FORCEINLINE
-FVector2D UKismetMathLibrary::Multiply_Vector2DFloat(FVector2D A, float B)
+FVector2D UKismetMathLibrary::Multiply_Vector2DFloat(FVector2D A, double B)
 {
 	return A * B;
 }
@@ -1181,7 +1187,7 @@ FVector2D UKismetMathLibrary::Multiply_Vector2DVector2D(FVector2D A, FVector2D B
 }
 
 KISMET_MATH_FORCEINLINE
-FVector2D UKismetMathLibrary::Divide_Vector2DFloat(FVector2D A, float B)
+FVector2D UKismetMathLibrary::Divide_Vector2DFloat(FVector2D A, double B)
 {
 	if (B == 0.f)
 	{
@@ -1205,13 +1211,13 @@ FVector2D UKismetMathLibrary::Divide_Vector2DVector2D(FVector2D A, FVector2D B)
 }
 
 KISMET_MATH_FORCEINLINE
-FVector2D UKismetMathLibrary::Add_Vector2DFloat(FVector2D A, float B)
+FVector2D UKismetMathLibrary::Add_Vector2DFloat(FVector2D A, double B)
 {
 	return A+B;
 }
 
 KISMET_MATH_FORCEINLINE
-FVector2D UKismetMathLibrary::Subtract_Vector2DFloat(FVector2D A, float B)
+FVector2D UKismetMathLibrary::Subtract_Vector2DFloat(FVector2D A, double B)
 {
 	return A-B;
 }
@@ -1247,37 +1253,37 @@ FVector2D UKismetMathLibrary::Negated2D(const FVector2D& A)
 }
 
 KISMET_MATH_FORCEINLINE
-void UKismetMathLibrary::Set2D(FVector2D& A, float X, float Y)
+void UKismetMathLibrary::Set2D(FVector2D& A, double X, double Y)
 {
 	A.Set(X, Y);
 }
 
 KISMET_MATH_FORCEINLINE
-FVector2D UKismetMathLibrary::ClampAxes2D(FVector2D A, float MinAxisVal, float MaxAxisVal)
+FVector2D UKismetMathLibrary::ClampAxes2D(FVector2D A, double MinAxisVal, double MaxAxisVal)
 {
 	return A.ClampAxes(MinAxisVal, MaxAxisVal);
 }
 
 KISMET_MATH_FORCEINLINE
-float UKismetMathLibrary::CrossProduct2D(FVector2D A, FVector2D B)
+double UKismetMathLibrary::CrossProduct2D(FVector2D A, FVector2D B)
 {
 	return FVector2D::CrossProduct(A, B);
 }
 
 KISMET_MATH_FORCEINLINE
-float UKismetMathLibrary::Distance2D(FVector2D V1, FVector2D V2)
+double UKismetMathLibrary::Distance2D(FVector2D V1, FVector2D V2)
 {
 	return FVector2D::Distance(V1, V2);
 }
 
 KISMET_MATH_FORCEINLINE
-float UKismetMathLibrary::DistanceSquared2D(FVector2D V1, FVector2D V2)
+double UKismetMathLibrary::DistanceSquared2D(FVector2D V1, FVector2D V2)
 {
 	return FVector2D::DistSquared(V1, V2);
 }
 
 KISMET_MATH_FORCEINLINE
-float UKismetMathLibrary::DotProduct2D(FVector2D A, FVector2D B)
+double UKismetMathLibrary::DotProduct2D(FVector2D A, FVector2D B)
 {
 	return FVector2D::DotProduct(A, B);
 }
@@ -1289,19 +1295,19 @@ FVector2D UKismetMathLibrary::GetAbs2D(FVector2D A)
 }
 
 KISMET_MATH_FORCEINLINE
-float UKismetMathLibrary::GetAbsMax2D(FVector2D A)
+double UKismetMathLibrary::GetAbsMax2D(FVector2D A)
 {
 	return A.GetAbsMax();
 }
 
 KISMET_MATH_FORCEINLINE
-float UKismetMathLibrary::GetMax2D(FVector2D A)
+double UKismetMathLibrary::GetMax2D(FVector2D A)
 {
 	return A.GetMax();
 }
 
 KISMET_MATH_FORCEINLINE
-float UKismetMathLibrary::GetMin2D(FVector2D A)
+double UKismetMathLibrary::GetMin2D(FVector2D A)
 {
 	return A.GetMin();
 }
@@ -1368,6 +1374,12 @@ void UKismetMathLibrary::ToDirectionAndLength2D(FVector2D A, FVector2D &OutDir,
 	OutLength = OutLengthDbl;
 }
 
+KISMET_MATH_FORCEINLINE
+void UKismetMathLibrary::ToDirectionAndLength2D(FVector2D A, FVector2D& OutDir, double& OutLength)
+{
+	A.ToDirectionAndLength(OutDir, OutLength);
+}
+
 KISMET_MATH_FORCEINLINE
 FVector2D UKismetMathLibrary::ToRounded2D(FVector2D A)
 {
@@ -1381,13 +1393,13 @@ FVector2D UKismetMathLibrary::ToSign2D(FVector2D A)
 }
 
 KISMET_MATH_FORCEINLINE
-float UKismetMathLibrary::VSize2D(FVector2D A)
+double UKismetMathLibrary::VSize2D(FVector2D A)
 {
 	return A.Size();
 }
 
 KISMET_MATH_FORCEINLINE
-float UKismetMathLibrary::VSize2DSquared(FVector2D A)
+double UKismetMathLibrary::VSize2DSquared(FVector2D A)
 {
 	return A.SizeSquared();
 }
@@ -1457,7 +1469,7 @@ void UKismetMathLibrary::Vector_Assign(FVector& A, const FVector& InVector)
 }
 
 KISMET_MATH_FORCEINLINE
-void UKismetMathLibrary::Vector_Set(FVector& A, float X, float Y, float Z)
+void UKismetMathLibrary::Vector_Set(FVector& A, double X, double Y, double Z)
 {
 	A.Set(X, Y, Z);
 }
@@ -1523,7 +1535,7 @@ FTransform UKismetMathLibrary::Conv_VectorToTransform(FVector InTranslation)
 }
 
 KISMET_MATH_FORCEINLINE
-FVector UKismetMathLibrary::Multiply_VectorFloat(FVector A, float B)
+FVector UKismetMathLibrary::Multiply_VectorFloat(FVector A, double B)
 {
 	return A * B;
 }
@@ -1541,7 +1553,7 @@ FVector UKismetMathLibrary::Multiply_VectorVector(FVector A, FVector B)
 }	
 
 KISMET_MATH_FORCEINLINE
-FVector UKismetMathLibrary::Divide_VectorFloat(FVector A, float B)
+FVector UKismetMathLibrary::Divide_VectorFloat(FVector A, double B)
 {
 	if (B == 0.f)
 	{
@@ -1583,7 +1595,7 @@ FVector UKismetMathLibrary::Add_VectorVector(FVector A, FVector B)
 }	
 
 KISMET_MATH_FORCEINLINE
-FVector UKismetMathLibrary::Add_VectorFloat(FVector A, float B)
+FVector UKismetMathLibrary::Add_VectorFloat(FVector A, double B)
 {
 	return A + B;
 }	
@@ -1601,7 +1613,7 @@ FVector UKismetMathLibrary::Subtract_VectorVector(FVector A, FVector B)
 }	
 
 KISMET_MATH_FORCEINLINE
-FVector UKismetMathLibrary::Subtract_VectorFloat(FVector A, float B)
+FVector UKismetMathLibrary::Subtract_VectorFloat(FVector A, double B)
 {
 	return A - B;
 }
@@ -1655,55 +1667,55 @@ FVector UKismetMathLibrary::Cross_VectorVector(FVector A, FVector B)
 }
 
 KISMET_MATH_FORCEINLINE
-float UKismetMathLibrary::Dot_VectorVector(FVector A, FVector B)
+double UKismetMathLibrary::Dot_VectorVector(FVector A, FVector B)
 {
 	return FVector::DotProduct(A, B);
 }
 
 KISMET_MATH_FORCEINLINE
-float UKismetMathLibrary::Vector_Distance(FVector V1, FVector V2)
+double UKismetMathLibrary::Vector_Distance(FVector V1, FVector V2)
 {
 	return FVector::Distance(V1, V2);
 }
 
 KISMET_MATH_FORCEINLINE
-float UKismetMathLibrary::Vector_DistanceSquared(FVector V1, FVector V2)
+double UKismetMathLibrary::Vector_DistanceSquared(FVector V1, FVector V2)
 {
 	return FVector::DistSquared(V1, V2);
 }
 
 KISMET_MATH_FORCEINLINE
-float UKismetMathLibrary::Vector_Distance2D(FVector V1, FVector V2)
+double UKismetMathLibrary::Vector_Distance2D(FVector V1, FVector V2)
 {
 	return FVector::DistXY(V1, V2);
 }
 
 KISMET_MATH_FORCEINLINE
-float UKismetMathLibrary::Vector_Distance2DSquared(FVector V1, FVector V2)
+double UKismetMathLibrary::Vector_Distance2DSquared(FVector V1, FVector V2)
 {
 	return FVector::DistSquaredXY(V1, V2);
 }
 
 KISMET_MATH_FORCEINLINE
-float UKismetMathLibrary::VSize(FVector A)
+double UKismetMathLibrary::VSize(FVector A)
 {
 	return A.Size();
 }
 
 KISMET_MATH_FORCEINLINE
-float UKismetMathLibrary::VSizeSquared(FVector A)
+double UKismetMathLibrary::VSizeSquared(FVector A)
 {
 	return A.SizeSquared();
 }
 
 KISMET_MATH_FORCEINLINE
-float UKismetMathLibrary::VSizeXY(FVector A)
+double UKismetMathLibrary::VSizeXY(FVector A)
 {
 	return A.Size2D();
 }	
 
 KISMET_MATH_FORCEINLINE
-float UKismetMathLibrary::VSizeXYSquared(FVector A)
+double UKismetMathLibrary::VSizeXYSquared(FVector A)
 {
 	return A.SizeSquared2D();
 }	
@@ -1919,49 +1931,49 @@ void UKismetMathLibrary::Vector_UnwindEuler(FVector& A)
 }
 
 KISMET_MATH_FORCEINLINE
-FVector UKismetMathLibrary::ClampVectorSize(FVector A, float Min, float Max)
+FVector UKismetMathLibrary::ClampVectorSize(FVector A, double Min, double Max)
 {
 	return A.GetClampedToSize(Min, Max);
 }
 
 KISMET_MATH_FORCEINLINE
-FVector UKismetMathLibrary::Vector_ClampSize2D(FVector A, float Min, float Max)
+FVector UKismetMathLibrary::Vector_ClampSize2D(FVector A, double Min, double Max)
 {
 	return A.GetClampedToSize2D(Min, Max);
 }
 
 KISMET_MATH_FORCEINLINE
-FVector UKismetMathLibrary::Vector_ClampSizeMax(FVector A, float Max)
+FVector UKismetMathLibrary::Vector_ClampSizeMax(FVector A, double Max)
 {
 	return A.GetClampedToMaxSize(Max);
 }
 
 KISMET_MATH_FORCEINLINE
-FVector UKismetMathLibrary::Vector_ClampSizeMax2D(FVector A, float Max)
+FVector UKismetMathLibrary::Vector_ClampSizeMax2D(FVector A, double Max)
 {
 	return A.GetClampedToMaxSize2D(Max);
 }
 
 KISMET_MATH_FORCEINLINE
-float UKismetMathLibrary::GetMinElement(FVector A)
+double UKismetMathLibrary::GetMinElement(FVector A)
 {
 	return A.GetMin();
 }
 
 KISMET_MATH_FORCEINLINE
-float UKismetMathLibrary::GetMaxElement(FVector A)
+double UKismetMathLibrary::GetMaxElement(FVector A)
 {
 	return A.GetMax();
 }
 
 KISMET_MATH_FORCEINLINE
-float UKismetMathLibrary::Vector_GetAbsMax(FVector A)
+double UKismetMathLibrary::Vector_GetAbsMax(FVector A)
 {
 	return A.GetAbsMax();
 }
 
 KISMET_MATH_FORCEINLINE
-float UKismetMathLibrary::Vector_GetAbsMin(FVector A)
+double UKismetMathLibrary::Vector_GetAbsMin(FVector A)
 {
 	return A.GetAbsMin();
 }
@@ -1997,13 +2009,13 @@ FVector UKismetMathLibrary::Vector_GetProjection(FVector A)
 }
 
 KISMET_MATH_FORCEINLINE
-float UKismetMathLibrary::Vector_HeadingAngle(FVector A)
+double UKismetMathLibrary::Vector_HeadingAngle(FVector A)
 {
 	return A.HeadingAngle();
 }
 
 KISMET_MATH_FORCEINLINE
-float UKismetMathLibrary::Vector_CosineAngle2D(FVector A, FVector B)
+double UKismetMathLibrary::Vector_CosineAngle2D(FVector A, FVector B)
 {
 	return A.CosineAngle2D(B);
 }
@@ -2159,7 +2171,7 @@ void UKismetMathLibrary::Vector4_Assign(FVector4& A, const FVector4& InVector)
 }
 
 KISMET_MATH_FORCEINLINE
-void UKismetMathLibrary::Vector4_Set(FVector4& A, float X, float Y, float Z, float W)
+void UKismetMathLibrary::Vector4_Set(FVector4& A, double X, double Y, double Z, double W)
 {
 	A.Set(X, Y, Z, W);
 }
@@ -2171,13 +2183,13 @@ FVector4 UKismetMathLibrary::Vector4_CrossProduct3(const FVector4& A, const FVec
 }
 
 KISMET_MATH_FORCEINLINE
-float UKismetMathLibrary::Vector4_DotProduct(const FVector4& A, const FVector4& B)
+double UKismetMathLibrary::Vector4_DotProduct(const FVector4& A, const FVector4& B)
 {
 	return Dot4(A, B);
 }
 
 KISMET_MATH_FORCEINLINE
-float UKismetMathLibrary::Vector4_DotProduct3(const FVector4& A, const FVector4& B)
+double UKismetMathLibrary::Vector4_DotProduct3(const FVector4& A, const FVector4& B)
 {
 	return Dot3(A, B);
 }
@@ -2201,25 +2213,25 @@ bool UKismetMathLibrary::Vector4_IsZero(const FVector4& A)
 }
 
 KISMET_MATH_FORCEINLINE
-float UKismetMathLibrary::Vector4_Size(const FVector4& A)
+double UKismetMathLibrary::Vector4_Size(const FVector4& A)
 {
 	return A.Size();
 }
 
 KISMET_MATH_FORCEINLINE
-float UKismetMathLibrary::Vector4_SizeSquared(const FVector4& A)
+double UKismetMathLibrary::Vector4_SizeSquared(const FVector4& A)
 {
 	return A.SizeSquared();
 }
 
 KISMET_MATH_FORCEINLINE
-float UKismetMathLibrary::Vector4_Size3(const FVector4& A)
+double UKismetMathLibrary::Vector4_Size3(const FVector4& A)
 {
 	return A.Size3();
 }
 
 KISMET_MATH_FORCEINLINE
-float UKismetMathLibrary::Vector4_SizeSquared3(const FVector4& A)
+double UKismetMathLibrary::Vector4_SizeSquared3(const FVector4& A)
 {
 	return A.SizeSquared3();
 }
@@ -2985,7 +2997,7 @@ FMatrix UKismetMathLibrary::Add_MatrixMatrix(const FMatrix& A, const FMatrix& B)
 }
 
 KISMET_MATH_FORCEINLINE
-FMatrix UKismetMathLibrary::Multiply_MatrixFloat(const FMatrix& A, float B)
+FMatrix UKismetMathLibrary::Multiply_MatrixFloat(const FMatrix& A, double B)
 {
 	return A * B;
 }
@@ -3763,7 +3775,7 @@ int32 UKismetMathLibrary::SelectInt(int32 A, int32 B, bool bSelectA)
 }
 
 KISMET_MATH_FORCEINLINE
-float UKismetMathLibrary::SelectFloat(float A, float B, bool bSelectA)
+double UKismetMathLibrary::SelectFloat(double A, double B, bool bSelectA)
 {
 	return bSelectA ? A : B;
 }
diff --git a/Engine/Source/Runtime/Engine/Private/KismetMathLibrary.cpp b/Engine/Source/Runtime/Engine/Private/KismetMathLibrary.cpp
--- a/Engine/Source/Runtime/Engine/Private/KismetMathLibrary.cpp
+++ b/Engine/Source/Runtime/Engine/Private/KismetMathLibrary.cpp
@@ -13,23 +13,24 @@
 #define LOCTEXT_NAMESPACE "UKismetMathLibrary"
 
 /** Interpolate a linear alpha value using an ease mode and function. */
-float EaseAlpha(float InAlpha, uint8 EasingFunc, float BlendExp, int32 Steps)
+template<typename FloatType, TEMPLATE_REQUIRES(TIsFloatingPoint<FloatType>::Value)>
+FloatType EaseAlpha(FloatType InAlpha, uint8 EasingFunc, FloatType BlendExp, int32 Steps)
 {
 	switch (EasingFunc)
 	{
-	case EEasingFunc::Step:					return FMath::InterpStep<float>(0.f, 1.f, InAlpha, Steps);
-	case EEasingFunc::SinusoidalIn:			return FMath::InterpSinIn<float>(0.f, 1.f, InAlpha);
-	case EEasingFunc::SinusoidalOut:		return FMath::InterpSinOut<float>(0.f, 1.f, InAlpha);
-	case EEasingFunc::SinusoidalInOut:		return FMath::InterpSinInOut<float>(0.f, 1.f, InAlpha);
-	case EEasingFunc::EaseIn:				return FMath::InterpEaseIn<float>(0.f, 1.f, InAlpha, BlendExp);
-	case EEasingFunc::EaseOut:				return FMath::InterpEaseOut<float>(0.f, 1.f, InAlpha, BlendExp);
-	case EEasingFunc::EaseInOut:			return FMath::InterpEaseInOut<float>(0.f, 1.f, InAlpha, BlendExp);
-	case EEasingFunc::ExpoIn:				return FMath::InterpExpoIn<float>(0.f, 1.f, InAlpha);
-	case EEasingFunc::ExpoOut:				return FMath::InterpExpoOut<float>(0.f, 1.f, InAlpha);
-	case EEasingFunc::ExpoInOut:			return FMath::InterpExpoInOut<float>(0.f, 1.f, InAlpha);
-	case EEasingFunc::CircularIn:			return FMath::InterpCircularIn<float>(0.f, 1.f, InAlpha);
-	case EEasingFunc::CircularOut:			return FMath::InterpCircularOut<float>(0.f, 1.f, InAlpha);
-	case EEasingFunc::CircularInOut:		return FMath::InterpCircularInOut<float>(0.f, 1.f, InAlpha);
+	case EEasingFunc::Step:					return FMath::InterpStep<FloatType>(0.f, 1.f, InAlpha, Steps);
+	case EEasingFunc::SinusoidalIn:			return FMath::InterpSinIn<FloatType>(0.f, 1.f, InAlpha);
+	case EEasingFunc::SinusoidalOut:		return FMath::InterpSinOut<FloatType>(0.f, 1.f, InAlpha);
+	case EEasingFunc::SinusoidalInOut:		return FMath::InterpSinInOut<FloatType>(0.f, 1.f, InAlpha);
+	case EEasingFunc::EaseIn:				return FMath::InterpEaseIn<FloatType>(0.f, 1.f, InAlpha, BlendExp);
+	case EEasingFunc::EaseOut:				return FMath::InterpEaseOut<FloatType>(0.f, 1.f, InAlpha, BlendExp);
+	case EEasingFunc::EaseInOut:			return FMath::InterpEaseInOut<FloatType>(0.f, 1.f, InAlpha, BlendExp);
+	case EEasingFunc::ExpoIn:				return FMath::InterpExpoIn<FloatType>(0.f, 1.f, InAlpha);
+	case EEasingFunc::ExpoOut:				return FMath::InterpExpoOut<FloatType>(0.f, 1.f, InAlpha);
+	case EEasingFunc::ExpoInOut:			return FMath::InterpExpoInOut<FloatType>(0.f, 1.f, InAlpha);
+	case EEasingFunc::CircularIn:			return FMath::InterpCircularIn<FloatType>(0.f, 1.f, InAlpha);
+	case EEasingFunc::CircularOut:			return FMath::InterpCircularOut<FloatType>(0.f, 1.f, InAlpha);
+	case EEasingFunc::CircularInOut:		return FMath::InterpCircularInOut<FloatType>(0.f, 1.f, InAlpha);
 	}
 	return InAlpha;
 }
@@ -85,7 +86,12 @@ void UKismetMathLibrary::ReportError_Divide_IntInt()
 
 void UKismetMathLibrary::ReportError_Divide_DoubleDouble()
 {
-	FFrame::KismetExecutionMessage(TEXT("Divide by zero: ReportError_Divide_DoubleDouble"), ELogVerbosity::Warning, DivideByZeroWarning);
+	FFrame::KismetExecutionMessage(TEXT("Divide by zero: Divide_DoubleDouble"), ELogVerbosity::Warning, DivideByZeroWarning);
+}
+
+void UKismetMathLibrary::ReportError_Divide_FloatFloat()
+{
+	FFrame::KismetExecutionMessage(TEXT("Divide by zero: Divide_FloatFloat"), ELogVerbosity::Warning, DivideByZeroWarning);
 }
 
 void UKismetMathLibrary::ReportError_Divide_Int64Int64()
@@ -187,47 +193,40 @@ bool UKismetMathLibrary::RandomBoolWithWeightFromStream(float Weight, const FRan
 
 }
 
-/* This function is custom thunked, the real function is GenericDivide_FloatFloat */
-float UKismetMathLibrary::Divide_FloatFloat(float A, float B)
-{
-	check(0);
-	return 0;
-}
-
 /* This function is custom thunked, the real function is GenericPercent_FloatFloat */
-float UKismetMathLibrary::Percent_FloatFloat(float A, float B)
+double UKismetMathLibrary::Percent_FloatFloat(double A, double B)
 {
 	check(0);
 	return 0;
 }
 
-float UKismetMathLibrary::GenericPercent_FloatFloat(float A, float B)
+double UKismetMathLibrary::GenericPercent_FloatFloat(double A, double B)
 {
-	return (B != 0.f) ? FMath::Fmod(A, B) : 0.f;
+	return (B != 0.0) ? FMath::Fmod(A, B) : 0.0;
 }
 
-bool UKismetMathLibrary::InRange_FloatFloat(float Value, float Min, float Max, bool InclusiveMin, bool InclusiveMax)
+bool UKismetMathLibrary::InRange_FloatFloat(double Value, double Min, double Max, bool InclusiveMin, bool InclusiveMax)
 {
 	return ((InclusiveMin ? (Value >= Min) : (Value > Min)) && (InclusiveMax ? (Value <= Max) : (Value < Max)));
 }
 
-float UKismetMathLibrary::Hypotenuse(float Width, float Height)
+double UKismetMathLibrary::Hypotenuse(double Width, double Height)
 {
 	// This implementation avoids overflow/underflow caused by squaring width and height:
-	float Min = FMath::Abs(Width);
-	float Max = FMath::Abs(Height);
+	double Min = FMath::Abs(Width);
+	double Max = FMath::Abs(Height);
 
 	if (Min > Max)
 	{
 		Swap(Min, Max);
 	}
 
-	return (FMath::IsNearlyZero(Max) ? 0.f : Max * FMath::Sqrt(1.f + FMath::Square(Min/Max)));
+	return (FMath::IsNearlyZero(Max) ? 0.0 : Max * FMath::Sqrt(1.0 + FMath::Square(Min/Max)));
 }
 
-float UKismetMathLibrary::Log(float A, float Base)
+double UKismetMathLibrary::Log(double A, double Base)
 {
-	float Result = 0.f;
+	double Result = 0.f;
 	if(Base <= 0.f)
 	{
 		FFrame::KismetExecutionMessage(TEXT("Divide by zero: Log"), ELogVerbosity::Warning, DivideByZeroWarning);
@@ -259,7 +258,47 @@ int32 UKismetMathLibrary::FMod(float Dividend, float Divisor, float& Remainder)
 	return Result;
 }
 
-float UKismetMathLibrary::NormalizeToRange(float Value, float RangeMin, float RangeMax)
+int32 UKismetMathLibrary::FMod(double Dividend, double Divisor, double& Remainder)
+{
+	int32 Result;
+	if (Divisor != 0.f)
+	{
+		const double Quotient = Dividend / Divisor;
+		Result = (Quotient < 0.f ? -1 : 1) * FMath::FloorToInt(FMath::Abs(Quotient));
+		Remainder = FMath::Fmod(Dividend, Divisor);
+	}
+	else
+	{
+		FFrame::KismetExecutionMessage(TEXT("Attempted modulo 0 - returning 0."), ELogVerbosity::Warning, DivideByZeroWarning);
+
+		Result = 0;
+		Remainder = 0.f;
+	}
+
+	return Result;
+}
+
+int64 UKismetMathLibrary::FMod64(double Dividend, double Divisor, double& Remainder)
+{
+	int64 Result;
+	if (Divisor != 0.f)
+	{
+		const double Quotient = Dividend / Divisor;
+		Result = (Quotient < 0.f ? -1 : 1) * FMath::FloorToInt64(FMath::Abs(Quotient));
+		Remainder = FMath::Fmod(Dividend, Divisor);
+	}
+	else
+	{
+		FFrame::KismetExecutionMessage(TEXT("Attempted modulo 0 - returning 0."), ELogVerbosity::Warning, DivideByZeroWarning);
+
+		Result = 0;
+		Remainder = 0.f;
+	}
+
+	return Result;
+}
+
+double UKismetMathLibrary::NormalizeToRange(double Value, double RangeMin, double RangeMax)
 {
 	if (RangeMin == RangeMax)
 	{
@@ -280,19 +319,19 @@ float UKismetMathLibrary::NormalizeToRange(float Value, float RangeMin, float Ra
 	return (Value - RangeMin) / (RangeMax - RangeMin);
 }
 
-float UKismetMathLibrary::MapRangeUnclamped(float Value, float InRangeA, float InRangeB, float OutRangeA, float OutRangeB)
+double UKismetMathLibrary::MapRangeUnclamped(double Value, double InRangeA, double InRangeB, double OutRangeA, double OutRangeB)
 {
-	return FMath::GetMappedRangeValueUnclamped(FVector2f(InRangeA, InRangeB), FVector2f(OutRangeA, OutRangeB), Value);
+	return FMath::GetMappedRangeValueUnclamped(FVector2D(InRangeA, InRangeB), FVector2D(OutRangeA, OutRangeB), Value);
 }
 
-float UKismetMathLibrary::MapRangeClamped(float Value, float InRangeA, float InRangeB, float OutRangeA, float OutRangeB)
+double UKismetMathLibrary::MapRangeClamped(double Value, double InRangeA, double InRangeB, double OutRangeA, double OutRangeB)
 {
-	return FMath::GetMappedRangeValueClamped(FVector2f(InRangeA, InRangeB), FVector2f(OutRangeA, OutRangeB), Value);
+	return FMath::GetMappedRangeValueClamped(FVector2D(InRangeA, InRangeB), FVector2D(OutRangeA, OutRangeB), Value);
 }
 
-float UKismetMathLibrary::FInterpEaseInOut(float A, float B, float Alpha, float Exponent)
+double UKismetMathLibrary::FInterpEaseInOut(double A, double B, double Alpha, double Exponent)
 {
-	return FMath::InterpEaseInOut<float>(A, B, Alpha, Exponent);
+	return FMath::InterpEaseInOut<double>(A, B, Alpha, Exponent);
 }
 
 float UKismetMathLibrary::MakePulsatingValue(float InCurrentTime, float InPulsesPerSecond, float InPhase)
@@ -300,7 +339,7 @@ float UKismetMathLibrary::MakePulsatingValue(float InCurrentTime, float InPulses
 	return FMath::MakePulsatingValue((double)InCurrentTime, InPulsesPerSecond, InPhase);
 }
 
-float UKismetMathLibrary::SafeDivide(float A, float B)
+double UKismetMathLibrary::SafeDivide(double A, double B)
 {
 	return (B != 0.0f) ? (A / B) : 0.0f;
 }
@@ -373,7 +412,7 @@ float UKismetMathLibrary::InverseLerp(float A, float B, float Value)
 	return NormalizeToRange(Value, A, B);
 }
 
-float UKismetMathLibrary::Ease(float A, float B, float Alpha, TEnumAsByte<EEasingFunc::Type> EasingFunc, float BlendExp, int32 Steps)
+double UKismetMathLibrary::Ease(double A, double B, double Alpha, TEnumAsByte<EEasingFunc::Type> EasingFunc, double BlendExp, int32 Steps)
 {
 	return Lerp(A, B, EaseAlpha(Alpha, EasingFunc, BlendExp, Steps));
 }
