Added memreport header info for location of local playcontroller(s).

#rnx
#test memreport
#pf none

diff --git a/Engine/Source/Runtime/Engine/Private/UnrealEngine.cpp b/Engine/Source/Runtime/Engine/Private/UnrealEngine.cpp
--- a/Engine/Source/Runtime/Engine/Private/UnrealEngine.cpp
+++ b/Engine/Source/Runtime/Engine/Private/UnrealEngine.cpp
@@ -7128,9 +7128,24 @@ bool UEngine::HandleMemReportDeferredCommand( const TCHAR* Cmd, FOutputDevice& A
 		ReportAr->Logf(TEXT("Config: %s"), *ConfigString);
 		ReportAr->Logf(TEXT("Device Name: %s"), *DeviceName);
 		ReportAr->Logf(TEXT("Device Profile: %s"), *DeviceProfile);
-
 		ReportAr->Logf(TEXT("CommandLine Options: %s"), FCommandLine::Get());
-		ReportAr->Logf(TEXT("Time Since Boot: %.02f Seconds") LINE_TERMINATOR, FPlatformTime::Seconds() - GStartTime);
+		ReportAr->Logf(TEXT("Time Since Boot: %.02f Seconds"), FPlatformTime::Seconds() - GStartTime);
+
+		// List Name, Location, and Rotation of each local PlayerController (useful context for client mem)
+		TArray<APlayerController*> LocalPlayerControllers;
+		GetAllLocalPlayerControllers(LocalPlayerControllers);
+		for (APlayerController* LocalPlayerController : LocalPlayerControllers)
+		{
+			FVector OutLocation;
+			FRotator OutRotation;
+			LocalPlayerController->GetPlayerViewPoint(OutLocation, OutRotation);
+			ReportAr->Logf(TEXT("Local PlayerController: %s View Location: %s View Rotation: %s"),
+				*LocalPlayerController->GetName(),
+				*OutLocation.ToString(),
+				*OutRotation.ToString());
+		}
+
+		ReportAr->Logf(TEXT(""));
 	}
 
 
