Fix for potential garbage read the end of Riff file parse.

diff --git a/Engine/Source/Runtime/Engine/Private/Audio.cpp b/Engine/Source/Runtime/Engine/Private/Audio.cpp
--- a/Engine/Source/Runtime/Engine/Private/Audio.cpp
+++ b/Engine/Source/Runtime/Engine/Private/Audio.cpp
@@ -1638,7 +1638,12 @@ bool FWaveModInfo::ReadWaveInfo( const uint8* WaveData, int32 WaveDataSize, FStr
 			FRiffListChunk* ListChunk = (FRiffListChunk*)((uint8*)RiffChunk);
 			
 			RiffChunk = (FRiffChunkOld*)(ListChunk + 1);
-			const uint8* ListChunkEnd = (uint8*)RiffChunk + ListChunk->ChunkDataSize + 4;
+			uint8* ListChunkEnd = (uint8*)RiffChunk + ListChunk->ChunkDataSize + 4;
+			if (ListChunkEnd > WaveDataEnd)
+			{
+				ListChunkEnd = (uint8*)WaveDataEnd;
+			}
+
 			while (((uint8*)RiffChunk + 8) <= ListChunkEnd)
 			{
 				// Labeled Text Chunk
