Seperate / Fix naming of the cvars declaration for QualityLevelMinLod feature.

diff --git a/Engine/Source/Runtime/Engine/Private/SkeletalMesh.cpp b/Engine/Source/Runtime/Engine/Private/SkeletalMesh.cpp
--- a/Engine/Source/Runtime/Engine/Private/SkeletalMesh.cpp
+++ b/Engine/Source/Runtime/Engine/Private/SkeletalMesh.cpp
@@ -154,7 +154,7 @@ static TAutoConsoleVariable<int32> CVarRayTracingSkeletalMeshLODBias(
 const TCHAR* GSkeletalMeshMinLodQualityLevelCVarName = TEXT("r.SkeletalMesh.MinLodQualityLevel");
 const TCHAR* GSkeletalMeshMinLodQualityLevelScalabilitySection = TEXT("ViewDistanceQuality");
 int32 GSkeletalMeshMinLodQualityLevel = -1;
-static FAutoConsoleVariableRef CVarMinLodQualityLevel(
+static FAutoConsoleVariableRef CVarSkeletalMeshMinLodQualityLevel(
 	GSkeletalMeshMinLodQualityLevelCVarName,
 	GSkeletalMeshMinLodQualityLevel,
 	TEXT("The quality level for the Min stripping LOD. \n"),
diff --git a/Engine/Source/Runtime/Engine/Private/StaticMesh.cpp b/Engine/Source/Runtime/Engine/Private/StaticMesh.cpp
--- a/Engine/Source/Runtime/Engine/Private/StaticMesh.cpp
+++ b/Engine/Source/Runtime/Engine/Private/StaticMesh.cpp
@@ -153,7 +153,7 @@ static TAutoConsoleVariable<int32> CVarStripDistanceFieldDataDuringLoad(
 const TCHAR* GMinLodQualityLevelCVarName = TEXT("r.StaticMesh.MinLodQualityLevel");
 const TCHAR* GMinLodQualityLevelScalabilitySection = TEXT("ViewDistanceQuality");
 int32 GMinLodQualityLevel = -1;
-static FAutoConsoleVariableRef CVarMinLodQualityLevel(
+static FAutoConsoleVariableRef CVarStaticMeshMinLodQualityLevel(
 	GMinLodQualityLevelCVarName,
 	GMinLodQualityLevel,
 	TEXT("The quality level for the Min stripping LOD. \n"),
