StreamingTextureManager : reduce log spew

diff --git a/Engine/Source/Runtime/Engine/Private/Streaming/StreamingManagerTexture.cpp b/Engine/Source/Runtime/Engine/Private/Streaming/StreamingManagerTexture.cpp
--- a/Engine/Source/Runtime/Engine/Private/Streaming/StreamingManagerTexture.cpp
+++ b/Engine/Source/Runtime/Engine/Private/Streaming/StreamingManagerTexture.cpp
@@ -1811,8 +1811,11 @@ int32 FRenderAssetStreamingManager::BlockTillAllRequestsFinished( float TimeLimi
 		{
 			if (bLogResults)
 			{
-				UE_LOG(LogContentStreaming, Log, TEXT("Blocking on texture streaming: %.1f ms (%d still in flight)"), (float)(FPlatformTime::Seconds() - StartTime) * 1000, NumOfInFlights);
-
+				float BlockedMillis = (float)(FPlatformTime::Seconds() - StartTime) * 1000;
+				if ( BlockedMillis > 0.1f )
+				{
+					UE_LOG(LogContentStreaming, Log, TEXT("Blocking on texture streaming: %.1f ms (%d still in flight)"), BlockedMillis, NumOfInFlights);
+				}
 			}
 			return NumOfInFlights;
 		}
