fix OutBuildSettings.LossyCompressionAmount was not being set

diff --git a/Engine/Source/Runtime/Engine/Private/TextureDerivedData.cpp b/Engine/Source/Runtime/Engine/Private/TextureDerivedData.cpp
--- a/Engine/Source/Runtime/Engine/Private/TextureDerivedData.cpp
+++ b/Engine/Source/Runtime/Engine/Private/TextureDerivedData.cpp
@@ -796,9 +796,9 @@ static void GetTextureBuildSettings(
 	OutBuildSettings.OodleTextureSdkVersion = ConditionalRemapOodleTextureSdkVersion(Texture.OodleTextureSdkVersion,&TargetPlatform);
 
 	// if LossyCompressionAmount is Default, inherit from LODGroup :
-	const FTextureLODGroup& LODGroup = TextureLODSettings.GetTextureLODGroup(Texture.LODGroup);
-	if ( OutBuildSettings.LossyCompressionAmount == TLCA_Default )
+	if ( Texture.LossyCompressionAmount == TLCA_Default )
 	{
+		const FTextureLODGroup& LODGroup = TextureLODSettings.GetTextureLODGroup(Texture.LODGroup);
 		OutBuildSettings.LossyCompressionAmount = LODGroup.LossyCompressionAmount;
 		if (OutBuildResultMetadata)
 		{
