VirtualTextureDataBuilder log warning when crunch is requested but forced off

diff --git a/Engine/Source/Runtime/Engine/Private/VT/VirtualTextureDataBuilder.cpp b/Engine/Source/Runtime/Engine/Private/VT/VirtualTextureDataBuilder.cpp
--- a/Engine/Source/Runtime/Engine/Private/VT/VirtualTextureDataBuilder.cpp
+++ b/Engine/Source/Runtime/Engine/Private/VT/VirtualTextureDataBuilder.cpp
@@ -624,6 +624,8 @@ void FVirtualTextureDataBuilder::BuildTiles(const TArray<FVTSourceTileEntry>& Ti
 		case TLCA_Highest: CrunchParameters.CompressionAmmount = 1.0f; break;
 		default: checkNoEntry(); CrunchParameters.CompressionAmmount = 0.0f; break;
 		}
+		
+		UE_LOG(LogVirtualTexturing, Verbose, TEXT("Crunch LCA=%d CompressionAmmount=%f"), (int)BuildSettingsLayer0.LossyCompressionAmount, CrunchParameters.CompressionAmmount);
 
 		// We can't split crunch compression into multiple tasks/threads, since all tiles need to compress together in order to generate the codec payload
 		// Instead we rely on internal Crunch threading to make this efficient
@@ -1303,13 +1305,22 @@ void FVirtualTextureDataBuilder::BuildSourcePixels(const FTextureSourceData& Sou
 
 		bool bUseCrunch = false;
 #if WITH_CRUNCH_COMPRESSION
+		if ( SettingsPerLayer[0].bVirtualTextureEnableCompressCrunch &&
+			SettingsPerLayer[0].LossyCompressionAmount != TLCA_None )
 		{
-			// Crunch compressor crashes if given size larger than 8k by 8k, needs to be updated to use 64bit sizes/offsets in various places
-			const bool bCrunchSizeValid = OutData.Width * OutData.Height <= FMath::Square(8u * 1024u);
-			bUseCrunch = bCrunchSizeValid &&
-				SettingsPerLayer[0].bVirtualTextureEnableCompressCrunch &&
-				SettingsPerLayer[0].LossyCompressionAmount != TLCA_None &&
-				CrunchCompression::IsValidFormat(TextureFormatName);
+			if ( ! CrunchCompression::IsValidFormat(TextureFormatName) )
+			{
+				UE_LOG(LogVirtualTexturing, Warning, TEXT("VT Format invalid for Crunch, disabled") );				
+			}
+			else if ( OutData.Width * OutData.Height > FMath::Square(8u * 1024u) )
+			{
+				// Crunch compressor crashes if given size larger than 8k by 8k, needs to be updated to use 64bit sizes/offsets in various places
+				UE_LOG(LogVirtualTexturing, Warning, TEXT("VT too big for Crunch, disabled") );
+			}
+			else
+			{
+				bUseCrunch = true;
+			}
 		}
 #endif // WITH_CRUNCH_COMPRESSION
 
