LossyCompressionAmount option on RuntimeVirtualTextureComponent
for control of Oodle RDO lambda

diff --git a/Engine/Source/Runtime/Engine/Classes/Components/RuntimeVirtualTextureComponent.h b/Engine/Source/Runtime/Engine/Classes/Components/RuntimeVirtualTextureComponent.h
--- a/Engine/Source/Runtime/Engine/Classes/Components/RuntimeVirtualTextureComponent.h
+++ b/Engine/Source/Runtime/Engine/Classes/Components/RuntimeVirtualTextureComponent.h
@@ -4,6 +4,7 @@
 
 #include "CoreMinimal.h"
 #include "SceneComponent.h"
+#include "Engine/TextureDefines.h"
 #include "RuntimeVirtualTextureComponent.generated.h"
 
 class URuntimeVirtualTexture;
@@ -64,6 +65,13 @@ protected:
 	/** Enable Crunch texture compression for the streaming low mips. Generic ZLib compression is used when Crunch is disabled. */
 	UPROPERTY(EditAnywhere, AdvancedDisplay, Category = VirtualTextureBuild, meta = (DisplayName = "Enable Crunch"))
 	bool bEnableCompressCrunch = false;
+	
+	/** How aggressively should any relevant lossy compression be applied. For compressors that support EncodeSpeed (i.e. Oodle), this is only
+	*	applied if enabled (see Project Settings -> Texture Encoding). Note that this is *in addition* to any
+	*	unavoidable loss due to the target format - selecting "No Lossy Compression" will not result in zero distortion for BCn formats.
+	*/
+	UPROPERTY(EditAnywhere, AdvancedDisplay, Category = VirtualTextureBuild)
+	TEnumAsByte<ETextureLossyCompressionAmount> LossyCompressionAmount = TLCA_Default;
 
 	/** Build the streaming low mips using debug coloring. This can help show where streaming mips are being used. */
 	UPROPERTY(EditAnywhere, BlueprintReadOnly, Transient, AdvancedDisplay, Category = VirtualTextureBuild, meta = (DisplayName = "Build Debug"))
diff --git a/Engine/Source/Runtime/Engine/Classes/VT/VirtualTextureBuilder.h b/Engine/Source/Runtime/Engine/Classes/VT/VirtualTextureBuilder.h
--- a/Engine/Source/Runtime/Engine/Classes/VT/VirtualTextureBuilder.h
+++ b/Engine/Source/Runtime/Engine/Classes/VT/VirtualTextureBuilder.h
@@ -23,6 +23,8 @@ struct FVirtualTextureBuildDesc
 	TEnumAsByte<enum TextureGroup> LODGroup;
 
 	bool bCrunchCompressed = false;
+	ETextureLossyCompressionAmount LossyCompressionAmount = TLCA_Default;
+
 	bool bContinuousUpdate = false;
 	bool bSinglePhysicalSpace = false;
 
diff --git a/Engine/Source/Runtime/Engine/Private/Components/RuntimeVirtualTextureComponent.cpp b/Engine/Source/Runtime/Engine/Private/Components/RuntimeVirtualTextureComponent.cpp
--- a/Engine/Source/Runtime/Engine/Private/Components/RuntimeVirtualTextureComponent.cpp
+++ b/Engine/Source/Runtime/Engine/Private/Components/RuntimeVirtualTextureComponent.cpp
@@ -140,6 +140,7 @@ uint64 URuntimeVirtualTextureComponent::CalculateStreamingTextureSettingsHash()
 		uint64 PackedValue;
 		struct
 		{
+			uint32 PackedSettingsVersion : 4;
 			uint32 MaterialType : 4;
 			uint32 TileSize : 12;
 			uint32 TileBorderSize : 4;
@@ -149,11 +150,13 @@ uint64 URuntimeVirtualTextureComponent::CalculateStreamingTextureSettingsHash()
 			uint32 EnableCompressCrunch : 1;
 			uint32 ContinuousUpdate : 1;
 			uint32 bUseLowQualityCompression : 1;
+			uint32 LossyCompressionAmount : 4;
 		};
 	};
 
 	FPackedSettings Settings;
 	Settings.PackedValue = 0;
+	Settings.PackedSettingsVersion = 1;
 	Settings.MaterialType = (uint32)VirtualTexture->GetMaterialType();
 	Settings.TileSize = (uint32)VirtualTexture->GetTileSize();
 	Settings.TileBorderSize = (uint32)VirtualTexture->GetTileBorderSize();
@@ -161,8 +164,9 @@ uint64 URuntimeVirtualTextureComponent::CalculateStreamingTextureSettingsHash()
 	Settings.CompressTextures = (uint32)VirtualTexture->GetCompressTextures();
 	Settings.ContinuousUpdate = (uint32)VirtualTexture->GetContinuousUpdate();
 	Settings.SinglePhysicalSpace = (uint32)VirtualTexture->GetSinglePhysicalSpace();
-	Settings.EnableCompressCrunch = (uint32)bEnableCompressCrunch;
 	Settings.bUseLowQualityCompression = (uint32)VirtualTexture->GetLQCompression();
+	Settings.EnableCompressCrunch = (uint32)bEnableCompressCrunch;
+	Settings.LossyCompressionAmount = (uint32) LossyCompressionAmount.GetValue();
 
 	return Settings.PackedValue;
 }
@@ -243,6 +247,7 @@ void URuntimeVirtualTextureComponent::InitializeStreamingTexture(uint32 InSizeX,
 		BuildDesc.TileBorderSize = VirtualTexture->GetTileBorderSize();
 		BuildDesc.LODGroup = VirtualTexture->GetLODGroup();
 		BuildDesc.bCrunchCompressed = bEnableCompressCrunch;
+		BuildDesc.LossyCompressionAmount = LossyCompressionAmount;
 
 		BuildDesc.LayerCount = VirtualTexture->GetLayerCount();
 		check(BuildDesc.LayerCount <= RuntimeVirtualTexture::MaxTextureLayers);
diff --git a/Engine/Source/Runtime/Engine/Private/VT/VirtualTextureBuilder.cpp b/Engine/Source/Runtime/Engine/Private/VT/VirtualTextureBuilder.cpp
--- a/Engine/Source/Runtime/Engine/Private/VT/VirtualTextureBuilder.cpp
+++ b/Engine/Source/Runtime/Engine/Private/VT/VirtualTextureBuilder.cpp
@@ -44,6 +44,7 @@ void UVirtualTextureBuilder::BuildTexture(FVirtualTextureBuildDesc const& BuildD
 	Texture->Settings.TileSize = BuildDesc.TileSize;
 	Texture->Settings.TileBorderSize = BuildDesc.TileBorderSize;
 	Texture->Settings.bEnableCompressCrunch = BuildDesc.bCrunchCompressed;
+	Texture->LossyCompressionAmount = BuildDesc.LossyCompressionAmount;
 
 	Texture->bContinuousUpdate = BuildDesc.bContinuousUpdate;
 	Texture->bSinglePhysicalSpace = BuildDesc.bSinglePhysicalSpace;
