Fix visual glitches when removing foliage instances in editor.
#rnx

diff --git a/Engine/Source/Runtime/Engine/Private/InstancedStaticMesh.cpp b/Engine/Source/Runtime/Engine/Private/InstancedStaticMesh.cpp
--- a/Engine/Source/Runtime/Engine/Private/InstancedStaticMesh.cpp
+++ b/Engine/Source/Runtime/Engine/Private/InstancedStaticMesh.cpp
@@ -1398,22 +1398,26 @@ void FInstancedStaticMeshSceneProxy::SetupProxy(UInstancedStaticMeshComponent* I
 	{
 		const TArray<int32>& InstanceReorderTable = InComponent->InstanceReorderTable;
 
+		// NumRenderInstances is the extent that the reorder table can map to.
+		// Temporarily when removing instances from a HISM this can be sparse so that InComponent->GetInstanceCount() < NumRenderInstances.
+		const int32 NumRenderInstances = FMath::Max<int32>(InComponent->PerInstanceRenderData->InstanceBuffer.GetNumInstances(), InComponent->GetInstanceCount());
+
 		bSupportsInstanceDataBuffer = true;
-		InstanceSceneData.SetNum(InComponent->GetInstanceCount());
+		InstanceSceneData.SetNumZeroed(NumRenderInstances);
 
 		bHasPerInstanceDynamicData = InComponent->PerInstancePrevTransform.Num() == InComponent->GetInstanceCount();
-		InstanceDynamicData.SetNumUninitialized(bHasPerInstanceDynamicData ? InComponent->GetInstanceCount() : 0);
+		InstanceDynamicData.SetNumUninitialized(bHasPerInstanceDynamicData ? NumRenderInstances : 0);
 
 		static const auto AllowStaticLightingVar = IConsoleManager::Get().FindTConsoleVariableDataInt(TEXT("r.AllowStaticLighting"));
 		const bool bAllowStaticLighting = (!AllowStaticLightingVar || AllowStaticLightingVar->GetValueOnAnyThread() != 0);
 		bHasPerInstanceLMSMUVBias = bAllowStaticLighting;
-		InstanceLightShadowUVBias.SetNumZeroed(bHasPerInstanceLMSMUVBias ? InComponent->GetInstanceCount() : 0);
+		InstanceLightShadowUVBias.SetNumZeroed(bHasPerInstanceLMSMUVBias ? NumRenderInstances : 0);
 
-		InstanceRandomID.SetNumZeroed(bHasPerInstanceRandom ? InComponent->GetInstanceCount() : 0); // Only allocate if material bound which uses this
+		InstanceRandomID.SetNumZeroed(bHasPerInstanceRandom ? NumRenderInstances : 0); // Only allocate if material bound which uses this
 
 #if WITH_EDITOR
 		bHasPerInstanceEditorData = true;
-		InstanceEditorData.SetNumZeroed(bHasPerInstanceEditorData ? InComponent->GetInstanceCount() : 0);
+		InstanceEditorData.SetNumZeroed(bHasPerInstanceEditorData ? NumRenderInstances : 0);
 #endif
 
 		// Only allocate if material bound which uses this
@@ -1427,43 +1431,38 @@ void FInstancedStaticMeshSceneProxy::SetupProxy(UInstancedStaticMeshComponent* I
 			bHasPerInstanceCustomData = false;
 		}
 
-		for (int32 InInstanceIndex = 0; InInstanceIndex < InstanceSceneData.Num(); ++InInstanceIndex)
-		{
-			// Make sure the instance is initialized, regardless of below remapping
-			{
-				FPrimitiveInstance& TmpSceneData = InstanceSceneData[InInstanceIndex];
-				TmpSceneData.LocalToPrimitive = FRenderTransform::Identity;
-			}
+		FVector TranslatedSpaceOffset = -InComponent->GetTranslatedInstanceSpaceOrigin();
 
-			int32 OutInstanceIndex = InInstanceIndex;
-			// GPUCULL_TODO: After deleting instances in a HISM the InstanceReorderTable often contains nonsense, this is then corrected
-			// by the async build, which re-creates the proxy in a nearby future frame. All this should be removed in favour of GPU-side culling.
-			if (OutInstanceIndex < InstanceReorderTable.Num() && InstanceReorderTable[OutInstanceIndex] < InstanceSceneData.Num())
+		// Add the visible instances. 
+		// Non visible ones (that are being removed) should not appear due to a zeroed render transform in the entries where we don't add anything here.
+		for (int32 InstanceIndex = 0; InstanceIndex < InComponent->GetInstanceCount(); ++InstanceIndex)
+		{
+			const int32 RenderInstanceIndex = InstanceReorderTable[InstanceIndex];
+			if (!ensure(RenderInstanceIndex < NumRenderInstances))
 			{
-				// Temporary workaround for out of bound array access
-				// TODO: fix this properly!!!!!!
-				OutInstanceIndex = InstanceReorderTable[OutInstanceIndex] != INDEX_NONE ? InstanceReorderTable[OutInstanceIndex] : OutInstanceIndex;
+				continue;
 			}
 
-			FPrimitiveInstance& SceneData = InstanceSceneData[OutInstanceIndex];
-
+			FPrimitiveInstance& SceneData = InstanceSceneData[RenderInstanceIndex];
 
 			FTransform InstanceTransform;
-			InComponent->GetInstanceTransform(InInstanceIndex, InstanceTransform);
+			InComponent->GetInstanceTransform(InstanceIndex, InstanceTransform);
+			InstanceTransform.AddToTranslation(TranslatedSpaceOffset);
 			SceneData.LocalToPrimitive = InstanceTransform.ToMatrixWithScale();
 
 			if (bHasPerInstanceDynamicData)
 			{
 				FTransform InstancePrevTransform;
-				const bool bHasPrevTransform = InComponent->GetInstancePrevTransform(InInstanceIndex, InstancePrevTransform);
+				const bool bHasPrevTransform = InComponent->GetInstancePrevTransform(InstanceIndex, InstancePrevTransform);
 				ensure(bHasPrevTransform); // Should always be true here
-				InstanceDynamicData[OutInstanceIndex].PrevLocalToPrimitive = InstancePrevTransform.ToMatrixWithScale();
+				InstancePrevTransform.AddToTranslation(TranslatedSpaceOffset);
+				InstanceDynamicData[RenderInstanceIndex].PrevLocalToPrimitive = InstancePrevTransform.ToMatrixWithScale();
 			}
 
 			if (bHasPerInstanceCustomData)
 			{
-				const int32 SrcCustomDataOffset = InInstanceIndex  * InComponent->NumCustomDataFloats;
-				const int32 DstCustomDataOffset = OutInstanceIndex * InComponent->NumCustomDataFloats;
+				const int32 SrcCustomDataOffset = InstanceIndex  * InComponent->NumCustomDataFloats;
+				const int32 DstCustomDataOffset = RenderInstanceIndex * InComponent->NumCustomDataFloats;
 				FMemory::Memcpy
 				(
 					&InstanceCustomData[DstCustomDataOffset],
