fix WaitOnTextureCompilation message showing %s

diff --git a/Engine/Source/Runtime/Engine/Private/Texture2D.cpp b/Engine/Source/Runtime/Engine/Private/Texture2D.cpp
--- a/Engine/Source/Runtime/Engine/Private/Texture2D.cpp
+++ b/Engine/Source/Runtime/Engine/Private/Texture2D.cpp
@@ -274,7 +274,7 @@ const FTexturePlatformData* UTexture2D::GetPlatformData() const
 		TRACE_CPUPROFILER_EVENT_SCOPE(UTexture2D::GetPlatformDataStall);
 		UE_LOG(LogTexture, Log, TEXT("Call to GetPlatformData() is forcing a wait on data that is not yet ready."));
 
-		FText Msg = FText::Format(LOCTEXT("WaitOnTextureCompilation", "Waiting on texture compilation (%s) ..."), FText::FromString(GetName()));
+		FText Msg = FText::Format(LOCTEXT("WaitOnTextureCompilation", "Waiting on texture compilation {0} ..."), FText::FromString(GetName()));
 		FScopedSlowTask Progress(1.f, Msg, true);
 		Progress.MakeDialog(true);
 		uint64 StartTime = FPlatformTime::Cycles64();
diff --git a/Engine/Source/Runtime/Engine/Private/TextureCube.cpp b/Engine/Source/Runtime/Engine/Private/TextureCube.cpp
--- a/Engine/Source/Runtime/Engine/Private/TextureCube.cpp
+++ b/Engine/Source/Runtime/Engine/Private/TextureCube.cpp
@@ -104,7 +104,7 @@ const FTexturePlatformData* UTextureCube::GetPlatformData() const
 		TRACE_CPUPROFILER_EVENT_SCOPE(UTextureCube::GetPlatformDataStall);
 		UE_LOG(LogTexture, Log, TEXT("Call to GetPlatformData() is forcing a wait on data that is not yet ready."));
 
-		FText Msg = FText::Format(LOCTEXT("WaitOnTextureCompilation", "Waiting on texture compilation (%s) ..."), FText::FromString(GetName()));
+		FText Msg = FText::Format(LOCTEXT("WaitOnTextureCompilation", "Waiting on texture compilation {0} ..."), FText::FromString(GetName()));
 		FScopedSlowTask Progress(1.f, Msg, true);
 		Progress.MakeDialog(true);
 		PrivatePlatformData->FinishCache();
