World Partition: handle the case where the world is duplicated for PIE travel which also duplicates always loaded actors.
Really just replaced a bad assert with an explanatory comment.

diff --git a/Engine/Source/Runtime/Engine/Private/WorldPartition/WorldPartitionActorDesc.cpp b/Engine/Source/Runtime/Engine/Private/WorldPartition/WorldPartitionActorDesc.cpp
--- a/Engine/Source/Runtime/Engine/Private/WorldPartition/WorldPartitionActorDesc.cpp
+++ b/Engine/Source/Runtime/Engine/Private/WorldPartition/WorldPartitionActorDesc.cpp
@@ -366,9 +366,16 @@ void FWorldPartitionActorDesc::Unload()
 {
 	if (AActor* Actor = GetActor())
 	{
-		// @todo_ow: FWorldPartitionCookPackageSplitter should mark each FWorldPartitionActorDesc as moved, which we would assert on here rather than asserting IsRunningCookCommandlet
-		// and the splitter should take responsbility for calling ClearFlags on every object in the package when it does the move
-		check(Actor->IsPackageExternal() || IsRunningCookCommandlet());
+		// At this point, it can happen that an actor isn't in an external package:
+		//
+		// PIE travel: 
+		//		in this case, actors referenced by the world package (an example is the level script) will be duplicated as part of the PIE world duplication and will end up
+		//		not being using an external package, which is fine because in that case they are considered as always loaded.
+		//
+		// FWorldPartitionCookPackageSplitter:
+		//		should mark each FWorldPartitionActorDesc as moved, and the splitter should take responsbility for calling ClearFlags on every object in 
+		//		the package when it does the move
+
 		if (Actor->IsPackageExternal())
 		{
 			ForEachObjectWithPackage(Actor->GetPackage(), [](UObject* Object)
