Allow linked animation graphs to be set to a null anim blueprint class both in the animation blueprint editor and at runtime when using UAnimInstance::LinkAnimGraphByTag

#tests BP Compiler + various linked anim layers cases

diff --git a/Engine/Source/Editor/AnimGraph/Public/AnimGraphNode_LinkedAnimGraphBase.h b/Engine/Source/Editor/AnimGraph/Public/AnimGraphNode_LinkedAnimGraphBase.h
--- a/Engine/Source/Editor/AnimGraph/Public/AnimGraphNode_LinkedAnimGraphBase.h
+++ b/Engine/Source/Editor/AnimGraph/Public/AnimGraphNode_LinkedAnimGraphBase.h
@@ -45,6 +45,7 @@ public:
 
 	// UAnimGraphNode_CustomProperty interface
 	virtual bool IsStructuralProperty(FProperty* InProperty) const override;
+	virtual bool NeedsToSpecifyValidTargetClass() const override { return false; }
 
 	// Node accessor
 	virtual FAnimNode_LinkedAnimGraph* GetLinkedAnimGraphNode() PURE_VIRTUAL(UAnimGraphNode_LinkedAnimGraphBase::GetLinkedAnimGraphNode, return nullptr;);
diff --git a/Engine/Source/Editor/AnimGraph/Public/AnimGraphNode_LinkedAnimLayer.h b/Engine/Source/Editor/AnimGraph/Public/AnimGraphNode_LinkedAnimLayer.h
--- a/Engine/Source/Editor/AnimGraph/Public/AnimGraphNode_LinkedAnimLayer.h
+++ b/Engine/Source/Editor/AnimGraph/Public/AnimGraphNode_LinkedAnimLayer.h
@@ -72,7 +72,6 @@ protected:
 	// Begin UAnimGraphNode_CustomProperty
 	virtual FAnimNode_CustomProperty* GetCustomPropertyNode() override { return &Node;  }
 	virtual const FAnimNode_CustomProperty* GetCustomPropertyNode() const override { return &Node; }
-	virtual bool NeedsToSpecifyValidTargetClass() const override { return false; }
 	virtual UClass* GetTargetSkeletonClass() const override;
 
 	// Begin UAnimGraphNode_LinkedAnimGraphBase
diff --git a/Engine/Source/Runtime/Engine/Classes/Animation/AnimNode_LinkedAnimGraph.h b/Engine/Source/Runtime/Engine/Classes/Animation/AnimNode_LinkedAnimGraph.h
--- a/Engine/Source/Runtime/Engine/Classes/Animation/AnimNode_LinkedAnimGraph.h
+++ b/Engine/Source/Runtime/Engine/Classes/Animation/AnimNode_LinkedAnimGraph.h
@@ -111,7 +111,7 @@ protected:
 	void ReinitializeLinkedAnimInstance(const UAnimInstance* InOwningAnimInstance, UAnimInstance* InNewAnimInstance = nullptr);
 
 	// Shutdown the currently running instance
-	void TeardownInstance();
+	void TeardownInstance(const UAnimInstance* InOwningAnimInstance);
 
 	// Get Target Class
 	virtual UClass* GetTargetClass() const override 
diff --git a/Engine/Source/Runtime/Engine/Private/Animation/AnimNode_LinkedAnimGraph.cpp b/Engine/Source/Runtime/Engine/Private/Animation/AnimNode_LinkedAnimGraph.cpp
--- a/Engine/Source/Runtime/Engine/Private/Animation/AnimNode_LinkedAnimGraph.cpp
+++ b/Engine/Source/Runtime/Engine/Private/Animation/AnimNode_LinkedAnimGraph.cpp
@@ -206,7 +206,7 @@ void FAnimNode_LinkedAnimGraph::OnInitializeAnimInstance(const FAnimInstanceProx
 	else if(InstanceToRun)
 	{
 		// We have an instance but no instance class
-		TeardownInstance();
+		TeardownInstance(InAnimInstance);
 	}
 
 	if(InstanceToRun)
@@ -215,12 +215,23 @@ void FAnimNode_LinkedAnimGraph::OnInitializeAnimInstance(const FAnimInstanceProx
 	}
 }
 
-void FAnimNode_LinkedAnimGraph::TeardownInstance()
+void FAnimNode_LinkedAnimGraph::TeardownInstance(const UAnimInstance* InOwningAnimInstance)
 {
 	UAnimInstance* InstanceToRun = GetTargetInstance<UAnimInstance>();
 	if (InstanceToRun)
 	{
-		InstanceToRun->UninitializeAnimation();
+		DynamicUnlink(const_cast<UAnimInstance*>(InOwningAnimInstance));
+		// Never delete the owning animation instance
+		if (InstanceToRun != InOwningAnimInstance)
+		{
+			USkeletalMeshComponent* MeshComp = InOwningAnimInstance->GetSkelMeshComponent();
+			check(MeshComp);
+			MeshComp->GetLinkedAnimInstances().Remove(InstanceToRun);
+			// Only call UninitializeAnimation if we are not the owning anim instance
+			InstanceToRun->UninitializeAnimation();
+			InstanceToRun->MarkAsGarbage();
+		}
+
 		InstanceToRun = nullptr;
 	}
 
@@ -233,25 +244,13 @@ void FAnimNode_LinkedAnimGraph::ReinitializeLinkedAnimInstance(const UAnimInstan
 
 	IAnimClassInterface* PriorAnimBPClass = InstanceToRun ? IAnimClassInterface::GetFromClass(InstanceToRun->GetClass()) : nullptr;
 
+	// Full reinit, kill old instances
+	TeardownInstance(InOwningAnimInstance);
+
 	if(*InstanceClass || InNewAnimInstance)
 	{
 		USkeletalMeshComponent* MeshComp = InOwningAnimInstance->GetSkelMeshComponent();
 		check(MeshComp);
-		// Full reinit, kill old instances
-		if(InstanceToRun)
-		{
-			DynamicUnlink(const_cast<UAnimInstance*>(InOwningAnimInstance));
-
-			// Never delete the owning animation instance
-			if (InstanceToRun != InOwningAnimInstance)
-			{
-				MeshComp->GetLinkedAnimInstances().Remove(InstanceToRun);
-				// Only call UninitializeAnimation if we are not the owning anim instance
-				InstanceToRun->UninitializeAnimation();
-				InstanceToRun->MarkAsGarbage();
-			}
-			InstanceToRun = nullptr;
-		}
 
 		// Need an instance to run, so create it now
 		InstanceToRun = InNewAnimInstance ? InNewAnimInstance : NewObject<UAnimInstance>(MeshComp, InstanceClass);
@@ -284,11 +283,6 @@ void FAnimNode_LinkedAnimGraph::ReinitializeLinkedAnimInstance(const UAnimInstan
 
 		RequestBlend(PriorAnimBPClass, NewAnimBPClass);
 	}
-	else if(InstanceToRun)
-	{
-		// We have an instance but no instance class
-		TeardownInstance();
-	}
 }
 
 void FAnimNode_LinkedAnimGraph::SetAnimClass(TSubclassOf<UAnimInstance> InClass, const UAnimInstance* InOwningAnimInstance)
