Refactor to simplify ray tracing geometry update logic for skeletal mesh.
*) All skin cache entries needing RT geometry update are processed altogether in end-of-frame update whether batch dispatching or not, simplifies the logic.
*) Transition buffers to read immediately after skin cache DoDispatch calls, simplifies the logic.
*) No longer cache the flag to indicate whether RT geometry needs rebuild, the logic is calculated directly when updating RT geometry.

diff --git a/Engine/Source/Runtime/Engine/Private/Components/SkinnedMeshComponent.cpp b/Engine/Source/Runtime/Engine/Private/Components/SkinnedMeshComponent.cpp
--- a/Engine/Source/Runtime/Engine/Private/Components/SkinnedMeshComponent.cpp
+++ b/Engine/Source/Runtime/Engine/Private/Components/SkinnedMeshComponent.cpp
@@ -784,6 +784,9 @@ void USkinnedMeshComponent::ClearMotionVector()
 
 		++CurrentBoneTransformRevisionNumber;
 		MeshObject->Update(UseLOD, this, ActiveMorphTargets, MorphTargetWeights, EPreviousBoneTransformUpdateMode::None);  // send to rendering thread
+
+		// Skin cache may need updating
+		MarkForNeededEndOfFrameUpdate();
 	}
 }
 
@@ -800,6 +803,9 @@ void USkinnedMeshComponent::ForceMotionVector()
 
 		++CurrentBoneTransformRevisionNumber;
 		MeshObject->Update(UseLOD, this, ActiveMorphTargets, MorphTargetWeights, EPreviousBoneTransformUpdateMode::None);
+
+		// Skin cache may need updating
+		MarkForNeededEndOfFrameUpdate();
 	}
 }
 
diff --git a/Engine/Source/Runtime/Engine/Private/GPUSkinCache.cpp b/Engine/Source/Runtime/Engine/Private/GPUSkinCache.cpp
--- a/Engine/Source/Runtime/Engine/Private/GPUSkinCache.cpp
+++ b/Engine/Source/Runtime/Engine/Private/GPUSkinCache.cpp
@@ -752,12 +752,12 @@ void FGPUSkinCache::Cleanup()
 
 	while (Entries.Num() > 0)
 	{
-		ReleaseSkinCacheEntry(Entries.Last());
+		Release(Entries.Last());
 	}
 	ensure(Allocations.Num() == 0);
 }
 
-void FGPUSkinCache::TransitionAllToReadable(FRHICommandList& RHICmdList)
+void FGPUSkinCache::TransitionAllToReadable(FRHICommandList& RHICmdList, const TSet<FSkinCacheRWBuffer*>& BuffersToTransitionToRead)
 {
 	TRACE_CPUPROFILER_EVENT_SCOPE(FGPUSkinCache::TransitionAllToReadable);
 
@@ -771,8 +771,6 @@ void FGPUSkinCache::TransitionAllToReadable(FRHICommandList& RHICmdList)
 			UAVs.Add((*SetIt)->UpdateAccessState(ERHIAccess::VertexOrIndexBuffer | ERHIAccess::SRVMask));
 		}
 		RHICmdList.Transition(UAVs);
-
-		BuffersToTransitionToRead.Empty(BuffersToTransitionToRead.Num());
 	}
 }
 
@@ -1207,6 +1205,8 @@ void FGPUSkinCache::DoDispatch(FRHICommandListImmediate& RHICmdList)
 		MakeBufferTransitions(RHICmdList, BuffersToTransitionForSkinning, ERHIAccess::UAVCompute);
 	}
 
+	TSet<FSkinCacheRWBuffer*> BuffersToTransitionToRead;
+
 	TArray<FRHIUnorderedAccessView*> SkinningBuffersToOverlap;
 	GetBufferUAVs(BuffersToTransitionForSkinning, SkinningBuffersToOverlap);
 	RHICmdList.BeginUAVOverlap(SkinningBuffersToOverlap);
@@ -1215,7 +1215,7 @@ void FGPUSkinCache::DoDispatch(FRHICommandListImmediate& RHICmdList)
 		for (int32 i = 0; i < BatchCount; ++i)
 		{
 			FDispatchEntry& DispatchItem = BatchDispatches[i];
-			DispatchUpdateSkinning(RHICmdList, DispatchItem.SkinCacheEntry, DispatchItem.Section, DispatchItem.RevisionNumber);
+			DispatchUpdateSkinning(RHICmdList, DispatchItem.SkinCacheEntry, DispatchItem.Section, DispatchItem.RevisionNumber, BuffersToTransitionToRead);
 		}
 	}
 	RHICmdList.EndUAVOverlap(SkinningBuffersToOverlap);
@@ -1294,6 +1294,8 @@ void FGPUSkinCache::DoDispatch(FRHICommandListImmediate& RHICmdList)
 		FDispatchEntry& DispatchItem = BatchDispatches[i];
 		DispatchItem.SkinCacheEntry->UpdateVertexFactoryDeclaration(DispatchItem.Section);
 	}
+
+	TransitionAllToReadable(RHICmdList, BuffersToTransitionToRead);
 }
 
 void FGPUSkinCache::DoDispatch(FRHICommandListImmediate& RHICmdList, FGPUSkinCacheEntry* SkinCacheEntry, int32 Section, int32 RevisionNumber)
@@ -1305,6 +1307,8 @@ void FGPUSkinCache::DoDispatch(FRHICommandListImmediate& RHICmdList, FGPUSkinCac
 
 	INC_DWORD_STAT(STAT_GPUSkinCache_TotalNumChunks);
 
+	TSet<FSkinCacheRWBuffer*> BuffersToTransitionToRead;
+
 	TArray<FSkinCacheRWBuffer*> BuffersToTransitionForSkinning;
 	PrepareUpdateSkinning(SkinCacheEntry, Section, RevisionNumber, &BuffersToTransitionForSkinning);
 	MakeBufferTransitions(RHICmdList, BuffersToTransitionForSkinning, ERHIAccess::UAVCompute);
@@ -1313,7 +1317,7 @@ void FGPUSkinCache::DoDispatch(FRHICommandListImmediate& RHICmdList, FGPUSkinCac
 	GetBufferUAVs(BuffersToTransitionForSkinning, SkinningBuffersToOverlap);
 	RHICmdList.BeginUAVOverlap(SkinningBuffersToOverlap);
 	{
-	DispatchUpdateSkinning(RHICmdList, SkinCacheEntry, Section, RevisionNumber);
+		DispatchUpdateSkinning(RHICmdList, SkinCacheEntry, Section, RevisionNumber, BuffersToTransitionToRead);
 	}
 	RHICmdList.EndUAVOverlap(SkinningBuffersToOverlap);
 
@@ -1343,6 +1347,8 @@ void FGPUSkinCache::DoDispatch(FRHICommandListImmediate& RHICmdList, FGPUSkinCac
 	}
 
 	SkinCacheEntry->UpdateVertexFactoryDeclaration(Section);
+
+	TransitionAllToReadable(RHICmdList, BuffersToTransitionToRead);
 }
 
 bool FGPUSkinCache::ProcessEntry(
@@ -1547,13 +1553,21 @@ bool FGPUSkinCache::ProcessEntry(
 
 	if (bShouldBatchDispatches)
 	{
-		BatchDispatches.Add({ InOutEntry, &LodData, RevisionNumber, uint32(Section) });
+		BatchDispatches.Add({ InOutEntry, RevisionNumber, uint32(Section) });
 	}
 	else
 	{
 		DoDispatch(RHICmdList, InOutEntry, Section, RevisionNumber);
 	}
 
+#if RHI_RAYTRACING
+	if (!Skin->ShouldUseSeparateSkinCacheEntryForRayTracing() || Mode == EGPUSkinCacheEntryMode::RayTracing)
+	{
+		// This is a RT skin cache entry
+		PendingProcessRTGeometryEntries.Add(InOutEntry);
+	}
+#endif
+
 	return true;
 }
 
@@ -1570,20 +1584,19 @@ bool FGPUSkinCache::IsGPUSkinCacheRayTracingSupported()
 
 #if RHI_RAYTRACING
 
-void FGPUSkinCache::ProcessRayTracingGeometryToUpdate(FRHICommandListImmediate& RHICmdList, FGPUSkinCacheEntry* SkinCacheEntry, FSkeletalMeshLODRenderData& LODModel)
+void FGPUSkinCache::ProcessRayTracingGeometryToUpdate(FRHICommandListImmediate& RHICmdList, FGPUSkinCacheEntry* SkinCacheEntry)
 {
 	if (IsGPUSkinCacheRayTracingSupported() && SkinCacheEntry && SkinCacheEntry->GPUSkin && SkinCacheEntry->GPUSkin->bSupportRayTracing)
 	{
-		if (SkinCacheEntry->GPUSkin->bRequireRecreatingRayTracingGeometry)
-		{
-			// We will need to build a new BVH so flush pending skin cache resource barriers.
-			TransitionAllToReadable(RHICmdList);
-		}
-
  		TArray<FBufferRHIRef> VertexBufffers;
  		SkinCacheEntry->GetRayTracingSegmentVertexBuffers(VertexBufffers);
 
- 		SkinCacheEntry->GPUSkin->UpdateRayTracingGeometry(LODModel, VertexBufffers);
+		const int32 LODIndex = SkinCacheEntry->LOD;
+		FSkeletalMeshRenderData& SkelMeshRenderData = SkinCacheEntry->GPUSkin->GetSkeletalMeshRenderData();
+		check(LODIndex < SkelMeshRenderData.LODRenderData.Num());
+		FSkeletalMeshLODRenderData& LODModel = SkelMeshRenderData.LODRenderData[LODIndex];
+
+ 		SkinCacheEntry->GPUSkin->UpdateRayTracingGeometry(LODModel, LODIndex, VertexBufffers);
 	}
 }
 
@@ -1605,31 +1618,12 @@ void FGPUSkinCache::EndBatchDispatch(FRHICommandListImmediate& RHICmdList)
 #if RHI_RAYTRACING
 	if (IsGPUSkinCacheRayTracingSupported())
 	{
-		TSet<FGPUSkinCacheEntry*> SkinCacheEntriesProcessed;
-
-		// Process batched dispatches in reverse order to filter out duplicated ones and keep the last one
-		for (int32 Index = BatchDispatches.Num() - 1; Index >= 0; Index--)
+		for (FGPUSkinCacheEntry* SkinCacheEntry : PendingProcessRTGeometryEntries)
 		{
-			FDispatchEntry& DispatchItem = BatchDispatches[Index];
-
-			FGPUSkinCacheEntry* SkinCacheEntry = DispatchItem.SkinCacheEntry;
-
-			if (SkinCacheEntry->GPUSkin->ShouldUseSeparateSkinCacheEntryForRayTracing() && SkinCacheEntry->Mode != EGPUSkinCacheEntryMode::RayTracing)
-			{
-				continue;
-			}
-
-			FSkeletalMeshLODRenderData& LODModel = *DispatchItem.LODModel;
-
-			if (SkinCacheEntriesProcessed.Contains(SkinCacheEntry))
-			{
-				continue;
-			}
-
-			SkinCacheEntriesProcessed.Add(SkinCacheEntry);
-
-			ProcessRayTracingGeometryToUpdate(RHICmdList, SkinCacheEntry, LODModel);
+			ProcessRayTracingGeometryToUpdate(RHICmdList, SkinCacheEntry);
 		}
+
+		PendingProcessRTGeometryEntries.Reset();
 	}
 #endif
 
@@ -1641,6 +1635,10 @@ void FGPUSkinCache::Release(FGPUSkinCacheEntry*& SkinCacheEntry)
 {
 	if (SkinCacheEntry)
 	{
+		FGPUSkinCache* SkinCache = SkinCacheEntry->SkinCache;
+		check(SkinCache);
+		SkinCache->PendingProcessRTGeometryEntries.Remove(SkinCacheEntry);
+
 		ReleaseSkinCacheEntry(SkinCacheEntry);
 		SkinCacheEntry = nullptr;
 	}
@@ -1746,7 +1744,7 @@ void FGPUSkinCache::PrepareUpdateSkinning(FGPUSkinCacheEntry* Entry, int32 Secti
 	check(DispatchData.PreviousPositionBuffer != DispatchData.PositionBuffer);
 }
 
-void FGPUSkinCache::DispatchUpdateSkinning(FRHICommandListImmediate& RHICmdList, FGPUSkinCacheEntry* Entry, int32 Section, uint32 RevisionNumber)
+void FGPUSkinCache::DispatchUpdateSkinning(FRHICommandListImmediate& RHICmdList, FGPUSkinCacheEntry* Entry, int32 Section, uint32 RevisionNumber, TSet<FSkinCacheRWBuffer*>& BuffersToTransitionToRead)
 {
 	FGPUSkinCacheEntry::FSectionDispatchData& DispatchData = Entry->DispatchData[Section];
 	FGPUBaseSkinVertexFactory::FShaderDataType& ShaderData = DispatchData.SourceVertexFactory->GetShaderData();
@@ -1949,7 +1947,6 @@ void FGPUSkinCache::ReleaseSkinCacheEntry(FGPUSkinCacheEntry* SkinCacheEntry)
 		DEC_MEMORY_STAT_BY(STAT_GPUSkinCache_TotalMemUsed, RequiredMemInBytes);
 
 		SkinCache->Allocations.Remove(PositionAllocation);
-		PositionAllocation->RemoveAllFromTransitionArray(SkinCache->BuffersToTransitionToRead);
 
 		delete PositionAllocation;
 
diff --git a/Engine/Source/Runtime/Engine/Private/LevelTick.cpp b/Engine/Source/Runtime/Engine/Private/LevelTick.cpp
--- a/Engine/Source/Runtime/Engine/Private/LevelTick.cpp
+++ b/Engine/Source/Runtime/Engine/Private/LevelTick.cpp
@@ -1021,10 +1021,6 @@ void EndSendEndOfFrameUpdatesDrawEvent(FSendAllEndOfFrameUpdates& SendAllEndOfFr
 				// Once all the individual components have received their DoDeferredRenderUpdates_Concurrent()
 				// allow the GPU Skin Cache system to update.
 				GPUSkinCache->EndBatchDispatch(RHICmdList);
-
-				// Flush any remaining pending resource barriers.
-				GPUSkinCache->TransitionAllToReadable(RHICmdList);
-
 			}
 
 			if (ComputeTaskWorkers.Num() > 0)
diff --git a/Engine/Source/Runtime/Engine/Private/SkeletalMeshDeformerHelpers.cpp b/Engine/Source/Runtime/Engine/Private/SkeletalMeshDeformerHelpers.cpp
--- a/Engine/Source/Runtime/Engine/Private/SkeletalMeshDeformerHelpers.cpp
+++ b/Engine/Source/Runtime/Engine/Private/SkeletalMeshDeformerHelpers.cpp
@@ -151,7 +151,7 @@ void FSkeletalMeshDeformerHelpers::UpdateRayTracingGeometry(
 		VertexBufffers[SectionIndex] = PositionBuffer->GetRHI();
 	}
 
-	MeshObjectGPU->UpdateRayTracingGeometry(LODModel, VertexBufffers);
+	MeshObjectGPU->UpdateRayTracingGeometry(LODModel, LODIndex, VertexBufffers);
 }
 
 #endif // RHI_RAYTRACING
diff --git a/Engine/Source/Runtime/Engine/Private/SkeletalRender.cpp b/Engine/Source/Runtime/Engine/Private/SkeletalRender.cpp
--- a/Engine/Source/Runtime/Engine/Private/SkeletalRender.cpp
+++ b/Engine/Source/Runtime/Engine/Private/SkeletalRender.cpp
@@ -39,7 +39,6 @@ FSkeletalMeshObject::FSkeletalMeshObject(USkinnedMeshComponent* InMeshComponent,
 ,	WorkingMaxDistanceFactor(0.f)
 ,   bHasBeenUpdatedAtLeastOnce(false)
 #if RHI_RAYTRACING
-, bRequireRecreatingRayTracingGeometry(false)
 , bSupportRayTracing(IsSkeletalMeshRayTracingSupported() && InMeshComponent->SkeletalMesh->GetSupportRayTracing())
 , bHiddenMaterialVisibilityDirtyForRayTracing(false)
 , RayTracingMinLOD(InMeshComponent->SkeletalMesh->GetRayTracingMinLOD())
diff --git a/Engine/Source/Runtime/Engine/Private/SkeletalRenderGPUSkin.cpp b/Engine/Source/Runtime/Engine/Private/SkeletalRenderGPUSkin.cpp
--- a/Engine/Source/Runtime/Engine/Private/SkeletalRenderGPUSkin.cpp
+++ b/Engine/Source/Runtime/Engine/Private/SkeletalRenderGPUSkin.cpp
@@ -441,10 +441,10 @@ void FSkeletalMeshObjectGPUSkin::Update(int32 LODIndex,USkinnedMeshComponent* In
 	// queue a call to update this data
 	FSkeletalMeshObjectGPUSkin* MeshObject = this;
 	ENQUEUE_RENDER_COMMAND(SkelMeshObjectUpdateDataCommand)(
-		[MeshObject, FrameNumberToPrepare, RevisionNumber, NewDynamicData, GPUSkinCache, bIsGameWorld = (InMeshComponent && InMeshComponent->GetWorld() != nullptr && InMeshComponent->GetWorld()->IsGameWorld())](FRHICommandListImmediate& RHICmdList)
+		[MeshObject, FrameNumberToPrepare, RevisionNumber, NewDynamicData, GPUSkinCache](FRHICommandListImmediate& RHICmdList)
 		{
 			FScopeCycleCounter Context(MeshObject->GetStatId());
-			MeshObject->UpdateDynamicData_RenderThread(GPUSkinCache, RHICmdList, NewDynamicData, nullptr, bIsGameWorld, FrameNumberToPrepare, RevisionNumber);
+			MeshObject->UpdateDynamicData_RenderThread(GPUSkinCache, RHICmdList, NewDynamicData, nullptr, FrameNumberToPrepare, RevisionNumber);
 		}
 	);
 }
@@ -507,7 +507,7 @@ FORCEINLINE bool IsDeferredSkeletalDynamicDataUpdateEnabled()
 	return CVarDeferSkeletalDynamicDataUpdateUntilGDME.GetValueOnRenderThread() > 0 && !bIsHairPluginEnable;
 }
 
-void FSkeletalMeshObjectGPUSkin::UpdateDynamicData_RenderThread(FGPUSkinCache* GPUSkinCache, FRHICommandListImmediate& RHICmdList, FDynamicSkelMeshObjectDataGPUSkin* InDynamicData, FSceneInterface* Scene, bool bIsGameWorld, uint32 FrameNumberToPrepare, uint32 RevisionNumber)
+void FSkeletalMeshObjectGPUSkin::UpdateDynamicData_RenderThread(FGPUSkinCache* GPUSkinCache, FRHICommandListImmediate& RHICmdList, FDynamicSkelMeshObjectDataGPUSkin* InDynamicData, FSceneInterface* Scene, uint32 FrameNumberToPrepare, uint32 RevisionNumber)
 {
 	TRACE_CPUPROFILER_EVENT_SCOPE(GPUSkin::UpdateDynamicData_RT);
 
@@ -523,29 +523,6 @@ void FSkeletalMeshObjectGPUSkin::UpdateDynamicData_RenderThread(FGPUSkinCache* G
 		!DynamicData->ActiveMorphTargetsEqual(InDynamicData->ActiveMorphTargets, InDynamicData->MorphTargetWeights))
 		: true);
 
-#if RHI_RAYTRACING
-	// When switching to a higher detailed LOD, skin cache destroys the old entry, if there isn't enough memory to create the new entry because higher LOD costs more memory, SkinCacheEntry becomes null.
-	// So make sure to set bRequireRecreatingRayTracingGeometry to true to account for this, so RT geometry is updated correctly.
-	bRequireRecreatingRayTracingGeometry = (DynamicData == nullptr || RayTracingGeometry.Initializer.Segments.Num() == 0 || // Newly created
-		(DynamicData != nullptr && DynamicData->RayTracingLODIndex != InDynamicData->RayTracingLODIndex) || // LOD level changed
-		bHiddenMaterialVisibilityDirtyForRayTracing ||
-		(DynamicData->bIsSkinCacheAllowed && GetSkinCacheEntryForRayTracing() == nullptr));
-
-	bHiddenMaterialVisibilityDirtyForRayTracing = false;
-	
-	if (!bRequireRecreatingRayTracingGeometry)
-	{
-		for (FRayTracingGeometrySegment& Segment : RayTracingGeometry.Initializer.Segments)
-		{
-			if (Segment.VertexBuffer == nullptr)
-			{
-				bRequireRecreatingRayTracingGeometry = true;
-				break;
-			}
-		}
-	}
-#endif // RHI_RAYTRACING
-
 	WaitForRHIThreadFenceForDynamicData();
 	if (DynamicData)
 	{
@@ -567,12 +544,8 @@ void FSkeletalMeshObjectGPUSkin::UpdateDynamicData_RenderThread(FGPUSkinCache* G
 	#if RHI_RAYTRACING
 		if (ShouldUseSeparateSkinCacheEntryForRayTracing() && FGPUSkinCache::IsGPUSkinCacheRayTracingSupported() && GPUSkinCache && SkeletalMeshRenderData->bSupportRayTracing)
 		{
-			// When in game, updates not sent in UWorld::SendAllEndOfFrameUpdates() (implied by IsBatchingDispatch()) are ignored.
-			if (!bIsGameWorld || GPUSkinCache->IsBatchingDispatch())
-			{
-				// Morph delta is updated in raster pass above, no need to update again for ray tracing
-				ProcessUpdatedDynamicData(EGPUSkinCacheEntryMode::RayTracing, GPUSkinCache, RHICmdList, FrameNumberToPrepare, RevisionNumber, /*bMorphNeedsUpdate=*/false, DynamicData->RayTracingLODIndex);
-			}
+			// Morph delta is updated in raster pass above, no need to update again for ray tracing
+			ProcessUpdatedDynamicData(EGPUSkinCacheEntryMode::RayTracing, GPUSkinCache, RHICmdList, FrameNumberToPrepare, RevisionNumber, /*bMorphNeedsUpdate=*/false, DynamicData->RayTracingLODIndex);
 		}
 		else
 		{
@@ -583,20 +556,11 @@ void FSkeletalMeshObjectGPUSkin::UpdateDynamicData_RenderThread(FGPUSkinCache* G
 	}
 
 #if RHI_RAYTRACING
-	// Immediate update path. Only used in the editor as the world can be not ticking. When in game, updates not sent in UWorld::SendAllEndOfFrameUpdates() are ignored.
-	if (!bIsGameWorld && FGPUSkinCache::IsGPUSkinCacheRayTracingSupported() && GPUSkinCache && SkeletalMeshRenderData->bSupportRayTracing && !GPUSkinCache->IsBatchingDispatch())
+	if (FGPUSkinCache::IsGPUSkinCacheRayTracingSupported() && SkeletalMeshRenderData->bSupportRayTracing && !GetSkinCacheEntryForRayTracing() && !DynamicData->bHasMeshDeformer)
 	{
-		if (GetSkinCacheEntryForRayTracing())
-		{
-			FSkeletalMeshLODRenderData& LODModel = this->SkeletalMeshRenderData->LODRenderData[DynamicData->RayTracingLODIndex];
-			GPUSkinCache->ProcessRayTracingGeometryToUpdate(RHICmdList, GetSkinCacheEntryForRayTracing(), LODModel);
-		}
-		else
-		{
-			// When SkinCacheEntry is gone, clear geometry
-			RayTracingGeometry.ReleaseRHI();
-			RayTracingGeometry.SetInitializer(FRayTracingGeometryInitializer {});
-		}		
+		// When SkinCacheEntry is gone, clear geometry
+		RayTracingGeometry.ReleaseRHI();
+		RayTracingGeometry.SetInitializer(FRayTracingGeometryInitializer{});
 	}
 #endif
 }
@@ -809,7 +773,7 @@ void FSkeletalMeshObjectGPUSkin::ProcessUpdatedDynamicData(EGPUSkinCacheEntryMod
 
 #if RHI_RAYTRACING
 
-void FSkeletalMeshObjectGPUSkin::UpdateRayTracingGeometry(FSkeletalMeshLODRenderData& LODModel, TArray<FBufferRHIRef>& VertexBufffers)
+void FSkeletalMeshObjectGPUSkin::UpdateRayTracingGeometry(FSkeletalMeshLODRenderData& LODModel, uint32 LODIndex, TArray<FBufferRHIRef>& VertexBufffers)
 {
 	check(IsInRenderingThread());
 
@@ -817,6 +781,20 @@ void FSkeletalMeshObjectGPUSkin::UpdateRayTracingGeometry(FSkeletalMeshLODRender
 	{
 		const bool bAnySegmentUsesWorldPositionOffset = DynamicData != nullptr ? DynamicData->bAnySegmentUsesWorldPositionOffset : false;
 
+		bool bRequireRecreatingRayTracingGeometry = (LODIndex != RayTracingGeometry.LODIndex || bHiddenMaterialVisibilityDirtyForRayTracing);
+		if (!bRequireRecreatingRayTracingGeometry)
+		{
+			for (FRayTracingGeometrySegment& Segment : RayTracingGeometry.Initializer.Segments)
+			{
+				if (Segment.VertexBuffer == nullptr)
+				{
+					bRequireRecreatingRayTracingGeometry = true;
+					break;
+				}
+			}
+		}
+		bHiddenMaterialVisibilityDirtyForRayTracing = false;
+
 		if (bRequireRecreatingRayTracingGeometry)
 		{
 			uint32 MemoryEstimation = 0;
@@ -856,7 +834,6 @@ void FSkeletalMeshObjectGPUSkin::UpdateRayTracingGeometry(FSkeletalMeshLODRender
 			Initializer.bAllowUpdate = true;
 
 			Initializer.Segments.Reserve(LODModel.RenderSections.Num());
-			const int32 LODIndex = GetLOD();
 
 			for (int32 SectionIndex = 0; SectionIndex < LODModel.RenderSections.Num(); ++SectionIndex)
 			{
@@ -889,6 +866,8 @@ void FSkeletalMeshObjectGPUSkin::UpdateRayTracingGeometry(FSkeletalMeshLODRender
 
 			Initializer.SourceGeometry = LODModel.SourceRayTracingGeometry.RayTracingGeometryRHI;
 
+			RayTracingGeometry.LODIndex = LODIndex;
+
 			// Update the new init data
 			RayTracingGeometry.SetInitializer(Initializer);
 
diff --git a/Engine/Source/Runtime/Engine/Private/SkeletalRenderGPUSkin.h b/Engine/Source/Runtime/Engine/Private/SkeletalRenderGPUSkin.h
--- a/Engine/Source/Runtime/Engine/Private/SkeletalRenderGPUSkin.h
+++ b/Engine/Source/Runtime/Engine/Private/SkeletalRenderGPUSkin.h
@@ -324,7 +324,7 @@ public:
 	virtual void InitResources(USkinnedMeshComponent* InMeshComponent) override;
 	virtual void ReleaseResources() override;
 	virtual void Update(int32 LODIndex,USkinnedMeshComponent* InMeshComponent,const TArray<FActiveMorphTarget>& ActiveMorphTargets, const TArray<float>& MorphTargetWeights, EPreviousBoneTransformUpdateMode PreviousBoneTransformUpdateMode) override;
-	void UpdateDynamicData_RenderThread(FGPUSkinCache* GPUSkinCache, FRHICommandListImmediate& RHICmdList, FDynamicSkelMeshObjectDataGPUSkin* InDynamicData, FSceneInterface* Scene, bool bIsGameWorld, uint32 FrameNumberToPrepare, uint32 RevisionNumber);
+	void UpdateDynamicData_RenderThread(FGPUSkinCache* GPUSkinCache, FRHICommandListImmediate& RHICmdList, FDynamicSkelMeshObjectDataGPUSkin* InDynamicData, FSceneInterface* Scene, uint32 FrameNumberToPrepare, uint32 RevisionNumber);
 	virtual void PreGDMECallback(FGPUSkinCache* GPUSkinCache, uint32 FrameNumber) override;
 	virtual const FVertexFactory* GetSkinVertexFactory(const FSceneView* View, int32 LODIndex,int32 ChunkIdx, ESkinVertexFactoryMode VFMode = ESkinVertexFactoryMode::Default) const override;
 	virtual void CacheVertices(int32 LODIndex, bool bForce) const override {}
@@ -363,7 +363,7 @@ public:
 	 * VSinCS path is still required for world position offset materials but this can still use 
 	 * the updated vertex buffers from here with a passthrough vertex factory.
 	 */
-	void UpdateRayTracingGeometry(FSkeletalMeshLODRenderData& LODModel, TArray<FBufferRHIRef>& VertexBufffers);
+	void UpdateRayTracingGeometry(FSkeletalMeshLODRenderData& LODModel, uint32 LODIndex, TArray<FBufferRHIRef>& VertexBufffers);
 
 #endif // RHI_RAYTRACING
 
diff --git a/Engine/Source/Runtime/Engine/Public/GPUSkinCache.h b/Engine/Source/Runtime/Engine/Public/GPUSkinCache.h
--- a/Engine/Source/Runtime/Engine/Public/GPUSkinCache.h
+++ b/Engine/Source/Runtime/Engine/Public/GPUSkinCache.h
@@ -144,7 +144,6 @@ public:
 	struct FDispatchEntry
 	{
 		FGPUSkinCacheEntry* SkinCacheEntry = nullptr;
-		FSkeletalMeshLODRenderData* LODModel = nullptr;
 		uint32 RevisionNumber = 0;
 		uint32 Section = 0;	
 	};
@@ -412,19 +411,16 @@ public:
 		const FVertexBufferAndSRV* BoneBuffers[NUM_BUFFERS];
 	};
 
-	ENGINE_API void TransitionAllToReadable(FRHICommandList& RHICmdList);
-
 	ENGINE_API FRWBuffer* GetPositionBuffer(uint32 ComponentId, uint32 SectionIndex) const;
 	ENGINE_API FRWBuffer* GetTangentBuffer(uint32 ComponentId, uint32 SectionIndex) const;
 	ENGINE_API FRHIShaderResourceView* GetBoneBuffer(uint32 ComponentId, uint32 SectionIndex) const;
 
 #if RHI_RAYTRACING
-	void ProcessRayTracingGeometryToUpdate(FRHICommandListImmediate& RHICmdList, FGPUSkinCacheEntry* SkinCacheEntry, FSkeletalMeshLODRenderData& LODModel);
+	void ProcessRayTracingGeometryToUpdate(FRHICommandListImmediate& RHICmdList, FGPUSkinCacheEntry* SkinCacheEntry);
 #endif // RHI_RAYTRACING
 
 	void BeginBatchDispatch(FRHICommandListImmediate& RHICmdList);
 	void EndBatchDispatch(FRHICommandListImmediate& RHICmdList);
-	bool IsBatchingDispatch() const { return bShouldBatchDispatches; }
 
 	inline ERHIFeatureLevel::Type GetFeatureLevel() const { return FeatureLevel; }
 
@@ -432,9 +428,9 @@ protected:
 	void MakeBufferTransitions(FRHICommandListImmediate& RHICmdList, TArray<FSkinCacheRWBuffer*>& Buffers, ERHIAccess ToState);
 	void GetBufferUAVs(const TArray<FSkinCacheRWBuffer*>& InBuffers, TArray<FRHIUnorderedAccessView*>& OutUAVs);
 
-	TSet<FSkinCacheRWBuffer*> BuffersToTransitionToRead;
 	TArray<FRWBuffersAllocation*> Allocations;
 	TArray<FGPUSkinCacheEntry*> Entries;
+	TSet<FGPUSkinCacheEntry*> PendingProcessRTGeometryEntries;
 	TArray<FDispatchEntry> BatchDispatches;
 
 	FRWBuffersAllocation* TryAllocBuffer(uint32 NumVertices, bool WithTangnents, bool UseIntermediateTangents, uint32 NumTriangles, FRHICommandListImmediate& RHICmdList);
@@ -453,10 +449,12 @@ protected:
 		FRHICommandListImmediate& RHICmdList, 
 		FGPUSkinCacheEntry* Entry, 
 		int32 Section, 
-		uint32 RevisionNumber
+		uint32 RevisionNumber,
+		TSet<FSkinCacheRWBuffer*>& BuffersToTransitionToRead
 		);
 
 	void Cleanup();
+	static void TransitionAllToReadable(FRHICommandList& RHICmdList, const TSet<FSkinCacheRWBuffer*>& BuffersToTransitionToRead);
 	static void ReleaseSkinCacheEntry(FGPUSkinCacheEntry* SkinCacheEntry);
 	static FGPUSkinBatchElementUserData* InternalGetFactoryUserData(FGPUSkinCacheEntry* Entry, int32 Section);
 	void InvalidateAllEntries();
diff --git a/Engine/Source/Runtime/Engine/Public/SkeletalRenderPublic.h b/Engine/Source/Runtime/Engine/Public/SkeletalRenderPublic.h
--- a/Engine/Source/Runtime/Engine/Public/SkeletalRenderPublic.h
+++ b/Engine/Source/Runtime/Engine/Public/SkeletalRenderPublic.h
@@ -251,7 +251,6 @@ public:
 	bool bHasBeenUpdatedAtLeastOnce;
 
 #if RHI_RAYTRACING
-	bool bRequireRecreatingRayTracingGeometry;
 	bool bSupportRayTracing;
 	bool bHiddenMaterialVisibilityDirtyForRayTracing;
 	int32 RayTracingMinLOD;
diff --git a/Engine/Source/Runtime/RenderCore/Public/RenderResource.h b/Engine/Source/Runtime/RenderCore/Public/RenderResource.h
--- a/Engine/Source/Runtime/RenderCore/Public/RenderResource.h
+++ b/Engine/Source/Runtime/RenderCore/Public/RenderResource.h
@@ -730,6 +730,9 @@ public:
 	FRayTracingGeometryRHIRef RayTracingGeometryRHI;
 	bool bRequiresBuild = false;
 
+	/** LOD of the mesh associated with this ray tracing geometry object (-1 if unknown) */
+	int8 LODIndex = -1;
+
 	void SetInitializer(const FRayTracingGeometryInitializer& InInitializer)
 	{
 		Initializer = InInitializer;
diff --git a/Engine/Source/Runtime/Renderer/Private/DeferredShadingRenderer.cpp b/Engine/Source/Runtime/Renderer/Private/DeferredShadingRenderer.cpp
--- a/Engine/Source/Runtime/Renderer/Private/DeferredShadingRenderer.cpp
+++ b/Engine/Source/Runtime/Renderer/Private/DeferredShadingRenderer.cpp
@@ -2259,11 +2259,6 @@ void FDeferredShadingSceneRenderer::Render(FRDGBuilder& GraphBuilder)
 		AddDitheredStencilFillPass(GraphBuilder, Views, SceneTextures.Depth.Target, DepthPass);
 	}
 
-	AddPass(GraphBuilder, RDG_EVENT_NAME("GPUSkinCache::Transitions"), [this](FRHICommandList& InRHICmdList)
-	{
-		RunGPUSkinCacheTransition(InRHICmdList, Scene, EGPUSkinCacheTransition::Renderer);
-	});
-
 	FHairStrandsBookmarkParameters& HairStrandsBookmarkParameters = *GraphBuilder.AllocObject<FHairStrandsBookmarkParameters>();
 	if (IsHairStrandsEnabled(EHairStrandsShaderType::All, Scene->GetShaderPlatform()))
 	{
diff --git a/Engine/Source/Runtime/Renderer/Private/SceneRendering.cpp b/Engine/Source/Runtime/Renderer/Private/SceneRendering.cpp
--- a/Engine/Source/Runtime/Renderer/Private/SceneRendering.cpp
+++ b/Engine/Source/Runtime/Renderer/Private/SceneRendering.cpp
@@ -5262,21 +5262,3 @@ void FSceneRenderer::UpdateSkyIrradianceGpuBuffer(FRHICommandListImmediate& RHIC
 	RHICmdList.Transition(FRHITransitionInfo(Scene->SkyIrradianceEnvironmentMap.UAV, ERHIAccess::Unknown, ERHIAccess::SRVMask));
 }
 
-void RunGPUSkinCacheTransition(FRHICommandList& RHICmdList, FScene* Scene, EGPUSkinCacheTransition Type)
-{
-	// * When hair strands is disabled, the skin cache sync point run later 
-	//   during the deferred render pass
-	// * When hair strands is enabled, the skin cache sync point is run earlier, during 
-	//   the init views pass, as the output of the skin cached is used by Niagara
-	const bool bHairStrandsEnabled = IsHairStrandsEnabled(EHairStrandsShaderType::All, Scene->GetShaderPlatform());
-	const bool bRun = 
-		(bHairStrandsEnabled && EGPUSkinCacheTransition::FrameSetup == Type) || 
-		(!bHairStrandsEnabled && EGPUSkinCacheTransition::FrameSetup != Type);
-	if (bRun)
-	{
-		if (FGPUSkinCache* GPUSkinCache = Scene->GetGPUSkinCache())
-		{
-			GPUSkinCache->TransitionAllToReadable(RHICmdList);
-		}
-	}
-}
diff --git a/Engine/Source/Runtime/Renderer/Private/SceneRendering.h b/Engine/Source/Runtime/Renderer/Private/SceneRendering.h
--- a/Engine/Source/Runtime/Renderer/Private/SceneRendering.h
+++ b/Engine/Source/Runtime/Renderer/Private/SceneRendering.h
@@ -2581,8 +2581,7 @@ enum class EGPUSkinCacheTransition
 
 extern bool IsStaticLightingAllowed();
 
-/* Run GPU skin cache resource transitions */
-void RunGPUSkinCacheTransition(class FRHICommandList& RHICmdList, class FScene* Scene, EGPUSkinCacheTransition Type);/** Resolves the view rect of scene color or depth using either a custom resolve or hardware resolve. */
+/** Resolves the view rect of scene color or depth using either a custom resolve or hardware resolve. */
 void AddResolveSceneColorPass(FRDGBuilder& GraphBuilder, const FViewInfo& View, FRDGTextureMSAA SceneColor);
 void AddResolveSceneDepthPass(FRDGBuilder& GraphBuilder, const FViewInfo& View, FRDGTextureMSAA SceneDepth);
 
diff --git a/Engine/Source/Runtime/Renderer/Private/SceneVisibility.cpp b/Engine/Source/Runtime/Renderer/Private/SceneVisibility.cpp
--- a/Engine/Source/Runtime/Renderer/Private/SceneVisibility.cpp
+++ b/Engine/Source/Runtime/Renderer/Private/SceneVisibility.cpp
@@ -47,6 +47,7 @@
 #include "RayTracing/RayTracingInstanceCulling.h"
 #include "RendererModule.h"
 
+
 /*------------------------------------------------------------------------------
 	Globals
 ------------------------------------------------------------------------------*/
@@ -3495,12 +3496,6 @@ void FSceneRenderer::PreVisibilityFrameSetup(FRDGBuilder& GraphBuilder, const FS
 		}
 	}
 
-	AddPass(GraphBuilder, RDG_EVENT_NAME("GPUSkinCache::Transitions"), [this](FRHICommandList& InRHICmdList)
-	{
-		RunGPUSkinCacheTransition(InRHICmdList, Scene, EGPUSkinCacheTransition::FrameSetup);
-	});
-
-
 	if (Views.Num() > 0 && !ViewFamily.EngineShowFlags.HitProxies)
 	{
 		FHairStrandsBookmarkParameters Parameters = CreateHairStrandsBookmarkParameters(Scene, Views);
