Adding back the wait on FUpdateStreamingRenderAssetsTask until we figure out whether it is causing memory stomps

diff --git a/Engine/Source/Runtime/Engine/Private/Streaming/StreamingManagerTexture.cpp b/Engine/Source/Runtime/Engine/Private/Streaming/StreamingManagerTexture.cpp
--- a/Engine/Source/Runtime/Engine/Private/Streaming/StreamingManagerTexture.cpp
+++ b/Engine/Source/Runtime/Engine/Private/Streaming/StreamingManagerTexture.cpp
@@ -1727,6 +1727,12 @@ void FRenderAssetStreamingManager::UpdateResourceStreaming( float DeltaTime, boo
 		IncrementalUpdate(1.f / (float)FMath::Max(NumRenderAssetProcessingStages - 1, 1), true); // -1 since we don't want to do anything at stage 0.
 		++ProcessingStage;
 
+		// TODO: figure out whether letting the task to run past is causing memory stomps
+		if (bOverlappedExecution && StreamingRenderAssetsSyncEvent.IsValid())
+		{
+			StreamingRenderAssetsSyncEvent->Wait(ENamedThreads::GameThread);
+		}
+
 		STAT(GatheredStats.UpdateStreamingDataCycles = FMath::Max<uint32>(ProcessingStage > 2 ? GatheredStats.UpdateStreamingDataCycles : 0, FPlatformTime::Cycles() - StartTime);)
 	}
 	else if (AsyncWork->IsDone())
