World Partiion: make sure the bIsPIE flag is set when going through PIE travel in WP maps.

diff --git a/Engine/Source/Runtime/Engine/Private/WorldPartition/WorldPartition.cpp b/Engine/Source/Runtime/Engine/Private/WorldPartition/WorldPartition.cpp
--- a/Engine/Source/Runtime/Engine/Private/WorldPartition/WorldPartition.cpp
+++ b/Engine/Source/Runtime/Engine/Private/WorldPartition/WorldPartition.cpp
@@ -416,8 +416,9 @@ void UWorldPartition::Initialize(UWorld* InWorld, const FTransform& InTransform)
 	const bool bIsEditor = !World->IsGameWorld();
 	const bool bIsCooking = IsRunningCookCommandlet();
 	const bool bIsDedicatedServer = IsRunningDedicatedServer();
+	const bool bPIEWorldTravel = (World->WorldType == EWorldType::PIE) && !StreamingPolicy;
 
-	UE_LOG(LogWorldPartition, Log, TEXT("UWorldPartition::Initialize(IsEditor=%d, IsGame=%d, IsCooking=%d)"), bIsEditor ? 1 : 0, bIsGame ? 1 : 0, bIsCooking ? 1 : 0);
+	UE_LOG(LogWorldPartition, Log, TEXT("UWorldPartition::Initialize(IsEditor=%d, bPIEWorldTravel=%d IsGame=%d, IsCooking=%d)"), bIsEditor ? 1 : 0, bPIEWorldTravel ? 1 : 0, bIsGame ? 1 : 0, bIsCooking ? 1 : 0);
 
 	if (bIsGame || bIsCooking)
 	{
@@ -436,9 +437,6 @@ void UWorldPartition::Initialize(UWorld* InWorld, const FTransform& InTransform)
 
 	check(RuntimeHash);
 
-	// Did we travel into a WP map in PIE (null StreamingPolicy means GenerateStreaming wasn't called)
-	const bool bPIEWorldTravel = World->WorldType == EWorldType::PIE && !StreamingPolicy;
-
 	if (bIsEditor || bIsGame || bPIEWorldTravel || bIsDedicatedServer)
 	{
 		UPackage* LevelPackage = OuterWorld->PersistentLevel->GetOutermost();
@@ -483,7 +481,7 @@ void UWorldPartition::Initialize(UWorld* InWorld, const FTransform& InTransform)
 				{
 					HashActorDesc(*ActorDescIterator);
 				}
-				}
+			}
 		}
 	}
 
@@ -533,6 +531,12 @@ void UWorldPartition::Initialize(UWorld* InWorld, const FTransform& InTransform)
 	{
 		if (bIsGame || bPIEWorldTravel || bIsDedicatedServer)
 		{
+			if (bPIEWorldTravel)
+			{
+				check(!bIsPIE);
+				bIsPIE = true;
+			}
+
 			OnBeginPlay();
 		}
 
