deprecate crunch and zlib compression support on VT tiles
remove crunch use from engine

diff --git a/Engine/Source/Developer/TextureBuild/Private/TextureBuildFunction.cpp b/Engine/Source/Developer/TextureBuild/Private/TextureBuildFunction.cpp
--- a/Engine/Source/Developer/TextureBuild/Private/TextureBuildFunction.cpp
+++ b/Engine/Source/Developer/TextureBuild/Private/TextureBuildFunction.cpp
@@ -160,8 +160,6 @@ static FTextureBuildSettings ReadBuildSettingsFromCompactBinary(const FCbObjectV
 	ReadCbField(Object["VirtualAddressingModeY"], BuildSettings.VirtualAddressingModeY);
 	ReadCbField(Object["VirtualTextureTileSize"], BuildSettings.VirtualTextureTileSize);
 	ReadCbField(Object["VirtualTextureBorderSize"], BuildSettings.VirtualTextureBorderSize);
-	BuildSettings.bVirtualTextureEnableCompressZlib = Object["bVirtualTextureEnableCompressZlib"].AsBool(BuildSettings.bVirtualTextureEnableCompressZlib);
-	BuildSettings.bVirtualTextureEnableCompressCrunch = Object["bVirtualTextureEnableCompressCrunch"].AsBool(BuildSettings.bVirtualTextureEnableCompressCrunch);
 	BuildSettings.OodleEncodeEffort = Object["OodleEncodeEffort"].AsUInt8(BuildSettings.OodleEncodeEffort);
 	BuildSettings.OodleUniversalTiling = Object["OodleUniversalTiling"].AsUInt8(BuildSettings.OodleUniversalTiling);
 	BuildSettings.bOodleUsesRDO = Object["bOodleUsesRDO"].AsBool(BuildSettings.bOodleUsesRDO);
diff --git a/Engine/Source/Developer/TextureCompressor/Public/TextureCompressorModule.h b/Engine/Source/Developer/TextureCompressor/Public/TextureCompressorModule.h
--- a/Engine/Source/Developer/TextureCompressor/Public/TextureCompressorModule.h
+++ b/Engine/Source/Developer/TextureCompressor/Public/TextureCompressorModule.h
@@ -200,10 +200,6 @@ struct FTextureBuildSettings
 	int32 VirtualTextureTileSize;
 	/** Size in pixels of border on virtual texture tile */
 	int32 VirtualTextureBorderSize;
-	/** Is zlib compression enabled */
-	uint32 bVirtualTextureEnableCompressZlib : 1;
-	/** Is crunch compression enabled */
-	uint32 bVirtualTextureEnableCompressCrunch : 1;
 
 	// Which encode speed this build settings represents.
 	// This is not sent to the build worker, it is used to
@@ -274,8 +270,6 @@ struct FTextureBuildSettings
 		, VirtualAddressingModeY(0)
 		, VirtualTextureTileSize(0)
 		, VirtualTextureBorderSize(0)
-		, bVirtualTextureEnableCompressZlib(false)
-		, bVirtualTextureEnableCompressCrunch(false)
 	{
 	}
 
diff --git a/Engine/Source/Editor/VirtualTexturingEditor/Private/RuntimeVirtualTextureBuildStreamingMips.cpp b/Engine/Source/Editor/VirtualTexturingEditor/Private/RuntimeVirtualTextureBuildStreamingMips.cpp
--- a/Engine/Source/Editor/VirtualTexturingEditor/Private/RuntimeVirtualTextureBuildStreamingMips.cpp
+++ b/Engine/Source/Editor/VirtualTexturingEditor/Private/RuntimeVirtualTextureBuildStreamingMips.cpp
@@ -174,7 +174,7 @@ namespace RuntimeVirtualTexture
 
 		// Spin up slow task UI
 		const float TaskWorkRender = NumTilesX * NumTilesY;
-		const float TextureBuildTaskMultiplier = InComponent->IsCrunchCompressed() ? 3.f : .25f; // Crunch compression is slow.
+		const float TextureBuildTaskMultiplier = 0.25f;
 		const float TaskWorkBuildBulkData = TaskWorkRender * TextureBuildTaskMultiplier;
 		FScopedSlowTask Task(TaskWorkRender + TaskWorkBuildBulkData, FText::AsCultureInvariant(InComponent->GetStreamingTexture()->GetName()));
 		Task.MakeDialog(true);
diff --git a/Engine/Source/Runtime/CrunchCompression/CrunchCompression.build.cs b/Engine/Source/Runtime/CrunchCompression/CrunchCompression.build.cs
--- a/Engine/Source/Runtime/CrunchCompression/CrunchCompression.build.cs
+++ b/Engine/Source/Runtime/CrunchCompression/CrunchCompression.build.cs
@@ -6,8 +6,10 @@ namespace UnrealBuildTool.Rules
 	{
 		public CrunchCompression(ReadOnlyTargetRules Target) : base(Target)
 		{
-           // PrivateIncludePaths.Add("Runtime/CrunchCompression/Private");
-            PublicIncludePaths.Add("Runtime/CrunchCompression/Public");
+			PublicDefinitions.Add("WITH_CRUNCH=1");
+
+			// PrivateIncludePaths.Add("Runtime/CrunchCompression/Private");
+			PublicIncludePaths.Add("Runtime/CrunchCompression/Public");
 
 			PrivateDependencyModuleNames.AddRange(
 				new string[]
diff --git a/Engine/Source/Runtime/CrunchCompression/Public/CrunchCompression.h b/Engine/Source/Runtime/CrunchCompression/Public/CrunchCompression.h
--- a/Engine/Source/Runtime/CrunchCompression/Public/CrunchCompression.h
+++ b/Engine/Source/Runtime/CrunchCompression/Public/CrunchCompression.h
@@ -4,8 +4,9 @@
 
 #include "CoreMinimal.h"
 
+// WITH_CRUNCH=1 is set by the build.cs
 #ifndef WITH_CRUNCH
-#define WITH_CRUNCH 1
+#define WITH_CRUNCH 0
 #endif
 
 /** Crunch compression lib is currently only built for Windows */
diff --git a/Engine/Source/Runtime/Engine/Classes/Components/RuntimeVirtualTextureComponent.h b/Engine/Source/Runtime/Engine/Classes/Components/RuntimeVirtualTextureComponent.h
--- a/Engine/Source/Runtime/Engine/Classes/Components/RuntimeVirtualTextureComponent.h
+++ b/Engine/Source/Runtime/Engine/Classes/Components/RuntimeVirtualTextureComponent.h
@@ -62,17 +62,6 @@ protected:
 	UPROPERTY(EditAnywhere, AdvancedDisplay, Category = VirtualTextureBuild, meta = (DisplayName = "View in Editor"))
 	bool bUseStreamingLowMipsInEditor = false;
 
-	/** Enable Crunch texture compression for the streaming low mips. Generic ZLib compression is used when Crunch is disabled. */
-	UPROPERTY(EditAnywhere, AdvancedDisplay, Category = VirtualTextureBuild, meta = (DisplayName = "Enable Crunch"))
-	bool bEnableCompressCrunch = false;
-	
-	/** How aggressively should any relevant lossy compression be applied. For compressors that support EncodeSpeed (i.e. Oodle), this is only
-	*	applied if enabled (see Project Settings -> Texture Encoding). Note that this is *in addition* to any
-	*	unavoidable loss due to the target format - selecting "No Lossy Compression" will not result in zero distortion for BCn formats.
-	*/
-	UPROPERTY(EditAnywhere, AdvancedDisplay, Category = VirtualTextureBuild)
-	TEnumAsByte<ETextureLossyCompressionAmount> LossyCompressionAmount = TLCA_Default;
-
 	/** Build the streaming low mips using debug coloring. This can help show where streaming mips are being used. */
 	UPROPERTY(EditAnywhere, BlueprintReadOnly, Transient, AdvancedDisplay, Category = VirtualTextureBuild, meta = (DisplayName = "Build Debug"))
 	bool bBuildDebugStreamingMips = false;
@@ -121,9 +110,6 @@ public:
 	/** Get if we want to use any streaming low mips on this component. */
 	bool IsStreamingLowMips() const;
 
-	/** Public getter for crunch compression flag. */
-	bool IsCrunchCompressed() const { return bEnableCompressCrunch; }
-
 	/** Public getter for debug streaming mips flag. */
 	bool IsBuildDebugStreamingMips() { return bBuildDebugStreamingMips; }
 
diff --git a/Engine/Source/Runtime/Engine/Classes/VT/RuntimeVirtualTexture.h b/Engine/Source/Runtime/Engine/Classes/VT/RuntimeVirtualTexture.h
--- a/Engine/Source/Runtime/Engine/Classes/VT/RuntimeVirtualTexture.h
+++ b/Engine/Source/Runtime/Engine/Classes/VT/RuntimeVirtualTexture.h
@@ -74,7 +74,14 @@ protected:
 
 	/** Texture group this texture belongs to */
 	UPROPERTY(EditAnywhere, BlueprintReadWrite, Category = LevelOfDetail, meta = (DisplayName = "Texture Group"), AssetRegistrySearchable)
-	TEnumAsByte<enum TextureGroup> LODGroup;
+	TEnumAsByte<enum TextureGroup> LODGroup = TEXTUREGROUP_World;
+	
+	/** How aggressively should any relevant lossy compression be applied. For compressors that support EncodeSpeed (i.e. Oodle), this is only
+	*	applied if enabled (see Project Settings -> Texture Encoding). Note that this is *in addition* to any
+	*	unavoidable loss due to the target format - selecting "No Lossy Compression" will not result in zero distortion for BCn formats.
+	*/
+	UPROPERTY(EditAnywhere, BlueprintReadWrite, Category = LevelOfDetail)
+	TEnumAsByte<ETextureLossyCompressionAmount> LossyCompressionAmount = TLCA_Default;
 
 	/** Deprecated size of virtual texture. */
 	UPROPERTY()
@@ -126,6 +133,7 @@ public:
 	bool GetLQCompression() const { return bUseLowQualityCompression; }
 	/** Public getter for texture LOD Group */
 	TEnumAsByte<enum TextureGroup> GetLODGroup() const { return LODGroup; }
+	TEnumAsByte<ETextureLossyCompressionAmount> GetLossyCompressionAmount() const { return LossyCompressionAmount; }
 
 	/** Structure containing additional settings for initializing the producer. */
 	struct FInitSettings
diff --git a/Engine/Source/Runtime/Engine/Classes/VT/VirtualTextureBuildSettings.h b/Engine/Source/Runtime/Engine/Classes/VT/VirtualTextureBuildSettings.h
--- a/Engine/Source/Runtime/Engine/Classes/VT/VirtualTextureBuildSettings.h
+++ b/Engine/Source/Runtime/Engine/Classes/VT/VirtualTextureBuildSettings.h
@@ -17,12 +17,6 @@ struct FVirtualTextureBuildSettings
 	UPROPERTY()
 	int32 TileBorderSize = 0;
 
-	UPROPERTY()
-	bool bEnableCompressCrunch = false;
-
-	UPROPERTY()
-	bool bEnableCompressZlib = false;
-
 	/** Initialize with the default build settings. These are defined by the current project setup. */
 	ENGINE_API void Init();
 };
diff --git a/Engine/Source/Runtime/Engine/Classes/VT/VirtualTextureBuilder.h b/Engine/Source/Runtime/Engine/Classes/VT/VirtualTextureBuilder.h
--- a/Engine/Source/Runtime/Engine/Classes/VT/VirtualTextureBuilder.h
+++ b/Engine/Source/Runtime/Engine/Classes/VT/VirtualTextureBuilder.h
@@ -20,9 +20,7 @@ struct FVirtualTextureBuildDesc
 	int32 TileSize = 0;
 	int32 TileBorderSize = 0;
 
-	TEnumAsByte<enum TextureGroup> LODGroup;
-
-	bool bCrunchCompressed = false;
+	TEnumAsByte<enum TextureGroup> LODGroup = TEXTUREGROUP_World;
 	ETextureLossyCompressionAmount LossyCompressionAmount = TLCA_Default;
 
 	bool bContinuousUpdate = false;
diff --git a/Engine/Source/Runtime/Engine/Engine.Build.cs b/Engine/Source/Runtime/Engine/Engine.Build.cs
--- a/Engine/Source/Runtime/Engine/Engine.Build.cs
+++ b/Engine/Source/Runtime/Engine/Engine.Build.cs
@@ -127,7 +127,6 @@ public class Engine : ModuleRules
 				"AudioMixer",
 				"AudioMixerCore",
 				"SignalProcessing",
-				"CrunchCompression",
 				"IntelISPC",
 				"TraceLog",
 				"ColorManagement",
diff --git a/Engine/Source/Runtime/Engine/Private/Components/RuntimeVirtualTextureComponent.cpp b/Engine/Source/Runtime/Engine/Private/Components/RuntimeVirtualTextureComponent.cpp
--- a/Engine/Source/Runtime/Engine/Private/Components/RuntimeVirtualTextureComponent.cpp
+++ b/Engine/Source/Runtime/Engine/Private/Components/RuntimeVirtualTextureComponent.cpp
@@ -147,7 +147,6 @@ uint64 URuntimeVirtualTextureComponent::CalculateStreamingTextureSettingsHash()
 			uint32 LODGroup : 8;
 			uint32 CompressTextures : 1;
 			uint32 SinglePhysicalSpace : 1;
-			uint32 EnableCompressCrunch : 1;
 			uint32 ContinuousUpdate : 1;
 			uint32 bUseLowQualityCompression : 1;
 			uint32 LossyCompressionAmount : 4;
@@ -156,7 +155,7 @@ uint64 URuntimeVirtualTextureComponent::CalculateStreamingTextureSettingsHash()
 
 	FPackedSettings Settings;
 	Settings.PackedValue = 0;
-	Settings.PackedSettingsVersion = 1;
+	Settings.PackedSettingsVersion = 2;
 	Settings.MaterialType = (uint32)VirtualTexture->GetMaterialType();
 	Settings.TileSize = (uint32)VirtualTexture->GetTileSize();
 	Settings.TileBorderSize = (uint32)VirtualTexture->GetTileBorderSize();
@@ -165,8 +164,7 @@ uint64 URuntimeVirtualTextureComponent::CalculateStreamingTextureSettingsHash()
 	Settings.ContinuousUpdate = (uint32)VirtualTexture->GetContinuousUpdate();
 	Settings.SinglePhysicalSpace = (uint32)VirtualTexture->GetSinglePhysicalSpace();
 	Settings.bUseLowQualityCompression = (uint32)VirtualTexture->GetLQCompression();
-	Settings.EnableCompressCrunch = (uint32)bEnableCompressCrunch;
-	Settings.LossyCompressionAmount = (uint32) LossyCompressionAmount.GetValue();
+	Settings.LossyCompressionAmount = (uint32) VirtualTexture->GetLossyCompressionAmount();
 
 	return Settings.PackedValue;
 }
@@ -246,8 +244,7 @@ void URuntimeVirtualTextureComponent::InitializeStreamingTexture(uint32 InSizeX,
 		BuildDesc.TileSize = VirtualTexture->GetTileSize();
 		BuildDesc.TileBorderSize = VirtualTexture->GetTileBorderSize();
 		BuildDesc.LODGroup = VirtualTexture->GetLODGroup();
-		BuildDesc.bCrunchCompressed = bEnableCompressCrunch;
-		BuildDesc.LossyCompressionAmount = LossyCompressionAmount;
+		BuildDesc.LossyCompressionAmount = VirtualTexture->GetLossyCompressionAmount();
 
 		BuildDesc.LayerCount = VirtualTexture->GetLayerCount();
 		check(BuildDesc.LayerCount <= RuntimeVirtualTexture::MaxTextureLayers);
@@ -280,10 +277,6 @@ bool URuntimeVirtualTextureComponent::CanEditChange(const FProperty* InProperty)
 	{
 		bCanEdit &= GetVirtualTexture() != nullptr && GetStreamingTexture() != nullptr;
 	}
-	else if (InProperty->GetFName() == TEXT("bEnableCompressCrunch"))
-	{
-		bCanEdit &= GetVirtualTexture() != nullptr && GetVirtualTexture()->GetCompressTextures();
-	}
 	else if (InProperty->GetFName() == TEXT("bBuildDebugStreamingMips"))
 	{
 		bCanEdit = GetVirtualTexture() != nullptr && GetVirtualTexture()->GetMaterialType() != ERuntimeVirtualTextureMaterialType::WorldHeight;
diff --git a/Engine/Source/Runtime/Engine/Private/TextureDerivedData.cpp b/Engine/Source/Runtime/Engine/Private/TextureDerivedData.cpp
--- a/Engine/Source/Runtime/Engine/Private/TextureDerivedData.cpp
+++ b/Engine/Source/Runtime/Engine/Private/TextureDerivedData.cpp
@@ -187,8 +187,9 @@ static void SerializeForKey(FArchive& Ar, const FTextureBuildSettings& Settings)
 		TempByte = Settings.VirtualAddressingModeY; Ar << TempByte;
 		TempUint32 = Settings.VirtualTextureTileSize; Ar << TempUint32;
 		TempUint32 = Settings.VirtualTextureBorderSize; Ar << TempUint32;
-		TempByte = Settings.bVirtualTextureEnableCompressZlib; Ar << TempByte;
-		TempByte = Settings.bVirtualTextureEnableCompressCrunch; Ar << TempByte;
+		// compresion options removed: keep serializing them as "off" to keep the key the same:
+		TempByte = 0; Ar << TempByte;
+		TempByte = 0; Ar << TempByte;
 		TempByte = Settings.LossyCompressionAmount; Ar << TempByte; // Lossy compression currently only used by VT
 		TempByte = Settings.bApplyYCoCgBlockScale; Ar << TempByte; // YCoCg currently only used by VT
 	}
@@ -841,8 +842,6 @@ static void GetTextureBuildSettings(
 
 		FVirtualTextureBuildSettings VirtualTextureBuildSettings;
 		Texture.GetVirtualTextureBuildSettings(VirtualTextureBuildSettings);
-		OutBuildSettings.bVirtualTextureEnableCompressZlib = VirtualTextureBuildSettings.bEnableCompressZlib;
-		OutBuildSettings.bVirtualTextureEnableCompressCrunch = VirtualTextureBuildSettings.bEnableCompressCrunch;
 		OutBuildSettings.VirtualTextureTileSize = FMath::RoundUpToPowerOfTwo(VirtualTextureBuildSettings.TileSize);
 
 		// Apply any LOD group tile size bias here
@@ -865,8 +864,6 @@ static void GetTextureBuildSettings(
 		OutBuildSettings.VirtualAddressingModeY = TA_Wrap;
 		OutBuildSettings.VirtualTextureTileSize = 0;
 		OutBuildSettings.VirtualTextureBorderSize = 0;
-		OutBuildSettings.bVirtualTextureEnableCompressZlib = false;
-		OutBuildSettings.bVirtualTextureEnableCompressCrunch = false;
 	}
 
 	// By default, initialize settings for layer0
diff --git a/Engine/Source/Runtime/Engine/Private/TextureDerivedDataBuildUtils.cpp b/Engine/Source/Runtime/Engine/Private/TextureDerivedDataBuildUtils.cpp
--- a/Engine/Source/Runtime/Engine/Private/TextureDerivedDataBuildUtils.cpp
+++ b/Engine/Source/Runtime/Engine/Private/TextureDerivedDataBuildUtils.cpp
@@ -175,8 +175,6 @@ static void WriteBuildSettings(FCbWriter& Writer, const FTextureBuildSettings& B
 	WriteCbFieldWithDefault(Writer, "VirtualAddressingModeY", BuildSettings.VirtualAddressingModeY, DefaultSettings.VirtualAddressingModeY);
 	WriteCbFieldWithDefault(Writer, "VirtualTextureTileSize", BuildSettings.VirtualTextureTileSize, DefaultSettings.VirtualTextureTileSize);
 	WriteCbFieldWithDefault(Writer, "VirtualTextureBorderSize", BuildSettings.VirtualTextureBorderSize, DefaultSettings.VirtualTextureBorderSize);
-	WriteCbFieldWithDefault<bool>(Writer, "bVirtualTextureEnableCompressZlib", BuildSettings.bVirtualTextureEnableCompressZlib, DefaultSettings.bVirtualTextureEnableCompressZlib);
-	WriteCbFieldWithDefault<bool>(Writer, "bVirtualTextureEnableCompressCrunch", BuildSettings.bVirtualTextureEnableCompressCrunch, DefaultSettings.bVirtualTextureEnableCompressCrunch);
 
 	WriteCbFieldWithDefault<uint8>(Writer, "OodleEncodeEffort", (uint8)BuildSettings.OodleEncodeEffort, (uint8)DefaultSettings.OodleEncodeEffort);	
 	WriteCbFieldWithDefault<uint8>(Writer, "OodleUniversalTiling", (uint8)BuildSettings.OodleUniversalTiling, (uint8)DefaultSettings.OodleUniversalTiling);
diff --git a/Engine/Source/Runtime/Engine/Private/VT/UploadingVirtualTexture.cpp b/Engine/Source/Runtime/Engine/Private/VT/UploadingVirtualTexture.cpp
--- a/Engine/Source/Runtime/Engine/Private/VT/UploadingVirtualTexture.cpp
+++ b/Engine/Source/Runtime/Engine/Private/VT/UploadingVirtualTexture.cpp
@@ -5,7 +5,6 @@
 #include "VirtualTexturing.h"
 #include "SceneUtils.h"
 #include "FileCache/FileCache.h"
-#include "CrunchCompression.h"
 #include "VirtualTextureChunkDDCCache.h"
 #include "UObject/PackageResourceManager.h"
 #include "Misc/PackageSegment.h"
@@ -218,18 +217,6 @@ void FVirtualTextureCodec::Init(IMemoryReadStreamRef& HeaderData)
 				CodecPayload = TempBuffer.GetData();
 			}
 		}
-
-		switch (Chunk.CodecType[LayerIndex])
-		{
-		case EVirtualTextureCodec::Crunch:
-#if WITH_CRUNCH
-			Contexts[LayerIndex] = CrunchCompression::InitializeDecoderContext(CodecPayload, CodecPayloadSize);
-#endif // WITH_CRUNCH
-			check(Contexts[LayerIndex]);
-			break;
-		default:
-			break;
-		}
 	}
 }
 
@@ -260,24 +247,6 @@ FVirtualTextureCodec::~FVirtualTextureCodec()
 
 		check(!IsLinked());
 
-		const FVirtualTextureBuiltData* VTData = Owner->GetVTData();
-		const FVirtualTextureDataChunk& Chunk = VTData->Chunks[ChunkIndex];
-		const uint32 NumLayers = VTData->GetNumLayers();
-		for (uint32 LayerIndex = 0; LayerIndex < NumLayers; ++LayerIndex)
-		{
-			switch (Chunk.CodecType[LayerIndex])
-			{
-			case EVirtualTextureCodec::Crunch:
-#if WITH_CRUNCH
-				check(Contexts[LayerIndex]);
-				CrunchCompression::DestroyDecoderContext(Contexts[LayerIndex]);
-#endif // WITH_CRUNCH
-				break;
-			default:
-				break;
-			}
-		}
-
 		check(NumCodecs > 0u);
 		--NumCodecs;
 		DEC_DWORD_STAT(STAT_NumCodecs);
diff --git a/Engine/Source/Runtime/Engine/Private/VT/VirtualTextureBuildSettings.cpp b/Engine/Source/Runtime/Engine/Private/VT/VirtualTextureBuildSettings.cpp
--- a/Engine/Source/Runtime/Engine/Private/VT/VirtualTextureBuildSettings.cpp
+++ b/Engine/Source/Runtime/Engine/Private/VT/VirtualTextureBuildSettings.cpp
@@ -20,6 +20,4 @@ void FVirtualTextureBuildSettings::Init()
 {
 	TileSize = CVarVTTileSize.GetValueOnAnyThread();
 	TileBorderSize = CVarVTTileBorderSize.GetValueOnAnyThread();
-	bEnableCompressCrunch = false;
-	bEnableCompressZlib = false;
 }
diff --git a/Engine/Source/Runtime/Engine/Private/VT/VirtualTextureBuilder.cpp b/Engine/Source/Runtime/Engine/Private/VT/VirtualTextureBuilder.cpp
--- a/Engine/Source/Runtime/Engine/Private/VT/VirtualTextureBuilder.cpp
+++ b/Engine/Source/Runtime/Engine/Private/VT/VirtualTextureBuilder.cpp
@@ -43,7 +43,6 @@ void UVirtualTextureBuilder::BuildTexture(FVirtualTextureBuildDesc const& BuildD
 	Texture->Settings.Init();
 	Texture->Settings.TileSize = BuildDesc.TileSize;
 	Texture->Settings.TileBorderSize = BuildDesc.TileBorderSize;
-	Texture->Settings.bEnableCompressCrunch = BuildDesc.bCrunchCompressed;
 	Texture->LossyCompressionAmount = BuildDesc.LossyCompressionAmount;
 
 	Texture->bContinuousUpdate = BuildDesc.bContinuousUpdate;
diff --git a/Engine/Source/Runtime/Engine/Private/VT/VirtualTextureBuiltData.cpp b/Engine/Source/Runtime/Engine/Private/VT/VirtualTextureBuiltData.cpp
--- a/Engine/Source/Runtime/Engine/Private/VT/VirtualTextureBuiltData.cpp
+++ b/Engine/Source/Runtime/Engine/Private/VT/VirtualTextureBuiltData.cpp
@@ -473,7 +473,7 @@ bool FVirtualTextureBuiltData::ValidateData(FStringView const& InDDCDebugContext
 					const uint32 PackedStride = TileWidthInBlocks * GPixelFormats[LayerFormat].BlockBytes;
 					const size_t PackedOutputSize = PackedStride * TileHeightInBlocks;
 
-					if (VTCodec == EVirtualTextureCodec::ZippedGPU)
+					if (VTCodec == EVirtualTextureCodec::ZippedGPU_DEPRECATED)
 					{
 						const uint32 TileOffset = GetTileOffset_Legacy(ChunkIndex, TileIndex);
 						const uint32 NextTileOffset = GetTileOffset_Legacy(ChunkIndex, TileIndex + 1);
diff --git a/Engine/Source/Runtime/Engine/Private/VT/VirtualTextureBuiltData.h b/Engine/Source/Runtime/Engine/Private/VT/VirtualTextureBuiltData.h
--- a/Engine/Source/Runtime/Engine/Private/VT/VirtualTextureBuiltData.h
+++ b/Engine/Source/Runtime/Engine/Private/VT/VirtualTextureBuiltData.h
@@ -23,8 +23,8 @@ enum class EVirtualTextureCodec : uint8
 	White,			//Special case codec, always outputs white pixels 255,255,255,255
 	Flat,			//Special case codec, always outputs 128,125,255,255 (flat normal map)
 	RawGPU,			//Uncompressed data in an GPU-ready format (e.g R8G8B8A8, BC7, ASTC, ...)
-	ZippedGPU,		//Same as RawGPU but with the data zipped
-	Crunch,			//Use the Crunch library to compress data
+	ZippedGPU_DEPRECATED,		//Same as RawGPU but with the data zipped
+	Crunch_DEPRECATED,			//Use the Crunch library to compress data
 	Max,			// Add new codecs before this entry
 };
 
diff --git a/Engine/Source/Runtime/Engine/Private/VT/VirtualTextureChunkManager.cpp b/Engine/Source/Runtime/Engine/Private/VT/VirtualTextureChunkManager.cpp
--- a/Engine/Source/Runtime/Engine/Private/VT/VirtualTextureChunkManager.cpp
+++ b/Engine/Source/Runtime/Engine/Private/VT/VirtualTextureChunkManager.cpp
@@ -7,7 +7,6 @@
 #include "UploadingVirtualTexture.h"
 #include "VirtualTextureBuiltData.h"
 #include "BlockCodingHelpers.h"
-#include "CrunchCompression.h"
 #include "FileCache/FileCache.h"
 #include "Containers/LruCache.h"
 
diff --git a/Engine/Source/Runtime/Engine/Private/VT/VirtualTextureDataBuilder.cpp b/Engine/Source/Runtime/Engine/Private/VT/VirtualTextureDataBuilder.cpp
--- a/Engine/Source/Runtime/Engine/Private/VT/VirtualTextureDataBuilder.cpp
+++ b/Engine/Source/Runtime/Engine/Private/VT/VirtualTextureDataBuilder.cpp
@@ -5,7 +5,6 @@
 #include "VirtualTextureDataBuilder.h"
 #include "Modules/ModuleManager.h"
 #include "Modules/ModuleManager.h"
-#include "CrunchCompression.h"
 #include "Async/TaskGraphInterfaces.h"
 #include "Async/ParallelFor.h"
 #include "Misc/DataDrivenPlatformInfoRegistry.h"
@@ -598,96 +597,12 @@ void FVirtualTextureDataBuilder::BuildTiles(const TArray<FVTSourceTileEntry>& Ti
 	const int32 TileSize = BuildSettingsLayer0.VirtualTextureTileSize;
 	const int32 BorderSize = BuildSettingsLayer0.VirtualTextureBorderSize;
 	const int32 PhysicalTileSize = TileSize + BorderSize * 2;
+	
+	UE_LOG(LogVirtualTexturing, Verbose, TEXT("VT BuildTiles LayerIndex=%d TileList.Num=%d Pixels=%lld"), LayerIndex, TileList.Num(), (int64)PhysicalTileSize*PhysicalTileSize*TileList.Num() );
 
 	FThreadSafeBool bCompressionError = false;
 	EPixelFormat CompressedFormat = PF_Unknown;
 
-#if WITH_CRUNCH_COMPRESSION
-	if (LayerData.bUseCrunch)
-	{
-		check(LayerData.ImageFormat == ERawImageFormat::BGRA8);
-
-		FCrunchEncodeParameters CrunchParameters;
-		CrunchParameters.ImageWidth = PhysicalTileSize;
-		CrunchParameters.ImageHeight = PhysicalTileSize;
-		CrunchParameters.bIsGammaCorrected = LayerData.GammaSpace != EGammaSpace::Linear;
-		CrunchParameters.OutputFormat = LayerData.TextureFormatName;
-		switch (BuildSettingsLayer0.LossyCompressionAmount)
-		{
-		// map Default to Lowest
-		//	(this used to be done outside in the Engine)
-		case TLCA_Default:
-		case TLCA_Lowest: CrunchParameters.CompressionAmmount = 0.0f; break;
-		case TLCA_Low: CrunchParameters.CompressionAmmount = 0.25f; break;
-		case TLCA_Medium: CrunchParameters.CompressionAmmount = 0.5f; break;
-		case TLCA_High: CrunchParameters.CompressionAmmount = 0.75f; break;
-		case TLCA_Highest: CrunchParameters.CompressionAmmount = 1.0f; break;
-		default: checkNoEntry(); CrunchParameters.CompressionAmmount = 0.0f; break;
-		}
-		
-		UE_LOG(LogVirtualTexturing, Verbose, TEXT("Crunch LCA=%d CompressionAmmount=%f"), (int)BuildSettingsLayer0.LossyCompressionAmount, CrunchParameters.CompressionAmmount);
-
-		// We can't split crunch compression into multiple tasks/threads, since all tiles need to compress together in order to generate the codec payload
-		// Instead we rely on internal Crunch threading to make this efficient
-		// Might be worth modifying Crunch to expose threading callbacks, so this can use UE task graph instead of Crunch's internal threadpool
-		if (bAllowAsync && FApp::ShouldUseThreadingForPerformance())
-		{
-			CrunchParameters.NumWorkerThreads = FTaskGraphInterface::Get().GetNumWorkerThreads();
-		}
-
-		CrunchParameters.RawImagesRGBA.Reserve(TileList.Num());
-		for (const FVTSourceTileEntry& Tile : TileList)
-		{
-			const FTextureSourceBlockData& Block = SourceBlocks[Tile.BlockIndex];
-			const FImage& SourceMip = Block.MipsPerLayer[LayerIndex][Tile.MipIndexInBlock];
-			check(SourceMip.Format == LayerData.ImageFormat);
-
-			FPixelDataRectangle SourceData(LayerData.SourceFormat,
-				SourceMip.SizeX,
-				SourceMip.SizeY,
-				const_cast<uint8*>(SourceMip.RawData.GetData()));
-
-			TArray<uint32>& RawImage = CrunchParameters.RawImagesRGBA.AddDefaulted_GetRef();
-			RawImage.AddUninitialized(PhysicalTileSize * PhysicalTileSize);
-			FPixelDataRectangle TileData(LayerData.SourceFormat, PhysicalTileSize, PhysicalTileSize, (uint8*)RawImage.GetData());
-
-			TileData.Clear();
-			TileData.CopyRectangleBordered(0, 0, SourceData,
-				Tile.TileInBlockX * TileSize - BorderSize,
-				Tile.TileInBlockY * TileSize - BorderSize,
-				PhysicalTileSize,
-				PhysicalTileSize,
-				(TextureAddress)BuildSettingsLayer0.VirtualAddressingModeX,
-				(TextureAddress)BuildSettingsLayer0.VirtualAddressingModeY);
-
-			// Convert input image to the format expected by Crunch library
-			for (int32 PixelIndex = 0u; PixelIndex < RawImage.Num(); ++PixelIndex)
-			{
-				const FColor Color(RawImage[PixelIndex]);
-				RawImage[PixelIndex] = Color.ToPackedABGR();
-			}
-		}
-
-		if (CrunchCompression::Encode(CrunchParameters, GeneratedData.CodecPayload, GeneratedData.TilePayload))
-		{
-			static FName NameDXT1(TEXT("DXT1"));
-			static FName NameDXT5(TEXT("DXT5"));
-			static FName NameBC4(TEXT("BC4"));
-			static FName NameBC5(TEXT("BC5"));
-			GeneratedData.Codec = EVirtualTextureCodec::Crunch;
-			if (LayerData.TextureFormatName == NameDXT1) CompressedFormat = PF_DXT1;
-			else if (LayerData.TextureFormatName == NameDXT5) CompressedFormat = PF_DXT5;
-			else if (LayerData.TextureFormatName == NameBC4) CompressedFormat = PF_BC4;
-			else if (LayerData.TextureFormatName == NameBC5) CompressedFormat = PF_BC5;
-			else CompressedFormat = PF_Unknown;
-		}
-		else
-		{
-			bCompressionError = true;
-		}
-	}
-	else
-#endif // WITH_CRUNCH_COMPRESSION
 	{
 		// Create settings for building the tile. These should be simple, "clean" settings
 		// just compressing the style to a GPU format not adding things like colour correction, ... 
@@ -762,7 +677,7 @@ void FVirtualTextureDataBuilder::BuildTiles(const TArray<FVTSourceTileEntry>& Ti
 #endif // SAVE_TILES
 
 			// give each tile a unique DebugTexturePathName for DebugDump option :
-			FString DebugTilePathName = FString::Printf(TEXT("%s_VT%04d"), *DebugTexturePathName, TileIndex);
+			FString DebugTilePathName = FString::Printf(TEXT("%s_L%d_VT%04d"), *DebugTexturePathName, LayerIndex, TileIndex);
 
 			TArray<FCompressedImage2D> CompressedMip;
 			TArray<FImage> EmptyList;
@@ -781,32 +696,11 @@ void FVirtualTextureDataBuilder::BuildTiles(const TArray<FVTSourceTileEntry>& Ti
 			CompressedFormat = (EPixelFormat)CompressedMip[0].PixelFormat;
 
 			const uint32 SizeRaw = CompressedMip[0].RawData.Num() * CompressedMip[0].RawData.GetTypeSize();
-			if (BuildSettingsLayer0.bVirtualTextureEnableCompressZlib)
-			{
-				TArray<uint8>& TilePayload = GeneratedData.TilePayload[TileIndex];
-				int32 CompressedTileSize = FCompression::CompressMemoryBound(NAME_Zlib, SizeRaw);
-				TilePayload.AddUninitialized(CompressedTileSize);
-				verify(FCompression::CompressMemory(NAME_Zlib, TilePayload.GetData(), CompressedTileSize, CompressedMip[0].RawData.GetData(), SizeRaw));
-				check(CompressedTileSize <= TilePayload.Num());
-
-				// Set the correct size of the compressed tile, but avoid reallocating/copying memory
-				TilePayload.SetNum(CompressedTileSize, false);
-			}
-			else
-			{
-				GeneratedData.TilePayload[TileIndex] = MoveTemp(CompressedMip[0].RawData);
-			}
+			GeneratedData.TilePayload[TileIndex] = MoveTemp(CompressedMip[0].RawData);
 		}, 
 		(bIsSingleThreaded ? EParallelForFlags::ForceSingleThread : EParallelForFlags::None));
 
-		if (BuildSettingsLayer0.bVirtualTextureEnableCompressZlib)
-		{
-			GeneratedData.Codec = EVirtualTextureCodec::ZippedGPU;
-		}
-		else
-		{
-			GeneratedData.Codec = EVirtualTextureCodec::RawGPU;
-		}
+		GeneratedData.Codec = EVirtualTextureCodec::RawGPU;
 	}
 
 	if (OutData.LayerTypes[LayerIndex] == EPixelFormat::PF_Unknown)
@@ -825,6 +719,17 @@ void FVirtualTextureDataBuilder::BuildTiles(const TArray<FVTSourceTileEntry>& Ti
 		GeneratedData.Codec = EVirtualTextureCodec::Max;
 		UE_LOG(LogVirtualTexturing, Fatal, TEXT("Failed build tile"));
 	}
+	else
+	{
+		int CodecPayloadSize = GeneratedData.CodecPayload.Num();
+		int64 TilePayloadTotalSize = 0;
+		for(int i=0;i<GeneratedData.TilePayload.Num();i++)
+		{
+			TilePayloadTotalSize += GeneratedData.TilePayload[i].Num();
+		}
+
+		UE_LOG(LogVirtualTexturing, Verbose, TEXT("VT CodecPayloadSize = %d TilePayloadTotalSize = %lld"), CodecPayloadSize, TilePayloadTotalSize );
+	}
 }
 
 void FVirtualTextureDataBuilder::PushDataToChunk(const TArray<FVTSourceTileEntry>& Tiles, const TArray<FLayerData>& LayerData)
@@ -1303,29 +1208,7 @@ void FVirtualTextureDataBuilder::BuildSourcePixels(const FTextureSourceData& Sou
 			}
 		}
 
-		bool bUseCrunch = false;
-#if WITH_CRUNCH_COMPRESSION
-		if ( SettingsPerLayer[0].bVirtualTextureEnableCompressCrunch &&
-			SettingsPerLayer[0].LossyCompressionAmount != TLCA_None )
-		{
-			if ( ! CrunchCompression::IsValidFormat(TextureFormatName) )
-			{
-				UE_LOG(LogVirtualTexturing, Warning, TEXT("VT Format invalid for Crunch, disabled") );				
-			}
-			else if ( OutData.Width * OutData.Height > FMath::Square(8u * 1024u) )
-			{
-				// Crunch compressor crashes if given size larger than 8k by 8k, needs to be updated to use 64bit sizes/offsets in various places
-				UE_LOG(LogVirtualTexturing, Warning, TEXT("VT too big for Crunch, disabled") );
-			}
-			else
-			{
-				bUseCrunch = true;
-			}
-		}
-#endif // WITH_CRUNCH_COMPRESSION
-
-		if (bUseCrunch // NOTE: Crunch expects a format with no prefixes. See GetCrnFormat().
-			|| TextureFormatPrefix.IsNone())
+		if ( TextureFormatPrefix.IsNone())
 		{
 			LayerData.TextureFormatName = TextureFormatName;
 		}
@@ -1333,27 +1216,6 @@ void FVirtualTextureDataBuilder::BuildSourcePixels(const FTextureSourceData& Sou
 		{
 			LayerData.TextureFormatName = *(TextureFormatPrefix.ToString() + TextureFormatName.ToString());
 		}
-		LayerData.bUseCrunch = bUseCrunch;
-
-		if (bUseCrunch && LayerData.ImageFormat != ERawImageFormat::BGRA8)
-		{
-			// Crunch input data must be in BGRA8 format
-			LayerData.FormatName = "BGRA8";
-			LayerData.PixelFormat = PF_B8G8R8A8;
-			LayerData.SourceFormat = TSF_BGRA8;
-			LayerData.ImageFormat = ERawImageFormat::BGRA8;
-			LayerData.GammaSpace = BuildSettingsForLayer.GetGammaSpace();
-
-			FImage ConvertedMip;
-			for(FTextureSourceBlockData& SourceBlockData : SourceBlocks)
-			{
-				for (FImage& SourceMip : SourceBlockData.MipsPerLayer[LayerIndex])
-				{
-					SourceMip.CopyTo(ConvertedMip, ERawImageFormat::BGRA8, LayerData.GammaSpace);
-					SourceMip = MoveTemp(ConvertedMip);
-				}
-			}
-		}
 	}
 }
 
diff --git a/Engine/Source/Runtime/Engine/Private/VT/VirtualTextureDataBuilder.h b/Engine/Source/Runtime/Engine/Private/VT/VirtualTextureDataBuilder.h
--- a/Engine/Source/Runtime/Engine/Private/VT/VirtualTextureDataBuilder.h
+++ b/Engine/Source/Runtime/Engine/Private/VT/VirtualTextureDataBuilder.h
@@ -40,7 +40,6 @@ struct FVirtualTextureSourceLayerData
 
 	EGammaSpace GammaSpace;
 	bool bHasAlpha;
-	bool bUseCrunch;
 };
 
 /**
diff --git a/Engine/Source/Runtime/Engine/Private/VT/VirtualTextureTranscodeCache.cpp b/Engine/Source/Runtime/Engine/Private/VT/VirtualTextureTranscodeCache.cpp
--- a/Engine/Source/Runtime/Engine/Private/VT/VirtualTextureTranscodeCache.cpp
+++ b/Engine/Source/Runtime/Engine/Private/VT/VirtualTextureTranscodeCache.cpp
@@ -4,7 +4,6 @@
 
 #include "BlockCodingHelpers.h"
 #include "EngineModule.h"
-#include "CrunchCompression.h"
 #include "RendererInterface.h"
 #include "UploadingVirtualTexture.h"
 #include "VirtualTextureChunkManager.h"
@@ -102,8 +101,7 @@ struct FTranscodeTask
 
 		// Used to allocate any temp memory needed to decode tile
 		// Inline allocator hopefully avoids heap allocation in most cases
-		// Most common allocation need here is to linearize compressed Crunch source tile
-		// 136x136 DXT5 tile is 18k uncompressed, so will generally be around 2-4k when compressed with Crunch
+		// 136x136 DXT5 tile is 18k uncompressed
 		// TaskGraph threads currently look to be created with 384k of stack space
 		TArray<uint8, TInlineAllocator<16u * 1024>> TempBuffer;
 
@@ -170,32 +168,12 @@ struct FTranscodeTask
 					}
 				}
 				break;
-			case EVirtualTextureCodec::Crunch:
+			case EVirtualTextureCodec::Crunch_DEPRECATED:
 			{
-#if WITH_CRUNCH
-				// See if we can access compressed tile as a single contiguous block of memory
-				int64 DataReadSize = 0;
-				const void* CompressedTile = Params.Data->Read(DataReadSize, DataOffset, TileLayerSize);
-				if (DataReadSize < TileLayerSize)
-				{
-					// Couldn't access the full block, need to allocate temp block of contiguous memory
-					TempBuffer.SetNumUninitialized(TileLayerSize, false);
-					Params.Data->CopyTo(TempBuffer.GetData(), DataOffset, TileLayerSize);
-					CompressedTile = TempBuffer.GetData();
-				}
-
-				check(Params.Codec);
-				const uint32 StagingBufferSize = StagingBufferForLayer.Stride * TileHeightInBlocks;
-				check(StagingBufferSize <= StagingBufferForLayer.MemorySize);
-				bDecodeResult = CrunchCompression::Decode(Params.Codec->Contexts[LayerIndex],
-					CompressedTile, TileLayerSize,
-					StagingBufferForLayer.Memory, StagingBufferSize, StagingBufferForLayer.Stride);
-#else
 				bDecodeResult = false;
-#endif
 				break;
 			}
-			case EVirtualTextureCodec::ZippedGPU:
+			case EVirtualTextureCodec::ZippedGPU_DEPRECATED:
 				if (StagingBufferForLayer.Stride == PackedStride)
 				{
 					// output buffer is tightly packed, can decompress directly
