Fixed Translucent material sampling stencil buffer not working when r.SeparateTranslucencyScreenPercentage<100.
It works but it does not account for alliasing. It is literally impossible to know if users want to expend stencil bit erode.

diff --git a/Engine/Shaders/Private/DeferredShadingCommon.ush b/Engine/Shaders/Private/DeferredShadingCommon.ush
--- a/Engine/Shaders/Private/DeferredShadingCommon.ush
+++ b/Engine/Shaders/Private/DeferredShadingCommon.ush
@@ -779,7 +779,9 @@ FGBufferData GetGBufferData(float2 UV, bool bGetNormalizedNormal = true)
 	float4 GBufferD = Texture2DSampleLevel(SceneTexturesStruct.GBufferDTexture, SceneTexturesStruct_GBufferDTextureSampler, UV, 0);
 	float CustomNativeDepth = Texture2DSampleLevel(SceneTexturesStruct.CustomDepthTexture, SceneTexturesStruct_CustomDepthTextureSampler, UV, 0).r;
 
-	int2 IntUV = (int2)trunc(UV * View.BufferSizeAndInvSize.xy);
+	// BufferToSceneTextureScale is necessary when translucent materials are rendered in a render target 
+	// that has a different resolution than the scene color textures, e.g. r.SeparateTranslucencyScreenPercentage < 100.
+	int2 IntUV = (int2)trunc(UV * View.BufferSizeAndInvSize.xy * View.BufferToSceneTextureScale.xy);
 	uint CustomStencil = SceneTexturesStruct.CustomStencilTexture.Load(int3(IntUV, 0)) STENCIL_COMPONENT_SWIZZLE;
 
 	#if ALLOW_STATIC_LIGHTING
diff --git a/Engine/Source/Runtime/Engine/Private/ShaderCompiler/ShaderGenerationUtil.cpp b/Engine/Source/Runtime/Engine/Private/ShaderCompiler/ShaderGenerationUtil.cpp
--- a/Engine/Source/Runtime/Engine/Private/ShaderCompiler/ShaderGenerationUtil.cpp
+++ b/Engine/Source/Runtime/Engine/Private/ShaderCompiler/ShaderGenerationUtil.cpp
@@ -1363,7 +1363,9 @@ static FString CreateGBufferDecodeFunctionVariation(const FGBufferInfo& BufferIn
 
 		if (FEATURE_LEVEL >= ERHIFeatureLevel::SM5)
 		{
-			FullStr += FString::Printf(TEXT("\tint2 IntUV = (int2)trunc(%s * View.BufferSizeAndInvSize.xy);\n"), CoordName.GetCharArray().GetData());
+			// BufferToSceneTextureScale is necessary when translucent materials are rendered in a render target 
+			// that has a different resolution than the scene color textures, e.g. r.SeparateTranslucencyScreenPercentage < 100.
+			FullStr += FString::Printf(TEXT("\tint2 IntUV = (int2)trunc(%s * View.BufferSizeAndInvSize.xy * View.BufferToSceneTextureScale.xy);\n"), CoordName.GetCharArray().GetData());
 			FullStr += TEXT("\tuint CustomStencil = SceneTexturesStruct.CustomStencilTexture.Load(int3(IntUV, 0)) STENCIL_COMPONENT_SWIZZLE;\n");
 		}
 		else
diff --git a/Engine/Source/Runtime/Engine/Public/SceneView.h b/Engine/Source/Runtime/Engine/Public/SceneView.h
--- a/Engine/Source/Runtime/Engine/Public/SceneView.h
+++ b/Engine/Source/Runtime/Engine/Public/SceneView.h
@@ -649,6 +649,7 @@ enum ETranslucencyVolumeCascade
 	VIEW_UNIFORM_BUFFER_MEMBER(FVector4f, BufferSizeAndInvSize) \
 	VIEW_UNIFORM_BUFFER_MEMBER(FVector4f, BufferBilinearUVMinMax) \
 	VIEW_UNIFORM_BUFFER_MEMBER(FVector4f, ScreenToViewSpace) \
+	VIEW_UNIFORM_BUFFER_MEMBER(FVector2f, BufferToSceneTextureScale) \
 	VIEW_UNIFORM_BUFFER_MEMBER(int32, NumSceneColorMSAASamples) \
 	VIEW_UNIFORM_BUFFER_MEMBER(float, PreExposure) \
 	VIEW_UNIFORM_BUFFER_MEMBER(float, OneOverPreExposure) \
diff --git a/Engine/Source/Runtime/Renderer/Private/SceneRendering.cpp b/Engine/Source/Runtime/Renderer/Private/SceneRendering.cpp
--- a/Engine/Source/Runtime/Renderer/Private/SceneRendering.cpp
+++ b/Engine/Source/Runtime/Renderer/Private/SceneRendering.cpp
@@ -1334,6 +1334,8 @@ void FViewInfo::SetupUniformBufferParameters(
 		}
 	}
 
+	ViewUniformShaderParameters.BufferToSceneTextureScale = FVector2f(1.0f, 1.0f);
+
 	FRHITexture* TransmittanceLutTextureFound = nullptr;
 	FRHITexture* SkyViewLutTextureFound = nullptr;
 	FRHITexture* CameraAerialPerspectiveVolumeFound = nullptr;
diff --git a/Engine/Source/Runtime/Renderer/Private/TranslucentRendering.cpp b/Engine/Source/Runtime/Renderer/Private/TranslucentRendering.cpp
--- a/Engine/Source/Runtime/Renderer/Private/TranslucentRendering.cpp
+++ b/Engine/Source/Runtime/Renderer/Private/TranslucentRendering.cpp
@@ -881,6 +881,8 @@ void SetupDownsampledTranslucencyViewParameters(
 	float ActualDownsampleX = float(ViewRect.Width()) / float(View.ViewRect.Width());
 	float ActualDownsampleY = float(ViewRect.Height()) / float(View.ViewRect.Height());
 	DownsampledTranslucencyViewParameters.LightProbeSizeRatioAndInvSizeRatio = FVector4f(ActualDownsampleX, ActualDownsampleY, 1.0f / ActualDownsampleX, 1.0f / ActualDownsampleY);
+
+	DownsampledTranslucencyViewParameters.BufferToSceneTextureScale = FVector2f(1.0f / ActualDownsampleX, 1.0f / ActualDownsampleY);
 }
 
 TRDGUniformBufferRef<FTranslucentBasePassUniformParameters> CreateTranslucentBasePassUniformBuffer(
