Cached HLSLTree on Material/MaterialInstance.  Generate the tree when material is loaded in editor.  Preparation for using tree to generate cached data for materials

diff --git a/Engine/Source/Runtime/Engine/Classes/Materials/Material.h b/Engine/Source/Runtime/Engine/Classes/Materials/Material.h
--- a/Engine/Source/Runtime/Engine/Classes/Materials/Material.h
+++ b/Engine/Source/Runtime/Engine/Classes/Materials/Material.h
@@ -1103,8 +1103,6 @@ public:
 
 	ENGINE_API bool IsUsingControlFlow() const;
 
-	ENGINE_API bool IsUsingNewHLSLGenerator() const;
-
 #if WITH_EDITOR
 	/** Allows material properties to be compiled with the option of being overridden by the material attributes input. */
 	ENGINE_API virtual int32 CompilePropertyEx( class FMaterialCompiler* Compiler, const FGuid& AttributeID ) override;
diff --git a/Engine/Source/Runtime/Engine/Classes/Materials/MaterialInstance.h b/Engine/Source/Runtime/Engine/Classes/Materials/MaterialInstance.h
--- a/Engine/Source/Runtime/Engine/Classes/Materials/MaterialInstance.h
+++ b/Engine/Source/Runtime/Engine/Classes/Materials/MaterialInstance.h
@@ -643,6 +643,9 @@ public:
 	virtual ENGINE_API const UMaterial* GetMaterial_Concurrent(TMicRecursionGuard RecursionGuard = TMicRecursionGuard()) const override;
 	virtual ENGINE_API void GetMaterialInheritanceChain(FMaterialInheritanceChain& OutChain) const override;
 	virtual ENGINE_API const FMaterialCachedExpressionData& GetCachedExpressionData(TMicRecursionGuard RecursionGuard = TMicRecursionGuard()) const override;
+#if WITH_EDITOR
+	virtual ENGINE_API const FMaterialCachedHLSLTree& GetCachedHLSLTree(TMicRecursionGuard RecursionGuard = TMicRecursionGuard()) const override;
+#endif
 	virtual ENGINE_API FMaterialResource* AllocatePermutationResource();
 	virtual ENGINE_API FMaterialResource* GetMaterialResource(ERHIFeatureLevel::Type InFeatureLevel, EMaterialQualityLevel::Type QualityLevel = EMaterialQualityLevel::Num) override;
 	virtual ENGINE_API const FMaterialResource* GetMaterialResource(ERHIFeatureLevel::Type InFeatureLevel, EMaterialQualityLevel::Type QualityLevel = EMaterialQualityLevel::Num) const override;
diff --git a/Engine/Source/Runtime/Engine/Classes/Materials/MaterialInterface.h b/Engine/Source/Runtime/Engine/Classes/Materials/MaterialInterface.h
--- a/Engine/Source/Runtime/Engine/Classes/Materials/MaterialInterface.h
+++ b/Engine/Source/Runtime/Engine/Classes/Materials/MaterialInterface.h
@@ -39,6 +39,7 @@ class UTexture;
 class UMaterialInstance;
 struct FMaterialParameterInfo;
 struct FMaterialResourceLocOnDisk;
+class FMaterialCachedHLSLTree;
 #if WITH_EDITORONLY_DATA
 struct FParameterChannelNames;
 #endif
@@ -231,6 +232,10 @@ private:
 
 public:
 
+	UMaterialInterface();
+	UMaterialInterface(FVTableHelper& Helper);
+	ENGINE_API virtual ~UMaterialInterface();
+
 	//~ Begin IInterface_AssetUserData Interface
 	ENGINE_API virtual void AddAssetUserData(UAssetUserData* InUserData) override;
 	ENGINE_API virtual void RemoveUserDataOfClass(TSubclassOf<UAssetUserData> InUserDataClass) override;
@@ -321,6 +326,12 @@ public:
 
 	ENGINE_API virtual const FMaterialCachedExpressionData& GetCachedExpressionData(TMicRecursionGuard RecursionGuard = TMicRecursionGuard()) const;
 
+#if WITH_EDITOR
+	ENGINE_API virtual const FMaterialCachedHLSLTree& GetCachedHLSLTree(TMicRecursionGuard RecursionGuard = TMicRecursionGuard()) const;
+#endif
+
+	ENGINE_API bool IsUsingNewHLSLGenerator() const;
+
 	/**
 	* Test this material for dependency on a given material.
 	* @param	TestDependency - The material to test for dependency upon.
@@ -958,6 +969,9 @@ protected:
 	/** Set if CachedExpressionData was loaded from disk, should typically be true when running with cooked data, and false in the editor */
 	bool bLoadedCachedExpressionData = false;
 
+#if WITH_EDITOR
+	TUniquePtr<FMaterialCachedHLSLTree> CachedHLSLTree;
+#endif // WITH_EDITOR
 private:
 	/**
 	 * Post loads all default materials.
diff --git a/Engine/Source/Runtime/Engine/Private/HLSLTree/HLSLTree.cpp b/Engine/Source/Runtime/Engine/Private/HLSLTree/HLSLTree.cpp
--- a/Engine/Source/Runtime/Engine/Private/HLSLTree/HLSLTree.cpp
+++ b/Engine/Source/Runtime/Engine/Private/HLSLTree/HLSLTree.cpp
@@ -377,7 +377,7 @@ void FExpressionLocalPHI::EmitValueShader(FEmitContext& Context, FEmitScope& Sco
 			// This is the first time we've emitted shader code for this PHI
 			// Create an expression and add it to the map first, so if this is called recursively this path will only be taken the first time
 			const int32 LocalPHIIndex = Context.NumExpressionLocalPHIs++;
-			const Shader::FType LocalType = GetType();
+			const Shader::FType LocalType = Context.GetType(this);
 
 			EmitExpression = OutResult.Code = Context.EmitInlineExpression(Scope,
 				LocalType,
@@ -441,7 +441,7 @@ void FExpressionLocalPHI::EmitValuePreshader(FEmitContext& Context, FEmitScope&
 		}
 	}
 
-	OutResult.Type = GetType();
+	OutResult.Type = Context.GetType(this);
 	if (ValueStackPosition == INDEX_NONE)
 	{
 		Private::FLocalPHILiveScopes LiveScopes;
@@ -520,11 +520,6 @@ void FExpressionFunctionCall::EmitValuePreshader(FEmitContext& Context, FEmitSco
 	OutResult.Type = Function->OutputExpressions[OutputIndex]->GetValuePreshader(Context, Scope, RequestedType, OutResult.Preshader);
 }
 
-void FExpression::Reset()
-{
-	PrepareValueResult = FPrepareValueResult();
-}
-
 FRequestedType::FRequestedType(ERequestedType InType)
 {
 	const int32 NumComponents = (uint8)InType;
@@ -1117,13 +1112,14 @@ FEmitShaderExpression* FExpression::GetValueShader(FEmitContext& Context, FEmitS
 {
 	FOwnerScope OwnerScope(*Context.Errors, GetOwner());
 	
-	const EExpressionEvaluation Evaluation = PrepareValueResult.PreparedType.GetEvaluation(Scope, RequestedType);
+	const FPreparedType PreparedType = Context.GetPreparedType(this);
+	const EExpressionEvaluation Evaluation = PreparedType.GetEvaluation(Scope, RequestedType);
 	check(Evaluation != EExpressionEvaluation::None && Evaluation != EExpressionEvaluation::Unknown);
 
 	FEmitShaderExpression* Value = nullptr;
 	if (Evaluation == EExpressionEvaluation::ConstantZero)
 	{
-		Value = Context.EmitConstantZero(Scope, GetType());
+		Value = Context.EmitConstantZero(Scope, PreparedType.GetType());
 	}
 	else if (Evaluation == EExpressionEvaluation::Constant || Evaluation == EExpressionEvaluation::Preshader)
 	{
@@ -1136,10 +1132,10 @@ FEmitShaderExpression* FExpression::GetValueShader(FEmitContext& Context, FEmitS
 		Value = Result.Code;
 	}
 
-	if (Value->Type.IsNumeric() && GetPreparedType().IsNumeric())
+	if (Value->Type.IsNumeric() && PreparedType.IsNumeric())
 	{
 		const Shader::FValueTypeDescription TypeDest = Shader::GetValueTypeDescription(Value->Type);
-		if (TypeDest.ComponentType == Shader::EValueComponentType::Double && GetPreparedType().ValueComponentType == Shader::EValueComponentType::Float)
+		if (TypeDest.ComponentType == Shader::EValueComponentType::Double && PreparedType.ValueComponentType == Shader::EValueComponentType::Float)
 		{
 			const Shader::EValueType CastType = Shader::MakeValueType(Shader::EValueComponentType::Float, TypeDest.NumComponents);
 			Value = Context.EmitCast(Scope, Value, CastType);
@@ -1167,7 +1163,8 @@ FEmitShaderExpression* FExpression::GetValueShader(FEmitContext& Context, FEmitS
 
 FEmitShaderExpression* FExpression::GetValueShader(FEmitContext& Context, FEmitScope& Scope)
 {
-	return GetValueShader(Context, Scope, GetRequestedType());
+	const FPreparedType PreparedType = Context.GetPreparedType(this);
+	return GetValueShader(Context, Scope, PreparedType.GetRequestedType());
 }
 
 Shader::FType FExpression::GetValuePreshader(FEmitContext& Context, FEmitScope& Scope, const FRequestedType& RequestedType, Shader::FPreshaderData& OutPreshader)
@@ -1175,14 +1172,15 @@ Shader::FType FExpression::GetValuePreshader(FEmitContext& Context, FEmitScope&
 	FOwnerScope OwnerScope(*Context.Errors, GetOwner());
 	
 	const int32 PrevStackPosition = Context.PreshaderStackPosition;
-	const EExpressionEvaluation Evaluation = PrepareValueResult.PreparedType.GetEvaluation(Scope, RequestedType);
+	const FPreparedType PreparedType = Context.GetPreparedType(this);
+	const EExpressionEvaluation Evaluation = PreparedType.GetEvaluation(Scope, RequestedType);
 	check(Evaluation != EExpressionEvaluation::None && Evaluation != EExpressionEvaluation::Unknown);
 
 	FEmitValuePreshaderResult Result(OutPreshader);
 	if (Evaluation == EExpressionEvaluation::ConstantZero)
 	{
 		Context.PreshaderStackPosition++;
-		Result.Type = GetType();
+		Result.Type = PreparedType.GetType();
 		OutPreshader.WriteOpcode(Shader::EPreshaderOpcode::ConstantZero).Write(Result.Type);
 	}
 	else if (Evaluation == EExpressionEvaluation::Constant)
@@ -1205,18 +1203,17 @@ Shader::FValue FExpression::GetValueConstant(FEmitContext& Context, FEmitScope&
 {
 	FOwnerScope OwnerScope(*Context.Errors, GetOwner());
 	
-	check(!bReentryFlag);
-	const EExpressionEvaluation Evaluation = PrepareValueResult.PreparedType.GetEvaluation(Scope, RequestedType);
+	const FPreparedType PreparedType = Context.GetPreparedType(this);
+	const EExpressionEvaluation Evaluation = PreparedType.GetEvaluation(Scope, RequestedType);
 	if (Evaluation == EExpressionEvaluation::ConstantZero)
 	{
-		return Shader::FValue(GetType());
+		return Shader::FValue(PreparedType.GetType());
 	}
 	else
 	{
 		check(Evaluation == EExpressionEvaluation::Constant || Evaluation == EExpressionEvaluation::ConstantLoop);
 		Shader::FPreshaderData ConstantPreshader;
 		{
-			FExpressionReentryScope ReentryScope(this);
 			const int32 PrevPreshaderStackPosition = Context.PreshaderStackPosition;
 			FEmitValuePreshaderResult PreshaderResult(ConstantPreshader);
 			EmitValuePreshader(Context, Scope, RequestedType, PreshaderResult);
@@ -1312,17 +1309,6 @@ UObject* FOwnerContext::GetCurrentOwner() const
 	return (OwnerStack.Num() > 0) ? OwnerStack.Last() : nullptr;
 }
 
-void FTree::ResetNodes()
-{
-	FNode* Node = Nodes;
-	while (Node)
-	{
-		FNode* Next = Node->NextNode;
-		Node->Reset();
-		Node = Next;
-	}
-}
-
 bool FTree::Finalize()
 {
 	// Resolve values for any PHI nodes that were generated
@@ -1358,6 +1344,7 @@ bool FTree::Finalize()
 		}
 	}
 
+	PHIExpressions.Shrink();
 	return true;
 }
 
@@ -1469,23 +1456,15 @@ FExpression* FTree::NewFunctionCall(FScope& Scope, FFunction* Function, int32 Ou
 	return NewExpression<FExpressionFunctionCall>(Function, OutputIndex);
 }
 
-const FExpressionDerivatives& FTree::GetAnalyticDerivatives(FExpression* InExpression)
+FExpressionDerivatives FTree::GetAnalyticDerivatives(FExpression* InExpression)
 {
-	if (!InExpression)
-	{
-		static const FExpressionDerivatives EmptyDerivatives;
-		return EmptyDerivatives;
-	}
-
-	if (!InExpression->bComputedDerivatives)
+	FExpressionDerivatives Derivatives;
+	if (InExpression)
 	{
-		FExpressionReentryScope ReentryScope(InExpression);
 		FOwnerScope OwnerScope(*this, InExpression->GetOwner()); // Associate any newly created nodes with the same owner as the input expression
-
-		InExpression->ComputeAnalyticDerivatives(*this, InExpression->Derivatives);
-		InExpression->bComputedDerivatives = true;
+		InExpression->ComputeAnalyticDerivatives(*this, Derivatives);
 	}
-	return InExpression->Derivatives;
+	return Derivatives;
 }
 
 FExpression* FTree::GetPreviousFrame(FExpression* InExpression, const FRequestedType& RequestedType)
@@ -1493,9 +1472,7 @@ FExpression* FTree::GetPreviousFrame(FExpression* InExpression, const FRequested
 	FExpression* Result = InExpression;
 	if (Result && !RequestedType.IsVoid())
 	{
-		FExpressionReentryScope ReentryScope(InExpression);
 		FOwnerScope OwnerScope(*this, InExpression->GetOwner()); // Associate any newly created nodes with the same owner as the input expression
-
 		FExpression* PrevFrameExpression = InExpression->ComputePreviousFrame(*this, RequestedType);
 		if (PrevFrameExpression)
 		{
diff --git a/Engine/Source/Runtime/Engine/Private/HLSLTree/HLSLTreeCommon.cpp b/Engine/Source/Runtime/Engine/Private/HLSLTree/HLSLTreeCommon.cpp
--- a/Engine/Source/Runtime/Engine/Private/HLSLTree/HLSLTreeCommon.cpp
+++ b/Engine/Source/Runtime/Engine/Private/HLSLTree/HLSLTreeCommon.cpp
@@ -2,6 +2,7 @@
 #include "HLSLTree/HLSLTreeCommon.h"
 #include "HLSLTree/HLSLTreeEmit.h"
 #include "Misc/StringBuilder.h"
+#include "RenderUtils.h"
 #include "Engine/Texture.h"
 
 namespace UE::HLSLTree
@@ -81,6 +82,11 @@ FSwizzleParameters MakeSwizzleMask(bool bInR, bool bInG, bool bInB, bool bInA)
 	return FSwizzleParameters(ComponentIndex[0], ComponentIndex[1], ComponentIndex[2], ComponentIndex[3]);
 }
 
+bool FExpressionError::PrepareValue(FEmitContext& Context, FEmitScope&, const FRequestedType&, FPrepareValueResult&) const
+{
+	return Context.Errors->AddError(ErrorMessage);
+}
+
 FExpression* FTree::NewConstant(const Shader::FValue& Value)
 {
 	return NewExpression<FExpressionConstant>(Value);
@@ -371,10 +377,10 @@ void FExpressionSetStructField::EmitValueShader(FEmitContext& Context, FEmitScop
 {
 	FRequestedType RequestedStructType(RequestedType);
 	RequestedStructType.ClearFieldRequested(Field);
-	const EExpressionEvaluation StructEvaluation = StructExpression->GetEvaluation(Scope, RequestedStructType);
+	const EExpressionEvaluation StructEvaluation = Context.GetEvaluation(StructExpression, Scope, RequestedStructType);
 
 	const FRequestedType RequestedFieldType = RequestedType.GetField(Field);
-	const EExpressionEvaluation FieldEvaluation = FieldExpression->GetEvaluation(Scope, RequestedFieldType);
+	const EExpressionEvaluation FieldEvaluation = Context.GetEvaluation(FieldExpression, Scope, RequestedStructType);
 
 	FEmitShaderExpression* StructValue = (StructEvaluation != EExpressionEvaluation::None) ? StructExpression->GetValueShader(Context, Scope, RequestedStructType) : nullptr;
 	FEmitShaderExpression* FieldValue = (FieldEvaluation != EExpressionEvaluation::None) ? FieldExpression->GetValueShader(Context, Scope, RequestedFieldType, Field->Type) : nullptr;
@@ -413,10 +419,10 @@ void FExpressionSetStructField::EmitValuePreshader(FEmitContext& Context, FEmitS
 {
 	FRequestedType RequestedStructType(RequestedType);
 	RequestedStructType.ClearFieldRequested(Field);
-	const EExpressionEvaluation StructEvaluation = StructExpression->GetEvaluation(Scope, RequestedStructType);
+	const EExpressionEvaluation StructEvaluation = Context.GetEvaluation(StructExpression, Scope, RequestedStructType);
 
 	const FRequestedType RequestedFieldType = RequestedType.GetField(Field);
-	const EExpressionEvaluation FieldEvaluation = FieldExpression->GetEvaluation(Scope, RequestedFieldType);
+	const EExpressionEvaluation FieldEvaluation = Context.GetEvaluation(FieldExpression, Scope, RequestedStructType);
 
 	if (StructEvaluation != EExpressionEvaluation::None)
 	{
@@ -485,7 +491,7 @@ bool FExpressionSelect::PrepareValue(FEmitContext& Context, FEmitScope& Scope, c
 
 void FExpressionSelect::EmitValueShader(FEmitContext& Context, FEmitScope& Scope, const FRequestedType& RequestedType, FEmitValueShaderResult& OutResult) const
 {
-	const EExpressionEvaluation ConditionEvaluation = ConditionExpression->GetEvaluation(Scope, ERequestedType::Scalar);
+	const EExpressionEvaluation ConditionEvaluation = Context.GetEvaluation(ConditionExpression, Scope, ERequestedType::Scalar);
 	if (ConditionEvaluation == EExpressionEvaluation::Constant)
 	{
 		const bool bCondition = ConditionExpression->GetValueConstant(Context, Scope, Shader::EValueType::Bool1).AsBoolScalar();
@@ -494,7 +500,7 @@ void FExpressionSelect::EmitValueShader(FEmitContext& Context, FEmitScope& Scope
 	}
 	else
 	{
-		const Shader::FType LocalType = GetType();
+		const Shader::FType LocalType = Context.GetType(this);
 		FEmitShaderExpression* TrueValue = TrueExpression->GetValueShader(Context, Scope, RequestedType, LocalType);
 		FEmitShaderExpression* FalseValue = FalseExpression->GetValueShader(Context, Scope, RequestedType, LocalType);
 
@@ -509,7 +515,7 @@ void FExpressionSelect::EmitValuePreshader(FEmitContext& Context, FEmitScope& Sc
 {
 	// TODO
 	Context.PreshaderStackPosition++;
-	OutResult.Type = GetType();
+	OutResult.Type = Context.GetType(this);
 	OutResult.Preshader.WriteOpcode(Shader::EPreshaderOpcode::ConstantZero).Write(OutResult.Type);
 }
 
@@ -559,7 +565,7 @@ void FExpressionDerivative::EmitValuePreshader(FEmitContext& Context, FEmitScope
 {
 	// Derivative of a constant is 0
 	Context.PreshaderStackPosition++;
-	OutResult.Type = GetType();
+	OutResult.Type = Context.GetType(this);
 	OutResult.Preshader.WriteOpcode(Shader::EPreshaderOpcode::ConstantZero).Write(OutResult.Type);
 }
 
@@ -802,7 +808,7 @@ FAppendTypes GetAppendTypes(const FRequestedType& RequestedType, Shader::EValueT
 
 void FExpressionAppend::EmitValueShader(FEmitContext& Context, FEmitScope& Scope, const FRequestedType& RequestedType, FEmitValueShaderResult& OutResult) const
 {
-	const Private::FAppendTypes Types = Private::GetAppendTypes(RequestedType, Lhs->GetType(), Rhs->GetType());
+	const Private::FAppendTypes Types = Private::GetAppendTypes(RequestedType, Context.GetType(Lhs), Context.GetType(Rhs));
 	FEmitShaderExpression* LhsValue = Lhs->GetValueShader(Context, Scope, Types.LhsRequestedType, Types.LhsType);
 
 	if (Types.RhsType == Shader::EValueType::Void)
@@ -830,7 +836,7 @@ void FExpressionAppend::EmitValueShader(FEmitContext& Context, FEmitScope& Scope
 
 void FExpressionAppend::EmitValuePreshader(FEmitContext& Context, FEmitScope& Scope, const FRequestedType& RequestedType, FEmitValuePreshaderResult& OutResult) const
 {
-	const Private::FAppendTypes Types = Private::GetAppendTypes(RequestedType, Lhs->GetType(), Rhs->GetType());
+	const Private::FAppendTypes Types = Private::GetAppendTypes(RequestedType, Context.GetType(Lhs), Context.GetType(Rhs));
 	Lhs->GetValuePreshader(Context, Scope, Types.LhsRequestedType, OutResult.Preshader);
 	if (Types.RhsType != Shader::EValueType::Void)
 	{
@@ -844,6 +850,76 @@ void FExpressionAppend::EmitValuePreshader(FEmitContext& Context, FEmitScope& Sc
 	OutResult.Type = Types.ResultType;
 }
 
+void FExpressionSwitchBase::ComputeAnalyticDerivatives(FTree& Tree, FExpressionDerivatives& OutResult) const
+{
+	FExpression* InputDdx[MaxInputs];
+	FExpression* InputDdy[MaxInputs];
+	for (int32 Index = 0; Index < NumInputs; ++Index)
+	{
+		const FExpressionDerivatives InputDerivative = Tree.GetAnalyticDerivatives(Input[Index]);
+		InputDdx[Index] = InputDerivative.ExpressionDdx;
+		InputDdy[Index] = InputDerivative.ExpressionDdy;
+	}
+	OutResult.ExpressionDdx = NewSwitch(Tree, MakeArrayView(InputDdx, NumInputs));
+	OutResult.ExpressionDdy = NewSwitch(Tree, MakeArrayView(InputDdy, NumInputs));
+}
+
+FExpression* FExpressionSwitchBase::ComputePreviousFrame(FTree& Tree, const FRequestedType& RequestedType) const
+{
+	FExpression* InputPrevFrame[MaxInputs];
+	for (int32 Index = 0; Index < NumInputs; ++Index)
+	{
+		InputPrevFrame[Index] = Tree.GetPreviousFrame(Input[Index], RequestedType);
+	}
+	return NewSwitch(Tree, MakeArrayView(InputPrevFrame, NumInputs));
+}
+
+bool FExpressionSwitchBase::PrepareValue(FEmitContext& Context, FEmitScope& Scope, const FRequestedType& RequestedType, FPrepareValueResult& OutResult) const
+{
+	const int32 InputIndex = GetInputIndex(Context);
+	check(InputIndex >= 0 && InputIndex < NumInputs);
+	const FPreparedType& InputType = Context.PrepareExpression(Input[InputIndex], Scope, RequestedType);
+	if (InputType.IsVoid())
+	{
+		return false;
+	}
+	return OutResult.SetType(Context, RequestedType, InputType);
+}
+
+void FExpressionSwitchBase::EmitValueShader(FEmitContext& Context, FEmitScope& Scope, const FRequestedType& RequestedType, FEmitValueShaderResult& OutResult) const
+{
+	const int32 InputIndex = GetInputIndex(Context);
+	check(InputIndex >= 0 && InputIndex < NumInputs);
+	OutResult.Code = Input[InputIndex]->GetValueShader(Context, Scope, RequestedType);
+}
+
+void FExpressionSwitchBase::EmitValuePreshader(FEmitContext& Context, FEmitScope& Scope, const FRequestedType& RequestedType, FEmitValuePreshaderResult& OutResult) const
+{
+	const int32 InputIndex = GetInputIndex(Context);
+	check(InputIndex >= 0 && InputIndex < NumInputs);
+	OutResult.Type = Input[InputIndex]->GetValuePreshader(Context, Scope, RequestedType, OutResult.Preshader);
+}
+
+int32 FExpressionFeatureLevelSwitch::GetInputIndex(const FEmitContext& Context) const
+{
+	return (int32)Context.TargetParameters.FeatureLevel;
+}
+
+int32 FExpressionShadingPathSwitch::GetInputIndex(const FEmitContext& Context) const
+{
+	const EShaderPlatform ShaderPlatform = Context.TargetParameters.ShaderPlatform;
+	ERHIShadingPath::Type ShadingPathToCompile = ERHIShadingPath::Deferred;
+	if (IsForwardShadingEnabled(ShaderPlatform))
+	{
+		ShadingPathToCompile = ERHIShadingPath::Forward;
+	}
+	else if (Context.TargetParameters.FeatureLevel < ERHIFeatureLevel::SM5)
+	{
+		ShadingPathToCompile = ERHIShadingPath::Mobile;
+	}
+	return (int32)ShadingPathToCompile;
+}
+
 bool FExpressionInlineCustomHLSL::PrepareValue(FEmitContext& Context, FEmitScope& Scope, const FRequestedType& RequestedType, FPrepareValueResult& OutResult) const
 {
 	return OutResult.SetType(Context, RequestedType, EExpressionEvaluation::Shader, Shader::FType(ResultType));
@@ -937,7 +1013,7 @@ bool FStatementIf::Prepare(FEmitContext& Context, FEmitScope& Scope) const
 void FStatementIf::EmitShader(FEmitContext& Context, FEmitScope& Scope) const
 {
 	FEmitShaderNode* Dependency = nullptr;
-	const EExpressionEvaluation ConditionEvaluation = ConditionExpression->GetEvaluation(Scope, ERequestedType::Scalar);
+	const EExpressionEvaluation ConditionEvaluation = Context.GetEvaluation(ConditionExpression, Scope, ERequestedType::Scalar);
 	if (ConditionEvaluation == EExpressionEvaluation::Constant)
 	{
 		const bool bCondition = ConditionExpression->GetValueConstant(Context, Scope, Shader::EValueType::Bool1).AsBoolScalar();
diff --git a/Engine/Source/Runtime/Engine/Private/HLSLTree/HLSLTreeEmit.cpp b/Engine/Source/Runtime/Engine/Private/HLSLTree/HLSLTreeEmit.cpp
--- a/Engine/Source/Runtime/Engine/Private/HLSLTree/HLSLTreeEmit.cpp
+++ b/Engine/Source/Runtime/Engine/Private/HLSLTree/HLSLTreeEmit.cpp
@@ -192,10 +192,14 @@ FEmitScope* FEmitScope::FindLoop()
 	return nullptr;
 }
 
-FEmitContext::FEmitContext(FMemStackBase& InAllocator, FErrorHandlerInterface& InErrors, const Shader::FStructTypeRegistry& InTypeRegistry)
+FEmitContext::FEmitContext(FMemStackBase& InAllocator,
+	const FTargetParameters& InTargetParameters,
+	FErrorHandlerInterface& InErrors,
+	const Shader::FStructTypeRegistry& InTypeRegistry)
 	: Allocator(&InAllocator)
 	, Errors(&InErrors)
 	, TypeRegistry(&InTypeRegistry)
+	, TargetParameters(InTargetParameters)
 {
 }
 
@@ -339,6 +343,28 @@ void FEmitContext::EmitDeclarationsCode(FStringBuilderBase& OutCode)
 	}
 }
 
+const FPreparedType& FEmitContext::GetPreparedType(const FExpression* Expression) const
+{
+	static const FPreparedType VoidType;
+	FPrepareValueResult* const* PrevResult = PrepareValueMap.Find(Expression);
+	return PrevResult ? (*PrevResult)->PreparedType : VoidType;
+}
+
+FRequestedType FEmitContext::GetRequestedType(const FExpression* Expression) const
+{
+	return GetPreparedType(Expression).GetRequestedType();
+}
+
+Shader::FType FEmitContext::GetType(const FExpression* Expression) const
+{
+	return GetPreparedType(Expression).GetType();
+}
+
+EExpressionEvaluation FEmitContext::GetEvaluation(const FExpression* Expression, const FEmitScope& Scope, const FRequestedType& RequestedType) const
+{
+	return GetPreparedType(Expression).GetEvaluation(Scope, RequestedType);
+}
+
 FPreparedType FEmitContext::PrepareExpression(FExpression* InExpression, FEmitScope& Scope, const FRequestedType& RequestedType)
 {
 	if (!InExpression || RequestedType.IsVoid())
@@ -346,29 +372,38 @@ FPreparedType FEmitContext::PrepareExpression(FExpression* InExpression, FEmitSc
 		return FPreparedType();
 	}
 
+	FPrepareValueResult** PrevResult = PrepareValueMap.Find(InExpression);
+	FPrepareValueResult* Result = PrevResult ? *PrevResult : nullptr;
+	if (!Result)
+	{
+		Result = new(*Allocator) FPrepareValueResult();
+		PrepareValueMap.Add(InExpression, Result);
+	}
+
+
 	FOwnerScope OwnerScope(*Errors, InExpression->GetOwner());
-	FPrepareValueResult& Result = InExpression->PrepareValueResult;
-	if (InExpression->bReentryFlag)
+	if (Result->bPreparingValue)
 	{
 		// Valid for this to be called reentrantly
 		// Code should ensure that the type is set before the reentrant call, otherwise type will not be valid here
 		// LocalPHI nodes rely on this to break loops
 		FEmitScope* LoopScope = Scope.FindLoop();
 		check(LoopScope);
-		Result.PreparedType.SetLoopEvaluation(*LoopScope, RequestedType);
-		return Result.PreparedType;
+		Result->PreparedType.SetLoopEvaluation(*LoopScope, RequestedType);
+		return Result->PreparedType;
 	}
 
 	bool bResult = false;
 	{
-		FExpressionReentryScope ReentryScope(InExpression);
-		bResult = InExpression->PrepareValue(*this, Scope, RequestedType, Result);
+		Result->bPreparingValue = true;
+		bResult = InExpression->PrepareValue(*this, Scope, RequestedType, *Result);
+		Result->bPreparingValue = false;
 	}
 
 	FPreparedType ResultType;
 	if (bResult)
 	{
-		ResultType = Result.PreparedType;
+		ResultType = Result->PreparedType;
 		check(!ResultType.IsVoid());
 	}
 	return ResultType;
@@ -843,7 +878,7 @@ FEmitShaderExpression* FEmitContext::EmitPreshaderOrConstant(FEmitScope& Scope,
 		const Shader::EValueType FieldType = Type.GetFlatFieldType(FieldIndex);
 		const Shader::FValueTypeDescription TypeDesc = Shader::GetValueTypeDescription(FieldType);
 		const int32 NumFieldComponents = TypeDesc.NumComponents;
-		const EExpressionEvaluation FieldEvaluation = Expression->GetPreparedType().GetFieldEvaluation(Scope, RequestedType, ComponentIndex, NumFieldComponents);
+		const EExpressionEvaluation FieldEvaluation = GetPreparedType(Expression).GetFieldEvaluation(Scope, RequestedType, ComponentIndex, NumFieldComponents);
 
 		if (FieldEvaluation == EExpressionEvaluation::Preshader)
 		{
@@ -1165,7 +1200,7 @@ FEmitShaderExpression* FEmitContext::EmitCustomHLSL(FEmitScope& Scope, FStringVi
 	EmitInputs.Reserve(Inputs.Num());
 	for (const FCustomHLSLInput& Input : Inputs)
 	{
-		EmitInputs.Add(FEmitCustomHLSLInput{ Input.Name, Input.Expression->GetType() });
+		EmitInputs.Add(FEmitCustomHLSLInput{ Input.Name, GetType(Input.Expression) });
 	}
 
 	FHasher Hasher;
@@ -1217,6 +1252,7 @@ void FEmitContext::Finalize()
 	
 	// Don't reset Expression/Preshader maps, allow future passes to share matching preshaders/expressions
 
+	PrepareValueMap.Reset();
 	EmitScopeMap.Reset();
 	EmitFunctionMap.Reset();
 	EmitLocalPHIMap.Reset();
diff --git a/Engine/Source/Runtime/Engine/Private/HLSLTree/HLSLTreeOperation.cpp b/Engine/Source/Runtime/Engine/Private/HLSLTree/HLSLTreeOperation.cpp
--- a/Engine/Source/Runtime/Engine/Private/HLSLTree/HLSLTreeOperation.cpp
+++ b/Engine/Source/Runtime/Engine/Private/HLSLTree/HLSLTreeOperation.cpp
@@ -180,10 +180,15 @@ FOperationRequestedTypes GetOperationRequestedTypes(EOperation Op, const FReques
 	return Types;
 }
 
-FOperationTypes GetOperationTypes(EOperation Op, TConstArrayView<FExpression*> Inputs)
+FOperationTypes GetOperationTypes(FEmitContext& Context, EOperation Op, TConstArrayView<FExpression*> Inputs)
 {
-	FOperationTypes Types;
+	FPreparedType InputPreparedType[FExpressionOperation::MaxInputs];
+	for (int32 Index = 0; Index < Inputs.Num(); ++Index)
+	{
+		InputPreparedType[Index] = Context.GetPreparedType(Inputs[Index]);
+	}
 
+	FOperationTypes Types;
 	if (Op == EOperation::VecMulMatrix3 ||
 		Op == EOperation::VecMulMatrix4 ||
 		Op == EOperation::Matrix3MulVec ||
@@ -193,9 +198,8 @@ FOperationTypes GetOperationTypes(EOperation Op, TConstArrayView<FExpression*> I
 		Shader::EValueComponentType IntermediateComponentType = Shader::EValueComponentType::Void;
 		for (int32 Index = 0; Index < Inputs.Num(); ++Index)
 		{
-			const FPreparedType& InputType = Inputs[Index]->GetPreparedType();
-			IntermediateComponent = CombineComponents(IntermediateComponent, InputType.GetMergedComponent());
-			IntermediateComponentType = Shader::CombineComponentTypes(IntermediateComponentType, InputType.ValueComponentType);
+			IntermediateComponent = CombineComponents(IntermediateComponent, InputPreparedType[Index].GetMergedComponent());
+			IntermediateComponentType = Shader::CombineComponentTypes(IntermediateComponentType, InputPreparedType[Index].ValueComponentType);
 		}
 
 		switch (Op)
@@ -225,8 +229,7 @@ FOperationTypes GetOperationTypes(EOperation Op, TConstArrayView<FExpression*> I
 		FPreparedType IntermediateType;
 		for (int32 Index = 0; Index < Inputs.Num(); ++Index)
 		{
-			const FPreparedType& InputType = Inputs[Index]->GetPreparedType();
-			IntermediateType = MergePreparedTypes(IntermediateType, InputType);
+			IntermediateType = MergePreparedTypes(IntermediateType, InputPreparedType[Index]);
 		}
 
 		const Shader::EValueType IntermediateValueType = IntermediateType.GetType();
@@ -299,13 +302,13 @@ FOperationTypes GetOperationTypes(EOperation Op, TConstArrayView<FExpression*> I
 		case EOperation::Min:
 			for (int32 Index = 0; Index < Types.ResultType.PreparedComponents.Num(); ++Index)
 			{
-				Types.ResultType.SetComponentBounds(Index, Shader::MinBound(Inputs[0]->GetPreparedType().GetComponentBounds(Index), Inputs[1]->GetPreparedType().GetComponentBounds(Index)));
+				Types.ResultType.SetComponentBounds(Index, Shader::MinBound(InputPreparedType[0].GetComponentBounds(Index), InputPreparedType[1].GetComponentBounds(Index)));
 			}
 			break;
 		case EOperation::Max:
 			for (int32 Index = 0; Index < Types.ResultType.PreparedComponents.Num(); ++Index)
 			{
-				Types.ResultType.SetComponentBounds(Index, Shader::MaxBound(Inputs[0]->GetPreparedType().GetComponentBounds(Index), Inputs[1]->GetPreparedType().GetComponentBounds(Index)));
+				Types.ResultType.SetComponentBounds(Index, Shader::MaxBound(InputPreparedType[0].GetComponentBounds(Index), InputPreparedType[1].GetComponentBounds(Index)));
 			}
 			break;
 		default:
@@ -545,7 +548,7 @@ bool FExpressionOperation::PrepareValue(FEmitContext& Context, FEmitScope& Scope
 		}
 	}
 
-	Private::FOperationTypes Types = Private::GetOperationTypes(Op, MakeArrayView(Inputs, OpDesc.NumInputs));
+	Private::FOperationTypes Types = Private::GetOperationTypes(Context, Op, MakeArrayView(Inputs, OpDesc.NumInputs));
 	if (OpDesc.PreshaderOpcode == Shader::EPreshaderOpcode::Nop)
 	{
 		// No preshader support
@@ -572,7 +575,7 @@ void FExpressionOperation::EmitValueShader(FEmitContext& Context, FEmitScope& Sc
 {
 	const FOperationDescription OpDesc = GetOperationDescription(Op);
 	const Private::FOperationRequestedTypes RequestedTypes = Private::GetOperationRequestedTypes(Op, RequestedType);
-	const Private::FOperationTypes Types = Private::GetOperationTypes(Op, MakeArrayView(Inputs, OpDesc.NumInputs));
+	const Private::FOperationTypes Types = Private::GetOperationTypes(Context, Op, MakeArrayView(Inputs, OpDesc.NumInputs));
 	FEmitShaderExpression* InputValue[MaxInputs] = { nullptr };
 	for (int32 Index = 0; Index < OpDesc.NumInputs; ++Index)
 	{
@@ -672,7 +675,7 @@ void FExpressionOperation::EmitValuePreshader(FEmitContext& Context, FEmitScope&
 {
 	const FOperationDescription OpDesc = GetOperationDescription(Op);
 	const Private::FOperationRequestedTypes RequestedTypes = Private::GetOperationRequestedTypes(Op, RequestedType);
-	const Private::FOperationTypes Types = Private::GetOperationTypes(Op, MakeArrayView(Inputs, OpDesc.NumInputs));
+	const Private::FOperationTypes Types = Private::GetOperationTypes(Context, Op, MakeArrayView(Inputs, OpDesc.NumInputs));
 	check(OpDesc.PreshaderOpcode != Shader::EPreshaderOpcode::Nop);
 
 	for (int32 Index = 0; Index < OpDesc.NumInputs; ++Index)
diff --git a/Engine/Source/Runtime/Engine/Private/Materials/Material.cpp b/Engine/Source/Runtime/Engine/Private/Materials/Material.cpp
--- a/Engine/Source/Runtime/Engine/Private/Materials/Material.cpp
+++ b/Engine/Source/Runtime/Engine/Private/Materials/Material.cpp
@@ -88,6 +88,7 @@
 
 
 #if WITH_EDITOR
+#include "MaterialCachedHLSLTree.h"
 #include "Logging/TokenizedMessage.h"
 #include "Logging/MessageLog.h"
 #include "Misc/UObjectToken.h"
@@ -133,15 +134,6 @@ static TAutoConsoleVariable<int32> CVarMaterialEnableControlFlow(
 	TEXT("Allows experemental control flow to be used in the material editor.\n"),
 	ECVF_RenderThreadSafe | ECVF_ReadOnly);
 
-static TAutoConsoleVariable<int32> CVarMaterialEnableNewHLSLGenerator(
-	TEXT("r.MaterialEnableNewHLSLGenerator"),
-	0,
-	TEXT("Enables the new (WIP) material HLSL generator.\n")
-	TEXT("0 - Don't allow\n")
-	TEXT("1 - Allow if enabled by material\n")
-	TEXT("2 - Force all materials to use new generator\n"),
-	ECVF_RenderThreadSafe | ECVF_ReadOnly);
-
 bool Engine_IsStrataEnabled()
 {
 	static const auto CVar = IConsoleManager::Get().FindTConsoleVariableDataInt(TEXT("r.Strata"));
@@ -1894,6 +1886,14 @@ void UMaterial::UpdateCachedExpressionData()
 		return;
 	}
 
+	FMaterialCachedHLSLTree* LocalTree = nullptr;
+	if (IsUsingNewHLSLGenerator())
+	{
+		LocalTree = new FMaterialCachedHLSLTree();
+		LocalTree->GenerateTree(this, nullptr);
+	}
+	CachedHLSLTree.Reset(LocalTree);
+
 	if (!CachedExpressionData)
 	{
 		CachedExpressionData.Reset(new FMaterialCachedExpressionData());
@@ -3480,16 +3480,6 @@ bool UMaterial::IsUsingControlFlow() const
 	return false;
 }
 
-bool UMaterial::IsUsingNewHLSLGenerator() const
-{
-	const int CVarValue = CVarMaterialEnableNewHLSLGenerator.GetValueOnAnyThread();
-	if (bEnableNewHLSLGenerator)
-	{
-		return CVarValue != 0;
-	}
-	return CVarValue == 2;
-}
-
 #if WITH_EDITOR
 void UMaterial::BeginCacheForCookedPlatformData( const ITargetPlatform *TargetPlatform )
 {
diff --git a/Engine/Source/Runtime/Engine/Private/Materials/MaterialCachedHLSLTree.cpp b/Engine/Source/Runtime/Engine/Private/Materials/MaterialCachedHLSLTree.cpp
new file mode 100644
--- /dev/null
+++ b/Engine/Source/Runtime/Engine/Private/Materials/MaterialCachedHLSLTree.cpp
@@ -0,0 +1,184 @@
+// Copyright Epic Games, Inc. All Rights Reserved.
+
+#include "MaterialCachedHLSLTree.h"
+
+#if WITH_EDITOR
+
+#include "Materials/Material.h"
+#include "Materials/MaterialExpressionCustomOutput.h"
+#include "MaterialHLSLGenerator.h"
+
+DECLARE_STATS_GROUP(TEXT("Material Memory"), STATGROUP_MaterialMemory, STATCAT_Advanced);
+DECLARE_MEMORY_STAT(TEXT("Material HLSLTree Memory"), STAT_MaterialMemory_HLSLTree, STATGROUP_MaterialMemory);
+
+const FMaterialCachedHLSLTree FMaterialCachedHLSLTree::EmptyTree;
+
+FMaterialCachedHLSLTree::FMaterialCachedHLSLTree()
+	: TypeRegistry(Allocator)
+{
+}
+
+FMaterialCachedHLSLTree::~FMaterialCachedHLSLTree()
+{
+	const SIZE_T AllocatedSize = GetAllocatedSize();
+	DEC_MEMORY_STAT_BY(STAT_MaterialMemory_HLSLTree, AllocatedSize);
+}
+
+SIZE_T FMaterialCachedHLSLTree::GetAllocatedSize() const
+{
+	return Allocator.GetByteCount();
+}
+
+static UE::Shader::EValueType GetShaderType(EMaterialValueType MaterialType)
+{
+	switch (MaterialType)
+	{
+	case MCT_Float1: return UE::Shader::EValueType::Float1;
+	case MCT_Float2: return UE::Shader::EValueType::Float2;
+	case MCT_Float3: return UE::Shader::EValueType::Float3;
+	case MCT_Float4: return UE::Shader::EValueType::Float4;
+	case MCT_Float: return UE::Shader::EValueType::Float1;
+	case MCT_StaticBool: return UE::Shader::EValueType::Bool1;
+	case MCT_MaterialAttributes: return UE::Shader::EValueType::Struct;
+	case MCT_ShadingModel: return UE::Shader::EValueType::Int1;
+	case MCT_LWCScalar: return UE::Shader::EValueType::Double1;
+	case MCT_LWCVector2: return UE::Shader::EValueType::Double2;
+	case MCT_LWCVector3: return UE::Shader::EValueType::Double3;
+	case MCT_LWCVector4: return UE::Shader::EValueType::Double4;
+	default:return UE::Shader::EValueType::Void;
+	}
+}
+
+bool FMaterialCachedHLSLTree::GenerateTree(UMaterial* Material, const FMaterialLayersFunctions* LayerOverrides)
+{
+	const EMaterialShadingModel DefaultShadingModel = Material->GetShadingModels().GetFirstShadingModel();
+
+	Material->GetAllCustomOutputExpressions(MaterialCustomOutputs);
+
+	TArray<UE::Shader::FStructFieldInitializer, TInlineAllocator<MP_MAX + 16>> MaterialAttributeFields;
+
+	const TArray<FGuid>& OrderedVisibleAttributes = FMaterialAttributeDefinitionMap::GetOrderedVisibleAttributeList();
+	for (const FGuid& AttributeID : OrderedVisibleAttributes)
+	{
+		const FString& PropertyName = FMaterialAttributeDefinitionMap::GetAttributeName(AttributeID);
+		const EMaterialValueType PropertyType = FMaterialAttributeDefinitionMap::GetValueType(AttributeID);
+		const UE::Shader::EValueType ValueType = GetShaderType(PropertyType);
+
+		if (ValueType != UE::Shader::EValueType::Void &&
+			ValueType != UE::Shader::EValueType::Struct)
+		{
+			MaterialAttributeFields.Emplace(PropertyName, ValueType);
+
+			if (PropertyType == MCT_ShadingModel)
+			{
+				check(ValueType == UE::Shader::EValueType::Int1);
+				MaterialAttributesDefaultValue.Component.Add((int32)DefaultShadingModel);
+			}
+			else
+			{
+				const UE::Shader::FValue DefaultValue = UE::Shader::Cast(FMaterialAttributeDefinitionMap::GetDefaultValue(AttributeID), ValueType);
+				MaterialAttributesDefaultValue.Component.Append(DefaultValue.Component);
+			}
+		}
+	}
+
+	TArray<TStringBuilder<256>, TInlineAllocator<4>> CustomOutputNames;
+	CustomOutputNames.Reserve(MaterialCustomOutputs.Num());
+	for (UMaterialExpressionCustomOutput* CustomOutput : MaterialCustomOutputs)
+	{
+		const int32 NumOutputs = CustomOutput->GetNumOutputs();
+		const FString OutputName = CustomOutput->GetFunctionName();
+
+		check(!CustomOutput->ShouldCompileBeforeAttributes()); // not supported yet, looks like this isn't currently being used
+
+		for (int32 OutputIndex = 0; OutputIndex < NumOutputs; ++OutputIndex)
+		{
+			const UE::Shader::EValueType ValueType = CustomOutput->GetCustomOutputType(OutputIndex);
+			FStringBuilderBase& FormattedName = CustomOutputNames.AddDefaulted_GetRef();
+			FormattedName.Appendf(TEXT("%s%d"), *OutputName, OutputIndex);
+			MaterialAttributeFields.Emplace(FormattedName.ToView(), ValueType);
+
+			const UE::Shader::FValue DefaultValue(ValueType);
+			MaterialAttributesDefaultValue.Component.Append(DefaultValue.Component);
+		}
+	}
+
+	MaterialAttributeFields.Emplace(TEXT("PrevWorldPositionOffset"), UE::Shader::EValueType::Float3);
+	MaterialAttributesDefaultValue.Component.Append({ 0.0f, 0.0f, 0.0f });
+
+	UE::Shader::FStructTypeInitializer MaterialAttributesInitializer;
+	MaterialAttributesInitializer.Name = TEXT("FMaterialAttributes");
+	MaterialAttributesInitializer.Fields = MaterialAttributeFields;
+	MaterialAttributesType = TypeRegistry.NewType(MaterialAttributesInitializer);
+
+	check(MaterialAttributesDefaultValue.Component.Num() == MaterialAttributesType->ComponentTypes.Num());
+	MaterialAttributesDefaultValue.Type = MaterialAttributesType;
+
+	HLSLTree = UE::HLSLTree::FTree::Create(Allocator);
+
+	FMaterialHLSLGenerator Generator(Material, LayerOverrides, *this);
+	const bool bResult = Generator.Generate();
+	if (bResult)
+	{
+		const SIZE_T AllocatedSize = GetAllocatedSize();
+		INC_MEMORY_STAT_BY(STAT_MaterialMemory_HLSLTree, AllocatedSize);
+	}
+	return bResult;
+}
+
+void FMaterialCachedHLSLTree::SetRequestedFields(EShaderFrequency ShaderFrequency, UE::HLSLTree::FRequestedType& OutRequestedType) const
+{
+	for (UMaterialExpressionCustomOutput* CustomOutput : MaterialCustomOutputs)
+	{
+		if (CustomOutput->GetShaderFrequency() != ShaderFrequency)
+		{
+			continue;
+		}
+
+		const int32 NumOutputs = CustomOutput->GetNumOutputs();
+		const FString OutputName = CustomOutput->GetFunctionName();
+
+		for (int32 OutputIndex = 0; OutputIndex < NumOutputs; ++OutputIndex)
+		{
+			TStringBuilder<256> FieldName;
+			FieldName.Appendf(TEXT("%s%d"), *OutputName, OutputIndex);
+			const UE::Shader::FStructField* CustomOutputField = GetMaterialAttributesType()->FindFieldByName(FieldName.ToString());
+			check(CustomOutputField);
+			OutRequestedType.SetFieldRequested(CustomOutputField);
+		}
+	}
+}
+
+void FMaterialCachedHLSLTree::EmitSharedCode(FStringBuilderBase& OutCode) const
+{
+	using namespace UE::Shader;
+	for (UMaterialExpressionCustomOutput* CustomOutput : MaterialCustomOutputs)
+	{
+		const int32 NumOutputs = CustomOutput->GetNumOutputs();
+		const FString OutputName = CustomOutput->GetFunctionName();
+		const EShaderFrequency ShaderFrequency = CustomOutput->GetShaderFrequency();
+
+		if (CustomOutput->NeedsCustomOutputDefines())
+		{
+			OutCode.Appendf(TEXT("#define NUM_MATERIAL_OUTPUTS_%s %d\n"), *OutputName.ToUpper(), NumOutputs);
+		}
+
+		for (int32 OutputIndex = 0; OutputIndex < NumOutputs; ++OutputIndex)
+		{
+			const EValueType ValueType = CustomOutput->GetCustomOutputType(OutputIndex);
+			const FValueTypeDescription ValueTypeDesc = GetValueTypeDescription(ValueType);
+
+			OutCode.Appendf(TEXT("#define HAVE_%s%d 1\n"), *OutputName, OutputIndex);
+
+			OutCode.Appendf(TEXT("%s %s%d(FMaterial%sParameters Parameters) { return Parameters.MaterialAttributes.%s%d; }\n"),
+				ValueTypeDesc.Name,
+				*OutputName, OutputIndex,
+				ShaderFrequency == SF_Pixel ? TEXT("Pixel") : TEXT("Vertex"),
+				*OutputName, OutputIndex);
+		}
+		OutCode.Append(TEXT("\n"));
+	}
+}
+
+
+#endif // WITH_EDITOR
diff --git a/Engine/Source/Runtime/Engine/Private/Materials/MaterialExpressionHLSL.cpp b/Engine/Source/Runtime/Engine/Private/Materials/MaterialExpressionHLSL.cpp
--- a/Engine/Source/Runtime/Engine/Private/Materials/MaterialExpressionHLSL.cpp
+++ b/Engine/Source/Runtime/Engine/Private/Materials/MaterialExpressionHLSL.cpp
@@ -143,12 +143,12 @@
 
 bool UMaterialExpression::GenerateHLSLExpression(FMaterialHLSLGenerator& Generator, UE::HLSLTree::FScope& Scope, int32 OutputIndex, UE::HLSLTree::FExpression*& OutExpression)
 {
-	return Generator.GetErrors().AddError(TEXT("Node does not support expressions"));
+	return Generator.Error(TEXT("Node does not support expressions"));
 }
 
 bool UMaterialExpression::GenerateHLSLStatements(FMaterialHLSLGenerator& Generator, UE::HLSLTree::FScope& Scope)
 {
-	return Generator.GetErrors().AddError(TEXT("Node does not support statements"));
+	return Generator.Error(TEXT("Node does not support statements"));
 }
 
 bool UMaterialExpressionReroute::GenerateHLSLExpression(FMaterialHLSLGenerator& Generator, UE::HLSLTree::FScope& Scope, int32 OutputIndex, UE::HLSLTree::FExpression*& OutExpression)
@@ -224,57 +224,68 @@ bool UMaterialExpressionStaticSwitch::GenerateHLSLExpression(FMaterialHLSLGenera
 
 bool UMaterialExpressionFeatureLevelSwitch::GenerateHLSLExpression(FMaterialHLSLGenerator& Generator, UE::HLSLTree::FScope& Scope, int32 OutputIndex, UE::HLSLTree::FExpression*& OutExpression)
 {
-	const ERHIFeatureLevel::Type FeatureLevelToCompile = Generator.GetCompileTarget().FeatureLevel;
-	check(FeatureLevelToCompile < UE_ARRAY_COUNT(Inputs));
-	FExpressionInput& FeatureInput = Inputs[FeatureLevelToCompile];
-
+	using namespace UE::HLSLTree;
 	if (!Default.GetTracedInput().Expression)
 	{
-		return Generator.GetErrors().AddError(TEXT("Missing default input"));
+		return Generator.Error(TEXT("Missing default input"));
 	}
 
-	if (FeatureInput.GetTracedInput().Expression)
-	{
-		OutExpression = FeatureInput.AcquireHLSLExpression(Generator, Scope);
-	}
-	else
+	FExpression* ExpressionDefault = nullptr;
+	FExpression* ExpressionInputs[ERHIFeatureLevel::Num] = { nullptr };
+	for (int32 Index = 0; Index < ERHIFeatureLevel::Num; ++Index)
 	{
-		OutExpression = Default.AcquireHLSLExpression(Generator, Scope);
+		FExpression* Expression = nullptr;
+		FExpressionInput& FeatureInput = Inputs[Index];
+		if (FeatureInput.GetTracedInput().Expression)
+		{
+			Expression = FeatureInput.AcquireHLSLExpression(Generator, Scope);
+		}
+		else
+		{
+			if (!ExpressionDefault)
+			{
+				ExpressionDefault = Default.AcquireHLSLExpression(Generator, Scope);
+			}
+			Expression = ExpressionDefault;
+		}
+		ExpressionInputs[Index] = Expression;
 	}
-	return OutExpression != nullptr;
+
+	OutExpression = Generator.GetTree().NewExpression<FExpressionFeatureLevelSwitch>(ExpressionInputs);
+	return true;
 }
 
 bool UMaterialExpressionShadingPathSwitch::GenerateHLSLExpression(FMaterialHLSLGenerator& Generator, UE::HLSLTree::FScope& Scope, int32 OutputIndex, UE::HLSLTree::FExpression*& OutExpression)
 {
-	const EShaderPlatform ShaderPlatform = Generator.GetCompileTarget().ShaderPlatform;
-	ERHIShadingPath::Type ShadingPathToCompile = ERHIShadingPath::Deferred;
-	if (IsForwardShadingEnabled(ShaderPlatform))
-	{
-		ShadingPathToCompile = ERHIShadingPath::Forward;
-	}
-	else if (Generator.GetCompileTarget().FeatureLevel < ERHIFeatureLevel::SM5)
+	using namespace UE::HLSLTree;
+	if (!Default.GetTracedInput().Expression)
 	{
-		ShadingPathToCompile = ERHIShadingPath::Mobile;
+		return Generator.Error(TEXT("Missing default input"));
 	}
 
-	check(ShadingPathToCompile < UE_ARRAY_COUNT(Inputs));
-	FExpressionInput ShadingPathInput = Inputs[ShadingPathToCompile].GetTracedInput();
-	FExpressionInput DefaultTraced = Default.GetTracedInput();
-
-	if (!DefaultTraced.Expression)
+	FExpression* ExpressionDefault = nullptr;
+	FExpression* ExpressionInputs[ERHIShadingPath::Num] = { nullptr };
+	for (int32 Index = 0; Index < ERHIShadingPath::Num; ++Index)
 	{
-		return Generator.GetErrors().AddError(TEXT("Missing default input"));
+		FExpression* Expression = nullptr;
+		FExpressionInput& FeatureInput = Inputs[Index];
+		if (FeatureInput.GetTracedInput().Expression)
+		{
+			Expression = FeatureInput.AcquireHLSLExpression(Generator, Scope);
+		}
+		else
+		{
+			if (!ExpressionDefault)
+			{
+				ExpressionDefault = Default.AcquireHLSLExpression(Generator, Scope);
+			}
+			Expression = ExpressionDefault;
+		}
+		ExpressionInputs[Index] = Expression;
 	}
 
-	if (ShadingPathInput.Expression)
-	{
-		OutExpression = ShadingPathInput.AcquireHLSLExpression(Generator, Scope);
-	}
-	else
-	{
-		OutExpression = Default.AcquireHLSLExpression(Generator, Scope);
-	}
-	return OutExpression != nullptr;
+	OutExpression = Generator.GetTree().NewExpression<FExpressionShadingPathSwitch>(ExpressionInputs);
+	return true;
 }
 
 bool UMaterialExpressionGetLocal::GenerateHLSLExpression(FMaterialHLSLGenerator& Generator, UE::HLSLTree::FScope& Scope, int32 OutputIndex, UE::HLSLTree::FExpression*& OutExpression)
@@ -282,7 +293,7 @@ bool UMaterialExpressionGetLocal::GenerateHLSLExpression(FMaterialHLSLGenerator&
 	OutExpression = Generator.GetTree().AcquireLocal(Scope, LocalName);
 	if (!OutExpression)
 	{
-		return Generator.GetErrors().AddError(TEXT("Local accessed before assigned"));
+		return Generator.Errorf(TEXT("Local '%s' accessed before assigned"), *LocalName.ToString());
 	}
 	return true;
 }
@@ -588,7 +599,7 @@ bool UMaterialExpressionTextureSample::GenerateHLSLExpressionBase(FMaterialHLSLG
 	using namespace UE::HLSLTree;
 	if (!TextureExpression)
 	{
-		return Generator.GetErrors().AddError(TEXT("Missing input texture"));
+		return Generator.Error(TEXT("Missing input texture"));
 	}
 
 	FExpression* TexCoordExpression = Coordinates.GetTracedInput().Expression ? Coordinates.TryAcquireHLSLExpression(Generator, Scope) : Generator.NewTexCoord(ConstCoordinate);
@@ -677,15 +688,15 @@ bool UMaterialExpressionFontSample::GenerateHLSLExpression(FMaterialHLSLGenerato
 
 	if (!Font)
 	{
-		return Generator.GetErrors().AddError(TEXT("Missing input Font"));
+		return Generator.Error(TEXT("Missing input Font"));
 	}
 	else if (Font->FontCacheType == EFontCacheType::Runtime)
 	{
-		return Generator.GetErrors().AddErrorf(TEXT("Font '%s' is runtime cached, but only offline cached fonts can be sampled"), *Font->GetName());
+		return Generator.Errorf(TEXT("Font '%s' is runtime cached, but only offline cached fonts can be sampled"), *Font->GetName());
 	}
 	else if (!Font->Textures.IsValidIndex(FontTexturePage))
 	{
-		return Generator.GetErrors().AddErrorf(TEXT("Invalid font page %d. Max allowed is %d"), FontTexturePage, Font->Textures.Num());
+		return Generator.Errorf(TEXT("Invalid font page %d. Max allowed is %d"), FontTexturePage, Font->Textures.Num());
 	}
 
 	UTexture* Texture = Font->Textures[FontTexturePage];
@@ -706,11 +717,11 @@ bool UMaterialExpressionFontSample::GenerateHLSLExpression(FMaterialHLSLGenerato
 		ExpectedSamplerType = Texture->SRGB ? SAMPLERTYPE_Color : SAMPLERTYPE_LinearColor;
 	}
 
-	FString SamplerTypeError;
+	/*FString SamplerTypeError;
 	if (!UMaterialExpressionTextureBase::VerifySamplerType(Generator.GetCompileTarget().FeatureLevel, Generator.GetCompileTarget().TargetPlatform, Texture, ExpectedSamplerType, SamplerTypeError))
 	{
-		return Generator.GetErrors().AddErrorf(TEXT("%s"), *SamplerTypeError);
-	}
+		return Generator.Errorf(TEXT("%s"), *SamplerTypeError);
+	}*/
 
 	const FTextureValue* TextureValue = Generator.AcquireTextureValue(FTextureValue(Texture, ExpectedSamplerType));
 	FExpression* TextureExpression = GenerateHLSLTextureExpression(Generator, TextureValue);
@@ -745,7 +756,7 @@ bool UMaterialExpressionSceneTexture::GenerateHLSLExpression(FMaterialHLSLGenera
 		return true;
 	}
 
-	return Generator.GetErrors().AddError(TEXT("Invalid input parameter"));
+	return Generator.Error(TEXT("Invalid input parameter"));
 }
 
 bool UMaterialExpressionNoise::GenerateHLSLExpression(FMaterialHLSLGenerator& Generator, UE::HLSLTree::FScope& Scope, int32 OutputIndex, UE::HLSLTree::FExpression*& OutExpression)
@@ -1274,7 +1285,7 @@ bool UMaterialExpressionGetMaterialAttributes::GenerateHLSLExpression(FMaterialH
 	const int32 AttributeIndex = OutputIndex - 1;
 	if (!AttributeGetTypes.IsValidIndex(AttributeIndex))
 	{
-		return Generator.GetErrors().AddError(TEXT("Invalid attribute"));
+		return Generator.Error(TEXT("Invalid attribute"));
 	}
 
 	const FGuid& AttributeID = AttributeGetTypes[AttributeIndex];
@@ -1380,7 +1391,7 @@ bool UMaterialExpressionBreakMaterialAttributes::GenerateHLSLExpression(FMateria
 	const EMaterialProperty* Property = PropertyToIOIndexMap.FindKey(OutputIndex);
 	if (!Property)
 	{
-		return Generator.GetErrors().AddError(TEXT("Invalid output"));
+		return Generator.Error(TEXT("Invalid output"));
 	}
 
 	FExpression* AttributesExpression = MaterialAttributes.AcquireHLSLExpression(Generator, Scope);
@@ -1429,11 +1440,11 @@ bool UMaterialExpressionMaterialAttributeLayers::GenerateHLSLExpression(FMateria
 {
 	using namespace UE::HLSLTree;
 
-	const FStaticParameterSet& StaticParameters = Generator.GetStaticParameters();
-	const FMaterialLayersFunctions& MaterialLayers = StaticParameters.bHasMaterialLayers ? StaticParameters.MaterialLayers : DefaultLayers;
+	const FMaterialLayersFunctions* LayerOverrides = Generator.GetLayerOverrides();
+	const FMaterialLayersFunctions& MaterialLayers = LayerOverrides ? *LayerOverrides : DefaultLayers;
 	if (MaterialLayers.Layers.Num() == 0)
 	{
-		return Generator.GetErrors().AddError(TEXT("No layers"));
+		return Generator.Error(TEXT("No layers"));
 	}
 
 	TArray<FFunctionExpressionInput> FunctionInputs;
@@ -1453,7 +1464,7 @@ bool UMaterialExpressionMaterialAttributeLayers::GenerateHLSLExpression(FMateria
 			const EMaterialFunctionUsage Usage = LayerFunction->GetMaterialFunctionUsage();
 			if (Usage != EMaterialFunctionUsage::MaterialLayer)
 			{
-				return Generator.GetErrors().AddErrorf(TEXT("Layer function %s is not a UMaterialFunctionMaterialLayer"),
+				return Generator.Errorf(TEXT("Layer function %s is not a UMaterialFunctionMaterialLayer"),
 					*LayerFunction->GetName());
 			}
 
@@ -1462,7 +1473,7 @@ bool UMaterialExpressionMaterialAttributeLayers::GenerateHLSLExpression(FMateria
 			LayerFunction->GetInputsAndOutputs(FunctionInputs, FunctionOutputs);
 			if (FunctionInputs.Num() > 1 || FunctionOutputs.Num() != 1)
 			{
-				return Generator.GetErrors().AddErrorf(TEXT("Layer function %s expected to have 0 or 1 inputs and 1 output, found %d inputs and %d outputs"),
+				return Generator.Errorf(TEXT("Layer function %s expected to have 0 or 1 inputs and 1 output, found %d inputs and %d outputs"),
 					*LayerFunction->GetName(), FunctionInputs.Num(), FunctionOutputs.Num());
 			}
 
@@ -1479,7 +1490,7 @@ bool UMaterialExpressionMaterialAttributeLayers::GenerateHLSLExpression(FMateria
 	FExpression* BottomLayerExpression = LayerExpressions[0];
 	if (!BottomLayerExpression)
 	{
-		return Generator.GetErrors().AddError(TEXT("No layers"));
+		return Generator.Error(TEXT("No layers"));
 	}
 
 	TArray<FExpression*, TInlineAllocator<2>> BlendInputExpressions;
@@ -1488,7 +1499,7 @@ bool UMaterialExpressionMaterialAttributeLayers::GenerateHLSLExpression(FMateria
 		const int32 LayerIndex = BlendIndex + 1;
 		if (!MaterialLayers.Layers.IsValidIndex(LayerIndex))
 		{
-			return Generator.GetErrors().AddErrorf(TEXT("Invalid number of layers (%d) and blends (%d)"), MaterialLayers.Layers.Num(), MaterialLayers.Blends.Num());
+			return Generator.Errorf(TEXT("Invalid number of layers (%d) and blends (%d)"), MaterialLayers.Layers.Num(), MaterialLayers.Blends.Num());
 		}
 
 		if (MaterialLayers.Layers[LayerIndex] && MaterialLayers.LayerStates[LayerIndex])
@@ -1496,7 +1507,7 @@ bool UMaterialExpressionMaterialAttributeLayers::GenerateHLSLExpression(FMateria
 			FExpression* LayerExpression = LayerExpressions[LayerIndex];
 			if (!LayerExpression)
 			{
-				return Generator.GetErrors().AddErrorf(TEXT("Missing layer %d"), LayerIndex);
+				return Generator.Errorf(TEXT("Missing layer %d"), LayerIndex);
 			}
 
 			UMaterialFunctionInterface* BlendFunction = MaterialLayers.Blends[BlendIndex];
@@ -1505,7 +1516,7 @@ bool UMaterialExpressionMaterialAttributeLayers::GenerateHLSLExpression(FMateria
 				const EMaterialFunctionUsage Usage = BlendFunction->GetMaterialFunctionUsage();
 				if (Usage != EMaterialFunctionUsage::MaterialLayerBlend)
 				{
-					return Generator.GetErrors().AddErrorf(TEXT("Blend function %s is not a UMaterialFunctionMaterialBlend"),
+					return Generator.Errorf(TEXT("Blend function %s is not a UMaterialFunctionMaterialBlend"),
 						*BlendFunction->GetName());
 				}
 
@@ -1514,7 +1525,7 @@ bool UMaterialExpressionMaterialAttributeLayers::GenerateHLSLExpression(FMateria
 				BlendFunction->GetInputsAndOutputs(FunctionInputs, FunctionOutputs);
 				if (FunctionInputs.Num() != 2 || FunctionOutputs.Num() != 1)
 				{
-					return Generator.GetErrors().AddErrorf(TEXT("Blend function %s expected to have 2 inputs and 1 output, found %d inputs and %d outputs"),
+					return Generator.Errorf(TEXT("Blend function %s expected to have 2 inputs and 1 output, found %d inputs and %d outputs"),
 						*BlendFunction->GetName(), FunctionInputs.Num(), FunctionOutputs.Num());
 				}
 
@@ -2045,7 +2056,7 @@ bool UMaterialExpressionVolumetricAdvancedMaterialInput::GenerateHLSLExpression(
 {
 	if (OutputIndex != 0)
 	{
-		return Generator.GetErrors().AddError(TEXT("Invalid output"));
+		return Generator.Error(TEXT("Invalid output"));
 	}
 	OutExpression = Generator.GetTree().NewExpression<UE::HLSLTree::FExpressionInlineCustomHLSL>(UE::Shader::EValueType::Float3, TEXT("MaterialExpressionVolumeSampleConservativeDensity(Parameters)"));
 	return true;
@@ -2062,7 +2073,7 @@ bool UMaterialExpressionVolumetricAdvancedMaterialOutput::GenerateHLSLExpression
 	case 4: OutExpression = MultiScatteringOcclusion.AcquireHLSLExpressionOrConstant(Generator, Scope, ConstMultiScatteringOcclusion); break;
 	case 5: OutExpression = MultiScatteringEccentricity.AcquireHLSLExpressionOrConstant(Generator, Scope, ConstMultiScatteringEccentricity); break;
 	case 6: OutExpression = ConservativeDensity.AcquireHLSLExpressionOrConstant(Generator, Scope, FVector3f(1.0f, 1.0f, 1.0f)); break;
-	default: return Generator.GetErrors().AddError(TEXT("Invlid output"));
+	default: return Generator.Error(TEXT("Invlid output"));
 	}
 	return OutExpression != nullptr;
 }
@@ -2099,7 +2110,7 @@ bool UMaterialExpressionHairAttributes::GenerateHLSLExpression(FMaterialHLSLGene
 	case 11: OutExpression = Generator.GetTree().NewExpression<FExpressionInlineCustomHLSL>(EValueType::Float4, TEXT("MaterialExpressionGetHairAuxilaryData(Parameters)")); break;
 	case 12: OutExpression = Generator.GetTree().NewExpression<FExpressionInlineCustomHLSL>(EValueType::Float2, TEXT("MaterialExpressionGetAtlasUVs(Parameters)")); break;
 	case 13: OutExpression = Generator.GetTree().NewExpression<FExpressionInlineCustomHLSL>(EValueType::Float1, TEXT("MaterialExpressionGetHairGroupIndex(Parameters)")); break;
-	default: return Generator.GetErrors().AddError(TEXT("Invalid output"));
+	default: return Generator.Error(TEXT("Invalid output"));
 	}
 	return OutExpression != nullptr;
 }
@@ -2115,7 +2126,7 @@ bool UMaterialExpressionCloudSampleAttribute::GenerateHLSLExpression(FMaterialHL
 	case 1: OutExpression = Generator.GetTree().NewExpression<FExpressionInlineCustomHLSL>(EValueType::Float1, TEXT("MaterialExpressionCloudSampleAltitudeInLayer(Parameters)")); break;
 	case 2: OutExpression = Generator.GetTree().NewExpression<FExpressionInlineCustomHLSL>(EValueType::Float1, TEXT("MaterialExpressionCloudSampleNormAltitudeInLayer(Parameters)")); break;
 	case 3: OutExpression = Generator.GetTree().NewExpression<FExpressionInlineCustomHLSL>(EValueType::Float1, TEXT("MaterialExpressionVolumeSampleShadowSampleDistance(Parameters)")); break;
-	default: return Generator.GetErrors().AddError(TEXT("Invalid output"));
+	default: return Generator.Error(TEXT("Invalid output"));
 	}
 	return OutExpression != nullptr;
 }
@@ -2168,7 +2179,7 @@ bool UMaterialExpressionCustom::GenerateHLSLExpression(FMaterialHLSLGenerator& G
 
 	if (OutputIndex < 0 || OutputIndex > AdditionalOutputs.Num())
 	{
-		return Generator.GetErrors().AddErrorf(TEXT("Invalid output index %d"), OutputIndex);
+		return Generator.Errorf(TEXT("Invalid output index %d"), OutputIndex);
 	}
 
 	FMemStackBase& Allocator = Generator.GetTree().GetAllocator();
@@ -2303,12 +2314,12 @@ bool UMaterialExpressionWhileLoop::GenerateHLSLStatements(FMaterialHLSLGenerator
 	using namespace UE::HLSLTree;
 	if (!Condition.IsConnected())
 	{
-		return Generator.GetErrors().AddError(TEXT("Missing condition connection"));
+		return Generator.Error(TEXT("Missing condition connection"));
 	}
 
 	if (!LoopBody.GetExpression())
 	{
-		return Generator.GetErrors().AddError(TEXT("Missing LoopBody connection"));
+		return Generator.Error(TEXT("Missing LoopBody connection"));
 	}
 
 	FStatementLoop* LoopStatement = Generator.GetTree().NewStatement<FStatementLoop>(Scope);
@@ -2348,7 +2359,7 @@ bool UMaterialExpressionForLoop::GenerateHLSLExpression(FMaterialHLSLGenerator&
 	FExpressionDataForLoop* ExpressionData = Generator.FindExpressionData<FExpressionDataForLoop>(this);
 	if (!ExpressionData || !Scope.HasParentScope(*ExpressionData->LoopScope))
 	{
-		return Generator.GetErrors().AddError(TEXT("For loop index accessed outside loop scope"));
+		return Generator.Error(TEXT("For loop index accessed outside loop scope"));
 	}
 
 	OutExpression = Generator.GetTree().AcquireLocal(Scope, ExpressionData->LocalName);
@@ -2360,7 +2371,7 @@ bool UMaterialExpressionForLoop::GenerateHLSLStatements(FMaterialHLSLGenerator&
 	using namespace UE::HLSLTree;
 	if (!LoopBody.GetExpression())
 	{
-		return Generator.GetErrors().AddError(TEXT("Missing LoopBody connection"));
+		return Generator.Error(TEXT("Missing LoopBody connection"));
 	}
 
 	FExpression* StartExpression = StartIndex.AcquireHLSLExpression(Generator, Scope);
diff --git a/Engine/Source/Runtime/Engine/Private/Materials/MaterialHLSLEmitter.cpp b/Engine/Source/Runtime/Engine/Private/Materials/MaterialHLSLEmitter.cpp
--- a/Engine/Source/Runtime/Engine/Private/Materials/MaterialHLSLEmitter.cpp
+++ b/Engine/Source/Runtime/Engine/Private/Materials/MaterialHLSLEmitter.cpp
@@ -12,7 +12,7 @@
 #include "Materials/MaterialExpressionVolumetricAdvancedMaterialOutput.h"
 #include "Materials/Material.h"
 #include "Materials/MaterialHLSLTree.h"
-#include "MaterialHLSLGenerator.h"
+#include "MaterialCachedHLSLTree.h"
 #include "ShaderCore.h"
 #include "HLSLTree/HLSLTree.h"
 #include "HLSLTree/HLSLTreeCommon.h"
@@ -724,7 +724,7 @@ public:
 	FStringBuilderMemstack(FMemStackBase& Allocator, int32 InSize) : TStringBuilderBase<TCHAR>((TCHAR*)Allocator.Alloc(sizeof(TCHAR)* InSize, alignof(TCHAR)), InSize) {}
 };
 
-static bool IsAttributeUsed(FMaterialHLSLGenerator& Generator,
+static bool IsAttributeUsed(const FMaterialCachedHLSLTree& CachedTree,
 	UE::HLSLTree::FEmitContext& Context,
 	UE::HLSLTree::FEmitScope& Scope,
 	const UE::HLSLTree::FPreparedType& ResultType,
@@ -734,7 +734,7 @@ static bool IsAttributeUsed(FMaterialHLSLGenerator& Generator,
 	using namespace UE::HLSLTree;
 	using namespace UE::Shader;
 
-	const FStructField* PropertyField = Generator.GetMaterialAttributesType()->FindFieldByName(*FMaterialAttributeDefinitionMap::GetAttributeName(Property));
+	const FStructField* PropertyField = CachedTree.GetMaterialAttributesType()->FindFieldByName(*FMaterialAttributeDefinitionMap::GetAttributeName(Property));
 	if (!PropertyField)
 	{
 		return false;
@@ -751,7 +751,7 @@ static bool IsAttributeUsed(FMaterialHLSLGenerator& Generator,
 	}
 	else if (Evaluation == EExpressionEvaluation::Constant || Evaluation == EExpressionEvaluation::ConstantZero)
 	{
-		const FValue& DefaultValue = Generator.GetMaterialAttributesDefaultValue();
+		const FValue& DefaultValue = CachedTree.GetMaterialAttributesDefaultValue();
 		const FValue ConstantValue = Expression->GetValueConstant(Context, Scope, RequestedType);
 		if (DefaultValue.GetType() != ConstantValue.GetType())
 		{
@@ -774,6 +774,56 @@ static bool IsAttributeUsed(FMaterialHLSLGenerator& Generator,
 	return true;
 }
 
+class FMaterialHLSLErrorHandler : public UE::HLSLTree::FErrorHandlerInterface
+{
+public:
+	explicit FMaterialHLSLErrorHandler(FMaterial& InOutMaterial)
+		: Material(&InOutMaterial)
+	{
+		Material->CompileErrors.Reset();
+		Material->ErrorExpressions.Reset();
+	}
+
+	virtual void AddErrorInternal(UObject* InOwner, FStringView InError) override
+	{
+		UMaterialExpression* MaterialExpressionOwner = Cast<UMaterialExpression>(InOwner);
+		UMaterialExpression* ExpressionToError = nullptr;
+		TStringBuilder<1024> FormattedError;
+
+		if (MaterialExpressionOwner)
+		{
+			if (MaterialExpressionOwner->GetClass() != UMaterialExpressionMaterialFunctionCall::StaticClass()
+				&& MaterialExpressionOwner->GetClass() != UMaterialExpressionFunctionInput::StaticClass()
+				&& MaterialExpressionOwner->GetClass() != UMaterialExpressionFunctionOutput::StaticClass())
+			{
+				// Add the expression currently being compiled to ErrorExpressions so we can draw it differently
+				ExpressionToError = MaterialExpressionOwner;
+
+				const int32 ChopCount = FCString::Strlen(TEXT("MaterialExpression"));
+				const FString ErrorClassName = MaterialExpressionOwner->GetClass()->GetName();
+
+				// Add the node type to the error message
+				FormattedError.Appendf(TEXT("(Node %s) "), *ErrorClassName.Right(ErrorClassName.Len() - ChopCount));
+			}
+		}
+
+		FormattedError.Append(InError);
+		const FString Error(FormattedError.ToView());
+
+		// Standard error handling, immediately append one-off errors and signal failure
+		Material->CompileErrors.AddUnique(Error);
+
+		if (ExpressionToError)
+		{
+			Material->ErrorExpressions.Add(ExpressionToError);
+			ExpressionToError->LastErrorText = Error;
+		}
+	}
+
+private:
+	FMaterial* Material;
+};
+
 bool MaterialEmitHLSL(const FMaterialCompileTargetParameters& InCompilerTarget,
 	const FStaticParameterSet& InStaticParameters,
 	FMaterial& InOutMaterial,
@@ -802,27 +852,23 @@ bool MaterialEmitHLSL(const FMaterialCompileTargetParameters& InCompilerTarget,
 	SharedPixelProperties[MP_ShadingModel] = true;
 	SharedPixelProperties[MP_FrontMaterial] = true;
 
-	FMemStackBase Allocator;
-	FStructTypeRegistry TypeRegistry(Allocator);
-	FTree* HLSLTree = FTree::Create(Allocator);
-
-	FMaterialHLSLGenerator Generator(InCompilerTarget, InStaticParameters, InOutMaterial, TypeRegistry, *HLSLTree);
-	if (!Generator.Generate())
-	{
-		return false;
-	}
-
 	bool bUsesWorldPositionOffset = false;
 	bool bUsesPixelDepthOffset = false;
 
-	FEmitContext EmitContext(Allocator, Generator.GetErrors(), TypeRegistry);
+	const FTargetParameters TargetParameters(InCompilerTarget.ShaderPlatform, InCompilerTarget.FeatureLevel, InCompilerTarget.TargetPlatform);
+	UMaterialInterface* MaterialInterface = InOutMaterial.GetMaterialInterface();
+	const FMaterialCachedHLSLTree& CachedTree = MaterialInterface->GetCachedHLSLTree();
+
+	FMaterialHLSLErrorHandler ErrorHandler(InOutMaterial);
+	FMemStackBase Allocator;
+	FEmitContext EmitContext(Allocator, TargetParameters, ErrorHandler, CachedTree.GetTypeRegistry());
 	EmitContext.Material = &InOutMaterial;
 	EmitContext.MaterialCompilationOutput = &OutCompilationOutput;
 
 	Material::FEmitData& EmitMaterialData = EmitContext.AcquireData<Material::FEmitData>();
 	EmitMaterialData.StaticParameters = &InStaticParameters;
 
-	const FStructField* NormalField = Generator.GetMaterialAttributesType()->FindFieldByName(*FMaterialAttributeDefinitionMap::GetAttributeName(MP_Normal));
+	const FStructField* NormalField = CachedTree.GetMaterialAttributesType()->FindFieldByName(*FMaterialAttributeDefinitionMap::GetAttributeName(MP_Normal));
 	check(NormalField);
 
 	// Prepare pixel shader code
@@ -833,17 +879,17 @@ bool MaterialEmitHLSL(const FMaterialCompileTargetParameters& InCompilerTarget,
 	{
 		EmitContext.ShaderFrequency = SF_Pixel;
 		EmitContext.bUseAnalyticDerivatives = (DerivativeIndex == 1);
-		FEmitScope* EmitResultScope = EmitContext.PrepareScope(&Generator.GetResultStatement()->GetParentScope());
+		FEmitScope* EmitResultScope = EmitContext.PrepareScope(&CachedTree.GetResultStatement()->GetParentScope());
 
 		// Prepare all fields *except* normal
 		FRequestedType RequestedPixelAttributesType;
-		Generator.SetRequestedFields(SF_Pixel, RequestedPixelAttributesType);
+		CachedTree.SetRequestedFields(SF_Pixel, RequestedPixelAttributesType);
 		for (const FGuid& AttributeID : OrderedVisibleAttributes)
 		{
 			if (FMaterialAttributeDefinitionMap::GetShaderFrequency(AttributeID) == SF_Pixel)
 			{
 				const FString& FieldName = FMaterialAttributeDefinitionMap::GetAttributeName(AttributeID);
-				const FStructField* Field = Generator.GetMaterialAttributesType()->FindFieldByName(*FieldName);
+				const FStructField* Field = CachedTree.GetMaterialAttributesType()->FindFieldByName(*FieldName);
 				if (Field && Field != NormalField)
 				{
 					RequestedPixelAttributesType.SetFieldRequested(Field);
@@ -851,13 +897,13 @@ bool MaterialEmitHLSL(const FMaterialCompileTargetParameters& InCompilerTarget,
 			}
 		}
 
-		const FPreparedType PixelResultType0 = EmitContext.PrepareExpression(Generator.GetResultExpression(), *EmitResultScope, RequestedPixelAttributesType);
+		const FPreparedType PixelResultType0 = EmitContext.PrepareExpression(CachedTree.GetResultExpression(), *EmitResultScope, RequestedPixelAttributesType);
 		if (PixelResultType0.IsVoid())
 		{
 			return false;
 		}
 
-		if (IsAttributeUsed(Generator, EmitContext, *EmitResultScope, PixelResultType0, Generator.GetResultExpression(), MP_PixelDepthOffset))
+		if (IsAttributeUsed(CachedTree, EmitContext, *EmitResultScope, PixelResultType0, CachedTree.GetResultExpression(), MP_PixelDepthOffset))
 		{
 			bUsesPixelDepthOffset = true;
 		}
@@ -866,29 +912,27 @@ bool MaterialEmitHLSL(const FMaterialCompileTargetParameters& InCompilerTarget,
 		{
 			// No access to material normal, can execute everything in phase0
 			RequestedPixelAttributesType.SetFieldRequested(NormalField);
-			const FPreparedType PixelResultType1 = EmitContext.PrepareExpression(Generator.GetResultExpression(), *EmitResultScope, RequestedPixelAttributesType);
+			const FPreparedType PixelResultType1 = EmitContext.PrepareExpression(CachedTree.GetResultExpression(), *EmitResultScope, RequestedPixelAttributesType);
 			if (PixelResultType1.IsVoid())
 			{
 				return false;
 			}
 
 			PixelCodePhase0[DerivativeIndex] = new(Allocator) FStringBuilderMemstack(Allocator, 128 * 1024);
-			HLSLTree->EmitShader(EmitContext, *PixelCodePhase0[DerivativeIndex]);
-			HLSLTree->ResetNodes();
+			CachedTree.GetTree().EmitShader(EmitContext, *PixelCodePhase0[DerivativeIndex]);
 		}
 		else
 		{
 			// Execute everything *except* normal in phase1
 			PixelCodePhase1[DerivativeIndex] = new(Allocator) FStringBuilderMemstack(Allocator, 128 * 1024);
-			HLSLTree->EmitShader(EmitContext, *PixelCodePhase1[DerivativeIndex]);
-			HLSLTree->ResetNodes();
+			CachedTree.GetTree().EmitShader(EmitContext, *PixelCodePhase1[DerivativeIndex]);
 
-			EmitResultScope = EmitContext.PrepareScope(&Generator.GetResultStatement()->GetParentScope());
+			EmitResultScope = EmitContext.PrepareScope(&CachedTree.GetResultStatement()->GetParentScope());
 
 			// Prepare code for just the normal
 			FRequestedType RequestedMaterialNormal;
 			RequestedMaterialNormal.SetFieldRequested(NormalField);
-			const FPreparedType PixelResultType1 = EmitContext.PrepareExpression(Generator.GetResultExpression(), *EmitResultScope, RequestedMaterialNormal);
+			const FPreparedType PixelResultType1 = EmitContext.PrepareExpression(CachedTree.GetResultExpression(), *EmitResultScope, RequestedMaterialNormal);
 			if (PixelResultType1.IsVoid())
 			{
 				return false;
@@ -896,8 +940,7 @@ bool MaterialEmitHLSL(const FMaterialCompileTargetParameters& InCompilerTarget,
 
 			// Execute the normal in phase0
 			PixelCodePhase0[DerivativeIndex] = new(Allocator) FStringBuilderMemstack(Allocator, 128 * 1024);
-			HLSLTree->EmitShader(EmitContext, *PixelCodePhase0[DerivativeIndex]);
-			HLSLTree->ResetNodes();
+			CachedTree.GetTree().EmitShader(EmitContext, *PixelCodePhase0[DerivativeIndex]);
 		}
 	}
 
@@ -905,34 +948,34 @@ bool MaterialEmitHLSL(const FMaterialCompileTargetParameters& InCompilerTarget,
 	FStringBuilderMemstack VertexCode(Allocator, 128 * 1024);
 	{
 		FRequestedType RequestedVertexAttributesType;
-		Generator.SetRequestedFields(SF_Vertex, RequestedVertexAttributesType);
+		CachedTree.SetRequestedFields(SF_Vertex, RequestedVertexAttributesType);
 		for (const FGuid& AttributeID : OrderedVisibleAttributes)
 		{
 			if (FMaterialAttributeDefinitionMap::GetShaderFrequency(AttributeID) == SF_Vertex)
 			{
 				const FString& FieldName = FMaterialAttributeDefinitionMap::GetAttributeName(AttributeID);
-				const FStructField* Field = Generator.GetMaterialAttributesType()->FindFieldByName(*FieldName);
+				const FStructField* Field = CachedTree.GetMaterialAttributesType()->FindFieldByName(*FieldName);
 				if (Field)
 				{
 					RequestedVertexAttributesType.SetFieldRequested(Field);
 				}
 			}
 		}
-		RequestedVertexAttributesType.SetFieldRequested(Generator.GetMaterialAttributesType()->FindFieldByName(TEXT("PrevWorldPositionOffset")));
+		RequestedVertexAttributesType.SetFieldRequested(CachedTree.GetMaterialAttributesType()->FindFieldByName(TEXT("PrevWorldPositionOffset")));
 
 		EmitContext.ShaderFrequency = SF_Vertex;
 		EmitContext.bUseAnalyticDerivatives = false;
-		FEmitScope* EmitResultScope = EmitContext.PrepareScope(&Generator.GetResultStatement()->GetParentScope());
+		FEmitScope* EmitResultScope = EmitContext.PrepareScope(&CachedTree.GetResultStatement()->GetParentScope());
 
-		const FPreparedType VertexResultType = EmitContext.PrepareExpression(Generator.GetResultExpression(), *EmitResultScope, RequestedVertexAttributesType);
+		const FPreparedType VertexResultType = EmitContext.PrepareExpression(CachedTree.GetResultExpression(), *EmitResultScope, RequestedVertexAttributesType);
 		if (VertexResultType.IsVoid())
 		{
 			return false;
 		}
 
-		bUsesWorldPositionOffset = IsAttributeUsed(Generator, EmitContext, *EmitResultScope, VertexResultType, Generator.GetResultExpression(), MP_WorldPositionOffset);
+		bUsesWorldPositionOffset = IsAttributeUsed(CachedTree, EmitContext, *EmitResultScope, VertexResultType, CachedTree.GetResultExpression(), MP_WorldPositionOffset);
 
-		HLSLTree->EmitShader(EmitContext, VertexCode);
+		CachedTree.GetTree().EmitShader(EmitContext, VertexCode);
 	}
 
 	if (EmitContext.Errors->HasErrors())
@@ -945,10 +988,10 @@ bool MaterialEmitHLSL(const FMaterialCompileTargetParameters& InCompilerTarget,
 	OutCompilationOutput.bUsesPixelDepthOffset = bUsesPixelDepthOffset;
 
 	FStringBuilderMemstack Declarations(Allocator, 32 * 1024);
-	TypeRegistry.EmitDeclarationsCode(Declarations);
+	CachedTree.GetTypeRegistry().EmitDeclarationsCode(Declarations);
 
 	FStringBuilderMemstack SharedCode(Allocator, 32 * 1024);
-	Generator.EmitSharedCode(SharedCode);
+	CachedTree.EmitSharedCode(SharedCode);
 	EmitContext.EmitDeclarationsCode(SharedCode);
 
 	FString MaterialTemplateSource;
diff --git a/Engine/Source/Runtime/Engine/Private/Materials/MaterialHLSLGenerator.cpp b/Engine/Source/Runtime/Engine/Private/Materials/MaterialHLSLGenerator.cpp
--- a/Engine/Source/Runtime/Engine/Private/Materials/MaterialHLSLGenerator.cpp
+++ b/Engine/Source/Runtime/Engine/Private/Materials/MaterialHLSLGenerator.cpp
@@ -9,164 +9,50 @@
 #include "Materials/MaterialExpressionFunctionInput.h"
 #include "Materials/MaterialExpressionFunctionOutput.h"
 #include "Materials/MaterialExpressionCustomOutput.h"
-#include "Materials/MaterialExpressionVolumetricAdvancedMaterialOutput.h"
 #include "Materials/MaterialExpressionExecBegin.h"
-#include "Materials/Material.h"
 #include "Materials/MaterialHLSLTree.h"
+#include "Materials/Material.h"
+#include "MaterialCachedHLSLTree.h"
 #include "ShaderCore.h"
 #include "HLSLTree/HLSLTree.h"
 #include "HLSLTree/HLSLTreeCommon.h"
 #include "Containers/LazyPrintf.h"
+#include "Misc/MemStackUtility.h"
 
-static UE::Shader::EValueType GetShaderType(EMaterialValueType MaterialType)
+FMaterialHLSLGenerator::FMaterialHLSLGenerator(UMaterial* Material, const FMaterialLayersFunctions* InLayerOverrides, FMaterialCachedHLSLTree& InTree)
+	: TargetMaterial(Material)
+	, LayerOverrides(InLayerOverrides)
+	, CachedTree(InTree)
+	, bGeneratedResult(false)
 {
-	switch (MaterialType)
-	{
-	case MCT_Float1: return UE::Shader::EValueType::Float1;
-	case MCT_Float2: return UE::Shader::EValueType::Float2;
-	case MCT_Float3: return UE::Shader::EValueType::Float3;
-	case MCT_Float4: return UE::Shader::EValueType::Float4;
-	case MCT_Float: return UE::Shader::EValueType::Float1;
-	case MCT_StaticBool: return UE::Shader::EValueType::Bool1;
-	case MCT_MaterialAttributes: return UE::Shader::EValueType::Struct;
-	case MCT_ShadingModel: return UE::Shader::EValueType::Int1;
-	case MCT_LWCScalar: return UE::Shader::EValueType::Double1;
-	case MCT_LWCVector2: return UE::Shader::EValueType::Double2;
-	case MCT_LWCVector3: return UE::Shader::EValueType::Double3;
-	case MCT_LWCVector4: return UE::Shader::EValueType::Double4;
-	default:return UE::Shader::EValueType::Void;
-	}
+	FunctionCallStack.Add(&RootFunctionCallEntry);
 }
 
-FMaterialHLSLErrorHandler::FMaterialHLSLErrorHandler(FMaterial& InOutMaterial)
-	: Material(&InOutMaterial)
+UE::HLSLTree::FTree& FMaterialHLSLGenerator::GetTree() const
 {
-	Material->CompileErrors.Reset();
-	Material->ErrorExpressions.Reset();
+	return CachedTree.GetTree();
 }
 
-void FMaterialHLSLErrorHandler::AddErrorInternal(UObject* InOwner, FStringView InError)
+UE::Shader::FStructTypeRegistry& FMaterialHLSLGenerator::GetTypeRegistry() const
 {
-	UMaterialExpression* MaterialExpressionOwner = Cast<UMaterialExpression>(InOwner);
-	UMaterialExpression* ExpressionToError = nullptr;
-	TStringBuilder<1024> FormattedError;
-
-	if (MaterialExpressionOwner)
-	{
-		if (MaterialExpressionOwner->GetClass() != UMaterialExpressionMaterialFunctionCall::StaticClass()
-			&& MaterialExpressionOwner->GetClass() != UMaterialExpressionFunctionInput::StaticClass()
-			&& MaterialExpressionOwner->GetClass() != UMaterialExpressionFunctionOutput::StaticClass())
-		{
-			// Add the expression currently being compiled to ErrorExpressions so we can draw it differently
-			ExpressionToError = MaterialExpressionOwner;
-
-			const int32 ChopCount = FCString::Strlen(TEXT("MaterialExpression"));
-			const FString ErrorClassName = MaterialExpressionOwner->GetClass()->GetName();
-
-			// Add the node type to the error message
-			FormattedError.Appendf(TEXT("(Node %s) "), *ErrorClassName.Right(ErrorClassName.Len() - ChopCount));
-		}
-	}
-
-	FormattedError.Append(InError);
-	const FString Error(FormattedError.ToView());
-
-	// Standard error handling, immediately append one-off errors and signal failure
-	Material->CompileErrors.AddUnique(Error);
-
-	if (ExpressionToError)
-	{
-		Material->ErrorExpressions.Add(ExpressionToError);
-		ExpressionToError->LastErrorText = Error;
-	}
+	return CachedTree.GetTypeRegistry();
 }
 
-FMaterialHLSLGenerator::FMaterialHLSLGenerator(const FMaterialCompileTargetParameters& InCompileTarget,
-	const FStaticParameterSet& InStaticParameters,
-	FMaterial& InOutMaterial,
-	UE::Shader::FStructTypeRegistry& InOutTypeRegistry,
-	UE::HLSLTree::FTree& InOutTree)
-	: CompileTarget(InCompileTarget)
-	, StaticParameters(InStaticParameters)
-	, Errors(InOutMaterial)
-	, HLSLTree(&InOutTree)
-	, TypeRegistry(&InOutTypeRegistry)
-	, bGeneratedResult(false)
+const UE::Shader::FStructType* FMaterialHLSLGenerator::GetMaterialAttributesType() const
 {
-	UMaterialInterface* MaterialInterface = InOutMaterial.GetMaterialInterface();
-	TargetMaterial = MaterialInterface->GetMaterial();
-	const EMaterialShadingModel DefaultShadingModel = TargetMaterial->GetShadingModels().GetFirstShadingModel();
-
-	FFunctionCallEntry* RootFunctionEntry = new(InOutTree.GetAllocator()) FFunctionCallEntry();
-	FunctionCallStack.Add(RootFunctionEntry);
-
-	InOutMaterial.GatherCustomOutputExpressions(MaterialCustomOutputs);
-
-	TArray<UE::Shader::FStructFieldInitializer, TInlineAllocator<MP_MAX + 16>> MaterialAttributeFields;
-
-	const TArray<FGuid>& OrderedVisibleAttributes = FMaterialAttributeDefinitionMap::GetOrderedVisibleAttributeList();
-	for (const FGuid& AttributeID : OrderedVisibleAttributes)
-	{
-		const FString& PropertyName = FMaterialAttributeDefinitionMap::GetAttributeName(AttributeID);
-		const EMaterialValueType PropertyType = FMaterialAttributeDefinitionMap::GetValueType(AttributeID);
-		const UE::Shader::EValueType ValueType = GetShaderType(PropertyType);
-		
-		if (ValueType != UE::Shader::EValueType::Void &&
-			ValueType != UE::Shader::EValueType::Struct)
-		{
-			MaterialAttributeFields.Emplace(PropertyName, ValueType);
-
-			if (PropertyType == MCT_ShadingModel)
-			{
-				check(ValueType == UE::Shader::EValueType::Int1);
-				MaterialAttributesDefaultValue.Component.Add((int32)DefaultShadingModel);
-			}
-			else
-			{
-				const UE::Shader::FValue DefaultValue = UE::Shader::Cast(FMaterialAttributeDefinitionMap::GetDefaultValue(AttributeID), ValueType);
-				MaterialAttributesDefaultValue.Component.Append(DefaultValue.Component);
-			}
-		}
-	}
-
-	TArray<TStringBuilder<256>, TInlineAllocator<4>> CustomOutputNames;
-	CustomOutputNames.Reserve(MaterialCustomOutputs.Num());
-	for (UMaterialExpressionCustomOutput* CustomOutput : MaterialCustomOutputs)
-	{
-		const int32 NumOutputs = CustomOutput->GetNumOutputs();
-		const FString OutputName = CustomOutput->GetFunctionName();
-
-		check(!CustomOutput->ShouldCompileBeforeAttributes()); // not supported yet, looks like this isn't currently being used
-
-		for (int32 OutputIndex = 0; OutputIndex < NumOutputs; ++OutputIndex)
-		{
-			const UE::Shader::EValueType ValueType = CustomOutput->GetCustomOutputType(OutputIndex);
-			FStringBuilderBase& FormattedName = CustomOutputNames.AddDefaulted_GetRef();
-			FormattedName.Appendf(TEXT("%s%d"), *OutputName, OutputIndex);
-			MaterialAttributeFields.Emplace(FormattedName.ToView(), ValueType);
-
-			const UE::Shader::FValue DefaultValue(ValueType);
-			MaterialAttributesDefaultValue.Component.Append(DefaultValue.Component);
-		}
-	}
-
-	MaterialAttributeFields.Emplace(TEXT("PrevWorldPositionOffset"), UE::Shader::EValueType::Float3);
-	MaterialAttributesDefaultValue.Component.Append({ 0.0f, 0.0f, 0.0f });
-
-	UE::Shader::FStructTypeInitializer MaterialAttributesInitializer;
-	MaterialAttributesInitializer.Name = TEXT("FMaterialAttributes");
-	MaterialAttributesInitializer.Fields = MaterialAttributeFields;
-	MaterialAttributesType = InOutTypeRegistry.NewType(MaterialAttributesInitializer);
+	return CachedTree.GetMaterialAttributesType();
+}
 
-	check(MaterialAttributesDefaultValue.Component.Num() == MaterialAttributesType->ComponentTypes.Num());
-	MaterialAttributesDefaultValue.Type = MaterialAttributesType;
+const UE::Shader::FValue& FMaterialHLSLGenerator::GetMaterialAttributesDefaultValue() const
+{
+	return CachedTree.GetMaterialAttributesDefaultValue();
 }
 
 bool FMaterialHLSLGenerator::Generate()
 {
 	using namespace UE::HLSLTree;
 
-	FScope& RootScope = HLSLTree->GetRootScope();
+	FScope& RootScope = CachedTree.GetTree().GetRootScope();
 
 	bool bResult = false;
 	if (TargetMaterial->IsUsingControlFlow())
@@ -174,7 +60,7 @@ bool FMaterialHLSLGenerator::Generate()
 		UMaterialExpression* BaseExpression = TargetMaterial->ExpressionExecBegin;
 		if (!BaseExpression)
 		{
-			bResult = Errors.AddError(TEXT("Missing ExpressionExecBegin"));
+			bResult = Error(TEXT("Missing ExpressionExecBegin"));
 		}
 		else
 		{
@@ -194,12 +80,12 @@ bool FMaterialHLSLGenerator::Generate()
 	check(FunctionCallStack.Num() == 1);
 	if (!bGeneratedResult)
 	{
-		return Errors.AddError(TEXT("Missing connection to material output"));
+		return Error(TEXT("Missing connection to material output"));
 	}
 
-	if (!ResultExpression || !ResultStatement)
+	if (!CachedTree.GetResultExpression() || !CachedTree.GetResultStatement())
 	{
-		return Errors.AddError(TEXT("Failed to initialize result"));
+		return Error(TEXT("Failed to initialize result"));
 	}
 
 	for (const auto& It : StatementMap)
@@ -208,70 +94,16 @@ bool FMaterialHLSLGenerator::Generate()
 		const FStatementEntry& Entry = It.Value;
 		if (Entry.NumInputs != Expression->NumExecutionInputs)
 		{
-			return Errors.AddError(TEXT("Invalid number of input connections"));
+			return Error(TEXT("Invalid number of input connections"));
 		}
 	}
 
 	if (JoinedScopeStack.Num() != 0)
 	{
-		return Errors.AddError(TEXT("Invalid control flow"));
+		return Error(TEXT("Invalid control flow"));
 	}
 
-	return HLSLTree->Finalize();
-}
-
-void FMaterialHLSLGenerator::SetRequestedFields(EShaderFrequency ShaderFrequency, UE::HLSLTree::FRequestedType& OutRequestedType)
-{
-	for (UMaterialExpressionCustomOutput* CustomOutput : MaterialCustomOutputs)
-	{
-		if (CustomOutput->GetShaderFrequency() != ShaderFrequency)
-		{
-			continue;
-		}
-
-		const int32 NumOutputs = CustomOutput->GetNumOutputs();
-		const FString OutputName = CustomOutput->GetFunctionName();
-
-		for (int32 OutputIndex = 0; OutputIndex < NumOutputs; ++OutputIndex)
-		{
-			TStringBuilder<256> FieldName;
-			FieldName.Appendf(TEXT("%s%d"), *OutputName, OutputIndex);
-			const UE::Shader::FStructField* CustomOutputField = GetMaterialAttributesType()->FindFieldByName(FieldName.ToString());
-			check(CustomOutputField);
-			OutRequestedType.SetFieldRequested(CustomOutputField);
-		}
-	}
-}
-
-void FMaterialHLSLGenerator::EmitSharedCode(FStringBuilderBase& OutCode) const
-{
-	using namespace UE::Shader;
-	for (UMaterialExpressionCustomOutput* CustomOutput : MaterialCustomOutputs)
-	{
-		const int32 NumOutputs = CustomOutput->GetNumOutputs();
-		const FString OutputName = CustomOutput->GetFunctionName();
-		const EShaderFrequency ShaderFrequency = CustomOutput->GetShaderFrequency();
-
-		if (CustomOutput->NeedsCustomOutputDefines())
-		{
-			OutCode.Appendf(TEXT("#define NUM_MATERIAL_OUTPUTS_%s %d\n"), *OutputName.ToUpper(), NumOutputs);
-		}
-
-		for (int32 OutputIndex = 0; OutputIndex < NumOutputs; ++OutputIndex)
-		{
-			const EValueType ValueType = CustomOutput->GetCustomOutputType(OutputIndex);
-			const FValueTypeDescription ValueTypeDesc = GetValueTypeDescription(ValueType);
-
-			OutCode.Appendf(TEXT("#define HAVE_%s%d 1\n"), *OutputName, OutputIndex);
-
-			OutCode.Appendf(TEXT("%s %s%d(FMaterial%sParameters Parameters) { return Parameters.MaterialAttributes.%s%d; }\n"),
-				ValueTypeDesc.Name,
-				*OutputName, OutputIndex,
-				ShaderFrequency == SF_Pixel ? TEXT("Pixel") : TEXT("Vertex"),
-				*OutputName, OutputIndex);
-		}
-		OutCode.Append(TEXT("\n"));
-	}
+	return GetTree().Finalize();
 }
 
 static UE::HLSLTree::FExpression* CompileMaterialInput(FMaterialHLSLGenerator& Generator,
@@ -328,16 +160,16 @@ bool FMaterialHLSLGenerator::GenerateResult(UE::HLSLTree::FScope& Scope)
 	}
 	else if (bGeneratedResult)
 	{
-		return Errors.AddError(TEXT("Multiple connections to execution output"));
+		return Error(TEXT("Multiple connections to execution output"));
 	}
 	else
 	{
-		check(!ResultStatement);
-		check(!ResultExpression);
+		check(!CachedTree.ResultStatement);
+		check(!CachedTree.ResultExpression);
 
 		if (TargetMaterial)
 		{
-			const FStructField* PrevWPOField = GetMaterialAttributesType()->FindFieldByName(TEXT("PrevWorldPositionOffset"));
+			const FStructField* PrevWPOField = CachedTree.GetMaterialAttributesType()->FindFieldByName(TEXT("PrevWorldPositionOffset"));
 
 			FExpression* AttributesExpression = nullptr;
 			if (TargetMaterial->bUseMaterialAttributes)
@@ -351,49 +183,49 @@ bool FMaterialHLSLGenerator::GenerateResult(UE::HLSLTree::FScope& Scope)
 					if (AttributesExpression)
 					{
 						const FString& WPOName = FMaterialAttributeDefinitionMap::GetAttributeName(MP_WorldPositionOffset);
-						const FStructField* WPOField = GetMaterialAttributesType()->FindFieldByName(*WPOName);
+						const FStructField* WPOField = CachedTree.GetMaterialAttributesType()->FindFieldByName(*WPOName);
 
 						FRequestedType PrevRequestedType;
 						PrevRequestedType.SetFieldRequested(WPOField);
 
-						FExpression* PrevAttributesExpression = HLSLTree->GetPreviousFrame(AttributesExpression, PrevRequestedType);
+						FExpression* PrevAttributesExpression = GetTree().GetPreviousFrame(AttributesExpression, PrevRequestedType);
 						ensure(PrevAttributesExpression);
-						FExpression* PrevWPOExpression = HLSLTree->NewExpression<FExpressionGetStructField>(GetMaterialAttributesType(), WPOField, PrevAttributesExpression);
-						AttributesExpression = HLSLTree->NewExpression<FExpressionSetStructField>(GetMaterialAttributesType(), PrevWPOField, AttributesExpression, PrevWPOExpression);
+						FExpression* PrevWPOExpression = GetTree().NewExpression<FExpressionGetStructField>(CachedTree.GetMaterialAttributesType(), WPOField, PrevAttributesExpression);
+						AttributesExpression = GetTree().NewExpression<FExpressionSetStructField>(CachedTree.GetMaterialAttributesType(), PrevWPOField, AttributesExpression, PrevWPOExpression);
 					}
 				}
 			}
 			else
 			{
-				AttributesExpression = HLSLTree->NewExpression<FExpressionConstant>(MaterialAttributesDefaultValue);
+				AttributesExpression = GetTree().NewExpression<FExpressionConstant>(CachedTree.GetMaterialAttributesDefaultValue());
 				for (int32 PropertyIndex = 0; PropertyIndex < MP_MAX; ++PropertyIndex)
 				{
 					const EMaterialProperty Property = (EMaterialProperty)PropertyIndex;
 
 					// We're only interesting in attributes that map to valid fields
-					const UE::Shader::FStructField* AttributeField = GetMaterialAttributesType()->FindFieldByName(*FMaterialAttributeDefinitionMap::GetAttributeName(Property));
+					const UE::Shader::FStructField* AttributeField = CachedTree.GetMaterialAttributesType()->FindFieldByName(*FMaterialAttributeDefinitionMap::GetAttributeName(Property));
 					if (AttributeField)
 					{
 						FExpression* InputExpression = CompileMaterialInput(*this, Scope, Property, TargetMaterial);
 						if (InputExpression)
 						{
 							AttributesExpression = GetTree().NewExpression<FExpressionSetStructField>(
-								GetMaterialAttributesType(),
+								CachedTree.GetMaterialAttributesType(),
 								AttributeField,
 								AttributesExpression,
 								InputExpression);
 							if (Property == MP_WorldPositionOffset)
 							{
-								FExpression* PrevWPOExpression = HLSLTree->GetPreviousFrame(InputExpression, ERequestedType::Vector3);
+								FExpression* PrevWPOExpression = GetTree().GetPreviousFrame(InputExpression, ERequestedType::Vector3);
 								ensure(PrevWPOExpression);
-								AttributesExpression = HLSLTree->NewExpression<FExpressionSetStructField>(GetMaterialAttributesType(), PrevWPOField, AttributesExpression, PrevWPOExpression);
+								AttributesExpression = GetTree().NewExpression<FExpressionSetStructField>(CachedTree.GetMaterialAttributesType(), PrevWPOField, AttributesExpression, PrevWPOExpression);
 							}
 						}
 					}
 				}
 			}
 
-			for (UMaterialExpressionCustomOutput* CustomOutput : MaterialCustomOutputs)
+			for (UMaterialExpressionCustomOutput* CustomOutput : CachedTree.MaterialCustomOutputs)
 			{
 				const int32 NumOutputs = CustomOutput->GetNumOutputs();
 				const FString OutputName = CustomOutput->GetFunctionName();
@@ -401,20 +233,20 @@ bool FMaterialHLSLGenerator::GenerateResult(UE::HLSLTree::FScope& Scope)
 				{
 					TStringBuilder<256> FieldName;
 					FieldName.Appendf(TEXT("%s%d"), *OutputName, OutputIndex);
-					const FStructField* CustomOutputField = GetMaterialAttributesType()->FindFieldByName(FieldName.ToString());
+					const FStructField* CustomOutputField = CachedTree.GetMaterialAttributesType()->FindFieldByName(FieldName.ToString());
 					check(CustomOutputField);
 
 					FExpression* CustomOutputExpression = AcquireExpression(Scope, CustomOutput, OutputIndex);
-					AttributesExpression = HLSLTree->NewExpression<FExpressionSetStructField>(GetMaterialAttributesType(), CustomOutputField, AttributesExpression, CustomOutputExpression);
+					AttributesExpression = GetTree().NewExpression<FExpressionSetStructField>(CachedTree.GetMaterialAttributesType(), CustomOutputField, AttributesExpression, CustomOutputExpression);
 				}
 			}
 
 			if (AttributesExpression)
 			{
-				FStatementReturn* ReturnStatement = HLSLTree->NewStatement<FStatementReturn>(Scope);
+				FStatementReturn* ReturnStatement = GetTree().NewStatement<FStatementReturn>(Scope);
 				ReturnStatement->Expression = AttributesExpression;
-				ResultExpression = AttributesExpression;
-				ResultStatement = ReturnStatement;
+				CachedTree.ResultExpression = AttributesExpression;
+				CachedTree.ResultStatement = ReturnStatement;
 				bResult = true;
 			}
 		}
@@ -429,7 +261,7 @@ bool FMaterialHLSLGenerator::GenerateResult(UE::HLSLTree::FScope& Scope)
 
 UE::HLSLTree::FScope* FMaterialHLSLGenerator::NewScope(UE::HLSLTree::FScope& Scope, EMaterialNewScopeFlag Flags)
 {
-	UE::HLSLTree::FScope* NewScope = HLSLTree->NewScope(Scope);
+	UE::HLSLTree::FScope* NewScope = GetTree().NewScope(Scope);
 	if (!EnumHasAllFlags(Flags, EMaterialNewScopeFlag::NoPreviousScope))
 	{
 		NewScope->AddPreviousScope(Scope);
@@ -440,43 +272,45 @@ UE::HLSLTree::FScope* FMaterialHLSLGenerator::NewScope(UE::HLSLTree::FScope& Sco
 
 UE::HLSLTree::FScope* FMaterialHLSLGenerator::NewOwnedScope(UE::HLSLTree::FStatement& Owner)
 {
-	UE::HLSLTree::FScope* NewScope = HLSLTree->NewOwnedScope(Owner);
+	UE::HLSLTree::FScope* NewScope = GetTree().NewOwnedScope(Owner);
 	NewScope->AddPreviousScope(Owner.GetParentScope());
 	return NewScope;
 }
 
 UE::HLSLTree::FScope* FMaterialHLSLGenerator::NewJoinedScope(UE::HLSLTree::FScope& Scope)
 {
-	UE::HLSLTree::FScope* NewScope = HLSLTree->NewScope(Scope);
+	UE::HLSLTree::FScope* NewScope = GetTree().NewScope(Scope);
 	JoinedScopeStack.Add(NewScope);
 	return NewScope;
 }
 
 UE::HLSLTree::FExpression* FMaterialHLSLGenerator::NewConstant(const UE::Shader::FValue& Value)
 {
-	return HLSLTree->NewConstant(Value);
+	return GetTree().NewConstant(Value);
 }
 
 UE::HLSLTree::FExpression* FMaterialHLSLGenerator::NewTexCoord(int32 Index)
 {
 	using namespace UE::HLSLTree::Material;
-	return HLSLTree->NewExpression<FExpressionExternalInput>(MakeInputTexCoord(Index));
+	return GetTree().NewExpression<FExpressionExternalInput>(MakeInputTexCoord(Index));
 }
 
 UE::HLSLTree::FExpression* FMaterialHLSLGenerator::NewSwizzle(const UE::HLSLTree::FSwizzleParameters& Params, UE::HLSLTree::FExpression* Input)
 {
-	return HLSLTree->NewExpression<UE::HLSLTree::FExpressionSwizzle>(Params, Input);
+	return GetTree().NewExpression<UE::HLSLTree::FExpressionSwizzle>(Params, Input);
 }
 
 const UE::Shader::FTextureValue* FMaterialHLSLGenerator::AcquireTextureValue(const UE::Shader::FTextureValue& InValue)
 {
 	using namespace UE::Shader;
-	FString SamplerTypeError;
+	
+	// Need to move this to HLSLTreeEmit
+	/*FString SamplerTypeError;
 	if (!UMaterialExpressionTextureBase::VerifySamplerType(CompileTarget.FeatureLevel, CompileTarget.TargetPlatform, InValue.Texture, InValue.SamplerType, SamplerTypeError))
 	{
 		Errors.AddError(SamplerTypeError);
 		return nullptr;
-	}
+	}*/
 
 	FXxHash64Builder Hasher;
 	Hasher.Update(&InValue.Texture, sizeof(InValue.Texture));
@@ -492,23 +326,36 @@ const UE::Shader::FTextureValue* FMaterialHLSLGenerator::AcquireTextureValue(con
 		return *PrevValue;
 	}
 
-	FTextureValue* Value = new(HLSLTree->GetAllocator()) FTextureValue(InValue);
+	FTextureValue* Value = new(GetTree().GetAllocator()) FTextureValue(InValue);
 	TextureValueMap.Add(Hash, Value);
 	return Value;
 }
 
+bool FMaterialHLSLGenerator::InternalError(FStringView ErrorMessage)
+{
+	if (CurrentErrorMessage.Len() > 0)
+	{
+		CurrentErrorMessage.AppendChar(TEXT('\n'));
+	}
+	CurrentErrorMessage.Append(ErrorMessage);
+	return false;
+}
+
 UE::HLSLTree::FExpression* FMaterialHLSLGenerator::AcquireExpression(UE::HLSLTree::FScope& Scope, UMaterialExpression* MaterialExpression, int32 OutputIndex)
 {
 	using namespace UE::HLSLTree;
 	FOwnerScope TreeOwnerScope(GetTree(), MaterialExpression);
-	FOwnerScope ErrorOwnerScope(Errors, MaterialExpression);
 
 	FExpression* Expression = nullptr;
 	if (MaterialExpression->GenerateHLSLExpression(*this, Scope, OutputIndex, Expression))
 	{
 		return Expression;
 	}
-	return nullptr;
+	
+	check(!Expression);
+	const FStringView ErrorMessage = UE::MemStack::AllocateStringView(GetTree().GetAllocator(), CurrentErrorMessage);
+	CurrentErrorMessage.Reset();
+	return GetTree().NewExpression<UE::HLSLTree::FExpressionError>(ErrorMessage);
 }
 
 UE::HLSLTree::FExpression* FMaterialHLSLGenerator::AcquireFunctionInputExpression(UE::HLSLTree::FScope& Scope, const UMaterialExpressionFunctionInput* MaterialExpression)
@@ -532,7 +379,7 @@ UE::HLSLTree::FExpression* FMaterialHLSLGenerator::AcquireFunctionInputExpressio
 		if (!bFoundInput)
 		{
 			// Finding a connected input is always expected if we're in a function call
-			Errors.AddError(TEXT("Invalid function input"));
+			Error(TEXT("Invalid function input"));
 			return nullptr;
 		}
 	}
@@ -551,17 +398,17 @@ UE::HLSLTree::FExpression* FMaterialHLSLGenerator::AcquireFunctionInputExpressio
 			case FunctionInput_Vector2: DefaultValue = FVector2f(PreviewValue.X, PreviewValue.Y); break;
 			case FunctionInput_Vector3: DefaultValue = FVector3f(PreviewValue.X, PreviewValue.Y, PreviewValue.Z); break;
 			case FunctionInput_Vector4: DefaultValue = PreviewValue; break;
-			case FunctionInput_MaterialAttributes: DefaultValue = GetMaterialAttributesDefaultValue(); break;
+			case FunctionInput_MaterialAttributes: DefaultValue = CachedTree.GetMaterialAttributesDefaultValue(); break;
 			case FunctionInput_Texture2D:
 			case FunctionInput_TextureCube:
 			case FunctionInput_Texture2DArray:
 			case FunctionInput_VolumeTexture:
 			case FunctionInput_StaticBool:
 			case FunctionInput_TextureExternal:
-				Errors.AddErrorf(TEXT("Missing Preview connection for function input '%s'"), *MaterialExpression->InputName.ToString());
+				Errorf(TEXT("Missing Preview connection for function input '%s'"), *MaterialExpression->InputName.ToString());
 				return nullptr;
 			default:
-				Errors.AddError(TEXT("Unknown input type"));
+				Error(TEXT("Unknown input type"));
 				return nullptr;
 			}
 
@@ -575,18 +422,17 @@ UE::HLSLTree::FExpression* FMaterialHLSLGenerator::AcquireFunctionInputExpressio
 bool FMaterialHLSLGenerator::GenerateStatements(UE::HLSLTree::FScope& Scope, UMaterialExpression* MaterialExpression)
 {
 	using namespace UE::HLSLTree;
-	FOwnerScope ErrorOwnerScope(Errors, MaterialExpression);
 
 	FStatementEntry& Entry = StatementMap.FindOrAdd(MaterialExpression);
 	check(Entry.NumInputs >= 0);
 
 	if (Entry.NumInputs >= MaterialExpression->NumExecutionInputs)
 	{
-		return Errors.AddErrorf(TEXT("Bad control flow, found %d inputs out of %d reported"), Entry.NumInputs, MaterialExpression->NumExecutionInputs);
+		return Errorf(TEXT("Bad control flow, found %d inputs out of %d reported"), Entry.NumInputs, MaterialExpression->NumExecutionInputs);
 	}
 	if (Entry.NumInputs == MaxNumPreviousScopes)
 	{
-		return Errors.AddErrorf(TEXT("Bad control flow, too many execution inputs"));
+		return Errorf(TEXT("Bad control flow, too many execution inputs"));
 	}
 
 	Entry.PreviousScope[Entry.NumInputs++] = &Scope;
@@ -599,7 +445,7 @@ bool FMaterialHLSLGenerator::GenerateStatements(UE::HLSLTree::FScope& Scope, UMa
 		{
 			if (JoinedScopeStack.Num() == 0)
 			{
-				return Errors.AddError(TEXT("Bad control flow"));
+				return Error(TEXT("Bad control flow"));
 			}
 
 			ScopeToUse = JoinedScopeStack.Pop(false);
@@ -635,7 +481,7 @@ UE::HLSLTree::FExpression* FMaterialHLSLGenerator::GenerateMaterialParameter(EMa
 		}
 	}
 
-	return HLSLTree->NewExpression<UE::HLSLTree::Material::FExpressionParameter>(InType, GetParameterInfo(InParameterName), DefaultValue);
+	return GetTree().NewExpression<UE::HLSLTree::Material::FExpressionParameter>(InType, GetParameterInfo(InParameterName), DefaultValue);
 }
 
 UE::HLSLTree::FExpression* FMaterialHLSLGenerator::GenerateFunctionCall(UE::HLSLTree::FScope& Scope,
@@ -649,7 +495,7 @@ UE::HLSLTree::FExpression* FMaterialHLSLGenerator::GenerateFunctionCall(UE::HLSL
 
 	if (!MaterialFunction)
 	{
-		Errors.AddError(TEXT("Missing material function"));
+		Error(TEXT("Missing material function"));
 		return nullptr;
 	}
 
@@ -659,14 +505,14 @@ UE::HLSLTree::FExpression* FMaterialHLSLGenerator::GenerateFunctionCall(UE::HLSL
 
 	if (FunctionInputs.Num() != ConnectedInputs.Num())
 	{
-		Errors.AddError(TEXT("Mismatched function inputs"));
+		Error(TEXT("Mismatched function inputs"));
 		return nullptr;
 	}
 
 	const UMaterialExpressionFunctionOutput* ExpressionOutput = FunctionOutputs.IsValidIndex(OutputIndex) ? FunctionOutputs[OutputIndex].ExpressionOutput.Get() : nullptr;
 	if (!ExpressionOutput)
 	{
-		Errors.AddError(TEXT("Invalid function output"));
+		Error(TEXT("Invalid function output"));
 		return nullptr;
 	}
 
@@ -680,13 +526,13 @@ UE::HLSLTree::FExpression* FMaterialHLSLGenerator::GenerateFunctionCall(UE::HLSL
 		ParameterIndex = PrevFunctionEntry->ParameterIndex;
 	}
 
-	FSHAHash Hash;
+	FXxHash64 Hash;
 	FFunctionInputArray LocalFunctionInputs;
 	{
-		FSHA1 Hasher;
-		Hasher.Update((uint8*)&MaterialFunction, sizeof(UMaterialFunctionInterface*));
-		Hasher.Update((uint8*)&ParameterAssociation, sizeof(EMaterialParameterAssociation));
-		Hasher.Update((uint8*)&ParameterIndex, sizeof(int32));
+		FXxHash64Builder Hasher;
+		Hasher.Update(&MaterialFunction, sizeof(UMaterialFunctionInterface*));
+		Hasher.Update(&ParameterAssociation, sizeof(EMaterialParameterAssociation));
+		Hasher.Update(&ParameterIndex, sizeof(int32));
 
 		for (int32 InputIndex = 0; InputIndex < ConnectedInputs.Num(); ++InputIndex)
 		{
@@ -704,13 +550,13 @@ UE::HLSLTree::FExpression* FMaterialHLSLGenerator::GenerateFunctionCall(UE::HLSL
 	}
 
 	const bool bInlineFunction = !MaterialFunction->IsUsingControlFlow();
-	FFunctionCallEntry** ExistingFunctionCall = FunctionCallMap.Find(Hash);
-	FFunctionCallEntry* FunctionCall = ExistingFunctionCall ? *ExistingFunctionCall : nullptr;
+	TUniquePtr<FFunctionCallEntry>* ExistingFunctionCall = FunctionCallMap.Find(Hash);
+	FFunctionCallEntry* FunctionCall = ExistingFunctionCall ? ExistingFunctionCall->Get() : nullptr;
 	if (!FunctionCall)
 	{
 		// Generate an HLSL function object, if this is not an inline function call
-		FFunction* HLSLFunction = !bInlineFunction ? HLSLTree->NewFunction() : nullptr;
-		FunctionCall = new(HLSLTree->GetAllocator()) FFunctionCallEntry();
+		FFunction* HLSLFunction = !bInlineFunction ? GetTree().NewFunction() : nullptr;
+		FunctionCall = new FFunctionCallEntry();
 		FunctionCall->MaterialFunction = MaterialFunction;
 		FunctionCall->ParameterAssociation = ParameterAssociation;
 		FunctionCall->ParameterIndex = ParameterIndex;
@@ -723,7 +569,7 @@ UE::HLSLTree::FExpression* FMaterialHLSLGenerator::GenerateFunctionCall(UE::HLSL
 			FunctionCall->FunctionOutputs.Add(Output.ExpressionOutput);
 		}
 
-		FunctionCallMap.Add(Hash, FunctionCall);
+		FunctionCallMap.Emplace(Hash, FunctionCall);
 
 		if (HLSLFunction)
 		{
@@ -748,11 +594,11 @@ UE::HLSLTree::FExpression* FMaterialHLSLGenerator::GenerateFunctionCall(UE::HLSL
 		check(HLSLFunction->OutputExpressions.Num() == FunctionOutputs.Num());
 		if (HLSLFunction->OutputExpressions[OutputIndex])
 		{
-			Result = HLSLTree->NewFunctionCall(Scope, HLSLFunction, OutputIndex);
+			Result = GetTree().NewFunctionCall(Scope, HLSLFunction, OutputIndex);
 		}
 		else
 		{
-			Errors.AddError(TEXT("Invalid function output"));
+			Error(TEXT("Invalid function output"));
 		}
 	}
 	verify(FunctionCallStack.Pop() == FunctionCall);
diff --git a/Engine/Source/Runtime/Engine/Private/Materials/MaterialHLSLTree.cpp b/Engine/Source/Runtime/Engine/Private/Materials/MaterialHLSLTree.cpp
--- a/Engine/Source/Runtime/Engine/Private/Materials/MaterialHLSLTree.cpp
+++ b/Engine/Source/Runtime/Engine/Private/Materials/MaterialHLSLTree.cpp
@@ -546,7 +546,8 @@ bool FExpressionNoise::PrepareValue(FEmitContext& Context, FEmitScope& Scope, co
 
 void FExpressionNoise::EmitValueShader(FEmitContext& Context, FEmitScope& Scope, const FRequestedType& RequestedType, FEmitValueShaderResult& OutResult) const
 {
-	bool bIsLWC = Shader::IsLWCType(PositionExpression->GetPreparedType().ValueComponentType);
+	const FPreparedType& PreparedType = Context.GetPreparedType(PositionExpression);
+	bool bIsLWC = Shader::IsLWCType(PreparedType.ValueComponentType);
 	FEmitShaderExpression* EmitPosition = PositionExpression->GetValueShader(Context, Scope, bIsLWC ? Shader::EValueType::Double3 : Shader::EValueType::Float3);
 	FEmitShaderExpression* EmitFilterWidth = FilterWidthExpression->GetValueShader(Context, Scope, Shader::EValueType::Float1);
 
diff --git a/Engine/Source/Runtime/Engine/Private/Materials/MaterialInstance.cpp b/Engine/Source/Runtime/Engine/Private/Materials/MaterialInstance.cpp
--- a/Engine/Source/Runtime/Engine/Private/Materials/MaterialInstance.cpp
+++ b/Engine/Source/Runtime/Engine/Private/Materials/MaterialInstance.cpp
@@ -820,6 +820,26 @@ const FMaterialCachedExpressionData& UMaterialInstance::GetCachedExpressionData(
 	return UMaterial::GetDefaultMaterial(MD_Surface)->GetCachedExpressionData();
 }
 
+#if WITH_EDITOR
+const FMaterialCachedHLSLTree& UMaterialInstance::GetCachedHLSLTree(TMicRecursionGuard RecursionGuard) const
+{
+	check(IsUsingNewHLSLGenerator());
+	const FMaterialCachedHLSLTree* LocalTree = CachedHLSLTree.Get();
+	if (LocalTree)
+	{
+		return *LocalTree;
+	}
+
+	if (Parent && !RecursionGuard.Contains(this))
+	{
+		RecursionGuard.Set(this);
+		return Parent->GetCachedHLSLTree(RecursionGuard);
+	}
+
+	return UMaterial::GetDefaultMaterial(MD_Surface)->GetCachedHLSLTree();
+}
+#endif // WITH_EDITOR
+
 bool UMaterialInstance::GetParameterOverrideValue(EMaterialParameterType Type, const FMemoryImageMaterialParameterInfo& ParameterInfo, FMaterialParameterMetadata& OutResult) const
 {
 	bool bResult = false;
diff --git a/Engine/Source/Runtime/Engine/Private/Materials/MaterialInstanceConstant.cpp b/Engine/Source/Runtime/Engine/Private/Materials/MaterialInstanceConstant.cpp
--- a/Engine/Source/Runtime/Engine/Private/Materials/MaterialInstanceConstant.cpp
+++ b/Engine/Source/Runtime/Engine/Private/Materials/MaterialInstanceConstant.cpp
@@ -8,6 +8,7 @@
 #include "Materials/MaterialInstanceSupport.h"
 #include "ProfilingDebugging/CookStats.h"
 #if WITH_EDITOR
+#include "MaterialCachedHLSLTree.h"
 #include "MaterialEditor/DEditorScalarParameterValue.h"
 #include "ObjectCacheEventSink.h"
 #endif
@@ -182,6 +183,7 @@ void UMaterialInstanceConstant::UpdateCachedData()
 	if (!bLoadedCachedExpressionData)
 	{
 		FMaterialCachedExpressionData* LocalCachedExpressionData = nullptr;
+		FMaterialCachedHLSLTree* LocalCachedTree = nullptr;
 
 		// If we have overriden material layers, need to create a local cached expression data
 		// Otherwise we can leave it as null, and use cached data from our parent
@@ -195,8 +197,15 @@ void UMaterialInstanceConstant::UpdateCachedData()
 			LocalCachedExpressionData = new FMaterialCachedExpressionData();
 			LocalCachedExpressionData->Reset();
 			LocalCachedExpressionData->UpdateForExpressions(Context, BaseMaterial->Expressions, GlobalParameter, INDEX_NONE);
+
+			if (IsUsingNewHLSLGenerator())
+			{
+				LocalCachedTree = new FMaterialCachedHLSLTree();
+				LocalCachedTree->GenerateTree(BaseMaterial, &LocalStaticParameters.MaterialLayers);
+			}
 		}
 		CachedExpressionData.Reset(LocalCachedExpressionData);
+		CachedHLSLTree.Reset(LocalCachedTree);
 
 		FObjectCacheEventSink::NotifyReferencedTextureChanged_Concurrent(this);
 	}
diff --git a/Engine/Source/Runtime/Engine/Private/Materials/MaterialInterface.cpp b/Engine/Source/Runtime/Engine/Private/Materials/MaterialInterface.cpp
--- a/Engine/Source/Runtime/Engine/Private/Materials/MaterialInterface.cpp
+++ b/Engine/Source/Runtime/Engine/Private/Materials/MaterialInterface.cpp
@@ -32,6 +32,7 @@
 
 #if WITH_EDITOR
 #include "ObjectCacheEventSink.h"
+#include "MaterialCachedHLSLTree.h"
 #endif
 
 /**
@@ -45,6 +46,15 @@
 
 UEnum* UMaterialInterface::SamplerTypeEnum = nullptr;
 
+static TAutoConsoleVariable<int32> CVarMaterialEnableNewHLSLGenerator(
+	TEXT("r.MaterialEnableNewHLSLGenerator"),
+	0,
+	TEXT("Enables the new (WIP) material HLSL generator.\n")
+	TEXT("0 - Don't allow\n")
+	TEXT("1 - Allow if enabled by material\n")
+	TEXT("2 - Force all materials to use new generator\n"),
+	ECVF_RenderThreadSafe | ECVF_ReadOnly);
+
 //////////////////////////////////////////////////////////////////////////
 
 bool IsHairStrandsGeometrySupported(const EShaderPlatform Platform)
@@ -91,6 +101,16 @@ void FMaterialRelevance::SetPrimitiveViewRelevance(FPrimitiveViewRelevance& OutV
 
 //////////////////////////////////////////////////////////////////////////
 
+UMaterialInterface::UMaterialInterface() = default;
+
+UMaterialInterface::UMaterialInterface(FVTableHelper& Helper)
+{
+}
+
+UMaterialInterface::~UMaterialInterface()
+{
+}
+
 UMaterialInterface::UMaterialInterface(const FObjectInitializer& ObjectInitializer)
 	: Super(ObjectInitializer)
 {
@@ -173,6 +193,31 @@ const FMaterialCachedExpressionData& UMaterialInterface::GetCachedExpressionData
 	return LocalData ? *LocalData : FMaterialCachedExpressionData::EmptyData;
 }
 
+#if WITH_EDITOR
+const FMaterialCachedHLSLTree& UMaterialInterface::GetCachedHLSLTree(TMicRecursionGuard) const
+{
+	check(IsUsingNewHLSLGenerator());
+	const FMaterialCachedHLSLTree* LocalTree = CachedHLSLTree.Get();
+	return LocalTree ? *LocalTree : FMaterialCachedHLSLTree::EmptyTree;
+}
+#endif // WITH_EDITOR
+
+bool UMaterialInterface::IsUsingNewHLSLGenerator() const
+{
+	const int CVarValue = CVarMaterialEnableNewHLSLGenerator.GetValueOnAnyThread();
+	if (CVarValue == 0)
+	{
+		return false;
+	}
+	else if (CVarValue == 2)
+	{
+		return true;
+	}
+
+	const UMaterial* BaseMaterial = GetMaterial_Concurrent();
+	return BaseMaterial ? BaseMaterial->bEnableNewHLSLGenerator : false;
+}
+
 void UMaterialInterface::GetQualityLevelUsage(TArray<bool, TInlineAllocator<EMaterialQualityLevel::Num> >& OutQualityLevelsUsed, EShaderPlatform ShaderPlatform, bool bCooking)
 {
 	OutQualityLevelsUsed = GetCachedExpressionData().QualityLevelsUsed;
diff --git a/Engine/Source/Runtime/Engine/Private/Materials/MaterialShared.cpp b/Engine/Source/Runtime/Engine/Private/Materials/MaterialShared.cpp
--- a/Engine/Source/Runtime/Engine/Private/Materials/MaterialShared.cpp
+++ b/Engine/Source/Runtime/Engine/Private/Materials/MaterialShared.cpp
@@ -299,7 +299,7 @@ UE::HLSLTree::FExpression* FExpressionInput::AcquireHLSLExpression(FMaterialHLSL
 	const FExpressionInput TracedInput = GetTracedInput();
 	if (!TracedInput.Expression)
 	{
-		Generator.GetErrors().AddError(TEXT("Missing input"));
+		Generator.Error(TEXT("Missing input"));
 		return nullptr;
 	}
 	return TryAcquireHLSLExpression(Generator, Scope);
diff --git a/Engine/Source/Runtime/Engine/Public/HLSLTree/HLSLTree.h b/Engine/Source/Runtime/Engine/Public/HLSLTree/HLSLTree.h
--- a/Engine/Source/Runtime/Engine/Public/HLSLTree/HLSLTree.h
+++ b/Engine/Source/Runtime/Engine/Public/HLSLTree/HLSLTree.h
@@ -116,7 +116,6 @@ class FNode
 {
 public:
 	virtual ~FNode() {}
-	virtual void Reset() {}
 
 	UObject* GetOwner() const { return Owner; }
 
@@ -338,6 +337,7 @@ private:
 	bool TryMergePreparedType(FEmitContext& Context, const Shader::FStructType* StructType, Shader::EValueComponentType ComponentType);
 
 	FPreparedType PreparedType;
+	bool bPreparingValue = false;
 
 	friend class FExpression;
 	friend class FEmitContext;
@@ -381,13 +381,6 @@ struct FExpressionDerivatives
 class FExpression : public FNode
 {
 public:
-	const FPreparedType& GetPreparedType() const { return PrepareValueResult.PreparedType; }
-	FRequestedType GetRequestedType() const { return PrepareValueResult.PreparedType.GetRequestedType(); }
-	Shader::FType GetType() const { return PrepareValueResult.PreparedType.GetType(); }
-	EExpressionEvaluation GetEvaluation(const FEmitScope& Scope, const FRequestedType& RequestedType) const { return PrepareValueResult.PreparedType.GetEvaluation(Scope, RequestedType); }
-
-	virtual void Reset() override;
-
 	FEmitShaderExpression* GetValueShader(FEmitContext& Context, FEmitScope& Scope, const FRequestedType& RequestedType);
 	Shader::FType GetValuePreshader(FEmitContext& Context, FEmitScope& Scope, const FRequestedType& RequestedType, Shader::FPreshaderData& OutPreshader);
 	Shader::FValue GetValueConstant(FEmitContext& Context, FEmitScope& Scope, const FRequestedType& RequestedType);
@@ -408,39 +401,8 @@ protected:
 	virtual void EmitValueShader(FEmitContext& Context, FEmitScope& Scope, const FRequestedType& RequestedType, FEmitValueShaderResult& OutResult) const;
 	virtual void EmitValuePreshader(FEmitContext& Context, FEmitScope& Scope, const FRequestedType& RequestedType, FEmitValuePreshaderResult& OutResult) const;
 
-private:
-	FExpressionDerivatives Derivatives;
-	FPrepareValueResult PrepareValueResult;
-	bool bReentryFlag = false;
-	bool bComputedDerivatives = false;
-
 	friend class FTree;
 	friend class FEmitContext;
-	friend class FExpressionReentryScope;
-};
-
-class FExpressionReentryScope
-{
-public:
-	FExpressionReentryScope(FExpression* InExpression) : Expression(InExpression)
-	{
-		if (Expression)
-		{
-			check(!Expression->bReentryFlag);
-			Expression->bReentryFlag = true;
-		}
-	}
-
-	~FExpressionReentryScope()
-	{
-		if (Expression)
-		{
-			check(Expression->bReentryFlag);
-			Expression->bReentryFlag = false;
-		}
-	}
-
-	FExpression* Expression;
 };
 
 class FFunction final : public FNode
@@ -499,8 +461,6 @@ public:
 
 	FMemStackBase& GetAllocator() { return *Allocator; }
 
-	void ResetNodes();
-
 	bool Finalize();
 
 	bool EmitShader(FEmitContext& Context, FStringBuilderBase& OutCode) const;
@@ -537,7 +497,7 @@ public:
 
 	FExpression* NewFunctionCall(FScope& Scope, FFunction* Function, int32 OutputIndex);
 
-	const FExpressionDerivatives& GetAnalyticDerivatives(FExpression* InExpression);
+	FExpressionDerivatives GetAnalyticDerivatives(FExpression* InExpression);
 	FExpression* GetPreviousFrame(FExpression* InExpression, const FRequestedType& RequestedType);
 
 	FScope* NewScope(FScope& Scope);
@@ -601,7 +561,6 @@ private:
 	FScope* RootScope = nullptr;
 	TMap<FXxHash64, FExpression*> ExpressionMap;
 	TArray<FExpressionLocalPHI*> PHIExpressions;
-	TArray<UObject*, TInlineAllocator<8>> OwnerStack;
 
 	friend class FExpressionLocalPHI;
 };
diff --git a/Engine/Source/Runtime/Engine/Public/HLSLTree/HLSLTreeCommon.h b/Engine/Source/Runtime/Engine/Public/HLSLTree/HLSLTreeCommon.h
--- a/Engine/Source/Runtime/Engine/Public/HLSLTree/HLSLTreeCommon.h
+++ b/Engine/Source/Runtime/Engine/Public/HLSLTree/HLSLTreeCommon.h
@@ -2,10 +2,23 @@
 #pragma once
 
 #include "HLSLTree/HLSLTree.h"
+#include "RHIDefinitions.h"
 
 namespace UE::HLSLTree
 {
 
+class FExpressionError : public FExpression
+{
+public:
+	explicit FExpressionError(FStringView InErrorMessage)
+		: ErrorMessage(InErrorMessage)
+	{}
+
+	virtual bool PrepareValue(FEmitContext& Context, FEmitScope& Scope, const FRequestedType& RequestedType, FPrepareValueResult& OutResult) const override;
+
+	FStringView ErrorMessage;
+};
+
 class FExpressionConstant : public FExpression
 {
 public:
@@ -178,6 +191,59 @@ public:
 	virtual void EmitValuePreshader(FEmitContext& Context, FEmitScope& Scope, const FRequestedType& RequestedType, FEmitValuePreshaderResult& OutResult) const override;
 };
 
+class FExpressionSwitchBase : public FExpression
+{
+public:
+	static constexpr int8 MaxInputs = 8;
+
+	FExpressionSwitchBase(TConstArrayView<FExpression*> InInputs) : NumInputs(InInputs.Num())
+	{
+		check(InInputs.Num() <= MaxInputs);
+		for (int32 i = 0; i < InInputs.Num(); ++i)
+		{
+			Input[i] = InInputs[i];
+		}
+	}
+
+	FExpression* Input[MaxInputs] = { nullptr };
+	int8 NumInputs = 0;
+
+	virtual FExpression* NewSwitch(FTree& Tree, TConstArrayView<FExpression*> InInputs) const = 0;
+	virtual int32 GetInputIndex(const FEmitContext& Context) const = 0;
+
+	virtual void ComputeAnalyticDerivatives(FTree& Tree, FExpressionDerivatives& OutResult) const override;
+	virtual FExpression* ComputePreviousFrame(FTree& Tree, const FRequestedType& RequestedType) const override;
+	virtual bool PrepareValue(FEmitContext& Context, FEmitScope& Scope, const FRequestedType& RequestedType, FPrepareValueResult& OutResult) const override;
+	virtual void EmitValueShader(FEmitContext& Context, FEmitScope& Scope, const FRequestedType& RequestedType, FEmitValueShaderResult& OutResult) const override;
+	virtual void EmitValuePreshader(FEmitContext& Context, FEmitScope& Scope, const FRequestedType& RequestedType, FEmitValuePreshaderResult& OutResult) const override;
+};
+
+class FExpressionFeatureLevelSwitch : public FExpressionSwitchBase
+{
+public:
+	static_assert(MaxInputs >= (int32)ERHIFeatureLevel::Num, "FExpressionSwitchBase is too small for FExpressionFeatureLevelSwitch");
+	FExpressionFeatureLevelSwitch(TConstArrayView<FExpression*> InInputs) : FExpressionSwitchBase(InInputs)
+	{
+		check(InInputs.Num() == (int32)ERHIFeatureLevel::Num);
+	}
+
+	virtual FExpression* NewSwitch(FTree& Tree, TConstArrayView<FExpression*> InInputs) const override { return Tree.NewExpression<FExpressionFeatureLevelSwitch>(InInputs); }
+	virtual int32 GetInputIndex(const FEmitContext& Context) const override;
+};
+
+class FExpressionShadingPathSwitch : public FExpressionSwitchBase
+{
+public:
+	static_assert(MaxInputs >= (int32)ERHIShadingPath::Num, "FExpressionSwitchBase is too small for FExpressionShadingPathSwitch");
+	FExpressionShadingPathSwitch(TConstArrayView<FExpression*> InInputs) : FExpressionSwitchBase(InInputs)
+	{
+		check(InInputs.Num() == (int32)ERHIShadingPath::Num);
+	}
+
+	virtual FExpression* NewSwitch(FTree& Tree, TConstArrayView<FExpression*> InInputs) const override { return Tree.NewExpression<FExpressionShadingPathSwitch>(InInputs); }
+	virtual int32 GetInputIndex(const FEmitContext& Context) const override;
+};
+
 /** Can be used to emit HLSL chunks with no inputs, where it's not worth the trouble of defining a new expression type */
 class FExpressionInlineCustomHLSL : public FExpression
 {
diff --git a/Engine/Source/Runtime/Engine/Public/HLSLTree/HLSLTreeEmit.h b/Engine/Source/Runtime/Engine/Public/HLSLTree/HLSLTreeEmit.h
--- a/Engine/Source/Runtime/Engine/Public/HLSLTree/HLSLTreeEmit.h
+++ b/Engine/Source/Runtime/Engine/Public/HLSLTree/HLSLTreeEmit.h
@@ -14,6 +14,7 @@
 
 class FMaterial;
 class FMaterialCompilationOutput;
+class ITargetPlatform;
 
 namespace UE::Shader
 {
@@ -30,6 +31,7 @@ class FExpression;
 class FFunction;
 class FRequestedType;
 class FPreparedType;
+class FPrepareValueResult;
 class FEmitScope;
 class FEmitShaderExpression;
 class FEmitShaderStatement;
@@ -243,11 +245,26 @@ enum class EEmitCastFlags : uint32
 };
 ENUM_CLASS_FLAGS(EEmitCastFlags)
 
+struct FTargetParameters
+{
+	FTargetParameters(EShaderPlatform InShaderPlatform, ERHIFeatureLevel::Type InFeatureLevel, const ITargetPlatform* InTargetPlatform)
+		: ShaderPlatform(InShaderPlatform), FeatureLevel(InFeatureLevel), TargetPlatform(InTargetPlatform)
+	{}
+
+	EShaderPlatform ShaderPlatform;
+	ERHIFeatureLevel::Type FeatureLevel;
+	const ITargetPlatform* TargetPlatform;
+};
+
 /** Tracks shared state while emitting HLSL code */
 class FEmitContext
 {
 public:
-	explicit FEmitContext(FMemStackBase& InAllocator, FErrorHandlerInterface& InErrors, const Shader::FStructTypeRegistry& InTypeRegistry);
+	explicit FEmitContext(FMemStackBase& InAllocator,
+		const FTargetParameters& InTargetParameters,
+		FErrorHandlerInterface& InErrors,
+		const Shader::FStructTypeRegistry& InTypeRegistry);
+
 	~FEmitContext();
 
 	template<typename T>
@@ -283,6 +300,11 @@ public:
 
 	void EmitDeclarationsCode(FStringBuilderBase& OutCode);
 
+	const FPreparedType& GetPreparedType(const FExpression* Expression) const;
+	FRequestedType GetRequestedType(const FExpression* Expression) const;
+	Shader::FType GetType(const FExpression* Expression) const;
+	EExpressionEvaluation GetEvaluation(const FExpression* Expression, const FEmitScope& Scope, const FRequestedType& RequestedType) const;
+
 	FPreparedType PrepareExpression(FExpression* InExpression, FEmitScope& Scope, const FRequestedType& RequestedType);
 
 	FEmitScope* InternalPrepareScope(FScope* Scope, FScope* ParentScope);
@@ -470,11 +492,13 @@ public:
 	FMemStackBase* Allocator = nullptr;
 	FErrorHandlerInterface* Errors = nullptr;
 	const Shader::FStructTypeRegistry* TypeRegistry = nullptr;
+	FTargetParameters TargetParameters;
 	EShaderFrequency ShaderFrequency = SF_Pixel;
 	bool bUseAnalyticDerivatives = false;
 
 	TArray<FEmitShaderNode*> EmitNodes;
 	TMap<const FScope*, FEmitScope*> EmitScopeMap;
+	TMap<const FExpression*, FPrepareValueResult*> PrepareValueMap;
 	TMap<const FExpression*, FEmitScope*> PrepareLocalPHIMap;
 	TMap<const FExpression*, FEmitShaderExpression*> EmitLocalPHIMap;
 	TMap<FXxHash64, FEmitShaderExpression*> EmitExpressionMap;
diff --git a/Engine/Source/Runtime/Engine/Public/MaterialCachedHLSLTree.h b/Engine/Source/Runtime/Engine/Public/MaterialCachedHLSLTree.h
new file mode 100644
--- /dev/null
+++ b/Engine/Source/Runtime/Engine/Public/MaterialCachedHLSLTree.h
@@ -0,0 +1,56 @@
+// Copyright Epic Games, Inc. All Rights Reserved.
+
+#pragma once
+
+#if WITH_EDITOR
+
+#include "HLSLTree/HLSLTree.h"
+#include "RHIDefinitions.h"
+#include "Misc/MemStack.h"
+
+class UMaterial;
+class UMaterialExpressionCustomOutput;
+struct FMaterialLayersFunctions;
+
+class FMaterialCachedHLSLTree
+{
+public:
+	static const FMaterialCachedHLSLTree EmptyTree;
+
+	FMaterialCachedHLSLTree();
+	~FMaterialCachedHLSLTree();
+
+	SIZE_T GetAllocatedSize() const;
+
+	bool GenerateTree(UMaterial* Material, const FMaterialLayersFunctions* LayerOverrides);
+
+	UE::Shader::FStructTypeRegistry& GetTypeRegistry() { return TypeRegistry; }
+	//UE::HLSLTree::FTree& GetTree() { return *HLSLTree; }
+
+	const UE::Shader::FStructTypeRegistry& GetTypeRegistry() const { return TypeRegistry; }
+	UE::HLSLTree::FTree& GetTree() const { return *HLSLTree; }
+
+	UE::HLSLTree::FExpression* GetResultExpression() const { return ResultExpression; }
+	UE::HLSLTree::FStatement* GetResultStatement() const { return ResultStatement; }
+
+	const UE::Shader::FStructType* GetMaterialAttributesType() const { return MaterialAttributesType; }
+	const UE::Shader::FValue& GetMaterialAttributesDefaultValue() const { return MaterialAttributesDefaultValue; }
+
+	void SetRequestedFields(EShaderFrequency ShaderFrequency, UE::HLSLTree::FRequestedType& OutRequestedType) const;
+	void EmitSharedCode(FStringBuilderBase& OutCode) const;
+
+private:
+	FMemStackBase Allocator;
+	UE::Shader::FStructTypeRegistry TypeRegistry;
+	UE::HLSLTree::FTree* HLSLTree = nullptr;
+	UE::HLSLTree::FExpression* ResultExpression = nullptr;
+	UE::HLSLTree::FStatement* ResultStatement = nullptr;
+
+	TArray<UMaterialExpressionCustomOutput*> MaterialCustomOutputs;
+	const UE::Shader::FStructType* MaterialAttributesType = nullptr;
+	UE::Shader::FValue MaterialAttributesDefaultValue;
+
+	friend class FMaterialHLSLGenerator;
+};
+
+#endif // WITH_EDITOR
diff --git a/Engine/Source/Runtime/Engine/Public/MaterialHLSLGenerator.h b/Engine/Source/Runtime/Engine/Public/MaterialHLSLGenerator.h
--- a/Engine/Source/Runtime/Engine/Public/MaterialHLSLGenerator.h
+++ b/Engine/Source/Runtime/Engine/Public/MaterialHLSLGenerator.h
@@ -9,11 +9,13 @@
 #include "Misc/MemStack.h"
 #include "Misc/EnumClassFlags.h"
 #include "Templates/RefCounting.h"
+#include "Templates/UniquePtr.h"
 #include "RHIDefinitions.h"
 #include "Materials/MaterialLayersFunctions.h"
 #include "HLSLTree/HLSLTree.h"
 
 class FMaterial;
+class FMaterialCachedHLSLTree;
 class FMaterialCompilationOutput;
 struct FSharedShaderCompilerEnvironment;
 struct FMaterialCompileTargetParameters;
@@ -53,44 +55,38 @@ struct TMaterialHLSLGeneratorType;
 #define DECLARE_MATERIAL_HLSLGENERATOR_DATA(T) \
 	template<> struct TMaterialHLSLGeneratorType<T> { static const FName& GetTypeName() { static const FName Name(TEXT(#T)); return Name; } }
 
-class FMaterialHLSLErrorHandler : public UE::HLSLTree::FErrorHandlerInterface
-{
-public:
-	explicit FMaterialHLSLErrorHandler(FMaterial& InOutMaterial);
-
-	virtual void AddErrorInternal(UObject* InOwner, FStringView InError) override;
-
-private:
-	FMaterial* Material;
-};
-
 /**
  * MaterialHLSLGenerator is a bridge between a material, and HLSLTree.  It facilitates generating HLSL source code for a given material, using HLSLTree
  */
 class FMaterialHLSLGenerator
 {
 public:
-	FMaterialHLSLGenerator(const FMaterialCompileTargetParameters& InCompilerTarget,
-		const FStaticParameterSet& InStaticParameters,
-		FMaterial& InOutMaterial,
-		UE::Shader::FStructTypeRegistry& InOutTypeRegistry,
-		UE::HLSLTree::FTree& InOutTree);
-
-	const FMaterialCompileTargetParameters& GetCompileTarget() const { return CompileTarget; }
-	const FStaticParameterSet& GetStaticParameters() const { return StaticParameters; }
-	
-	bool Generate();
+	FMaterialHLSLGenerator(UMaterial* Material,
+		const FMaterialLayersFunctions* InLayerOverrides,
+		FMaterialCachedHLSLTree& InTree);
 
-	UE::HLSLTree::FTree& GetTree() const { return *HLSLTree; }
-	UE::Shader::FStructTypeRegistry& GetTypeRegistry() const { return *TypeRegistry; }
-	FMaterialHLSLErrorHandler& GetErrors() { return Errors; }
+	const FMaterialLayersFunctions* GetLayerOverrides() const { return LayerOverrides; }
 
-	UE::HLSLTree::FExpression* GetResultExpression() { return ResultExpression; }
-	UE::HLSLTree::FStatement* GetResultStatement() { return ResultStatement; }
+	UE::HLSLTree::FTree& GetTree() const;
+	UE::Shader::FStructTypeRegistry& GetTypeRegistry() const;
+	const UE::Shader::FStructType* GetMaterialAttributesType() const;
+	const UE::Shader::FValue& GetMaterialAttributesDefaultValue() const;
 
-	void SetRequestedFields(EShaderFrequency ShaderFrequency, UE::HLSLTree::FRequestedType& OutRequestedType);
+	template<typename StringType>
+	inline bool Error(const StringType& InError)
+	{
+		return InternalError(FStringView(InError));
+	}
 
-	void EmitSharedCode(FStringBuilderBase& OutCode) const;
+	template<typename FormatType, typename... Types>
+	inline bool Errorf(const FormatType& Format, Types... Args)
+	{
+		TStringBuilder<2048> String;
+		String.Appendf(Format, Forward<Types>(Args)...);
+		return InternalError(FStringView(String.ToString(), String.Len()));
+	}
+	
+	bool Generate();
 
 	bool GenerateResult(UE::HLSLTree::FScope& Scope);
 
@@ -150,9 +146,6 @@ public:
 		return Data;
 	}
 
-	const UE::Shader::FStructType* GetMaterialAttributesType() const { return MaterialAttributesType; }
-	const UE::Shader::FValue& GetMaterialAttributesDefaultValue() const { return MaterialAttributesDefaultValue; }
-
 	bool GetParameterOverrideValueForCurrentFunction(EMaterialParameterType ParameterType, FName ParameterName, FMaterialParameterMetadata& OutResult) const;
 	FMaterialParameterInfo GetParameterInfo(const FName& ParameterName) const;
 
@@ -199,27 +192,21 @@ private:
 		int32 NumInputs = 0;
 	};
 
+	bool InternalError(FStringView ErrorMessage);
+
 	void InternalRegisterExpressionData(const FName& Type, const UMaterialExpression* MaterialExpression, void* Data);
 	void* InternalFindExpressionData(const FName& Type, const UMaterialExpression* MaterialExpression);
 
-	const FMaterialCompileTargetParameters& CompileTarget;
-	const FStaticParameterSet& StaticParameters;
-	UMaterial* TargetMaterial = nullptr;
-	FMaterialHLSLErrorHandler Errors;
-
-	TArray<UMaterialExpressionCustomOutput*> MaterialCustomOutputs;
-	const UE::Shader::FStructType* MaterialAttributesType = nullptr;
-	UE::Shader::FValue MaterialAttributesDefaultValue;
-
-	UE::HLSLTree::FTree* HLSLTree;
-	UE::Shader::FStructTypeRegistry* TypeRegistry = nullptr;
-	UE::HLSLTree::FExpression* ResultExpression = nullptr;
-	UE::HLSLTree::FStatement* ResultStatement = nullptr;
+	UMaterial* TargetMaterial;
+	const FMaterialLayersFunctions* LayerOverrides;
+	FMaterialCachedHLSLTree& CachedTree;
+	FString CurrentErrorMessage;
 
+	FFunctionCallEntry RootFunctionCallEntry;
 	TArray<FFunctionCallEntry*, TInlineAllocator<8>> FunctionCallStack;
 	TArray<UE::HLSLTree::FScope*> JoinedScopeStack;
 	TMap<FXxHash64, const UE::Shader::FTextureValue*> TextureValueMap;
-	TMap<FSHAHash, FFunctionCallEntry*> FunctionCallMap;
+	TMap<FXxHash64, TUniquePtr<FFunctionCallEntry>> FunctionCallMap;
 	TMap<UMaterialFunctionInterface*, UE::HLSLTree::FFunction*> FunctionMap;
 	TMap<UMaterialExpression*, FStatementEntry> StatementMap;
 	TMap<FExpressionDataKey, void*> ExpressionDataMap;
