[Backout] - CL19173027

#rnx

diff --git a/Engine/Source/Runtime/Engine/Private/InstancedStaticMesh.cpp b/Engine/Source/Runtime/Engine/Private/InstancedStaticMesh.cpp
--- a/Engine/Source/Runtime/Engine/Private/InstancedStaticMesh.cpp
+++ b/Engine/Source/Runtime/Engine/Private/InstancedStaticMesh.cpp
@@ -1398,26 +1398,22 @@ void FInstancedStaticMeshSceneProxy::SetupProxy(UInstancedStaticMeshComponent* I
 	{
 		const TArray<int32>& InstanceReorderTable = InComponent->InstanceReorderTable;
 
-		// NumRenderInstances is the extent that the reorder table can map to.
-		// Temporarily when removing instances from a HISM this can be sparse so that InComponent->GetInstanceCount() < NumRenderInstances.
-		const int32 NumRenderInstances = FMath::Max<int32>(InComponent->PerInstanceRenderData->InstanceBuffer.GetNumInstances(), InComponent->GetInstanceCount());
-
 		bSupportsInstanceDataBuffer = true;
-		InstanceSceneData.SetNumZeroed(NumRenderInstances);
+		InstanceSceneData.SetNum(InComponent->GetInstanceCount());
 
 		bHasPerInstanceDynamicData = InComponent->PerInstancePrevTransform.Num() == InComponent->GetInstanceCount();
-		InstanceDynamicData.SetNumUninitialized(bHasPerInstanceDynamicData ? NumRenderInstances : 0);
+		InstanceDynamicData.SetNumUninitialized(bHasPerInstanceDynamicData ? InComponent->GetInstanceCount() : 0);
 
 		static const auto AllowStaticLightingVar = IConsoleManager::Get().FindTConsoleVariableDataInt(TEXT("r.AllowStaticLighting"));
 		const bool bAllowStaticLighting = (!AllowStaticLightingVar || AllowStaticLightingVar->GetValueOnAnyThread() != 0);
 		bHasPerInstanceLMSMUVBias = bAllowStaticLighting;
-		InstanceLightShadowUVBias.SetNumZeroed(bHasPerInstanceLMSMUVBias ? NumRenderInstances : 0);
+		InstanceLightShadowUVBias.SetNumZeroed(bHasPerInstanceLMSMUVBias ? InComponent->GetInstanceCount() : 0);
 
-		InstanceRandomID.SetNumZeroed(bHasPerInstanceRandom ? NumRenderInstances : 0); // Only allocate if material bound which uses this
+		InstanceRandomID.SetNumZeroed(bHasPerInstanceRandom ? InComponent->GetInstanceCount() : 0); // Only allocate if material bound which uses this
 
 #if WITH_EDITOR
 		bHasPerInstanceEditorData = true;
-		InstanceEditorData.SetNumZeroed(bHasPerInstanceEditorData ? NumRenderInstances : 0);
+		InstanceEditorData.SetNumZeroed(bHasPerInstanceEditorData ? InComponent->GetInstanceCount() : 0);
 #endif
 
 		// Only allocate if material bound which uses this
@@ -1431,38 +1427,43 @@ void FInstancedStaticMeshSceneProxy::SetupProxy(UInstancedStaticMeshComponent* I
 			bHasPerInstanceCustomData = false;
 		}
 
-		FVector TranslatedSpaceOffset = -InComponent->GetTranslatedInstanceSpaceOrigin();
-
-		// Add the visible instances. 
-		// Non visible ones (that are being removed) should not appear due to a zeroed render transform in the entries where we don't add anything here.
-		for (int32 InstanceIndex = 0; InstanceIndex < InComponent->GetInstanceCount(); ++InstanceIndex)
+		for (int32 InInstanceIndex = 0; InInstanceIndex < InstanceSceneData.Num(); ++InInstanceIndex)
 		{
-			const int32 RenderInstanceIndex = InstanceReorderTable[InstanceIndex];
-			if (!ensure(RenderInstanceIndex < NumRenderInstances))
+			// Make sure the instance is initialized, regardless of below remapping
+			{
+				FPrimitiveInstance& TmpSceneData = InstanceSceneData[InInstanceIndex];
+				TmpSceneData.LocalToPrimitive = FRenderTransform::Identity;
+			}
+
+			int32 OutInstanceIndex = InInstanceIndex;
+			// GPUCULL_TODO: After deleting instances in a HISM the InstanceReorderTable often contains nonsense, this is then corrected
+			// by the async build, which re-creates the proxy in a nearby future frame. All this should be removed in favour of GPU-side culling.
+			if (OutInstanceIndex < InstanceReorderTable.Num() && InstanceReorderTable[OutInstanceIndex] < InstanceSceneData.Num())
 			{
-				continue;
+				// Temporary workaround for out of bound array access
+				// TODO: fix this properly!!!!!!
+				OutInstanceIndex = InstanceReorderTable[OutInstanceIndex] != INDEX_NONE ? InstanceReorderTable[OutInstanceIndex] : OutInstanceIndex;
 			}
 
-			FPrimitiveInstance& SceneData = InstanceSceneData[RenderInstanceIndex];
+			FPrimitiveInstance& SceneData = InstanceSceneData[OutInstanceIndex];
+
 
 			FTransform InstanceTransform;
-			InComponent->GetInstanceTransform(InstanceIndex, InstanceTransform);
-			InstanceTransform.AddToTranslation(TranslatedSpaceOffset);
+			InComponent->GetInstanceTransform(InInstanceIndex, InstanceTransform);
 			SceneData.LocalToPrimitive = InstanceTransform.ToMatrixWithScale();
 
 			if (bHasPerInstanceDynamicData)
 			{
 				FTransform InstancePrevTransform;
-				const bool bHasPrevTransform = InComponent->GetInstancePrevTransform(InstanceIndex, InstancePrevTransform);
+				const bool bHasPrevTransform = InComponent->GetInstancePrevTransform(InInstanceIndex, InstancePrevTransform);
 				ensure(bHasPrevTransform); // Should always be true here
-				InstancePrevTransform.AddToTranslation(TranslatedSpaceOffset);
-				InstanceDynamicData[RenderInstanceIndex].PrevLocalToPrimitive = InstancePrevTransform.ToMatrixWithScale();
+				InstanceDynamicData[OutInstanceIndex].PrevLocalToPrimitive = InstancePrevTransform.ToMatrixWithScale();
 			}
 
 			if (bHasPerInstanceCustomData)
 			{
-				const int32 SrcCustomDataOffset = InstanceIndex  * InComponent->NumCustomDataFloats;
-				const int32 DstCustomDataOffset = RenderInstanceIndex * InComponent->NumCustomDataFloats;
+				const int32 SrcCustomDataOffset = InInstanceIndex  * InComponent->NumCustomDataFloats;
+				const int32 DstCustomDataOffset = OutInstanceIndex * InComponent->NumCustomDataFloats;
 				FMemory::Memcpy
 				(
 					&InstanceCustomData[DstCustomDataOffset],
