* Added bTwoSided and bUsesPixelDepthOffset on FMaterialRelevance
* Define USES_WORLD_POSITION_OFFSET in material shaders on all platforms (not just Metal)
* Deprecated UsesWorldPositionOffset_GameThread() and implemented MaterialUsesWorldPositionOffset_[Game/Render]Thread()
* Deprecated MaterialUsesPixelDepthOffset() and implemented MaterialUsesPixelDepthOffset_[Game/Render]Thread()
* Fixed a number of typos.

diff --git a/Engine/Source/Runtime/Engine/Classes/Materials/MaterialRelevance.h b/Engine/Source/Runtime/Engine/Classes/Materials/MaterialRelevance.h
--- a/Engine/Source/Runtime/Engine/Classes/Materials/MaterialRelevance.h
+++ b/Engine/Source/Runtime/Engine/Classes/Materials/MaterialRelevance.h
@@ -21,6 +21,7 @@ struct ENGINE_API FMaterialRelevance
 			uint8 bMasked : 1;
 			uint8 bDistortion : 1;
 			uint8 bHairStrands : 1;
+			uint8 bTwoSided : 1;
 			uint8 bSeparateTranslucency : 1; // Translucency After DOF
 			uint8 bSeparateTranslucencyModulate : 1;
 			uint8 bPostMotionBlurTranslucency : 1;
@@ -29,6 +30,7 @@ struct ENGINE_API FMaterialRelevance
 			uint8 bOutputsTranslucentVelocity : 1;
 			uint8 bUsesGlobalDistanceField : 1;
 			uint8 bUsesWorldPositionOffset : 1;
+			uint8 bUsesPixelDepthOffset : 1;
 			uint8 bDecal : 1;
 			uint8 bTranslucentSurfaceLighting : 1;
 			uint8 bUsesSceneDepth : 1;
diff --git a/Engine/Source/Runtime/Engine/Private/Materials/HLSLMaterialTranslator.cpp b/Engine/Source/Runtime/Engine/Private/Materials/HLSLMaterialTranslator.cpp
--- a/Engine/Source/Runtime/Engine/Private/Materials/HLSLMaterialTranslator.cpp
+++ b/Engine/Source/Runtime/Engine/Private/Materials/HLSLMaterialTranslator.cpp
@@ -1663,10 +1663,7 @@ void FHLSLMaterialTranslator::GetMaterialEnvironment(EShaderPlatform InPlatform,
 	OutEnvironment.SetDefine(TEXT("NEEDS_INSTANCE_WORLD_TO_LOCAL_PS"), bUsesInstanceWorldToLocalPS);
 	OutEnvironment.SetDefine(TEXT("USES_TRANSFORM_VECTOR"), bUsesTransformVector);
 	OutEnvironment.SetDefine(TEXT("WANT_PIXEL_DEPTH_OFFSET"), bUsesPixelDepthOffset);
-	if (IsMetalPlatform(InPlatform))
-	{
-		OutEnvironment.SetDefine(TEXT("USES_WORLD_POSITION_OFFSET"), bUsesWorldPositionOffset);
-	}
+	OutEnvironment.SetDefine(TEXT("USES_WORLD_POSITION_OFFSET"), bUsesWorldPositionOffset);
 	OutEnvironment.SetDefine(TEXT("USES_EMISSIVE_COLOR"), bUsesEmissiveColor);
 	// Distortion uses tangent space transform 
 	OutEnvironment.SetDefine(TEXT("USES_DISTORTION"), Material->IsDistorted()); 
diff --git a/Engine/Source/Runtime/Engine/Private/Materials/MaterialHLSLEmitter.cpp b/Engine/Source/Runtime/Engine/Private/Materials/MaterialHLSLEmitter.cpp
--- a/Engine/Source/Runtime/Engine/Private/Materials/MaterialHLSLEmitter.cpp
+++ b/Engine/Source/Runtime/Engine/Private/Materials/MaterialHLSLEmitter.cpp
@@ -509,10 +509,7 @@ static void GetMaterialEnvironment(EShaderPlatform InPlatform,
 	OutEnvironment.SetDefine(TEXT("NEEDS_PARTICLE_WORLD_TO_LOCAL"), false);// bUsesParticleWorldToLocal);
 	OutEnvironment.SetDefine(TEXT("USES_TRANSFORM_VECTOR"), false);// bUsesTransformVector);
 	OutEnvironment.SetDefine(TEXT("WANT_PIXEL_DEPTH_OFFSET"), false);// bUsesPixelDepthOffset);
-	if (IsMetalPlatform(InPlatform))
-	{
-		OutEnvironment.SetDefine(TEXT("USES_WORLD_POSITION_OFFSET"), false);// bUsesWorldPositionOffset);
-	}
+	OutEnvironment.SetDefine(TEXT("USES_WORLD_POSITION_OFFSET"), false);// bUsesWorldPositionOffset);
 	OutEnvironment.SetDefine(TEXT("USES_EMISSIVE_COLOR"), false);// bUsesEmissiveColor);
 	// Distortion uses tangent space transform 
 	OutEnvironment.SetDefine(TEXT("USES_DISTORTION"), InMaterial.IsDistorted());
diff --git a/Engine/Source/Runtime/Engine/Private/Materials/MaterialInterface.cpp b/Engine/Source/Runtime/Engine/Private/Materials/MaterialInterface.cpp
--- a/Engine/Source/Runtime/Engine/Private/Materials/MaterialInterface.cpp
+++ b/Engine/Source/Runtime/Engine/Private/Materials/MaterialInterface.cpp
@@ -374,6 +374,7 @@ FMaterialRelevance UMaterialInterface::GetRelevance_Internal(const UMaterial* Ma
 			MaterialRelevance.bMasked = IsMasked();
 			MaterialRelevance.bDistortion = MaterialResource->IsDistorted();
 			MaterialRelevance.bHairStrands = IsCompatibleWithHairStrands(MaterialResource, InFeatureLevel);
+			MaterialRelevance.bTwoSided = MaterialResource->IsTwoSided();
 			MaterialRelevance.bSeparateTranslucency = (TranslucencyPass == MTP_AfterDOF) && StrataBlendMode != SBM_ColoredTransmittanceOnly;
 			MaterialRelevance.bSeparateTranslucencyModulate = bMaterialSeparateModulation || StrataBlendMode == SBM_ColoredTransmittanceOnly;
 			MaterialRelevance.bPostMotionBlurTranslucency = (TranslucencyPass == MTP_AfterMotionBlur);
@@ -382,7 +383,8 @@ FMaterialRelevance UMaterialInterface::GetRelevance_Internal(const UMaterial* Ma
 			MaterialRelevance.bUsesSceneColorCopy = bIsTranslucent && MaterialResource->RequiresSceneColorCopy_GameThread();
 			MaterialRelevance.bOutputsTranslucentVelocity = Material->IsTranslucencyWritingVelocity();
 			MaterialRelevance.bUsesGlobalDistanceField = MaterialResource->UsesGlobalDistanceField_GameThread();
-			MaterialRelevance.bUsesWorldPositionOffset = MaterialResource->UsesWorldPositionOffset_GameThread();
+			MaterialRelevance.bUsesWorldPositionOffset = MaterialResource->MaterialUsesWorldPositionOffset_GameThread();
+			MaterialRelevance.bUsesPixelDepthOffset = MaterialResource->MaterialUsesPixelDepthOffset_GameThread();
 			ETranslucencyLightingMode TranslucencyLightingMode = MaterialResource->GetTranslucencyLightingMode();
 			MaterialRelevance.bTranslucentSurfaceLighting = bIsTranslucent && (TranslucencyLightingMode == TLM_SurfacePerPixelLighting || TranslucencyLightingMode == TLM_Surface);
 			MaterialRelevance.bUsesSceneDepth = MaterialResource->MaterialUsesSceneDepthLookup_GameThread();
diff --git a/Engine/Source/Runtime/Engine/Private/Materials/MaterialShared.cpp b/Engine/Source/Runtime/Engine/Private/Materials/MaterialShared.cpp
--- a/Engine/Source/Runtime/Engine/Private/Materials/MaterialShared.cpp
+++ b/Engine/Source/Runtime/Engine/Private/Materials/MaterialShared.cpp
@@ -885,7 +885,13 @@ bool FMaterial::UsesGlobalDistanceField_GameThread() const
 	return GameThreadShaderMap.GetReference() ? GameThreadShaderMap->UsesGlobalDistanceField() : false; 
 }
 
-bool FMaterial::UsesWorldPositionOffset_GameThread() const 
+bool FMaterial::MaterialUsesWorldPositionOffset_RenderThread() const
+{
+	check(IsInParallelRenderingThread());
+	return RenderingThreadShaderMap ? RenderingThreadShaderMap->UsesWorldPositionOffset() : false;
+}
+
+bool FMaterial::MaterialUsesWorldPositionOffset_GameThread() const
 { 
 	return GameThreadShaderMap.GetReference() ? GameThreadShaderMap->UsesWorldPositionOffset() : false; 
 }
@@ -893,18 +899,14 @@ bool FMaterial::UsesWorldPositionOffset_GameThread() const
 bool FMaterial::MaterialModifiesMeshPosition_RenderThread() const
 { 
 	check(IsInParallelRenderingThread());
-	bool bUsesWPO = RenderingThreadShaderMap ? RenderingThreadShaderMap->ModifiesMeshPosition() : false;
-
-	return bUsesWPO;
+	return RenderingThreadShaderMap ? RenderingThreadShaderMap->ModifiesMeshPosition() : false;
 }
 
 bool FMaterial::MaterialModifiesMeshPosition_GameThread() const
 {
 	check(IsInParallelGameThread() || IsInGameThread());
 	FMaterialShaderMap* ShaderMap = GameThreadShaderMap.GetReference();
-	bool bUsesWPO = ShaderMap ? ShaderMap->ModifiesMeshPosition() : false;
-
-	return bUsesWPO;
+	return ShaderMap ? ShaderMap->ModifiesMeshPosition() : false;
 }
 
 bool FMaterial::MaterialMayModifyMeshPosition() const
@@ -914,7 +916,12 @@ bool FMaterial::MaterialMayModifyMeshPosition() const
 	return HasVertexPositionOffsetConnected() || HasPixelDepthOffsetConnected();
 }
 
-bool FMaterial::MaterialUsesPixelDepthOffset() const
+bool FMaterial::MaterialUsesPixelDepthOffset_GameThread() const
+{
+	return GameThreadShaderMap.GetReference() ? GameThreadShaderMap->UsesPixelDepthOffset() : false;
+}
+
+bool FMaterial::MaterialUsesPixelDepthOffset_RenderThread() const
 {
 	check(IsInParallelRenderingThread());
 	return RenderingThreadShaderMap ? RenderingThreadShaderMap->UsesPixelDepthOffset() : false;
diff --git a/Engine/Source/Runtime/Engine/Public/MaterialShared.h b/Engine/Source/Runtime/Engine/Public/MaterialShared.h
--- a/Engine/Source/Runtime/Engine/Public/MaterialShared.h
+++ b/Engine/Source/Runtime/Engine/Public/MaterialShared.h
@@ -714,7 +714,7 @@ public:
 	/** Bitfield of runtime virtual texture output attributes. */
 	LAYOUT_FIELD(uint8, RuntimeVirtualTextureOutputAttributeMask);
 
-	/** true if the material needs the scenetexture lookups. */
+	/** true if the material needs the scene texture lookups. */
 	LAYOUT_BITFIELD(uint8, bNeedsSceneTextures, 1);
 
 	/** true if the material uses the EyeAdaptationLookup */
@@ -1938,15 +1938,31 @@ public:
 	ENGINE_API bool NeedsGBuffer() const;
 	ENGINE_API bool UsesEyeAdaptation() const;	
 	ENGINE_API bool UsesGlobalDistanceField_GameThread() const;
-	ENGINE_API bool UsesWorldPositionOffset_GameThread() const;
+
+	/** Does the material use world position offset. */
+	ENGINE_API bool MaterialUsesWorldPositionOffset_RenderThread() const;
+	ENGINE_API bool MaterialUsesWorldPositionOffset_GameThread() const;
+
+	UE_DEPRECATED(5.0, "This function is deprecated. Use MaterialUsesWorldPositionOffset_GameThread() instead.")
+	inline bool UsesWorldPositionOffset_GameThread() const
+	{
+		return MaterialUsesWorldPositionOffset_GameThread();
+	}
+
+	/** Does the material use a pixel depth offset. */
+	ENGINE_API bool MaterialUsesPixelDepthOffset_RenderThread() const;
+	ENGINE_API bool MaterialUsesPixelDepthOffset_GameThread() const;
+
+	UE_DEPRECATED(5.0, "This function is deprecated. Use MaterialUsesPixelDepthOffset_RenderThread() instead.")
+	inline bool MaterialUsesPixelDepthOffset() const
+	{
+		return MaterialUsesPixelDepthOffset_RenderThread();
+	}
 
 	/** Does the material modify the mesh position. */
 	ENGINE_API bool MaterialModifiesMeshPosition_RenderThread() const;
 	ENGINE_API bool MaterialModifiesMeshPosition_GameThread() const;
 
-	/** Does the material use a pixel depth offset. */
-	ENGINE_API bool MaterialUsesPixelDepthOffset() const;
-
 	/** Does the material use a distance cull fade. */
 	ENGINE_API bool MaterialUsesDistanceCullFade_GameThread() const;
 
@@ -2141,7 +2157,7 @@ protected:
 	virtual FString GetBaseMaterialPathName() const { return TEXT(""); }
 	virtual FString GetDebugName() const { return GetBaseMaterialPathName(); }
 
-	UE_DEPRECATED(4.26, "Parameter bInHasQualityLevelUsage is depreceated")
+	UE_DEPRECATED(4.26, "Parameter bInHasQualityLevelUsage is deprecated")
 	void SetQualityLevelProperties(EMaterialQualityLevel::Type InQualityLevel, bool bInHasQualityLevelUsage, ERHIFeatureLevel::Type InFeatureLevel)
 	{
 		QualityLevel = InQualityLevel;
@@ -2419,8 +2435,8 @@ public:
 
 	/**
 	 * Finds the FMaterial to use for rendering this FMaterialRenderProxy.  Will fall back to a default material if needed due to a content error, or async compilation.
-	 * The returned FMaterial is guaranteed to have a complete shader map, so all relevant shaders should be availiable
-	 * OutFallbackMaterialRenderProxy - The proxy that coorisponds to the returned FMaterial, should be used for further rendering.  May be a fallback material, or 'this' if no fallback was needed
+	 * The returned FMaterial is guaranteed to have a complete shader map, so all relevant shaders should be available
+	 * OutFallbackMaterialRenderProxy - The proxy that corresponds to the returned FMaterial, should be used for further rendering.  May be a fallback material, or 'this' if no fallback was needed
 	 */
 	const FMaterial& GetMaterialWithFallback(ERHIFeatureLevel::Type InFeatureLevel, const FMaterialRenderProxy*& OutFallbackMaterialRenderProxy) const;
 
@@ -2640,7 +2656,7 @@ public:
 		SetQualityLevelProperties(InFeatureLevel, InQualityLevel);
 	}
 
-	UE_DEPRECATED(4.26, "Parameter bInHasQualityLevelUsage is depreceated")
+	UE_DEPRECATED(4.26, "Parameter bInHasQualityLevelUsage is deprecated")
 	void SetMaterial(UMaterial* InMaterial, EMaterialQualityLevel::Type InQualityLevel, bool bInHasQualityLevelUsage, ERHIFeatureLevel::Type InFeatureLevel, UMaterialInstance* InInstance = nullptr)
 	{
 		SetMaterial(InMaterial, InInstance, InFeatureLevel, InQualityLevel);
diff --git a/Engine/Source/Runtime/Renderer/Private/DepthRendering.cpp b/Engine/Source/Runtime/Renderer/Private/DepthRendering.cpp
--- a/Engine/Source/Runtime/Renderer/Private/DepthRendering.cpp
+++ b/Engine/Source/Runtime/Renderer/Private/DepthRendering.cpp
@@ -177,7 +177,7 @@ bool GetDepthPassShaders(
 	}
 	else
 	{
-		const bool bNeedsPixelShader = bUsesMobileColorValue || !Material.WritesEveryPixel() || Material.MaterialUsesPixelDepthOffset() || Material.IsTranslucencyWritingCustomDepth();
+		const bool bNeedsPixelShader = bUsesMobileColorValue || !Material.WritesEveryPixel() || Material.MaterialUsesPixelDepthOffset_RenderThread() || Material.IsTranslucencyWritingCustomDepth();
 		if (bNeedsPixelShader)
 		{
 			ShaderTypes.AddShaderType<FDepthOnlyPS<bUsesMobileColorValue>>();
