Bring back lost changes in 19121149 plus added corresponding cvars to be able to reenable them of specific projects.

diff --git a/Engine/Source/Runtime/Engine/Private/WorldPartition/RuntimeSpatialHash/RuntimeSpatialHashGridHelper.cpp b/Engine/Source/Runtime/Engine/Private/WorldPartition/RuntimeSpatialHash/RuntimeSpatialHashGridHelper.cpp
--- a/Engine/Source/Runtime/Engine/Private/WorldPartition/RuntimeSpatialHash/RuntimeSpatialHashGridHelper.cpp
+++ b/Engine/Source/Runtime/Engine/Private/WorldPartition/RuntimeSpatialHash/RuntimeSpatialHashGridHelper.cpp
@@ -5,11 +5,23 @@
 #include "ProfilingDebugging/ScopedTimers.h"
 
 bool GRuntimeSpatialHashUseAlignedGridLevels = true;
-static FAutoConsoleVariableRef CVarUseGrow(
+static FAutoConsoleVariableRef CVarRuntimeSpatialHashUseAlignedGridLevels(
 	TEXT("wp.Runtime.RuntimeSpatialHashUseAlignedGridLevels"),
 	GRuntimeSpatialHashUseAlignedGridLevels,
 	TEXT("Set RuntimeSpatialHashUseAlignedGridLevels to false to help break the pattern caused by world partition promotion of actors to upper grid levels that are always aligned on child levels."));
 
+bool GRuntimeSpatialHashSnapNonAlignedGridLevelsToLowerLevels = true;
+static FAutoConsoleVariableRef CVarRuntimeSpatialHashSnapNonAlignedGridLevelsToLowerLevels(
+	TEXT("wp.Runtime.RuntimeSpatialHashSnapNonAlignedGridLevelsToLowerLevels"),
+	GRuntimeSpatialHashSnapNonAlignedGridLevelsToLowerLevels,
+	TEXT("Set RuntimeSpatialHashSnapNonAlignedGridLevelsToLowerLevels to false to avoid snapping higher levels cells to child cells. Only used when GRuntimeSpatialHashUseAlignedGridLevels is false."));
+
+bool GRuntimeSpatialHashPlaceSmallActorsUsingLocation = false;
+static FAutoConsoleVariableRef CVarRuntimeSpatialHashPlaceSmallActorsUsingLocation(
+	TEXT("wp.Runtime.RuntimeSpatialHashPlaceSmallActorsUsingLocation"),
+	GRuntimeSpatialHashPlaceSmallActorsUsingLocation,
+	TEXT("Set RuntimeSpatialHashPlaceSmallActorsUsingLocation to true to place actors smaller than a cell size into their corresponding cell using their location instead of their bounding box."));
+
 FSquare2DGridHelper::FSquare2DGridHelper(const FBox& InWorldBounds, const FVector& InOrigin, int32 InCellSize)
 	: WorldBounds(InWorldBounds)
 	, Origin(InOrigin)
@@ -171,6 +183,7 @@ FSquare2DGridHelper GetPartitionedActors(const UWorldPartition* WorldPartition,
 		}
 	}
 
+	const float CellArea = PartitionedActors.GetLowestLevel().CellSize * PartitionedActors.GetLowestLevel().CellSize;
 	for (const FActorClusterInstance* ClusterInstance : GridActors)
 	{
 		const FActorCluster* ActorCluster = ClusterInstance->Cluster;
@@ -180,24 +193,36 @@ FSquare2DGridHelper GetPartitionedActors(const UWorldPartition* WorldPartition,
 
 		if (ActorCluster->bIsSpatiallyLoaded)
 		{
-			// Find grid level cell that encompasses the actor cluster bounding box and put actors in it.
-			const FVector2D ClusterSize = ClusterBounds.GetSize();
-			const float MinRequiredCellExtent = FMath::Max(ClusterSize.X, ClusterSize.Y);
-			const int32 FirstPotentialGridLevel = FMath::Max(FMath::CeilToFloat(FMath::Log2(MinRequiredCellExtent / (float)PartitionedActors.CellSize)), 0);
-
-			for (int32 GridLevelIndex = FirstPotentialGridLevel; GridLevelIndex < PartitionedActors.Levels.Num(); GridLevelIndex++)
+			if (GRuntimeSpatialHashPlaceSmallActorsUsingLocation && (ClusterBounds.GetArea() <= CellArea))
 			{
-				FSquare2DGridHelper::FGridLevel& GridLevel = PartitionedActors.Levels[GridLevelIndex];
+				// Find grid level cell that contains the actor cluster pivot and put actors in it.
+				FIntVector2 CellCoords;
+				if (PartitionedActors.GetLowestLevel().GetCellCoords(ClusterBounds.GetCenter(), CellCoords))
+				{
+					GridCell = &PartitionedActors.GetLowestLevel().GetCell(CellCoords);
+				}
+			}
+			else
+			{
+				// Find grid level cell that encompasses the actor cluster bounding box and put actors in it.
+				const FVector2D ClusterSize = ClusterBounds.GetSize();
+				const float MinRequiredCellExtent = FMath::Max(ClusterSize.X, ClusterSize.Y);
+				const int32 FirstPotentialGridLevel = FMath::Max(FMath::CeilToFloat(FMath::Log2(MinRequiredCellExtent / (float)PartitionedActors.CellSize)), 0);
 
-				if (GridLevel.GetNumIntersectingCells(ClusterInstance->Bounds) == 1)
+				for (int32 GridLevelIndex = FirstPotentialGridLevel; GridLevelIndex < PartitionedActors.Levels.Num(); GridLevelIndex++)
 				{
-					GridLevel.ForEachIntersectingCells(ClusterInstance->Bounds, [&GridLevel, ActorCluster, ClusterInstance, &GridCell](const FIntVector2& Coords)
-					{
-						check(!GridCell);
-						GridCell = &GridLevel.GetCell(Coords);
-					});
+					FSquare2DGridHelper::FGridLevel& GridLevel = PartitionedActors.Levels[GridLevelIndex];
 
-					break;
+					if (GridLevel.GetNumIntersectingCells(ClusterInstance->Bounds) == 1)
+					{
+						GridLevel.ForEachIntersectingCells(ClusterInstance->Bounds, [&GridLevel, ActorCluster, ClusterInstance, &GridCell](const FIntVector2& Coords)
+						{
+							check(!GridCell);
+							GridCell = &GridLevel.GetCell(Coords);
+						});
+
+						break;
+					}
 				}
 			}
 		}
diff --git a/Engine/Source/Runtime/Engine/Private/WorldPartition/RuntimeSpatialHash/RuntimeSpatialHashGridHelper.h b/Engine/Source/Runtime/Engine/Private/WorldPartition/RuntimeSpatialHash/RuntimeSpatialHashGridHelper.h
--- a/Engine/Source/Runtime/Engine/Private/WorldPartition/RuntimeSpatialHash/RuntimeSpatialHashGridHelper.h
+++ b/Engine/Source/Runtime/Engine/Private/WorldPartition/RuntimeSpatialHash/RuntimeSpatialHashGridHelper.h
@@ -11,6 +11,7 @@
 #include "WorldPartition/DataLayer/DataLayersID.h"
 
 extern ENGINE_API bool GRuntimeSpatialHashUseAlignedGridLevels;
+extern ENGINE_API bool GRuntimeSpatialHashSnapNonAlignedGridLevelsToLowerLevels;
 
 /**
   * Square 2D grid helper
diff --git a/Engine/Source/Runtime/Engine/Private/WorldPartition/WorldPartitionRuntimeSpatialHash.cpp b/Engine/Source/Runtime/Engine/Private/WorldPartition/WorldPartitionRuntimeSpatialHash.cpp
--- a/Engine/Source/Runtime/Engine/Private/WorldPartition/WorldPartitionRuntimeSpatialHash.cpp
+++ b/Engine/Source/Runtime/Engine/Private/WorldPartition/WorldPartitionRuntimeSpatialHash.cpp
@@ -223,7 +223,7 @@ void FSpatialHashStreamingGrid::GetCells(const TArray<FWorldPartitionStreamingSo
 								{
 									check(Source.TargetState == EStreamingSourceTargetState::Activated);
 									OutActivateCells.AddCell(Cell, Info);
-									bAddedActivatedCell = (!GRuntimeSpatialHashUseAlignedGridLevels);
+									bAddedActivatedCell = !GRuntimeSpatialHashUseAlignedGridLevels && GRuntimeSpatialHashSnapNonAlignedGridLevelsToLowerLevels;
 								}
 							}
 							else if (DataLayerSubsystem && DataLayerSubsystem->IsAnyDataLayerInEffectiveRuntimeState(Cell->GetDataLayers(), EDataLayerRuntimeState::Loaded))
@@ -243,7 +243,7 @@ void FSpatialHashStreamingGrid::GetCells(const TArray<FWorldPartitionStreamingSo
 
 	GetAlwaysLoadedCells(DataLayerSubsystem, OutActivateCells.GetCells(), OutLoadCells.GetCells());
 
-	if (!GRuntimeSpatialHashUseAlignedGridLevels)
+	if (!GRuntimeSpatialHashUseAlignedGridLevels && GRuntimeSpatialHashSnapNonAlignedGridLevelsToLowerLevels)
 	{
 		auto FindIntersectingParents = [&Helper, this](const FIntersectingCells& InAllCells, const FIntersectingCells& InTestCells, FIntersectingCells& OutIntersectingCells)
 		{
