When running with ODSC do not try and fetch fully completed shader maps from the DDC.  This greatly improves latency.

- This change is essentially added code back in that was removed in CL# 18795062
- Since everyone is running with ODSC now fully completed shader maps are never (or rarely) actually cached to the remote DDC now.
- This resulted in DDC misses and 100-200ms blocking DDC Get calls that would never return anything.
- For now the most optimal solution is to remove the query to improve editor startup and loads.

diff --git a/Engine/Source/Runtime/Engine/Private/Materials/MaterialShared.cpp b/Engine/Source/Runtime/Engine/Private/Materials/MaterialShared.cpp
--- a/Engine/Source/Runtime/Engine/Private/Materials/MaterialShared.cpp
+++ b/Engine/Source/Runtime/Engine/Private/Materials/MaterialShared.cpp
@@ -2309,13 +2309,19 @@ bool FMaterial::CacheShaders(const FMaterialShaderMapId& ShaderMapId, EShaderPla
 				}
 			}
 
+			// If we are loading individual shaders from the shader job cache don't attempt to load full maps.
+			const bool bSkipCompilationOnPostLoad = IsShaderJobCacheDDCEnabled();
+
 			// Attempt to load from the derived data cache if we are uncooked and don't have any shadermap.
 			// If we have an incomplete shadermap, continue with it to prevent creation of duplicate shadermaps for the same ShaderMapId
 			if (!ShaderMap && !FPlatformProperties::RequiresCookedData())
 			{
-				TRefCountPtr<FMaterialShaderMap> LoadedShaderMap;
-				FMaterialShaderMap::LoadFromDerivedDataCache(this, ShaderMapId, Platform, TargetPlatform, LoadedShaderMap, DDCKeyHash);
-				ShaderMap = LoadedShaderMap;
+				if (bSkipCompilationOnPostLoad == false || IsDefaultMaterial())
+				{
+					TRefCountPtr<FMaterialShaderMap> LoadedShaderMap;
+					FMaterialShaderMap::LoadFromDerivedDataCache(this, ShaderMapId, Platform, TargetPlatform, LoadedShaderMap, DDCKeyHash);
+					ShaderMap = LoadedShaderMap;
+				}
 			}
 
 			check(!ShaderMap || ShaderMap->GetFrozenContentSize() > 0u);
