- Fix LWC VectorField on CPU

#rnx

diff --git a/Engine/Source/Runtime/Engine/Private/VectorField.cpp b/Engine/Source/Runtime/Engine/Private/VectorField.cpp
--- a/Engine/Source/Runtime/Engine/Private/VectorField.cpp
+++ b/Engine/Source/Runtime/Engine/Private/VectorField.cpp
@@ -433,11 +433,11 @@ void UVectorFieldStatic::UpdateCPUData(bool bDiscardData)
 			Ar << Ptr[SampleIt].B;
 		}
 #else
-		CPUData.Reset(Align(SampleCount * sizeof(FVector), sizeof(FVector4)));
+		CPUData.Reset(Align(SampleCount * sizeof(FVector3f), sizeof(FVector4f)));
 
 		for (size_t SampleIt = 0; SampleIt < SampleCount; ++SampleIt)
 		{
-			FVector Value(Ptr[SampleIt].R.GetFloat(), Ptr[SampleIt].G.GetFloat(), Ptr[SampleIt].B.GetFloat());
+			FVector3f Value(Ptr[SampleIt].R.GetFloat(), Ptr[SampleIt].G.GetFloat(), Ptr[SampleIt].B.GetFloat());
 			Ar << Value;
 		}
 #endif
@@ -462,9 +462,9 @@ FVector UVectorFieldStatic::FilteredSample(const FVector& SamplePosition, const
 	};
 
 	static auto FVectorFloor = [](const FVector& v) {
-		return FVector(FGenericPlatformMath::FloorToFloat(v.X),
-			FGenericPlatformMath::FloorToFloat(v.Y),
-			FGenericPlatformMath::FloorToFloat(v.Z));
+		return FVector(FGenericPlatformMath::FloorToDouble(v.X),
+			FGenericPlatformMath::FloorToDouble(v.Y),
+			FGenericPlatformMath::FloorToDouble(v.Z));
 	};
 
 	const FVector Size(SizeX, SizeY, SizeZ);
