Speculative fix for GC crash with effect presets

diff --git a/Engine/Source/Runtime/Engine/Classes/Sound/SoundEffectBase.h b/Engine/Source/Runtime/Engine/Classes/Sound/SoundEffectBase.h
--- a/Engine/Source/Runtime/Engine/Classes/Sound/SoundEffectBase.h
+++ b/Engine/Source/Runtime/Engine/Classes/Sound/SoundEffectBase.h
@@ -28,8 +28,8 @@ class USoundEffectPreset;
 #define EFFECT_PRESET_NAME(CLASS_NAME)			 EFFECT_PRESET_NAME1(CLASS_NAME, Preset)
 
 #define GET_EFFECT_SETTINGS(EFFECT_NAME) \
-		U##EFFECT_NAME##Preset* _Preset = CastChecked<U##EFFECT_NAME##Preset>(Preset); \
-		F##EFFECT_NAME##Settings Settings = _Preset->GetSettings(); \
+		U##EFFECT_NAME##Preset* _Preset = Cast<U##EFFECT_NAME##Preset>(Preset); \
+		F##EFFECT_NAME##Settings Settings = _Preset != nullptr ? _Preset->GetSettings() : F##EFFECT_NAME##Settings(); \
 
 #define EFFECT_PRESET_METHODS(EFFECT_NAME) \
 		virtual bool CanFilter() const override { return false; } \
diff --git a/Engine/Source/Runtime/Engine/Private/SoundEffectSubmix.cpp b/Engine/Source/Runtime/Engine/Private/SoundEffectSubmix.cpp
--- a/Engine/Source/Runtime/Engine/Private/SoundEffectSubmix.cpp
+++ b/Engine/Source/Runtime/Engine/Private/SoundEffectSubmix.cpp
@@ -15,8 +15,8 @@ void FSoundEffectSubmix::ProcessAudio(FSoundEffectSubmixInputData& InData, FSoun
 
 	Update();
 
-	// Only process the effect if the effect is active
-	if (bIsActive)
+	// Only process the effect if the effect is active and the preset is valid
+	if (bIsActive && Preset.IsValid())
 	{
 		OnProcessAudio(InData, OutData);
 	}
