Refactored FHierarchicalStaticMeshSceneProxy into its own header instead of the cpp, and added ENGINE_API to FHierarchicalStaticMeshSceneProxy, FInstancedStaticMeshVertexFactory, FInstancedStaticMeshVertexFactoryShaderParameters, and FInstancedStaticMeshSceneProxy

diff --git a/Engine/Source/Runtime/Engine/Classes/Engine/InstancedStaticMesh.h b/Engine/Source/Runtime/Engine/Classes/Engine/InstancedStaticMesh.h
--- a/Engine/Source/Runtime/Engine/Classes/Engine/InstancedStaticMesh.h
+++ b/Engine/Source/Runtime/Engine/Classes/Engine/InstancedStaticMesh.h
@@ -255,7 +255,7 @@ struct FInstancedStaticMeshDataType
 /**
  * A vertex factory for instanced static meshes
  */
-struct FInstancedStaticMeshVertexFactory : public FLocalVertexFactory
+struct ENGINE_API FInstancedStaticMeshVertexFactory : public FLocalVertexFactory
 {
 	DECLARE_VERTEX_FACTORY_TYPE(FInstancedStaticMeshVertexFactory);
 public:
@@ -339,7 +339,7 @@ private:
 	TUniformBufferRef<FInstancedStaticMeshVertexFactoryUniformShaderParameters> UniformBuffer;
 };
 
-class FInstancedStaticMeshVertexFactoryShaderParameters : public FLocalVertexFactoryShaderParametersBase
+class ENGINE_API FInstancedStaticMeshVertexFactoryShaderParameters : public FLocalVertexFactoryShaderParametersBase
 {
 	DECLARE_TYPE_LAYOUT(FInstancedStaticMeshVertexFactoryShaderParameters, NonVirtual);
 public:
@@ -525,7 +525,7 @@ private:
 	FInstancedStaticMeshSceneProxy
 -----------------------------------------------------------------------------*/
 
-class FInstancedStaticMeshSceneProxy : public FStaticMeshSceneProxy
+class ENGINE_API FInstancedStaticMeshSceneProxy : public FStaticMeshSceneProxy
 {
 public:
 	SIZE_T GetTypeHash() const override;
diff --git a/Engine/Source/Runtime/Engine/Private/HierarchicalInstancedStaticMesh.cpp b/Engine/Source/Runtime/Engine/Private/HierarchicalInstancedStaticMesh.cpp
--- a/Engine/Source/Runtime/Engine/Private/HierarchicalInstancedStaticMesh.cpp
+++ b/Engine/Source/Runtime/Engine/Private/HierarchicalInstancedStaticMesh.cpp
@@ -23,10 +23,6 @@
 #include "RendererInterface.h"
 #include "Engine/StaticMesh.h"
 #include "UnrealEngine.h"
-#include "Components/InstancedStaticMeshComponent.h"
-#include "StaticMeshResources.h"
-#include "Components/HierarchicalInstancedStaticMeshComponent.h"
-#include "Engine/InstancedStaticMesh.h"
 #include "InstancedStaticMeshDelegates.h"
 #include "SceneManagement.h"
 #include "HAL/LowLevelMemTracker.h"
@@ -38,6 +34,7 @@
 #endif
 #include "PrimitiveSceneInfo.h"
 #include "NaniteSceneProxy.h"
+#include "HierarchicalStaticMeshSceneProxy.h"
 
 #if WITH_EDITOR
 static float GDebugBuildTreeAsyncDelayInSeconds = 0.f;
@@ -825,211 +822,138 @@ void ToggleFreezeFoliageCulling()
 #endif
 }
 
-
-struct FFoliageOcclusionResults
-{
-	TArray<bool> Results; // we keep a copy from the View as the view will get destroyed too often
-	int32 ResultsStart;
-	int32 NumResults;
-	uint32 FrameNumberRenderThread;
-
-	FFoliageOcclusionResults(TArray<bool>* InResults, int32 InResultsStart, int32 InNumResults)
-		: Results(*InResults)
-		, ResultsStart(InResultsStart)
-		, NumResults(InNumResults)
-		, FrameNumberRenderThread(GFrameNumberRenderThread)
-	{
-
-	}
-};
-
-struct FFoliageElementParams;
-struct FFoliageRenderInstanceParams;
-struct FFoliageCullInstanceParams;
-
 FHierarchicalInstancedStaticMeshDelegates::FOnTreeBuilt FHierarchicalInstancedStaticMeshDelegates::OnTreeBuilt;
 
-class FHierarchicalStaticMeshSceneProxy final : public FInstancedStaticMeshSceneProxy
+SIZE_T FHierarchicalStaticMeshSceneProxy::GetTypeHash() const
 {
-	TSharedRef<TArray<FClusterNode>, ESPMode::ThreadSafe> ClusterTreePtr;
-	const TArray<FClusterNode>& ClusterTree;
-
-	TArray<FBox> UnbuiltBounds;
-	int32 FirstUnbuiltIndex;
-	int32 InstanceCountToRender;
-
-	int32 FirstOcclusionNode;
-	int32 LastOcclusionNode;
-	TArray<FBoxSphereBounds> OcclusionBounds;
-	TMap<uint32, FFoliageOcclusionResults> OcclusionResults;
-	EHISMViewRelevanceType ViewRelevance;
-	bool bDitheredLODTransitions;
-	uint32 SceneProxyCreatedFrameNumberRenderThread;
-
-#if !(UE_BUILD_SHIPPING || UE_BUILD_TEST)
-	mutable TArray<uint32> SingleDebugRuns[MAX_STATIC_MESH_LODS];
-	mutable int32 SingleDebugTotalInstances[MAX_STATIC_MESH_LODS];
-	mutable TArray<uint32> MultipleDebugRuns[MAX_STATIC_MESH_LODS];
-	mutable int32 MultipleDebugTotalInstances[MAX_STATIC_MESH_LODS];
-	mutable int32 CaptureTag;
-#endif
+	static size_t UniquePointer;
+	return reinterpret_cast<size_t>(&UniquePointer);
+}
 
-public:
-	SIZE_T GetTypeHash() const override
-	{
-		static size_t UniquePointer;
-		return reinterpret_cast<size_t>(&UniquePointer);
-	}
-
-	FHierarchicalStaticMeshSceneProxy(UHierarchicalInstancedStaticMeshComponent* InComponent, ERHIFeatureLevel::Type InFeatureLevel)
-		: FInstancedStaticMeshSceneProxy(InComponent, InFeatureLevel)
-		, ClusterTreePtr(InComponent->ClusterTreePtr.ToSharedRef())
-		, ClusterTree(*InComponent->ClusterTreePtr)
-		, UnbuiltBounds(InComponent->UnbuiltInstanceBoundsList)
-		, FirstUnbuiltIndex(InComponent->NumBuiltInstances > 0 ? InComponent->NumBuiltInstances : InComponent->NumBuiltRenderInstances)
-		, InstanceCountToRender(InComponent->InstanceCountToRender)
-		, ViewRelevance(InComponent->GetViewRelevanceType())
-		, bDitheredLODTransitions(InComponent->SupportsDitheredLODTransitions(InFeatureLevel))
-		, SceneProxyCreatedFrameNumberRenderThread(UINT32_MAX)
+FHierarchicalStaticMeshSceneProxy::FHierarchicalStaticMeshSceneProxy(UHierarchicalInstancedStaticMeshComponent* InComponent, ERHIFeatureLevel::Type InFeatureLevel)
+: FInstancedStaticMeshSceneProxy(InComponent, InFeatureLevel)
+, ClusterTreePtr(InComponent->ClusterTreePtr.ToSharedRef())
+, ClusterTree(*InComponent->ClusterTreePtr)
+, UnbuiltBounds(InComponent->UnbuiltInstanceBoundsList)
+, FirstUnbuiltIndex(InComponent->NumBuiltInstances > 0 ? InComponent->NumBuiltInstances : InComponent->NumBuiltRenderInstances)
+, InstanceCountToRender(InComponent->InstanceCountToRender)
+, ViewRelevance(InComponent->GetViewRelevanceType())
+, bDitheredLODTransitions(InComponent->SupportsDitheredLODTransitions(InFeatureLevel))
+, SceneProxyCreatedFrameNumberRenderThread(UINT32_MAX)
 #if !(UE_BUILD_SHIPPING || UE_BUILD_TEST)
-		, CaptureTag(0)
+, CaptureTag(0)
 #endif
-	{
-		SetupOcclusion(InComponent);
-		bSupportsMeshCardRepresentation = false;
-	}
+{
+	SetupOcclusion(InComponent);
+	bSupportsMeshCardRepresentation = false;
+}
 
-	void SetupOcclusion(UHierarchicalInstancedStaticMeshComponent* InComponent)
+void FHierarchicalStaticMeshSceneProxy::SetupOcclusion(UHierarchicalInstancedStaticMeshComponent* InComponent)
+{
+	FirstOcclusionNode = 0;
+	LastOcclusionNode = 0;
+	if (ClusterTree.Num() && InComponent->OcclusionLayerNumNodes)
 	{
-		FirstOcclusionNode = 0;
-		LastOcclusionNode = 0;
-		if (ClusterTree.Num() && InComponent->OcclusionLayerNumNodes)
+		while (true)
 		{
-			while (true)
-			{
-				int32 NextFirstOcclusionNode = ClusterTree[FirstOcclusionNode].FirstChild;
-				int32 NextLastOcclusionNode = ClusterTree[LastOcclusionNode].LastChild;
+			int32 NextFirstOcclusionNode = ClusterTree[FirstOcclusionNode].FirstChild;
+			int32 NextLastOcclusionNode = ClusterTree[LastOcclusionNode].LastChild;
 
-				if (NextFirstOcclusionNode < 0 || NextLastOcclusionNode < 0)
-				{
-					break;
-				}
-				int32 NumNodes = 1 + NextLastOcclusionNode - NextFirstOcclusionNode;
-				if (NumNodes > InComponent->OcclusionLayerNumNodes)
-				{
-					break;
-				}
-				FirstOcclusionNode = NextFirstOcclusionNode;
-				LastOcclusionNode = NextLastOcclusionNode;
-			}
-		}
-		int32 NumNodes = 1 + LastOcclusionNode - FirstOcclusionNode;
-		if (NumNodes < 2)
-		{
-			FirstOcclusionNode = -1;
-			LastOcclusionNode = -1;
-			NumNodes = 0;
-			if (ClusterTree.Num())
+			if (NextFirstOcclusionNode < 0 || NextLastOcclusionNode < 0)
 			{
-				//UE_LOG(LogTemp, Display, TEXT("No SubOcclusion %d inst"), 1 + ClusterTree[0].LastInstance - ClusterTree[0].FirstInstance);
+				break;
 			}
-		}
-		else
-		{
-			//int32 NumPerNode = (1 + ClusterTree[0].LastInstance - ClusterTree[0].FirstInstance) / NumNodes;
-			//UE_LOG(LogTemp, Display, TEXT("Occlusion level %d   %d inst / node"), NumNodes, NumPerNode);
-			OcclusionBounds.Reserve(NumNodes);
-			FMatrix XForm = InComponent->GetRenderMatrix();
-			for (int32 Index = FirstOcclusionNode; Index <= LastOcclusionNode; Index++)
+			int32 NumNodes = 1 + NextLastOcclusionNode - NextFirstOcclusionNode;
+			if (NumNodes > InComponent->OcclusionLayerNumNodes)
 			{
-				OcclusionBounds.Add(FBoxSphereBounds(FBox(ClusterTree[Index].BoundMin, ClusterTree[Index].BoundMax).TransformBy(XForm)));
+				break;
 			}
+			FirstOcclusionNode = NextFirstOcclusionNode;
+			LastOcclusionNode = NextLastOcclusionNode;
 		}
 	}
-
-	// FPrimitiveSceneProxy interface.
-	
-	virtual void CreateRenderThreadResources() override
+	int32 NumNodes = 1 + LastOcclusionNode - FirstOcclusionNode;
+	if (NumNodes < 2)
 	{
-		FInstancedStaticMeshSceneProxy::CreateRenderThreadResources();
-		SceneProxyCreatedFrameNumberRenderThread = GFrameNumberRenderThread;
-	}
-	
-	virtual FPrimitiveViewRelevance GetViewRelevance(const FSceneView* View) const override
-	{
-		FPrimitiveViewRelevance Result;
-		bool bShowInstancedMesh = true;
-		switch (ViewRelevance)
-		{
-		case EHISMViewRelevanceType::Grass:
-			bShowInstancedMesh = View->Family->EngineShowFlags.InstancedGrass;
-			break;
-		case EHISMViewRelevanceType::Foliage:
-			bShowInstancedMesh = View->Family->EngineShowFlags.InstancedFoliage;
-			break;
-		case EHISMViewRelevanceType::HISM:
-			bShowInstancedMesh = View->Family->EngineShowFlags.InstancedStaticMeshes;
-			break;
-		default:
-			break;
-		}
-		if (bShowInstancedMesh)
-		{
-			Result = FStaticMeshSceneProxy::GetViewRelevance(View);
-			Result.bDynamicRelevance = true;
-			Result.bStaticRelevance = false;
-
-			// Remove relevance for primitives marked for runtime virtual texture only.
-			if (RuntimeVirtualTextures.Num() > 0 && !ShouldRenderInMainPass())
-			{
-				Result.bDynamicRelevance = false;
-			}
+		FirstOcclusionNode = -1;
+		LastOcclusionNode = -1;
+		NumNodes = 0;
+		if (ClusterTree.Num())
+		{
+			//UE_LOG(LogTemp, Display, TEXT("No SubOcclusion %d inst"), 1 + ClusterTree[0].LastInstance - ClusterTree[0].FirstInstance);
 		}
-		return Result;
 	}
-
-#if RHI_RAYTRACING
-	virtual bool IsRayTracingStaticRelevant() const override
+	else
 	{
-		return false;
+		//int32 NumPerNode = (1 + ClusterTree[0].LastInstance - ClusterTree[0].FirstInstance) / NumNodes;
+		//UE_LOG(LogTemp, Display, TEXT("Occlusion level %d   %d inst / node"), NumNodes, NumPerNode);
+		OcclusionBounds.Reserve(NumNodes);
+		FMatrix XForm = InComponent->GetRenderMatrix();
+		for (int32 Index = FirstOcclusionNode; Index <= LastOcclusionNode; Index++)
+		{
+			OcclusionBounds.Add(FBoxSphereBounds(FBox(ClusterTree[Index].BoundMin, ClusterTree[Index].BoundMax).TransformBy(XForm)));
+		}
 	}
-#endif
-
-	virtual void GetDynamicMeshElements(const TArray<const FSceneView*>& Views, const FSceneViewFamily& ViewFamily, uint32 VisibilityMap, FMeshElementCollector& Collector) const override;
+}
 
-	virtual const TArray<FBoxSphereBounds>* GetOcclusionQueries(const FSceneView* View) const override;
-	virtual void AcceptOcclusionResults(const FSceneView* View, TArray<bool>* Results, int32 ResultsStart, int32 NumResults) override;
-	virtual bool HasSubprimitiveOcclusionQueries() const override
+void FHierarchicalStaticMeshSceneProxy::CreateRenderThreadResources()
+{
+	FInstancedStaticMeshSceneProxy::CreateRenderThreadResources();
+	SceneProxyCreatedFrameNumberRenderThread = GFrameNumberRenderThread;
+}
+	
+FPrimitiveViewRelevance FHierarchicalStaticMeshSceneProxy::GetViewRelevance(const FSceneView* View) const
+{
+	FPrimitiveViewRelevance Result;
+	bool bShowInstancedMesh = true;
+	switch (ViewRelevance)
 	{
-		return FirstOcclusionNode > 0;
+	case EHISMViewRelevanceType::Grass:
+		bShowInstancedMesh = View->Family->EngineShowFlags.InstancedGrass;
+		break;
+	case EHISMViewRelevanceType::Foliage:
+		bShowInstancedMesh = View->Family->EngineShowFlags.InstancedFoliage;
+		break;
+	case EHISMViewRelevanceType::HISM:
+		bShowInstancedMesh = View->Family->EngineShowFlags.InstancedStaticMeshes;
+		break;
+	default:
+		break;
 	}
-
-
-	virtual void DrawStaticElements(FStaticPrimitiveDrawInterface* PDI) override
+	if (bShowInstancedMesh)
 	{
-		if (RuntimeVirtualTextures.Num() > 0)
+		Result = FStaticMeshSceneProxy::GetViewRelevance(View);
+		Result.bDynamicRelevance = true;
+		Result.bStaticRelevance = false;
+
+		// Remove relevance for primitives marked for runtime virtual texture only.
+		if (RuntimeVirtualTextures.Num() > 0 && !ShouldRenderInMainPass())
 		{
-			// Create non-hierachichal static mesh batches for use by the runtime virtual texture rendering.
-			//todo[vt]: Build an acceleration structure better suited for VT rendering maybe with batches aligned to VT pages?
-			FInstancedStaticMeshSceneProxy::DrawStaticElements(PDI);
+			Result.bDynamicRelevance = false;
 		}
 	}
+	return Result;
+}
 
-	virtual void ApplyWorldOffset(FVector InOffset) override
+void FHierarchicalStaticMeshSceneProxy::DrawStaticElements(FStaticPrimitiveDrawInterface* PDI)
+{
+	if (RuntimeVirtualTextures.Num() > 0)
 	{
-		FInstancedStaticMeshSceneProxy::ApplyWorldOffset(InOffset);
-		
-		for (FBoxSphereBounds& Item : OcclusionBounds)
-		{
-			Item.Origin+= InOffset;
-		}
+		// Create non-hierarchal static mesh batches for use by the runtime virtual texture rendering.
+		//todo[vt]: Build an acceleration structure better suited for VT rendering maybe with batches aligned to VT pages?
+		FInstancedStaticMeshSceneProxy::DrawStaticElements(PDI);
 	}
+}
 
-	void FillDynamicMeshElements(FMeshElementCollector& Collector, const FFoliageElementParams& ElementParams, const FFoliageRenderInstanceParams& Instances) const;
-
-	template<bool TUseVector>
-	void Traverse(const FFoliageCullInstanceParams& Params, int32 Index, int32 MinLOD, int32 MaxLOD, bool bFullyContained = false) const;
-};
+void FHierarchicalStaticMeshSceneProxy::ApplyWorldOffset(FVector InOffset)
+{
+	FInstancedStaticMeshSceneProxy::ApplyWorldOffset(InOffset);
+		
+	for (FBoxSphereBounds& Item : OcclusionBounds)
+	{
+		Item.Origin+= InOffset;
+	}
+}
 
 struct FFoliageRenderInstanceParams
 {
diff --git a/Engine/Source/Runtime/Engine/Public/HierarchicalStaticMeshSceneProxy.h b/Engine/Source/Runtime/Engine/Public/HierarchicalStaticMeshSceneProxy.h
new file mode 100644
--- /dev/null
+++ b/Engine/Source/Runtime/Engine/Public/HierarchicalStaticMeshSceneProxy.h
@@ -0,0 +1,98 @@
+// Copyright Epic Games, Inc. All Rights Reserved.
+
+#pragma once
+
+#include "CoreMinimal.h"
+#include "PrimitiveSceneProxy.h"
+#include "PrimitiveViewRelevance.h"
+#include "Engine/InstancedStaticMesh.h"
+#include "Components/InstancedStaticMeshComponent.h"
+#include "StaticMeshResources.h"
+#include "Components/HierarchicalInstancedStaticMeshComponent.h"
+
+struct FClusterNode;
+struct FFoliageElementParams;
+struct FFoliageRenderInstanceParams;
+struct FFoliageCullInstanceParams;
+
+struct FFoliageOcclusionResults
+{
+	TArray<bool> Results; // we keep a copy from the View as the view will get destroyed too often
+	int32 ResultsStart;
+	int32 NumResults;
+	uint32 FrameNumberRenderThread;
+
+	FFoliageOcclusionResults(TArray<bool>* InResults, int32 InResultsStart, int32 InNumResults)
+	: Results(*InResults)
+	, ResultsStart(InResultsStart)
+	, NumResults(InNumResults)
+	, FrameNumberRenderThread(GFrameNumberRenderThread)
+	{
+	}
+};
+
+class ENGINE_API FHierarchicalStaticMeshSceneProxy final : public FInstancedStaticMeshSceneProxy
+{
+	TSharedRef<TArray<FClusterNode>, ESPMode::ThreadSafe> ClusterTreePtr;
+	const TArray<FClusterNode>& ClusterTree;
+
+	TArray<FBox> UnbuiltBounds;
+	int32 FirstUnbuiltIndex;
+	int32 InstanceCountToRender;
+
+	int32 FirstOcclusionNode;
+	int32 LastOcclusionNode;
+	TArray<FBoxSphereBounds> OcclusionBounds;
+	TMap<uint32, FFoliageOcclusionResults> OcclusionResults;
+	EHISMViewRelevanceType ViewRelevance;
+	bool bDitheredLODTransitions;
+	uint32 SceneProxyCreatedFrameNumberRenderThread;
+
+#if !(UE_BUILD_SHIPPING || UE_BUILD_TEST)
+	mutable TArray<uint32> SingleDebugRuns[MAX_STATIC_MESH_LODS];
+	mutable int32 SingleDebugTotalInstances[MAX_STATIC_MESH_LODS];
+	mutable TArray<uint32> MultipleDebugRuns[MAX_STATIC_MESH_LODS];
+	mutable int32 MultipleDebugTotalInstances[MAX_STATIC_MESH_LODS];
+	mutable int32 CaptureTag;
+#endif
+
+public:
+	SIZE_T GetTypeHash() const override;
+
+	FHierarchicalStaticMeshSceneProxy(UHierarchicalInstancedStaticMeshComponent* InComponent, ERHIFeatureLevel::Type InFeatureLevel);
+
+	void SetupOcclusion(UHierarchicalInstancedStaticMeshComponent* InComponent);
+
+	// FPrimitiveSceneProxy interface.
+	
+	virtual void CreateRenderThreadResources() override;
+	
+	virtual FPrimitiveViewRelevance GetViewRelevance(const FSceneView* View) const override;
+
+#if RHI_RAYTRACING
+	virtual bool IsRayTracingStaticRelevant() const override
+	{
+		return false;
+	}
+#endif
+
+	virtual void GetDynamicMeshElements(const TArray<const FSceneView*>& Views, const FSceneViewFamily& ViewFamily, uint32 VisibilityMap, FMeshElementCollector& Collector) const override;
+
+	virtual const TArray<FBoxSphereBounds>* GetOcclusionQueries(const FSceneView* View) const override;
+	
+	virtual void AcceptOcclusionResults(const FSceneView* View, TArray<bool>* Results, int32 ResultsStart, int32 NumResults) override;
+	
+	virtual bool HasSubprimitiveOcclusionQueries() const override
+	{
+		return FirstOcclusionNode > 0;
+	}
+
+	virtual void DrawStaticElements(FStaticPrimitiveDrawInterface* PDI) override;
+
+	virtual void ApplyWorldOffset(FVector InOffset) override;
+
+	void FillDynamicMeshElements(FMeshElementCollector& Collector, const FFoliageElementParams& ElementParams, const FFoliageRenderInstanceParams& Instances) const;
+
+	template<bool TUseVector>
+	void Traverse(const FFoliageCullInstanceParams& Params, int32 Index, int32 MinLOD, int32 MaxLOD, bool bFullyContained = false) const;
+};
