Disable (gray out) shadow map method selection drop down when Forward Shading is enable to indicate it has no effect.

diff --git a/Engine/Source/Runtime/Engine/Classes/Engine/RendererSettings.h b/Engine/Source/Runtime/Engine/Classes/Engine/RendererSettings.h
--- a/Engine/Source/Runtime/Engine/Classes/Engine/RendererSettings.h
+++ b/Engine/Source/Runtime/Engine/Classes/Engine/RendererSettings.h
@@ -491,7 +491,7 @@ class ENGINE_API URendererSettings : public UDeveloperSettings
 
 	UPROPERTY(config, EditAnywhere, Category = Shadows, meta = (
 		ConsoleVariable = "r.Shadow.Virtual.Enable", DisplayName = "Shadow Map Method",
-		ToolTip = "Shadow Map Method"))
+		ToolTip = "Select the primary shadow mapping method. Automatically uses 'Shadow Maps' when Forward Shading is enabled for the project as Virtual Shadow Maps are not supported."))
 	TEnumAsByte<EShadowMapMethod::Type> ShadowMapMethod;
 
 	/**
diff --git a/Engine/Source/Runtime/Engine/Private/RendererSettings.cpp b/Engine/Source/Runtime/Engine/Private/RendererSettings.cpp
--- a/Engine/Source/Runtime/Engine/Private/RendererSettings.cpp
+++ b/Engine/Source/Runtime/Engine/Private/RendererSettings.cpp
@@ -303,7 +303,8 @@ bool URendererSettings::CanEditChange(const FProperty* InProperty) const
 	}
 
 	if (InProperty->GetFName() == GET_MEMBER_NAME_CHECKED(URendererSettings, DynamicGlobalIllumination)
-		|| InProperty->GetFName() == GET_MEMBER_NAME_CHECKED(URendererSettings, Reflections))
+		|| InProperty->GetFName() == GET_MEMBER_NAME_CHECKED(URendererSettings, Reflections)
+		|| InProperty->GetFName() == GET_MEMBER_NAME_CHECKED(URendererSettings, ShadowMapMethod))
 	{
 		return !bForwardShading;
 	}
