Chaos: Implement SetLocalTransform for Chaos interface.
 - Re-enabled copying of scaled implicits
 - Refactored UpdateBodyScale to work with transformed shapes and match behaviour with 4.27 (clears current local transform)
 - Added support to all geometry paths for transformed and scaled objects

diff --git a/Engine/Source/Runtime/Engine/Private/PhysicsEngine/BodyInstance.cpp b/Engine/Source/Runtime/Engine/Private/PhysicsEngine/BodyInstance.cpp
--- a/Engine/Source/Runtime/Engine/Private/PhysicsEngine/BodyInstance.cpp
+++ b/Engine/Source/Runtime/Engine/Private/PhysicsEngine/BodyInstance.cpp
@@ -1942,53 +1942,79 @@ bool FBodyInstance::UpdateBodyScale(const FVector& InScale3D, bool bForceUpdate)
 		for (FPhysicsShapeHandle& ShapeHandle : Shapes)
 		{
 			const Chaos::FImplicitObject& ImplicitObject = ShapeHandle.GetGeometry();
-			EImplicitObjectType ImplicitType = ImplicitObject.GetType();
-			EImplicitObjectType GeomType = GetInnerType(ImplicitType);
-
+			
+			EImplicitObjectType OuterType = ImplicitObject.GetType();
+			EImplicitObjectType WrappedOrConcreteType = GetInnerType(OuterType);
+			
 			const FTransform& RelativeTM = GetRelativeBodyTransform(ShapeHandle);
 
 			bool bIsTransformed = false;
-			if(ImplicitType == ImplicitObjectType::Transformed)
+			bool bIsScaled = false;
+			bool bIsInstanced = false;
+
+			// Unwrap the shape in order:
+			// Transformed -> Scaled | Instanced -> Concrete
+
+			if(OuterType == ImplicitObjectType::Transformed)
 			{
 				bIsTransformed = true;
 
 				// Get GeomType that is transformed
-				const TImplicitObjectTransformed<FReal, 3>& ImplicitObjectTransformed = static_cast<const TImplicitObjectTransformed<FReal, 3>&>(ImplicitObject);
-				GeomType = ImplicitObjectTransformed.GetTransformedObject()->GetType();
-				CHAOS_ENSURE(!IsInstanced(GeomType));
+				WrappedOrConcreteType = static_cast<const TImplicitObjectTransformed<FReal, 3>&>(ImplicitObject).GetTransformedObject()->GetType();
+			}
+			else
+			{
+				// If we didn't find a wrapper (transformed) then our wrapped or concrete type is the outer type
+				WrappedOrConcreteType = OuterType;
 			}
 
-			FKShapeElem* ShapeElem = FChaosUserData::Get<FKShapeElem>(FPhysicsInterface::GetUserData(ShapeHandle));
+			// Strip out the scaled and instance wrappers leaving us with a definite concrete type
+			EImplicitObjectType ConcreteType = WrappedOrConcreteType;
 
-			switch (GeomType)
+			if(IsScaled(ConcreteType))
 			{
-				case ImplicitObjectType::Sphere:
-				{
-					FKSphereElem* SphereElem = ShapeElem->GetShapeCheck<FKSphereElem>();
-					ensure(ScaleMode == EScaleMode::LockedXYZ);
+				bIsScaled = true;
 
-					FReal Radius = FMath::Max<FReal>(SphereElem->Radius * AdjustedScale3DAbs.X, FCollisionShape::MinSphereRadius());
+				ConcreteType ^= ImplicitObjectType::IsScaled;
+				ensure(!IsInstanced(ConcreteType));
+			}
+			else if(IsInstanced(ConcreteType))
+			{
+				bIsInstanced = true;
+
+				ConcreteType ^= ImplicitObjectType::IsInstanced;
+				ensure(!IsScaled(ConcreteType));
+			}
 
+			FKShapeElem* ShapeElem = FChaosUserData::Get<FKShapeElem>(FPhysicsInterface::GetUserData(ShapeHandle));
 
-					if (!CHAOS_ENSURE(!IsScaled(ImplicitType) && !bIsTransformed && !IsInstanced(ImplicitType)))
+			switch(ConcreteType)
+			{
+				case ImplicitObjectType::Sphere:
+				{
+					if (!CHAOS_ENSURE(!bIsInstanced && !bIsScaled))
 					{
-						// No support for Scaled, Instanced,or transformed.
+						// No support for Instanced, Scaled not supported as we bake the scale below
 						break;
 					}
 
-					FVec3 Center = RelativeTM.TransformPosition(SphereElem->Center) * AdjustedScale3D.X;
+					FKSphereElem* SphereElem = ShapeElem->GetShapeCheck<FKSphereElem>();
+					ensure(ScaleMode == EScaleMode::LockedXYZ);
 
+					FReal Radius = FMath::Max<FReal>(SphereElem->Radius * AdjustedScale3DAbs.X, FCollisionShape::MinSphereRadius());
+					FVec3 Center = RelativeTM.TransformPosition(SphereElem->Center) * InScale3D;
 					TUniquePtr<TSphere<FReal, 3>> NewSphere = MakeUnique<TSphere<FReal, 3>>(Center, Radius);
 
 					NewGeometry.Emplace(MoveTemp(NewSphere));
 					bSuccess = true;
+
 					break;
 				}
 				case ImplicitObjectType::Box:
 				{
-					if (!CHAOS_ENSURE(!IsScaled(ImplicitType) && !IsInstanced(ImplicitType)))
+					if (!CHAOS_ENSURE(!bIsScaled && !bIsInstanced))
 					{
-						// No support for ScaledImplicit Box yet or instanced
+						// No support for ScaledImplicit Box yet or Instanced, scale is baked below
 						break;
 					}
 
@@ -2031,23 +2057,17 @@ bool FBodyInstance::UpdateBodyScale(const FVector& InScale3D, bool bForceUpdate)
 				}
 				case ImplicitObjectType::Capsule:
 				{
-					ensure(ScaleMode == EScaleMode::LockedXY || ScaleMode == EScaleMode::LockedXYZ);
-
-
-
-					FReal ScaleRadius = FMath::Max(AdjustedScale3DAbs.X, AdjustedScale3DAbs.Y);
-					FReal ScaleLength = AdjustedScale3DAbs.Z;
-
-					FKSphylElem* SphylElem = ShapeElem->GetShapeCheck<FKSphylElem>();
-
-					if (!CHAOS_ENSURE(!IsScaled(ImplicitType) && !bIsTransformed && !IsInstanced(ImplicitType)))
+					if (!CHAOS_ENSURE(!bIsInstanced && !bIsScaled))
 					{
-						// No support for Scaled, Instanced, or Transformed Capsule yet
+						// No support for Instanced
 						break;
 					}
 
-					const FCapsule * CapsuleGeometry = static_cast<const FCapsule*>(&ImplicitObject);
+					ensure(ScaleMode == EScaleMode::LockedXY || ScaleMode == EScaleMode::LockedXYZ);
 
+					FReal ScaleRadius = FMath::Max(AdjustedScale3DAbs.X, AdjustedScale3DAbs.Y);
+					FReal ScaleLength = AdjustedScale3DAbs.Z;
+					FKSphylElem* SphylElem = ShapeElem->GetShapeCheck<FKSphylElem>();
 
 					const FReal InitialHeight = SphylElem->Radius * 2.0f + SphylElem->Length;
 					FReal Radius = FMath::Max(SphylElem->Radius * ScaleRadius, (FReal)0.1);
@@ -2057,15 +2077,13 @@ bool FBodyInstance::UpdateBodyScale(const FVector& InScale3D, bool bForceUpdate)
 					FReal HalfLength = HalfHeight - Radius;
 					HalfLength = FMath::Max((FReal)FCollisionShape::MinCapsuleAxisHalfHeight(), HalfLength);
 
-
 					// TODO: For Transformed implicit, do not bake this in. Set Transform instead.
-					FVec3 Center = RelativeTM.TransformPosition(SphylElem->Center) * AdjustedScale3D;
+					FVec3 Center = RelativeTM.TransformPosition(SphylElem->Center) * InScale3D;
 					const FVec3 Axis = SphylElem->Rotation.RotateVector(Chaos::FVec3(0, 0, 1));
 
 					const FVec3 X1 = Center - HalfLength * Axis;
 					const FVec3 X2 = Center + HalfLength * Axis;
 
-
 					TUniquePtr<FCapsule> NewCapsule =  MakeUnique<FCapsule>(X1, X2, Radius);
 					NewGeometry.Emplace(MoveTemp(NewCapsule));
 
@@ -2075,29 +2093,33 @@ bool FBodyInstance::UpdateBodyScale(const FVector& InScale3D, bool bForceUpdate)
 				}
 				case ImplicitObjectType::Convex:
 				{
-					if(!CHAOS_ENSURE(IsInstanced(ImplicitType) || IsScaled(ImplicitType)))
+					if(!CHAOS_ENSURE(bIsInstanced || bIsScaled))
 					{
-						CHAOS_ENSURE(false); // Expecting instanced or scaled.
+						// Expecting instanced or scaled.
 						break;
 					}
 
 					FKConvexElem* ConvexElem = ShapeElem->GetShapeCheck<FKConvexElem>();
-					const auto& ConvexImplicit = ConvexElem->GetChaosConvexMesh();
-
-					// Ensure no rotation/translation in relative transform. PhysX supports this, we currently do not.
-					CHAOS_ENSURE(RelativeTM.GetRotation() == FQuat::Identity);
-					CHAOS_ENSURE(RelativeTM.GetTranslation() == FVector(0, 0, 0));
+					const TSharedPtr<Chaos::FConvex, ESPMode::ThreadSafe>& ConvexImplicit = ConvexElem->GetChaosConvexMesh();
 
+					TUniquePtr<FImplicitObject> NewConvex = nullptr;
 					if (AdjustedScale3D == FVector(1.0f, 1.0f, 1.0f))
 					{
-						TUniquePtr<TImplicitObjectInstanced<FConvex>> NewConvex = MakeUnique<TImplicitObjectInstanced<FConvex>>(ConvexImplicit);
-						NewGeometry.Emplace(MoveTemp(NewConvex));
+						NewConvex = MakeUnique<TImplicitObjectInstanced<FConvex>>(ConvexImplicit);
 					}
 					else
 					{
-						TUniquePtr<TImplicitObjectScaled<FConvex>> NewConvex = MakeUnique<TImplicitObjectScaled<FConvex>>(ConvexImplicit, AdjustedScale3D);
-						NewGeometry.Emplace(MoveTemp(NewConvex));
+						NewConvex = MakeUnique<TImplicitObjectScaled<FConvex>>(ConvexImplicit, AdjustedScale3D);
+					}
+
+					if(RelativeTM.GetRotation() != FQuat::Identity || RelativeTM.GetTranslation() != FVector::ZeroVector)
+					{
+						FTransform AdjustedTransform = RelativeTM;
+						AdjustedTransform.SetTranslation(RelativeTM.GetTranslation() * AdjustedScale3D);
+						NewConvex = MakeUnique<TImplicitObjectTransformed<FReal, 3>>(MoveTemp(NewConvex), AdjustedTransform);
 					}
+					
+					NewGeometry.Emplace(MoveTemp(NewConvex));
 
 					bSuccess = true;
 
@@ -2105,55 +2127,46 @@ bool FBodyInstance::UpdateBodyScale(const FVector& InScale3D, bool bForceUpdate)
 				}
 				case ImplicitObjectType::TriangleMesh:
 				{
-					if(!CHAOS_ENSURE(IsScaled(ImplicitType) || IsInstanced(ImplicitType)))
+					if(!CHAOS_ENSURE(bIsScaled || bIsInstanced))
 					{
 						// Currently assuming all triangle meshes are scaled or instanced (if scale == 1).
 						break;
 					}
 
-					// PhysX supports translation, we currently do not.
-					CHAOS_ENSURE(RelativeTM.GetTranslation() == FVector(0, 0, 0));
-
-					auto CreateTriGeomInstanced = [](auto InObject, TArray<TUniquePtr<FImplicitObject>>& OutGeoArray)
-					{
-						TUniquePtr<TImplicitObjectInstanced<FTriangleMeshImplicitObject>> NewTriangleMesh = MakeUnique<TImplicitObjectInstanced<FTriangleMeshImplicitObject>>(InObject);
-						OutGeoArray.Emplace(MoveTemp(NewTriangleMesh));
-					};
-
-					auto CreateTriGeomScaled = [](auto InObject, TArray<TUniquePtr<FImplicitObject>>& OutGeoArray, const FVec3& InScale)
-					{
-						TUniquePtr<TImplicitObjectScaled<FTriangleMeshImplicitObject>> NewTriangleMesh = MakeUnique<TImplicitObjectScaled<FTriangleMeshImplicitObject>>(MoveTemp(InObject), InScale);
-						OutGeoArray.Emplace(MoveTemp(NewTriangleMesh));
-					};
-
-					auto CreateTriGeomAuto = [&CreateTriGeomInstanced, &CreateTriGeomScaled](auto InObject, TArray<TUniquePtr<FImplicitObject>>& OutGeoArray, const FVec3& InScale)
+					auto CreateTriGeomAuto = [](auto InObject, TArray<TUniquePtr<FImplicitObject>>& OutGeoArray, const FVec3& InScale) -> TUniquePtr<FImplicitObject>
 					{
-						if(InScale == FVec3(1.0f, 1.0f, 1.0f))
+						if(InScale == Chaos::FVec3(1.0f, 1.0f, 1.0f))
 						{
-							CreateTriGeomInstanced(InObject, OutGeoArray);
+							return MakeUnique<Chaos::TImplicitObjectInstanced<Chaos::FTriangleMeshImplicitObject>>(InObject);
 						}
 						else
 						{
-							CreateTriGeomScaled(InObject, OutGeoArray, InScale);
+							return MakeUnique<Chaos::TImplicitObjectScaled<Chaos::FTriangleMeshImplicitObject>>(MoveTemp(InObject), InScale);
 						}
 					};
 
+					const FImplicitObject* TrimeshContainer = &ImplicitObject;
+					if(bIsTransformed)
+					{
+						TrimeshContainer = static_cast<const TImplicitObjectTransformed<FReal, 3>*>(TrimeshContainer)->GetTransformedObject();
+					}
+
 					TSharedPtr<FTriangleMeshImplicitObject, ESPMode::ThreadSafe> InnerTriangleMesh = nullptr;
-					if (IsScaled(ImplicitType))
+					if (bIsScaled)
 					{
-						const TImplicitObjectScaled<FTriangleMeshImplicitObject>* ScaledTriangleMesh = (static_cast<const TImplicitObjectScaled<FTriangleMeshImplicitObject>*>(&ImplicitObject));
+						const TImplicitObjectScaled<FTriangleMeshImplicitObject>* ScaledTriangleMesh = (static_cast<const TImplicitObjectScaled<FTriangleMeshImplicitObject>*>(TrimeshContainer));
 						InnerTriangleMesh = ScaledTriangleMesh->GetSharedObject();
 
 						if(!InnerTriangleMesh)
 						{
 							// While a body setup will instantiate the triangle mesh as a shared geometry, other methods might not (e.g. retopologized landscape)
 							TImplicitObjectScaled<FTriangleMeshImplicitObject>::ObjectType InnerObject = ScaledTriangleMesh->Object();
-							CreateTriGeomScaled(ScaledTriangleMesh->Object(), NewGeometry, AdjustedScale3D);
+							NewGeometry.Emplace(MakeUnique<Chaos::TImplicitObjectScaled<Chaos::FTriangleMeshImplicitObject>>(ScaledTriangleMesh->Object(), AdjustedScale3D));
 						}
 					}
-					else if (IsInstanced(ImplicitType))
+					else if (bIsInstanced)
 					{
-						const TImplicitObjectInstanced<FTriangleMeshImplicitObject>* InstancedTriangleMesh = (static_cast<const TImplicitObjectInstanced<FTriangleMeshImplicitObject>*>(&ImplicitObject));
+						const TImplicitObjectInstanced<FTriangleMeshImplicitObject>* InstancedTriangleMesh = (static_cast<const TImplicitObjectInstanced<FTriangleMeshImplicitObject>*>(TrimeshContainer));
 						InnerTriangleMesh = InstancedTriangleMesh->Object();
 					}
 					else
@@ -2164,7 +2177,17 @@ bool FBodyInstance::UpdateBodyScale(const FVector& InScale3D, bool bForceUpdate)
 
 					if(InnerTriangleMesh)
 					{
-						CreateTriGeomAuto(MoveTempIfPossible(InnerTriangleMesh), NewGeometry, AdjustedScale3D);
+						TUniquePtr<FImplicitObject> NewTrimesh = CreateTriGeomAuto(MoveTempIfPossible(InnerTriangleMesh), NewGeometry, AdjustedScale3D);
+
+						// If we have a transform - wrap the trimesh
+						if(RelativeTM.GetRotation() != FQuat::Identity || RelativeTM.GetTranslation() != FVector::ZeroVector)
+						{
+							FTransform AdjustedTransform = RelativeTM;
+							AdjustedTransform.SetTranslation(RelativeTM.GetTranslation() * AdjustedScale3D);
+							NewTrimesh = MakeUnique<TImplicitObjectTransformed<FReal, 3>>(MoveTemp(NewTrimesh), AdjustedTransform);
+						}
+
+						NewGeometry.Emplace(MoveTemp(NewTrimesh));
 					}
 
 					bSuccess = true;
@@ -2174,7 +2197,7 @@ bool FBodyInstance::UpdateBodyScale(const FVector& InScale3D, bool bForceUpdate)
 				default:
 				{
 					CHAOS_ENSURE(false);
-					UE_LOG(LogPhysics, Warning, TEXT("UpdateBodyScale: Unimplemented ImplicitObject of type: %d skipped."), ImplicitType);
+					UE_LOG(LogPhysics, Warning, TEXT("UpdateBodyScale: Unimplemented ImplicitObject of type: %d skipped."), OuterType);
 				}
 			}// end switch
 		}
@@ -2184,6 +2207,7 @@ bool FBodyInstance::UpdateBodyScale(const FVector& InScale3D, bool bForceUpdate)
 		if (CHAOS_ENSURE(NewGeometry.Num() == Shapes.Num()))
 		{
 			ActorHandle->GetGameThreadAPI().SetGeometry(MakeUnique<Chaos::FImplicitObjectUnion>(MoveTemp(NewGeometry)));
+			FPhysicsInterface::WakeUp_AssumesLocked(ActorHandle);
 		}
 		else
 		{
@@ -2234,13 +2258,13 @@ bool FBodyInstance::UpdateBodyScale(const FVector& InScale3D, bool bForceUpdate)
 			PxGeometry* UpdatedGeometry = NULL;
 			FTransform LocalTransform = FPhysicsInterface::GetLocalTransform(Shape);
 
-			ECollisionShapeType GeomType = FPhysicsInterface::GetShapeType(Shape);
+			ECollisionShapeType WrappedOrConcreteType = FPhysicsInterface::GetShapeType(Shape);
 			FKShapeElem* ShapeElem = FPhysxUserData::Get<FKShapeElem>(FPhysicsInterface::GetUserData(Shape));
 			const FTransform& RelativeTM = GetRelativeBodyTransform(Shape);
 
 			FPhysicsGeometryCollection GeoCollection = FPhysicsInterface::GetGeometryCollection(Shape);
 
-			switch (GeomType)
+			switch (WrappedOrConcreteType)
 			{
 				case ECollisionShapeType::Sphere:
 				{
diff --git a/Engine/Source/Runtime/Experimental/Chaos/Public/Chaos/ImplicitObjectScaled.h b/Engine/Source/Runtime/Experimental/Chaos/Public/Chaos/ImplicitObjectScaled.h
--- a/Engine/Source/Runtime/Experimental/Chaos/Public/Chaos/ImplicitObjectScaled.h
+++ b/Engine/Source/Runtime/Experimental/Chaos/Public/Chaos/ImplicitObjectScaled.h
@@ -976,12 +976,12 @@ public:
 		return MObject->GetMaterialIndex(HintIndex);
 	}
 
-#if 0
+
 	virtual TUniquePtr<FImplicitObject> Copy() const override
 	{
 		return TUniquePtr<FImplicitObject>(CopyHelper(this));
 	}
-#endif
+
 private:
 	ObjectType MObject;
 	TSharedPtr<TConcrete, ESPMode::ThreadSafe> MSharedPtrForRefCount; // Temporary solution to force ref counting on trianglemesh from body setup.
@@ -999,7 +999,10 @@ private:
 
 	static TImplicitObjectScaled<TConcrete, false>* CopyHelper(const TImplicitObjectScaled<TConcrete, false>* Obj)
 	{
-		return new TImplicitObjectScaled<TConcrete, false>(Obj->MObject->Copy(), Obj->MScale, Obj->OuterMargin);
+		TUniquePtr<FImplicitObject> DuplicatedShape = Obj->MObject->Copy();
+		
+		// We know the actual type of the underlying object pointer so we can cast it to the required type to make a copy of this implicit
+		return new TImplicitObjectScaled<TConcrete, false>(reinterpret_cast<TUniquePtr<TConcrete>&&>(DuplicatedShape), Obj->MScale, Obj->OuterMargin);
 	}
 
 	void UpdateBounds()
diff --git a/Engine/Source/Runtime/Experimental/Chaos/Public/Chaos/ImplicitObjectTransformed.h b/Engine/Source/Runtime/Experimental/Chaos/Public/Chaos/ImplicitObjectTransformed.h
--- a/Engine/Source/Runtime/Experimental/Chaos/Public/Chaos/ImplicitObjectTransformed.h
+++ b/Engine/Source/Runtime/Experimental/Chaos/Public/Chaos/ImplicitObjectTransformed.h
@@ -6,6 +6,7 @@
 #include "Chaos/Transform.h"
 #include "ChaosArchive.h"
 #include "Templates/EnableIf.h"
+#include "AABB.h"
 
 namespace Chaos
 {
@@ -213,7 +214,7 @@ public:
 	const TRigidTransform<T, d>& GetTransform() const { return MTransform; }
 	void SetTransform(const TRigidTransform<T, d>& InTransform)
 	{
-		MLocalBoundingBox = MObject->BoundingBox().TransformedBox(InTransform);
+		MLocalBoundingBox = MObject->BoundingBox().TransformedAABB(InTransform);
 		MTransform = InTransform;
 	}
 
@@ -303,6 +304,33 @@ private:
 	friend FImplicitObject;	//needed for serialization
 };
 
+namespace Utilities
+{
+	inline TUniquePtr<FImplicitObject> DuplicateImplicitWithTransform(const FImplicitObject* const InObject, FTransform NewTransform)
+	{
+		if(!InObject)
+		{
+			return nullptr;
+		}
+
+		const EImplicitObjectType OuterType = InObject->GetType();
+
+		if(GetInnerType(OuterType) == ImplicitObjectType::Transformed)
+		{
+			TUniquePtr<FImplicitObject> NewTransformed = InObject->Copy();
+			TImplicitObjectTransformed<FReal, 3>* InnerTransformed = static_cast<TImplicitObjectTransformed<FReal, 3>*>(NewTransformed.Get());
+			InnerTransformed->SetTransform(NewTransform);
+
+			return MoveTemp(NewTransformed);
+		}
+		else
+		{
+			TUniquePtr<FImplicitObject> NewInnerObject = InObject->Copy();
+			return MakeUnique<TImplicitObjectTransformed<FReal, 3>>(MoveTemp(NewInnerObject), NewTransform);
+		}
+	}
+}
+
 template <typename T, int d>
 using TImplicitObjectTransformedNonSerializable = TImplicitObjectTransformed<T, d, false>;
 
diff --git a/Engine/Source/Runtime/PhysicsCore/Private/ChaosEngineInterface.cpp b/Engine/Source/Runtime/PhysicsCore/Private/ChaosEngineInterface.cpp
--- a/Engine/Source/Runtime/PhysicsCore/Private/ChaosEngineInterface.cpp
+++ b/Engine/Source/Runtime/PhysicsCore/Private/ChaosEngineInterface.cpp
@@ -1586,30 +1586,48 @@ FTransform FChaosEngineInterface::GetLocalTransform(const FPhysicsShapeReference
 
 void FChaosEngineInterface::SetLocalTransform(const FPhysicsShapeHandle& InShape,const FTransform& NewLocalTransform)
 {
-#if !WITH_CHAOS_NEEDS_TO_BE_FIXED
-	if(InShape.ActorRef.IsValid())
-	{
-		TArray<RigidBodyId> Ids ={InShape.ActorRef.GetId()};
-		const auto Index = InShape.ActorRef.GetScene()->GetIndexFromId(InShape.ActorRef.GetId());
-		if(InShape.Object->GetType() == Chaos::ImplicitObjectType::Transformed)
-		{
-			// @todo(mlentine): We can avoid creating a new object here by adding delayed update support for the object transforms
-			LocalParticles.SetDynamicGeometry(Index,MakeUnique<Chaos::TImplicitObjectTransformed<Chaos::FReal,3>>(InShape.Object->GetObject<Chaos::TImplicitObjectTransformed<Chaos::FReal,3>>()->Object(),NewLocalTransform));
-		} else
-		{
-			LocalParticles.SetDynamicGeometry(Index,MakeUnique<Chaos::TImplicitObjectTransformed<Chaos::FReal,3>>(InShape.Object,NewLocalTransform));
-		}
-	}
+	using namespace Chaos;
+
+	FSingleParticlePhysicsProxy* Particle = InShape.ActorRef;
+	if(Particle)
 	{
-		if(InShape.Object->GetType() == Chaos::ImplicitObjectType::Transformed)
-		{
-			InShape.Object->GetObject<Chaos::TImplicitObjectTransformed<Chaos::FReal,3>>()->SetTransform(NewLocalTransform);
-		} else
+		Chaos::FRigidBodyHandle_External& BodyHandle = Particle->GetGameThreadAPI();
+
+		const FImplicitObject* CurrentGeom = BodyHandle.Geometry().Get();
+		if(ensure(CurrentGeom && CurrentGeom->GetType() == FImplicitObjectUnion::StaticType()))
 		{
-			const_cast<FPhysicsShapeHandle&>(InShape).Object = new Chaos::TImplicitObjectTransformed<Chaos::FReal,3>(InShape.Object,NewLocalTransform);
+			const FImplicitObjectUnion* AsUnion = static_cast<const FImplicitObjectUnion*>(CurrentGeom);
+			const int32 ShapeIndex = InShape.Shape->GetShapeIndex();
+			const TArray<TUniquePtr<FImplicitObject>>& ObjectArray = AsUnion->GetObjects();
+
+			if(ensure(ShapeIndex < ObjectArray.Num()))
+			{
+				TArray<TUniquePtr<FImplicitObject>> NewGeoms;
+				NewGeoms.Reserve(ObjectArray.Num());
+
+				// Duplicate the union and either set transforms, or wrap in transforms
+				int32 CurrentIndex = 0;
+				for(const TUniquePtr<FImplicitObject>& Obj : ObjectArray)
+				{
+					if(CurrentIndex == ShapeIndex)
+					{
+						NewGeoms.Emplace(Utilities::DuplicateImplicitWithTransform(Obj.Get(), NewLocalTransform));
+					}
+					else
+					{
+						NewGeoms.Emplace(Obj->Copy());
+					}
+
+					CurrentIndex++;
+				}
+
+				if(ensure(NewGeoms.Num() == ObjectArray.Num()))
+				{
+					BodyHandle.SetGeometry(MakeUnique<FImplicitObjectUnion>(MoveTemp(NewGeoms)));
+				}
+			}
 		}
 	}
-#endif
 }
 
 template<typename AllocatorType>
diff --git a/Engine/Source/Runtime/PhysicsCore/Private/ChaosInterfaceWrapperCore.cpp b/Engine/Source/Runtime/PhysicsCore/Private/ChaosInterfaceWrapperCore.cpp
--- a/Engine/Source/Runtime/PhysicsCore/Private/ChaosInterfaceWrapperCore.cpp
+++ b/Engine/Source/Runtime/PhysicsCore/Private/ChaosInterfaceWrapperCore.cpp
@@ -30,14 +30,14 @@ namespace ChaosInterface
 	ECollisionShapeType GetImplicitType(const Chaos::FImplicitObject& InGeometry)
 	{
 		using namespace Chaos;
-		int32 ImplicitObjectType = GetInnerType(InGeometry.GetType());
+		int32 ResultObjectType = GetInnerType(InGeometry.GetType());
 
-		if (ImplicitObjectType == ImplicitObjectType::Transformed)
+		if (ResultObjectType == ImplicitObjectType::Transformed)
 		{
-			ImplicitObjectType = static_cast<const TImplicitObjectTransformed<FReal, 3>*>(&InGeometry)->Object()->GetType();
+			ResultObjectType = GetInnerType(static_cast<const TImplicitObjectTransformed<FReal, 3>*>(&InGeometry)->Object()->GetType());
 		}
 
-		return ImplicitTypeToCollisionType(ImplicitObjectType);
+		return ImplicitTypeToCollisionType(ResultObjectType);
 	}
 
 	Chaos::FReal GetRadius(const Chaos::FCapsule& InCapsule)
