Cleanup invalid level actor folders configurations.

diff --git a/Engine/Source/Runtime/Engine/Private/ActorEditor.cpp b/Engine/Source/Runtime/Engine/Private/ActorEditor.cpp
--- a/Engine/Source/Runtime/Engine/Private/ActorEditor.cpp
+++ b/Engine/Source/Runtime/Engine/Private/ActorEditor.cpp
@@ -1055,30 +1055,46 @@ UActorFolder* AActor::GetActorFolder(bool bSkipDeleted) const
 void AActor::FixupActorFolder()
 {
 	check(GetLevel());
-	check(IsUsingActorFolders(this));
-	
-	// First detect and fixup reference to deleted actor folders
-	UActorFolder* ActorFolder = GetActorFolder(/*bSkipDeleted*/ false);
-	if (ActorFolder)
+
+	if (!IsUsingActorFolders(this))
 	{
-		// Remap to skip deleted actor folder
-		if (ActorFolder->IsMarkedAsDeleted())
+		if (FolderGuid.IsValid())
 		{
-			ActorFolder = ActorFolder->GetParent();
-			SetFolderGuidInternal(ActorFolder ? ActorFolder->GetGuid() : FGuid(), /*bBroadcastChange*/ false);
+			UE_LOG(LogLevel, Warning, TEXT("Actor folder %s for actor %s encountered when not using actor folders"), *FolderGuid.ToString(), *GetName());
+			FolderGuid = FGuid();
 		}
-		// We found actor folder using its path, update actor folder guid
-		else if (!FolderPath.IsNone())
+	}
+	else
+	{
+		// First detect and fixup reference to deleted actor folders
+		UActorFolder* ActorFolder = GetActorFolder(/*bSkipDeleted*/ false);
+		if (ActorFolder)
 		{
-			SetFolderGuidInternal(ActorFolder ? ActorFolder->GetGuid() : FGuid(), /*bBroadcastChange*/ false);
+			// Remap to skip deleted actor folder
+			if (ActorFolder->IsMarkedAsDeleted())
+			{
+				ActorFolder = ActorFolder->GetParent();
+				SetFolderGuidInternal(ActorFolder ? ActorFolder->GetGuid() : FGuid(), /*bBroadcastChange*/ false);
+			}
+			// We found actor folder using its path, update actor folder guid
+			else if (!FolderPath.IsNone())
+			{
+				SetFolderGuidInternal(ActorFolder ? ActorFolder->GetGuid() : FGuid(), /*bBroadcastChange*/ false);
+			}
 		}
-	}
 
-	// If still invalid, warn and fallback to root
-	if (!IsActorFolderValid())
-	{
-		UE_LOG(LogLevel, Warning, TEXT("Missing actor folder for actor %s"), *GetName());
-		SetFolderGuidInternal(FGuid(), /*bBroadcastChange*/ false);
+		// If still invalid, warn and fallback to root
+		if (!IsActorFolderValid())
+		{
+			UE_LOG(LogLevel, Warning, TEXT("Missing actor folder for actor %s"), *GetName());
+			SetFolderGuidInternal(FGuid(), /*bBroadcastChange*/ false);
+		}
+
+		if (!FolderPath.IsNone())
+		{
+			UE_LOG(LogLevel, Warning, TEXT("Actor folder path %s for actor %s encountered when using actor folders"), *FolderPath.ToString(), *GetName());
+			FolderPath = NAME_None;
+		}
 	}
 }
 
diff --git a/Engine/Source/Runtime/Engine/Private/Level.cpp b/Engine/Source/Runtime/Engine/Private/Level.cpp
--- a/Engine/Source/Runtime/Engine/Private/Level.cpp
+++ b/Engine/Source/Runtime/Engine/Private/Level.cpp
@@ -2323,48 +2323,51 @@ void ULevel::CommitModelSurfaces()
 #if WITH_EDITOR
 void ULevel::FixupActorFolders()
 {
-	if (!IsUsingActorFolders() || !IsActorFolderObjectsFeatureAvailable())
+	if (!IsActorFolderObjectsFeatureAvailable())
 	{
 		return;
 	}
 
-	// At this point, LoadedExternalActorFolders are fully loaded, transfer them to the ActorFolders list.
-	for (UActorFolder* LoadedActorFolder : LoadedExternalActorFolders)
+	if (IsUsingActorFolders())
 	{
-		check(LoadedActorFolder->GetGuid().IsValid());
-		ActorFolders.Add(LoadedActorFolder->GetGuid(), LoadedActorFolder);
-	}
-	LoadedExternalActorFolders.Empty();
-	
-	// Discover duplicate paths first to prioritize non-duplicate paths
-	TSet<FName> FolderPaths;
-	TArray<UActorFolder*> DuplicateFolders;
-	ForEachActorFolder([&FolderPaths, &DuplicateFolders](UActorFolder* ActorFolder)
-	{
-		// Detects and clears invalid parent folder
-		ActorFolder->FixupParentFolder();
-
-		bool bIsAlreadyInSet = false;
-		FolderPaths.Add(ActorFolder->GetPath(), &bIsAlreadyInSet);
-		if (bIsAlreadyInSet)
+		// At this point, LoadedExternalActorFolders are fully loaded, transfer them to the ActorFolders list.
+		for (UActorFolder* LoadedActorFolder : LoadedExternalActorFolders)
 		{
-			DuplicateFolders.Add(ActorFolder);
+			check(LoadedActorFolder->GetGuid().IsValid());
+			ActorFolders.Add(LoadedActorFolder->GetGuid(), LoadedActorFolder);
 		}
-		return true;
-	}, /*bSkipDeleted*/ true);
-
-	if (!IsRunningCommandlet())
-	{
-		// Rename duplicates to a new valid/unique name
-		for (UActorFolder* ActorFolder : DuplicateFolders)
+		LoadedExternalActorFolders.Empty();
+	
+		// Discover duplicate paths first to prioritize non-duplicate paths
+		TSet<FName> FolderPaths;
+		TArray<UActorFolder*> DuplicateFolders;
+		ForEachActorFolder([&FolderPaths, &DuplicateFolders](UActorFolder* ActorFolder)
 		{
-			FFolder Folder = ActorFolder->GetFolder();
-			const FFolder NewPath = FActorFolders::Get().GetFolderName(*GetWorld(), Folder.GetParent(), Folder.GetLeafName());
-			ActorFolder->SetLabel(NewPath.GetLeafName().ToString());
+			// Detects and clears invalid parent folder
+			ActorFolder->FixupParentFolder();
+
 			bool bIsAlreadyInSet = false;
 			FolderPaths.Add(ActorFolder->GetPath(), &bIsAlreadyInSet);
-			check(!bIsAlreadyInSet);
-			UE_LOG(LogLevel, Warning, TEXT("Found duplicate actor folder %s, renamed to %s."), *Folder.GetPath().ToString(), *NewPath.GetPath().ToString());
+			if (bIsAlreadyInSet)
+			{
+				DuplicateFolders.Add(ActorFolder);
+			}
+			return true;
+		}, /*bSkipDeleted*/ true);
+
+		if (!IsRunningCommandlet())
+		{
+			// Rename duplicates to a new valid/unique name
+			for (UActorFolder* ActorFolder : DuplicateFolders)
+			{
+				FFolder Folder = ActorFolder->GetFolder();
+				const FFolder NewPath = FActorFolders::Get().GetFolderName(*GetWorld(), Folder.GetParent(), Folder.GetLeafName());
+				ActorFolder->SetLabel(NewPath.GetLeafName().ToString());
+				bool bIsAlreadyInSet = false;
+				FolderPaths.Add(ActorFolder->GetPath(), &bIsAlreadyInSet);
+				check(!bIsAlreadyInSet);
+				UE_LOG(LogLevel, Warning, TEXT("Found duplicate actor folder %s, renamed to %s."), *Folder.GetPath().ToString(), *NewPath.GetPath().ToString());
+			}
 		}
 	}
 	
@@ -2380,11 +2383,14 @@ void ULevel::FixupActorFolders()
 		}
 	}
 	
-	ForEachActorFolder([](UActorFolder* ActorFolder)
+	if (IsUsingActorFolders())
 	{
-		GEngine->BroadcastActorFolderAdded(ActorFolder);
-		return true;
-	}, /*bSkipDeleted*/ true);
+		ForEachActorFolder([](UActorFolder* ActorFolder)
+		{
+			GEngine->BroadcastActorFolderAdded(ActorFolder);
+			return true;
+		}, /*bSkipDeleted*/ true);
+	}
 }
 #endif
 
