Added AWorldSettings::bEnableLargeWorlds wrapping disabling bounds checks, octree visibility, and octree shadow culling.

diff --git a/Engine/Plugins/Runtime/GeoReferencing/Source/GeoReferencing/Private/RoundPlanetPawn.cpp b/Engine/Plugins/Runtime/GeoReferencing/Source/GeoReferencing/Private/RoundPlanetPawn.cpp
--- a/Engine/Plugins/Runtime/GeoReferencing/Source/GeoReferencing/Private/RoundPlanetPawn.cpp
+++ b/Engine/Plugins/Runtime/GeoReferencing/Source/GeoReferencing/Private/RoundPlanetPawn.cpp
@@ -178,7 +178,7 @@ void ARoundPlanetPawn::BeginPlay()
 
 	// Round planet is huge, if we enable world bounds check, the pawn behavior might be unpredictable when flying...
 	// LWC_TODO - should not be needed after WORLD_MAX updates. 
-	GetWorldSettings()->bEnableWorldBoundsChecks = false;
+	GetWorldSettings()->bEnableLargeWorlds = true;
 
 	if (!GeoReferencingSystem)
 	{
diff --git a/Engine/Source/Editor/UnrealEd/Private/EditorComponents.cpp b/Engine/Source/Editor/UnrealEd/Private/EditorComponents.cpp
--- a/Engine/Source/Editor/UnrealEd/Private/EditorComponents.cpp
+++ b/Engine/Source/Editor/UnrealEd/Private/EditorComponents.cpp
@@ -531,7 +531,7 @@ void FEditorCommonDrawHelper::DrawOldGrid(const FSceneView* View,FPrimitiveDrawI
 			DrawOriginAxisLine( &StartY, &EndY, &StartY.Z, &EndY.Z, View, PDI, AxisColors[1] );
 		}
 
-		if( bDrawKillZ && ( bIsOrthoXZ || bIsOrthoYZ ) && GWorld->GetWorldSettings()->bEnableWorldBoundsChecks )
+		if( bDrawKillZ && ( bIsOrthoXZ || bIsOrthoYZ ) && GWorld->GetWorldSettings()->AreWorldBoundsChecksEnabled() )
 		{
 			float KillZ = GWorld->GetWorldSettings()->KillZ;
 
diff --git a/Engine/Source/Runtime/Engine/Classes/GameFramework/WorldSettings.h b/Engine/Source/Runtime/Engine/Classes/GameFramework/WorldSettings.h
--- a/Engine/Source/Runtime/Engine/Classes/GameFramework/WorldSettings.h
+++ b/Engine/Source/Runtime/Engine/Classes/GameFramework/WorldSettings.h
@@ -445,8 +445,13 @@ class ENGINE_API AWorldSettings : public AInfo, public IInterface_AssetUserData
 
 	/** DEFAULT BASIC PHYSICS SETTINGS **/
 
-	/** If true, enables CheckStillInWorld checks */
+	/** If true, configures engine for large world testing. Disables CheckStillInWorld checks and octree visibility testing.
+	 *	See UE_USE_UE4_WORLD_MAX for a more correct alternative. This setting will be removed once UE_LARGE_WORLD_MAX is considered stable. */
 	UPROPERTY(EditAnywhere, BlueprintReadOnly, Category=World, AdvancedDisplay)
+	uint8 bEnableLargeWorlds:1;
+
+	/** If true, enables CheckStillInWorld checks. Note: Do not set this manually if experimenting with large worlds. @see bEnableLargeWorlds */
+	UPROPERTY(EditAnywhere, BlueprintReadOnly, Category=World, AdvancedDisplay, meta=(EditCondition = "!bEnableLargeWorlds"))
 	uint8 bEnableWorldBoundsChecks:1;
 
 protected:
@@ -584,11 +589,11 @@ public:
 	float WorldToMeters;
 
 	// any actor falling below this level gets destroyed
-	UPROPERTY(config, EditAnywhere, BlueprintReadOnly, Category=World, meta=(editcondition = "bEnableWorldBoundsChecks"))
+	UPROPERTY(config, EditAnywhere, BlueprintReadOnly, Category=World, meta=(editcondition = "bEnableWorldBoundsChecks && !bEnableLargeWorlds"))
 	float KillZ;
 
 	// The type of damage inflicted when a actor falls below KillZ
-	UPROPERTY(EditAnywhere, BlueprintReadOnly, Category=World, AdvancedDisplay, meta=(editcondition = "bEnableWorldBoundsChecks"))
+	UPROPERTY(EditAnywhere, BlueprintReadOnly, Category=World, AdvancedDisplay, meta=(editcondition = "bEnableWorldBoundsChecks && !bEnableLargeWorlds"))
 	TSubclassOf<UDamageType> KillZDamageType;
 
 	// current gravity actually being used
@@ -828,6 +833,9 @@ public:
 	/** @return whether given world is configured to host an AISystem */
 	bool IsAISystemEnabled() const { return bEnableAISystem; }
 
+	/** @return whether given world is restricting actors to +-HALF_WORLD_MAX bounds, destroying actors that move below KillZ */
+	bool AreWorldBoundsChecksEnabled() const { return bEnableWorldBoundsChecks && !bEnableLargeWorlds; }
+	
 	/**
 	 * Called from GameStateBase, calls BeginPlay on all actors
 	 */
@@ -878,6 +886,8 @@ private:
 	void SanitizeBookmarkClasses();
 	void UpdateBookmarkClass();
 
+	void UpdateEnableLargeWorldsCVars(bool bEnable) const;
+
 	/**
 	 * Maximum number of bookmarks allowed.
 	 * Changing this will change the allocation of the bookmarks array, and when shrinking
diff --git a/Engine/Source/Runtime/Engine/Private/Actor.cpp b/Engine/Source/Runtime/Engine/Private/Actor.cpp
--- a/Engine/Source/Runtime/Engine/Private/Actor.cpp
+++ b/Engine/Source/Runtime/Engine/Private/Actor.cpp
@@ -1696,7 +1696,7 @@ bool AActor::CheckStillInWorld()
 	// check the variations of KillZ
 	AWorldSettings* WorldSettings = MyWorld->GetWorldSettings( true );
 
-	if (!WorldSettings->bEnableWorldBoundsChecks)
+	if (!WorldSettings->AreWorldBoundsChecksEnabled())
 	{
 		return true;
 	}
diff --git a/Engine/Source/Runtime/Engine/Private/Components/InterpToMovementComponent.cpp b/Engine/Source/Runtime/Engine/Private/Components/InterpToMovementComponent.cpp
--- a/Engine/Source/Runtime/Engine/Private/Components/InterpToMovementComponent.cpp
+++ b/Engine/Source/Runtime/Engine/Private/Components/InterpToMovementComponent.cpp
@@ -371,7 +371,7 @@ bool UInterpToMovementComponent::CheckStillInWorld()
 	}
 	// check the variations of KillZ
 	AWorldSettings* WorldSettings = MyWorld->GetWorldSettings(true);
-	if (!WorldSettings->bEnableWorldBoundsChecks)
+	if (!WorldSettings->AreWorldBoundsChecksEnabled())
 	{
 		return true;
 	}
diff --git a/Engine/Source/Runtime/Engine/Private/Components/ProjectileMovementComponent.cpp b/Engine/Source/Runtime/Engine/Private/Components/ProjectileMovementComponent.cpp
--- a/Engine/Source/Runtime/Engine/Private/Components/ProjectileMovementComponent.cpp
+++ b/Engine/Source/Runtime/Engine/Private/Components/ProjectileMovementComponent.cpp
@@ -621,7 +621,7 @@ bool UProjectileMovementComponent::CheckStillInWorld()
 
 	// check the variations of KillZ
 	AWorldSettings* WorldSettings = MyWorld->GetWorldSettings( true );
-	if (!WorldSettings->bEnableWorldBoundsChecks)
+	if (!WorldSettings->AreWorldBoundsChecksEnabled())
 	{
 		return true;
 	}
diff --git a/Engine/Source/Runtime/Engine/Private/FloatingPawnMovement.cpp b/Engine/Source/Runtime/Engine/Private/FloatingPawnMovement.cpp
--- a/Engine/Source/Runtime/Engine/Private/FloatingPawnMovement.cpp
+++ b/Engine/Source/Runtime/Engine/Private/FloatingPawnMovement.cpp
@@ -84,7 +84,7 @@ void UFloatingPawnMovement::TickComponent(float DeltaTime, enum ELevelTick TickT
 bool UFloatingPawnMovement::LimitWorldBounds()
 {
 	AWorldSettings* WorldSettings = PawnOwner ? PawnOwner->GetWorldSettings() : NULL;
-	if (!WorldSettings || !WorldSettings->bEnableWorldBoundsChecks || !UpdatedComponent)
+	if (!WorldSettings || !WorldSettings->AreWorldBoundsChecksEnabled() || !UpdatedComponent)
 	{
 		return false;
 	}
diff --git a/Engine/Source/Runtime/Engine/Private/WorldSettings.cpp b/Engine/Source/Runtime/Engine/Private/WorldSettings.cpp
--- a/Engine/Source/Runtime/Engine/Private/WorldSettings.cpp
+++ b/Engine/Source/Runtime/Engine/Private/WorldSettings.cpp
@@ -70,6 +70,7 @@ AWorldSettings::AWorldSettings(const FObjectInitializer& ObjectInitializer)
 	};
 	static FConstructorStatics ConstructorStatics;
 
+	bEnableLargeWorlds = false;
 	bEnableWorldBoundsChecks = true;
 PRAGMA_DISABLE_DEPRECATION_WARNINGS
 	bEnableNavigationSystem = true;
@@ -209,6 +210,11 @@ void AWorldSettings::PostRegisterAllComponents()
 	{
 		AudioDevice->SetDefaultAudioSettings(World, DefaultReverbSettings, DefaultAmbientZoneSettings);
 	}
+	
+	if(bEnableLargeWorlds)
+	{
+		UpdateEnableLargeWorldsCVars(bEnableLargeWorlds);
+	}	
 }
 
 UWorldPartition* AWorldSettings::GetWorldPartition() const
@@ -272,6 +278,12 @@ void AWorldSettings::NotifyBeginPlay()
 
 		World->bBegunPlay = true;
 	}
+
+	if(bEnableLargeWorlds)
+	{
+		// Done post BeginPlay to give code a chance to set bEnableLargeWorlds.
+		UpdateEnableLargeWorldsCVars(bEnableLargeWorlds);
+	}	
 }
 
 void AWorldSettings::NotifyMatchStarted()
@@ -281,6 +293,13 @@ void AWorldSettings::NotifyMatchStarted()
 	World->bMatchStarted = true;
 }
 
+void AWorldSettings::UpdateEnableLargeWorldsCVars(bool bEnable) const
+{
+	// LWC_TODO: Large world support. This will be removed once UE_LARGE_WORLD_MAX is stable.
+	IConsoleManager::Get().FindConsoleVariable(TEXT("r.UseVisibilityOctree"))->Set(!bEnable);
+	IConsoleManager::Get().FindConsoleVariable(TEXT("r.Shadow.UseOctreeForCulling"))->Set(!bEnable); 
+}
+
 void AWorldSettings::GetLifetimeReplicatedProps( TArray< FLifetimeProperty > & OutLifetimeProps ) const
 {
 	Super::GetLifetimeReplicatedProps( OutLifetimeProps );
@@ -756,6 +775,10 @@ void AWorldSettings::InternalPostPropertyChanged(FName PropertyName)
 	{
 		UpdateBookmarkClass();
 	}
+	else if(PropertyName == GET_MEMBER_NAME_CHECKED(AWorldSettings, bEnableLargeWorlds))
+	{
+		UpdateEnableLargeWorldsCVars(bEnableLargeWorlds);
+	}
 
 	if (GetWorld() != nullptr && GetWorld()->PersistentLevel && GetWorld()->PersistentLevel->GetWorldSettings() == this)
 	{
diff --git a/Engine/Source/Runtime/Renderer/Private/ShadowSetup.cpp b/Engine/Source/Runtime/Renderer/Private/ShadowSetup.cpp
--- a/Engine/Source/Runtime/Renderer/Private/ShadowSetup.cpp
+++ b/Engine/Source/Runtime/Renderer/Private/ShadowSetup.cpp
@@ -4592,7 +4592,7 @@ struct FGatherShadowPrimitivesPrepareTask
 		}
 		else
 		{
-			const int32 PacketSize = CVarParallelGatherNumPrimitivesPerPacket.GetValueOnRenderThread();
+			const int32 PacketSize = CVarParallelGatherNumPrimitivesPerPacket.GetValueOnAnyThread();
 			const int32 NumPackets = FMath::DivideAndRoundUp(TaskData.Scene->Primitives.Num(), PacketSize);
 
 			TaskData.Packets.Reserve(NumPackets);
