Clear out bad compressed animation data following on from CL

diff --git a/Engine/Source/Runtime/Engine/Private/Animation/AnimCompressionDerivedData.cpp b/Engine/Source/Runtime/Engine/Private/Animation/AnimCompressionDerivedData.cpp
--- a/Engine/Source/Runtime/Engine/Private/Animation/AnimCompressionDerivedData.cpp
+++ b/Engine/Source/Runtime/Engine/Private/Animation/AnimCompressionDerivedData.cpp
@@ -35,7 +35,7 @@ const TCHAR* FDerivedDataAnimationCompression::GetVersionString() const
 	// This is a version string that mimics the old versioning scheme. If you
 	// want to bump this version, generate a new guid using VS->Tools->Create GUID and
 	// return it here. Ex.
-	return TEXT("5A39C0ABF95744778F168FC8DA1619D4");
+	return TEXT("3A3F9FBB7A4047278CABE35820CC44D7");
 }
 
 bool FDerivedDataAnimationCompression::Build( TArray<uint8>& OutDataArray )
