In UStaticMeshComponent::PostEditUndo, for the resource initialization check, use GetCurrentFirstLODIdx on the render data instead of 0 to be sure to use a valid index.

diff --git a/Engine/Source/Runtime/Engine/Private/Components/StaticMeshComponent.cpp b/Engine/Source/Runtime/Engine/Private/Components/StaticMeshComponent.cpp
--- a/Engine/Source/Runtime/Engine/Private/Components/StaticMeshComponent.cpp
+++ b/Engine/Source/Runtime/Engine/Private/Components/StaticMeshComponent.cpp
@@ -1765,7 +1765,7 @@ void UStaticMeshComponent::PostEditUndo()
 	// Debug check command trying to track down undo related uninitialized resource
 	if (GetStaticMesh() != NULL && GetStaticMesh()->GetRenderData() && GetStaticMesh()->GetRenderData()->LODResources.Num() > 0)
 	{
-		FRenderResource* Resource = &GetStaticMesh()->GetRenderData()->LODResources[0].IndexBuffer;
+		FRenderResource* Resource = &GetStaticMesh()->GetRenderData()->LODResources[GetStaticMesh()->GetRenderData()->GetCurrentFirstLODIdx(0)].IndexBuffer;
 		ENQUEUE_RENDER_COMMAND(ResourceCheckCommand)(
 			[Resource](FRHICommandList& RHICmdList)
 			{
