Reduce Low Mesh bias amount based on MinLODIdx for static meshes.

diff --git a/Engine/Source/Runtime/Engine/Private/StaticMeshRender.cpp b/Engine/Source/Runtime/Engine/Private/StaticMeshRender.cpp
--- a/Engine/Source/Runtime/Engine/Private/StaticMeshRender.cpp
+++ b/Engine/Source/Runtime/Engine/Private/StaticMeshRender.cpp
@@ -249,7 +249,10 @@ FStaticMeshSceneProxy::FStaticMeshSceneProxy(UStaticMeshComponent* InComponent,
 	extern int32 GUseMobileLODBiasOnDesktopES31;
 	if (GUseMobileLODBiasOnDesktopES31 != 0 && FeatureLevel == ERHIFeatureLevel::ES3_1)
 	{
-		EffectiveMinLOD += InComponent->GetStaticMesh()->GetRenderData()->LODBiasModifier;
+		int32 BiasModifier = InComponent->GetStaticMesh()->GetRenderData()->LODBiasModifier;
+		int32 MinLODIdx = FMath::Max(StaticMesh->GetMinLODIdx(), 0);
+		BiasModifier = FMath::Max(BiasModifier - MinLODIdx, 0);
+		EffectiveMinLOD += BiasModifier;
 	}
 #endif
 
diff --git a/Engine/Source/Runtime/Engine/Private/Streaming/StreamingTexture.cpp b/Engine/Source/Runtime/Engine/Private/Streaming/StreamingTexture.cpp
--- a/Engine/Source/Runtime/Engine/Private/Streaming/StreamingTexture.cpp
+++ b/Engine/Source/Runtime/Engine/Private/Streaming/StreamingTexture.cpp
@@ -211,7 +211,21 @@ void FStreamingRenderAsset::UpdateDynamicData(const int32* NumStreamedMips, int3
 #if PLATFORM_DESKTOP
 			if (GUseMobileLODBiasOnDesktopES31 != 0 && GMaxRHIFeatureLevel == ERHIFeatureLevel::ES3_1)
 			{
-				LODBias += ResourceState.LODBiasModifier;
+				if (RenderAssetType == EStreamableRenderAssetType::StaticMesh)
+				{
+					const UStaticMesh* StaticMesh = Cast<UStaticMesh>(RenderAsset);
+					if (StaticMesh)
+					{
+						int32 BiasModifier = ResourceState.LODBiasModifier;
+						int32 MinLODIdx = FMath::Max(StaticMesh->GetMinLODIdx(), 0);
+						BiasModifier = FMath::Max(BiasModifier - MinLODIdx, 0);
+						LODBias += BiasModifier;
+					}
+				}
+				else
+				{
+					LODBias += ResourceState.LODBiasModifier;
+				}
 			}
 #endif
 		}
