Fixed race condition in ISM when track bounces is disabled

diff --git a/Engine/Source/Runtime/Engine/Classes/Engine/InstancedStaticMesh.h b/Engine/Source/Runtime/Engine/Classes/Engine/InstancedStaticMesh.h
--- a/Engine/Source/Runtime/Engine/Classes/Engine/InstancedStaticMesh.h
+++ b/Engine/Source/Runtime/Engine/Classes/Engine/InstancedStaticMesh.h
@@ -432,13 +432,13 @@ private:
 	 */
 	void UpdateBoundsTransforms_Concurrent();
 	void UpdateBoundsTransforms();
-	void EnsureInstanceDataUpdated();
+	void EnsureInstanceDataUpdated(bool bForceUpdate = false);
 
 	TArray<FVector4f> PerInstanceBounds;
 	TArray<FRenderTransform> PerInstanceTransforms;
 	FGraphEventRef UpdateBoundsTask;
 	const FBox InstanceLocalBounds;
-	const bool bTrackBounds;
+	bool bTrackBounds;
 	bool bBoundsTransformsDirty;
 };
 
diff --git a/Engine/Source/Runtime/Engine/Private/InstancedStaticMesh.cpp b/Engine/Source/Runtime/Engine/Private/InstancedStaticMesh.cpp
--- a/Engine/Source/Runtime/Engine/Private/InstancedStaticMesh.cpp
+++ b/Engine/Source/Runtime/Engine/Private/InstancedStaticMesh.cpp
@@ -1189,7 +1189,7 @@ void FPerInstanceRenderData::UpdateBoundsTransforms()
 }
 
 
-void FPerInstanceRenderData::EnsureInstanceDataUpdated()
+void FPerInstanceRenderData::EnsureInstanceDataUpdated(bool bForceUpdate)
 {
 	check(IsInRenderingThread());
 
@@ -1202,7 +1202,7 @@ void FPerInstanceRenderData::EnsureInstanceDataUpdated()
 	}
 
 	// manually update if there is no pending update task
-	if (bBoundsTransformsDirty)
+	if (bBoundsTransformsDirty || bForceUpdate)
 	{
 		UpdateBoundsTransforms();
 		bBoundsTransformsDirty = false;
@@ -1211,8 +1211,10 @@ void FPerInstanceRenderData::EnsureInstanceDataUpdated()
 
 const TArray<FVector4f>& FPerInstanceRenderData::GetPerInstanceBounds()
 {
-	check(bTrackBounds);
-	EnsureInstanceDataUpdated();
+	// if bounds tracking has not been enabled, it needs to be enabled and instance data must be forced to update
+	bool bForceUpdate = !bTrackBounds;
+	bTrackBounds = true;
+	EnsureInstanceDataUpdated(bForceUpdate);
 	return PerInstanceBounds;
 }
 
