Don't allocate PendingOutRec for internal ack (replay) connections
Make max channels a connection config value, as well as an override that can be set on a driver definition (net.maxchannelsize overrides both)

diff --git a/Engine/Source/Runtime/Engine/Classes/Engine/Engine.h b/Engine/Source/Runtime/Engine/Classes/Engine/Engine.h
--- a/Engine/Source/Runtime/Engine/Classes/Engine/Engine.h
+++ b/Engine/Source/Runtime/Engine/Classes/Engine/Engine.h
@@ -224,10 +224,14 @@ struct FNetDriverDefinition
 	UPROPERTY()
 	FName DriverClassNameFallback;
 
+	UPROPERTY()
+	int32 MaxChannelsOverride;
+
 	FNetDriverDefinition() :
 		DefName(NAME_None),
 		DriverClassName(NAME_None),
-		DriverClassNameFallback(NAME_None)
+		DriverClassNameFallback(NAME_None),
+		MaxChannelsOverride(INDEX_NONE)
 	{
 	}
 };
diff --git a/Engine/Source/Runtime/Engine/Classes/Engine/NetConnection.h b/Engine/Source/Runtime/Engine/Classes/Engine/NetConnection.h
--- a/Engine/Source/Runtime/Engine/Classes/Engine/NetConnection.h
+++ b/Engine/Source/Runtime/Engine/Classes/Engine/NetConnection.h
@@ -329,10 +329,10 @@ public:
 	ENGINE_API virtual ~UNetConnection() {};
 
 private:
-	uint32 bInternalAck : 1;	// Internally ack all packets, for 100% reliable connections.
-	uint32 bReplay : 1;			// Flag to indicate a replay connection, independent of reliability
-	uint32 bForceInitialDirty : 1;	// Force all properties dirty on initial replication
-	uint32 bUnlimitedBunchSizeAllowed : 1; // Ignore the value of net.MaxConstructedPartialBunchSizeBytes
+	uint32 bInternalAck : 1;				// Internally ack all packets, for 100% reliable connections.
+	uint32 bReplay : 1;						// Flag to indicate a replay connection, independent of reliability
+	uint32 bForceInitialDirty : 1;			// Force all properties dirty on initial replication
+	uint32 bUnlimitedBunchSizeAllowed : 1;	// Ignore the value of net.MaxConstructedPartialBunchSizeBytes
 
 public:
 	struct FURL			URL;				// URL of the other side.
@@ -570,9 +570,16 @@ public:
 	bool			bLastHasServerFrameTime;
 
 	// Channel table.
+
+	UPROPERTY(config)
+	int32 DefaultMaxChannelSize;
+
+	UE_DEPRECATED(5.1, "Deprecated in favor of DefaultMaxChannelSize config property.")
 	static const int32 DEFAULT_MAX_CHANNEL_SIZE;
 
+	UE_DEPRECATED(5.1, "No longer used")
 	int32 MaxChannelSize;
+
 	TArray<UChannel*>	Channels;
 	TArray<int32>		OutReliable;
 	TArray<int32>		InReliable;
@@ -675,6 +682,8 @@ private:
 
 	UReplicationConnectionDriver* ReplicationConnectionDriver;
 
+	void InitChannelData();
+
 	/** Engine package version for compatibility */
 	FPackageFileVersion PackageVersionUE;
 
diff --git a/Engine/Source/Runtime/Engine/Classes/Engine/NetDriver.h b/Engine/Source/Runtime/Engine/Classes/Engine/NetDriver.h
--- a/Engine/Source/Runtime/Engine/Classes/Engine/NetDriver.h
+++ b/Engine/Source/Runtime/Engine/Classes/Engine/NetDriver.h
@@ -872,7 +872,15 @@ private:
 	UPROPERTY()
 	TArray<TObjectPtr<UChannel>> ActorChannelPool;
 
+	/** Name of net driver definition used to create this driver */
+	FName NetDriverDefinition;
+
+	/** Cached copy of MaxChannelsOverride from the net driver definition to avoid extra lookups */
+	int32 MaxChannelsOverride;
+
 public:
+	/** Get the value of MaxChannelsOverride cached from the net driver definition */
+	int32 GetMaxChannelsOverride() const { return MaxChannelsOverride; }
 
 	/** Creates a new channel of the specified type name. If the type is pooled, it will return a pre-created channel */
 	UChannel* GetOrCreateChannelByName(const FName& ChName);
@@ -883,6 +891,12 @@ public:
 	/** Change the NetDriver's NetDriverName. This will also reinit packet simulation settings so that settings can be qualified to a specific driver. */
 	void SetNetDriverName(FName NewNetDriverNamed);
 
+	/** Set the NetDriver's NetDriverDefintion. */
+	void SetNetDriverDefinition(FName NewNetDriverDefinition);
+
+	/** Get the NetDriver's NetDriverDefintion. */
+	FName GetNetDriverDefintion() const { return NetDriverDefinition; }
+
 	void InitPacketSimulationSettings();
 
 	/** Returns true during the duration of a packet loss burst triggered by the net.pktlossburst command. */
diff --git a/Engine/Source/Runtime/Engine/Private/DataChannel.cpp b/Engine/Source/Runtime/Engine/Private/DataChannel.cpp
--- a/Engine/Source/Runtime/Engine/Private/DataChannel.cpp
+++ b/Engine/Source/Runtime/Engine/Private/DataChannel.cpp
@@ -2423,7 +2423,8 @@ void UActorChannel::SetChannelActor(AActor* InActor, ESetChannelActorFlags Flags
 
 	UE_LOG(LogNetTraffic, VeryVerbose, TEXT("SetChannelActor: ChIndex: %i, Actor: %s, NetGUID: %s"), ChIndex, Actor ? *Actor->GetFullName() : TEXT("NULL"), *ActorNetGUID.ToString() );
 
-	if (ChIndex >= 0 && Connection->PendingOutRec[ChIndex] > 0)
+	// Internal ack connections never add to PendingOutRec
+	if (Connection->PendingOutRec.IsValidIndex(ChIndex) && Connection->PendingOutRec[ChIndex] > 0)
 	{
 		// send empty reliable bunches to synchronize both sides
 		// UE_LOG(LogNetTraffic, Log, TEXT("%i Actor %s WILL BE sending %i vs first %i"), ChIndex, *Actor->GetName(), Connection->PendingOutRec[ChIndex],Connection->OutReliable[ChIndex]);
diff --git a/Engine/Source/Runtime/Engine/Private/NetConnection.cpp b/Engine/Source/Runtime/Engine/Private/NetConnection.cpp
--- a/Engine/Source/Runtime/Engine/Private/NetConnection.cpp
+++ b/Engine/Source/Runtime/Engine/Private/NetConnection.cpp
@@ -62,8 +62,8 @@ static TAutoConsoleVariable<int32> CVarTickAllOpenChannels(TEXT("net.TickAllOpen
 static TAutoConsoleVariable<int32> CVarRandomizeSequence(TEXT("net.RandomizeSequence"), 1,
 	TEXT("Randomize initial packet sequence, can provide some obfuscation"));
 
-static TAutoConsoleVariable<int32> CVarMaxChannelSize(TEXT("net.MaxChannelSize"), UNetConnection::DEFAULT_MAX_CHANNEL_SIZE,
-	TEXT("The maximum number of network channels allowed across the entire server"));
+static TAutoConsoleVariable<int32> CVarMaxChannelSize(TEXT("net.MaxChannelSize"), 0,
+	TEXT("The maximum number of network channels allowed across the entire server, if <= 0 the connection DefaultMaxChannelSize will be used."));
 
 #if !UE_BUILD_SHIPPING
 static TAutoConsoleVariable<int32> CVarForceNetFlush(TEXT("net.ForceNetFlush"), 0,
@@ -205,7 +205,10 @@ static SIZE_T CountBytes(FWrittenChannelsRecord& WrittenChannelsRecord);
 
 };
 
+PRAGMA_DISABLE_DEPRECATION_WARNINGS
 const int32 UNetConnection::DEFAULT_MAX_CHANNEL_SIZE = 32767;
+PRAGMA_ENABLE_DEPRECATION_WARNINGS
+
 /*-----------------------------------------------------------------------------
 	UNetConnection implementation.
 -----------------------------------------------------------------------------*/
@@ -220,7 +223,6 @@ UNetConnection::UNetConnection(const FObjectInitializer& ObjectInitializer)
 ,	ViewTarget			( nullptr )
 ,   OwningActor			( nullptr )
 ,	MaxPacket			( 0 )
-,	InternalAck			( false )
 ,	bInternalAck		( false )
 ,	bReplay				( false )
 ,	bForceInitialDirty	( false )
@@ -236,7 +238,6 @@ UNetConnection::UNetConnection(const FObjectInitializer& ObjectInitializer)
 ,	QueuedBits			( 0 )
 ,	TickCount			( 0 )
 ,	LastProcessedFrame	( 0 )
-,	ConnectTime			( 0.0 )
 ,	LastOSReceiveTime	()
 ,	bIsOSReceiveTimeLocal(false)
 ,   PreviousJitterTimeDelta(0)
@@ -279,12 +280,13 @@ UNetConnection::UNetConnection(const FObjectInitializer& ObjectInitializer)
 ,	OutPacketId			( 0 ) // must be initialized as OutAckPacketId + 1 so loss of first packet can be detected
 ,	OutAckPacketId		( -1 )
 ,	bLastHasServerFrameTime( false )
+	, DefaultMaxChannelSize(32767)
 ,	InitOutReliable		( 0 )
 ,	InitInReliable		( 0 )
 ,	EngineNetworkProtocolVersion( FNetworkVersion::GetEngineNetworkProtocolVersion() )
 ,	GameNetworkProtocolVersion( FNetworkVersion::GetGameNetworkProtocolVersion() )
-,	PackageVersionUE( GPackageFileUEVersion )
-,	PackageVersionLicenseeUE( GPackageFileLicenseeUEVersion )
+	, PackageVersionUE( GPackageFileUEVersion )
+	, PackageVersionLicenseeUE( GPackageFileLicenseeUEVersion )
 ,	ResendAllDataState( EResendAllDataState::None )
 #if !UE_BUILD_SHIPPING
 ,	ReceivedRawPacketDel()
@@ -302,29 +304,51 @@ UNetConnection::UNetConnection(const FObjectInitializer& ObjectInitializer)
 ,	bFlushingPacketOrderCache(false)
 ,	ConnectionId(0)
 {
-	// This isn't ideal, because it won't capture memory derived classes are creating dynamically.
-	// The allocations could *probably* be moved somewhere else (like InitBase), but that
-	// causes failure to connect for some reason, and for now this is easier.
+}
+
+void UNetConnection::InitChannelData()
+{
 	LLM_SCOPE_BYTAG(NetConnection);
 
-	MaxChannelSize = CVarMaxChannelSize.GetValueOnAnyThread();
-	if (MaxChannelSize <= 0)
+	if (!ensureMsgf(Channels.Num() == 0, TEXT("InitChannelData: Already initialized!")))
 	{
-		UE_LOG(LogNet, Warning, TEXT("CVarMaxChannelSize of %d is less than or equal to 0, using the default number of channels."), MaxChannelSize);
-		MaxChannelSize = DEFAULT_MAX_CHANNEL_SIZE;
+		return;
 	}
 	
+	check(Driver);
+	check(!HasAnyFlags(EObjectFlags::RF_ClassDefaultObject | EObjectFlags::RF_ArchetypeObject));
+
+	int32 ChannelSize = CVarMaxChannelSize.GetValueOnAnyThread();
+	if (ChannelSize <= 0)
+	{
+		// set from the connection default
+		ChannelSize = DefaultMaxChannelSize;
 	
-	if (!HasAnyFlags(EObjectFlags::RF_ClassDefaultObject | EObjectFlags::RF_ArchetypeObject))
+		// allow the driver to override
+		const int32 MaxChannelsOverride = Driver->GetMaxChannelsOverride();
+		if (MaxChannelsOverride > 0)
 	{
-		Channels.AddDefaulted(MaxChannelSize);
-		OutReliable.AddDefaulted(MaxChannelSize);
-		InReliable.AddDefaulted(MaxChannelSize);
-		PendingOutRec.AddDefaulted(MaxChannelSize);
+			ChannelSize = MaxChannelsOverride;
+		}
+	}
+
+	UE_LOG(LogNet, Log, TEXT("%s setting maximum channels to: %d"), *GetNameSafe(this), ChannelSize);
+
+	PRAGMA_DISABLE_DEPRECATION_WARNINGS
+	MaxChannelSize = ChannelSize;
+	PRAGMA_ENABLE_DEPRECATION_WARNINGS
+
+	Channels.AddDefaulted(ChannelSize);
+	OutReliable.AddDefaulted(ChannelSize);
+	InReliable.AddDefaulted(ChannelSize);
+
+	if (!IsInternalAck())
+	{
+		PendingOutRec.AddDefaulted(ChannelSize);
+	}
 
 		PacketNotify.Init(InPacketId, OutPacketId);
 	}	
-}
 
 /**
  * Initialize common settings for this connection instance
@@ -346,6 +370,8 @@ void UNetConnection::InitBase(UNetDriver* InDriver,class FSocket* InSocket, cons
 	// Owning net driver
 	Driver = InDriver;
 
+	InitChannelData();
+
 	// Cache instance id
 #if UE_NET_TRACE_ENABLED
 	NetTraceId = Driver->GetNetTraceId();
@@ -496,6 +522,8 @@ void UNetConnection::InitConnection(UNetDriver* InDriver, EConnectionState InSta
 {
 	Driver = InDriver;
 
+	InitChannelData();
+
 	// We won't be sending any packets, so use a default size
 	MaxPacket = (InMaxPacket == 0 || InMaxPacket > MAX_PACKET_SIZE) ? MAX_PACKET_SIZE : InMaxPacket;
 	PacketOverhead = 0;
@@ -595,6 +623,8 @@ void UNetConnection::InitHandler()
 
 void UNetConnection::InitSequence(int32 IncomingSequence, int32 OutgoingSequence)
 {
+	checkf(InReliable.Num() > 0 && OutReliable.Num() > 0, TEXT("InitChannelData must be called prior to InitSequence."));
+
 	// Make sure the sequence hasn't already been initialized on the server, and ignore multiple initializations on the client
 	check(InPacketId == -1 || Driver->ServerConnection != nullptr);
 
@@ -2718,7 +2748,7 @@ void UNetConnection::ReceivedPacket( FBitReader& Reader, bool bIsReinjectedPacke
 					uint32 ChIndex;
 					Reader.SerializeIntPacked(ChIndex);
 
-					if (ChIndex >= (uint32)MaxChannelSize)
+					if (ChIndex >= (uint32)Channels.Num())
 					{
 						UE_LOG(LogNet, Warning, TEXT("Bunch channel index exceeds channel limit"));
 
@@ -3745,7 +3775,7 @@ UChannel* UNetConnection::CreateChannelByName(const FName& ChName, EChannelCreat
 			if (!bHasWarnedAboutChannelLimit)
 			{
 				bHasWarnedAboutChannelLimit = true;
-				UE_LOG(LogNetTraffic, Warning, TEXT("No free channel could be found in the channel list (current limit is %d channels) for connection with owner %s. Consider increasing the max channels allowed using net.MaxChannelSize."), MaxChannelSize, *GetNameSafe(OwningActor));
+				UE_LOG(LogNetTraffic, Warning, TEXT("No free channel could be found in the channel list (current limit is %d channels) for connection with owner %s. Consider increasing the max channels allowed using net.MaxChannelSize."), Channels.Num(), *GetNameSafe(OwningActor));
 			}
 
 			return nullptr;
diff --git a/Engine/Source/Runtime/Engine/Private/NetDriver.cpp b/Engine/Source/Runtime/Engine/Private/NetDriver.cpp
--- a/Engine/Source/Runtime/Engine/Private/NetDriver.cpp
+++ b/Engine/Source/Runtime/Engine/Private/NetDriver.cpp
@@ -367,6 +367,8 @@ UNetDriver::UNetDriver(const FObjectInitializer& ObjectInitializer)
 ,	AnalyticsProvider()
 ,	AnalyticsAggregator()
 ,   World(nullptr)
+,	NetDriverDefinition(NAME_None)
+,	MaxChannelsOverride(0)
 ,   Notify(nullptr)
 ,	ElapsedTime( 0.0 )
 ,	bInTick(false)
@@ -5354,6 +5356,16 @@ void UNetDriver::SetNetDriverName(FName NewNetDriverNamed)
 	InitPacketSimulationSettings();
 }
 
+void UNetDriver::SetNetDriverDefinition(FName NewNetDriverDefinition)
+{
+	NetDriverDefinition = NewNetDriverDefinition;
+
+	if (const FNetDriverDefinition* DriverDef = GEngine->NetDriverDefinitions.FindByPredicate([this](const FNetDriverDefinition& Def) { return (Def.DefName == NetDriverDefinition); }))
+	{
+		MaxChannelsOverride = DriverDef->MaxChannelsOverride;
+	}
+}
+
 #if NET_DEBUG_RELEVANT_ACTORS
 void UNetDriver::PrintDebugRelevantActors()
 {
diff --git a/Engine/Source/Runtime/Engine/Private/UnrealEngine.cpp b/Engine/Source/Runtime/Engine/Private/UnrealEngine.cpp
--- a/Engine/Source/Runtime/Engine/Private/UnrealEngine.cpp
+++ b/Engine/Source/Runtime/Engine/Private/UnrealEngine.cpp
@@ -12924,6 +12924,7 @@ UNetDriver* CreateNetDriver_Local(UEngine* Engine, FWorldContext& Context, FName
 			check(ReturnVal != nullptr);
 
 			ReturnVal->SetNetDriverName(ReturnVal->GetFName());
+			ReturnVal->SetNetDriverDefinition(NetDriverDefinition);
 
 			new(Context.ActiveNetDrivers) FNamedNetDriver(ReturnVal, Definition);
 		}
