Fix FMorphTargetVertexInfoBuffers::Serialize assert in loading packaged build

diff --git a/Engine/Source/Runtime/Engine/Private/Rendering/MorphTargetVertexInfoBuffers.cpp b/Engine/Source/Runtime/Engine/Private/Rendering/MorphTargetVertexInfoBuffers.cpp
--- a/Engine/Source/Runtime/Engine/Private/Rendering/MorphTargetVertexInfoBuffers.cpp
+++ b/Engine/Source/Runtime/Engine/Private/Rendering/MorphTargetVertexInfoBuffers.cpp
@@ -12,6 +12,7 @@ void FMorphTargetVertexInfoBuffers::InitRHI()
 	SCOPED_LOADTIMER(FFMorphTargetVertexInfoBuffers_InitRHI);
 
 	check(NumTotalBatches > 0);
+	check(!bRHIIntialized);
 
 	const uint32 BufferSize = MorphData.Num() * sizeof(uint32);
 	FRHIResourceCreateInfo CreateInfo(TEXT("MorphData"));
@@ -70,9 +71,9 @@ void FMorphTargetVertexInfoBuffers::Serialize(FArchive& Ar)
 		check(bIsMorphCPUDataValid);
 		ValidateVertexBuffers(true);
 	}
-	else
+	else if(Ar.IsLoading())
 	{
-		check(!bRHIIntialized);
+		ResetCPUData();
 	}
 
 	Ar << MorphData
diff --git a/Engine/Source/Runtime/Engine/Public/Rendering/MorphTargetVertexInfoBuffers.h b/Engine/Source/Runtime/Engine/Public/Rendering/MorphTargetVertexInfoBuffers.h
--- a/Engine/Source/Runtime/Engine/Public/Rendering/MorphTargetVertexInfoBuffers.h
+++ b/Engine/Source/Runtime/Engine/Public/Rendering/MorphTargetVertexInfoBuffers.h
@@ -76,7 +76,8 @@ public:
 	FBufferRHIRef MorphDataBuffer;
 	FShaderResourceViewRHIRef MorphDataSRV;
 
-	void Reset()
+protected:
+	void ResetCPUData()
 	{
 		MorphData.Empty();
 		MaximumValuePerMorph.Empty();
@@ -87,12 +88,9 @@ public:
 		PositionPrecision = 0.0f;
 		TangentZPrecision = 0.0f;
 		bResourcesInitialized = false;
-		bRHIIntialized = false;
 		bIsMorphCPUDataValid = false;
 	}
 
-protected:
-
 	void ValidateVertexBuffers(bool bMorphTargetsShouldBeValid);
 	void Serialize(FArchive& Ar);
 
