Add number of queued static mesh compilation to Insights Counter

diff --git a/Engine/Source/Runtime/Engine/Private/StaticMeshCompiler.cpp b/Engine/Source/Runtime/Engine/Private/StaticMeshCompiler.cpp
--- a/Engine/Source/Runtime/Engine/Private/StaticMeshCompiler.cpp
+++ b/Engine/Source/Runtime/Engine/Private/StaticMeshCompiler.cpp
@@ -15,6 +15,7 @@
 #include "Misc/ScopedSlowTask.h"
 #include "UObject/StrongObjectPtr.h"
 #include "Misc/IQueuedWork.h"
+#include "ProfilingDebugging/CountersTrace.h"
 #include "Components/HierarchicalInstancedStaticMeshComponent.h"
 #include "LevelEditor.h"
 #include "SLevelViewport.h"
@@ -217,8 +218,10 @@ bool FStaticMeshCompilingManager::IsAsyncStaticMeshCompilationEnabled() const
 	return CVarAsyncStaticMeshStandard.AsyncCompilation.GetValueOnAnyThread() != 0;
 }
 
+TRACE_DECLARE_INT_COUNTER(QueuedStaticMeshCompilation, TEXT("AsyncCompilation/QueuedStaticMesh"));
 void FStaticMeshCompilingManager::UpdateCompilationNotification()
 {
+	TRACE_COUNTER_SET(QueuedStaticMeshCompilation, GetNumRemainingMeshes());
 	Notification.Update(GetNumRemainingMeshes());
 }
 
@@ -355,6 +358,8 @@ void FStaticMeshCompilingManager::AddStaticMeshes(TArrayView<UStaticMesh* const>
 		check(StaticMesh->AsyncTask != nullptr);
 		RegisteredStaticMesh.Emplace(StaticMesh);
 	}
+
+	TRACE_COUNTER_SET(QueuedStaticMeshCompilation, GetNumRemainingMeshes());
 }
 
 void FStaticMeshCompilingManager::FinishCompilation(TArrayView<UStaticMesh* const> InStaticMeshes)
