Changed default value for WPO culling distance in ray tracing to 120m

diff --git a/Engine/Source/Runtime/Engine/Private/StaticMeshRender.cpp b/Engine/Source/Runtime/Engine/Private/StaticMeshRender.cpp
--- a/Engine/Source/Runtime/Engine/Private/StaticMeshRender.cpp
+++ b/Engine/Source/Runtime/Engine/Private/StaticMeshRender.cpp
@@ -185,8 +185,8 @@ static TAutoConsoleVariable<int32> CVarRayTracingStaticMeshesWPOCulling(
 
 static TAutoConsoleVariable<float> CVarRayTracingStaticMeshesWPOCullingRadius(
 	TEXT("r.RayTracing.Geometry.StaticMeshes.WPO.CullingRadius"),
-	5000.0f, // 50 m
-	TEXT("Do not evaluate world position offset for static meshes outside of this radius in ray tracing effects (default = 5000 (50m))"));
+	12000.0f, // 120 m
+	TEXT("Do not evaluate world position offset for static meshes outside of this radius in ray tracing effects (default = 12000 (120m))"));
 
 
 /** Initialization constructor. */
