Fix issue where game instance player registration can return an invalid index

diff --git a/Engine/Source/Runtime/Engine/Private/GameInstance.cpp b/Engine/Source/Runtime/Engine/Private/GameInstance.cpp
--- a/Engine/Source/Runtime/Engine/Private/GameInstance.cpp
+++ b/Engine/Source/Runtime/Engine/Private/GameInstance.cpp
@@ -794,10 +794,8 @@ int32 UGameInstance::AddLocalPlayer(ULocalPlayer* NewLocalPlayer, int32 Controll
 		return INDEX_NONE;
 	}
 
-	const int32 InsertIndex = LocalPlayers.Num();
-
 	// Add to list
-	LocalPlayers.AddUnique(NewLocalPlayer);
+	const int32 InsertIndex = LocalPlayers.AddUnique(NewLocalPlayer);
 
 	// Notify the player he/she was added
 	NewLocalPlayer->PlayerAdded(GetGameViewportClient(), ControllerId);
