Fix Engine Scalability Settings Preview

diff --git a/Engine/Source/Runtime/Engine/Private/Scalability.cpp b/Engine/Source/Runtime/Engine/Private/Scalability.cpp
--- a/Engine/Source/Runtime/Engine/Private/Scalability.cpp
+++ b/Engine/Source/Runtime/Engine/Private/Scalability.cpp
@@ -361,49 +361,7 @@ void ApplyScalabilityGroupFromPlatformIni(const TCHAR* InSectionName, const TCHA
 
 void ChangeScalabilityPreviewPlatform(FName NewPlatformScalabilityName)
 {
-	if (PlatformScalabilityName != NAME_None)
-	{
-		// restore any modified CVar values and reapply the scalability settings for the default Editor platform
-		UndoPlatformScalability();
-		PlatformScalabilityName = NAME_None;
-		FQualityLevels State = Scalability::GetQualityLevels();
-		Scalability::SetQualityLevels(State);
-	}
-
-	if (NewPlatformScalabilityName != NAME_None)
-	{
-		PlatformScalabilityName = NewPlatformScalabilityName;
-		FString PlatformString = NewPlatformScalabilityName.ToString();
-		FConfigCacheIni::LoadGlobalIniFile(PlatformScalabilityIniFilename, TEXT("Scalability"), *PlatformString, true);
-
-		// load allow/deny lists of cvars we can set when previewing this platform
-		PlatformScalabilityCVarAllowList.Empty();
-		TArray<FString> AllowListCVarNames;
-		GConfig->GetArray(TEXT("ScalabilityPreview"), TEXT("CVarAllowList"), AllowListCVarNames, *PlatformScalabilityIniFilename);
-		for (const FString& CVarName : AllowListCVarNames)
-		{
-			const IConsoleVariable* CVar = IConsoleManager::Get().FindConsoleVariable(*CVarName);
-			if (CVar)
-			{
-				PlatformScalabilityCVarAllowList.Add(CVar);
-			}
-		}
-		PlatformScalabilityCVarDenyList.Empty();
-		TArray<FString> DenyListCVarNames;
-		GConfig->GetArray(TEXT("ScalabilityPreview"), TEXT("CVarDenyList"), DenyListCVarNames, *PlatformScalabilityIniFilename);
-		for (const FString& CVarName : DenyListCVarNames)
-		{
-			const IConsoleVariable* CVar = IConsoleManager::Get().FindConsoleVariable(*CVarName);
-			if (CVar)
-			{
-				PlatformScalabilityCVarDenyList.Add(CVar);
-			}
-		}
-		
-		// apply scalability
-		FQualityLevels State = Scalability::GetQualityLevels();
-		Scalability::SetQualityLevels(State);
-	}
+	PlatformScalabilityName = NewPlatformScalabilityName;
 }
 #endif
 
@@ -415,7 +373,28 @@ static void SetGroupQualityLevel(const TCHAR* InGroupName, int32 InQualityLevel,
 #if WITH_EDITOR
 	if (PlatformScalabilityName != NAME_None)
 	{
-		ApplyScalabilityGroupFromPlatformIni(*Section, *PlatformScalabilityIniFilename);
+		if (FConfigCacheIni* PlatformIni = FConfigCacheIni::ForPlatform(PlatformScalabilityName))
+		{
+			TArray<FString> CVarData;
+			PlatformIni->GetSection(*Section, CVarData, GScalabilityIni);
+			{
+				// Check all cvars against current state to see if they match
+				for (int32 DataIndex = 0; DataIndex < CVarData.Num(); ++DataIndex)
+				{
+					const FString& CVarString = CVarData[DataIndex];
+					FString CVarName, CVarValue;
+					if (CVarString.Split(TEXT("="), &CVarName, &CVarValue))
+					{
+						IConsoleVariable* CVar = IConsoleManager::Get().FindConsoleVariable(*CVarName);
+						if (CVar)
+						{
+							float Value = FCString::Atof(*CVarValue);
+							CVar->SetWithCurrentPriority(Value);
+						}
+					}
+				}
+			}
+		}
 	}
 	else
 #endif
