Take string length into account when drawing `stat version` to the screen so it doesn't draw off screen.

- This doesn't draw exactly where I'd like it to, but I'm limited to the arguments passed into the `FEngineStatRender` render func.

diff --git a/Engine/Source/Runtime/Engine/Private/UnrealEngine.cpp b/Engine/Source/Runtime/Engine/Private/UnrealEngine.cpp
--- a/Engine/Source/Runtime/Engine/Private/UnrealEngine.cpp
+++ b/Engine/Source/Runtime/Engine/Private/UnrealEngine.cpp
@@ -16405,7 +16405,12 @@ int32 UEngine::RenderStatVersion(UWorld* World, FViewport* Viewport, FCanvas* Ca
 	{
 		if (!bSuppressMapWarnings)
 		{
-			FCanvasTextItem TextItem(FVector2D(X - 40, Y), FText::FromString(Viewport->AppVersionString), GetSmallFont(), FLinearColor::Yellow);
+			FText Text = FText::FromString(Viewport->AppVersionString);
+			int32 XSize = 0;
+			int32 YSize = 0;
+			StringSize(GEngine->GetSmallFont(), XSize, YSize, *Text.ToString());
+
+			FCanvasTextItem TextItem(FVector2D(X - XSize, Y), Text, GetSmallFont(), FLinearColor::Yellow);
 			TextItem.EnableShadow(FLinearColor::Black);
 			Canvas->DrawItem(TextItem);
 			Y += TextItem.DrawnSize.Y;
