Save high res editor screenshot as .png

diff --git a/Engine/Source/Runtime/Engine/Private/GameViewportClient.cpp b/Engine/Source/Runtime/Engine/Private/GameViewportClient.cpp
--- a/Engine/Source/Runtime/Engine/Private/GameViewportClient.cpp
+++ b/Engine/Source/Runtime/Engine/Private/GameViewportClient.cpp
@@ -1983,8 +1983,8 @@ bool UGameViewportClient::ProcessScreenShots(FViewport* InViewport)
 				}
 
 				// Save the contents of the array to a png file.
-				TArray<uint8> CompressedBitmap;
-				FImageUtils::CompressImageArray(Size.X, Size.Y, Bitmap, CompressedBitmap);
+				TArray64<uint8> CompressedBitmap;
+				FImageUtils::PNGCompressImageArray(Size.X, Size.Y, Bitmap, CompressedBitmap);
 				bIsScreenshotSaved = FFileHelper::SaveArrayToFile(CompressedBitmap, *ScreenShotName);
 
 			}
