Don't allow adding an abstract component to an Actor

diff --git a/Engine/Source/Runtime/Engine/Classes/GameFramework/Actor.h b/Engine/Source/Runtime/Engine/Classes/GameFramework/Actor.h
--- a/Engine/Source/Runtime/Engine/Classes/GameFramework/Actor.h
+++ b/Engine/Source/Runtime/Engine/Classes/GameFramework/Actor.h
@@ -1571,7 +1571,7 @@ public:
 	 * @param bDeferredFinish				Whether or not to immediately complete the creation and registration process for this component. Will be false if there are expose on spawn properties being set
 	 */
 	UFUNCTION(BlueprintCallable, meta=(ScriptNoExport, BlueprintInternalUseOnly = "true", InternalUseParam="bDeferredFinish"))
-	UActorComponent* AddComponentByClass(TSubclassOf<UActorComponent> Class, bool bManualAttachment, const FTransform& RelativeTransform, bool bDeferredFinish);
+	UActorComponent* AddComponentByClass(UPARAM(meta = (AllowAbstract = "false")) TSubclassOf<UActorComponent> Class, bool bManualAttachment, const FTransform& RelativeTransform, bool bDeferredFinish);
 
 	/** 
 	 * Completes the creation of a new actor component. Called either from blueprint after
