Fix material editor preview thumbnails when using new HLSL translator

diff --git a/Engine/Source/Editor/MaterialEditor/Private/MaterialEditor.cpp b/Engine/Source/Editor/MaterialEditor/Private/MaterialEditor.cpp
--- a/Engine/Source/Editor/MaterialEditor/Private/MaterialEditor.cpp
+++ b/Engine/Source/Editor/MaterialEditor/Private/MaterialEditor.cpp
@@ -147,6 +147,7 @@
 #include "MaterialEditorTabs.h"
 #include "MaterialEditorModes.h"
 #include "Materials/MaterialExpression.h"
+#include "MaterialCachedHLSLTree.h"
 
 #include "SMaterialParametersOverviewWidget.h"
 #include "SMaterialEditorCustomPrimitiveDataWidget.h"
@@ -173,6 +174,42 @@ static TAutoConsoleVariable<int32> CVarMaterialEdUseDevShaders(
 // FMatExpressionPreview //
 ///////////////////////////
 
+FMatExpressionPreview::FMatExpressionPreview()
+	: FMaterial()
+	, FMaterialRenderProxy(TEXT("FMatExpressionPreview"))
+	, UnrelatedNodesOpacity(1.0f)
+{
+	// Register this FMaterial derivative with AddEditorLoadedMaterialResource since it does not have a corresponding UMaterialInterface
+	FMaterial::AddEditorLoadedMaterialResource(this);
+	SetQualityLevelProperties(GMaxRHIFeatureLevel);
+}
+
+FMatExpressionPreview::FMatExpressionPreview(UMaterialExpression* InExpression)
+	: FMaterial()
+	, FMaterialRenderProxy(GetPathNameSafe(InExpression->Material))
+	, UnrelatedNodesOpacity(1.0f)
+	, Expression(InExpression)
+{
+	FMaterial::AddEditorLoadedMaterialResource(this);
+	FPlatformMisc::CreateGuid(Id);
+
+	check(InExpression->Material && InExpression->Material->Expressions.Contains(InExpression));
+	ReferencedTextures = InExpression->Material->GetReferencedTextures();
+	SetQualityLevelProperties(GMaxRHIFeatureLevel);
+
+	UMaterial* BaseMaterial = InExpression->Material;
+	if (BaseMaterial->IsUsingNewHLSLGenerator())
+	{
+		FMaterialCachedHLSLTree* LocalTree = new FMaterialCachedHLSLTree();
+		LocalTree->GenerateTree(BaseMaterial, nullptr, InExpression);
+		CachedHLSLTree.Reset(LocalTree);
+	}
+}
+
+FMatExpressionPreview::~FMatExpressionPreview()
+{
+}
+
 bool FMatExpressionPreview::ShouldCache(EShaderPlatform Platform, const FShaderType* ShaderType, const FVertexFactoryType* VertexFactoryType) const
 {
 	if(VertexFactoryType == FindVertexFactoryType(FName(TEXT("FLocalVertexFactory"), FNAME_Find)))
@@ -275,11 +312,7 @@ void FMatExpressionPreview::NotifyCompilationFinished()
 
 const FMaterialCachedHLSLTree* FMatExpressionPreview::GetCachedHLSLTree() const
 {
-	if (Expression.IsValid() && Expression->Material)
-	{
-		return &Expression->Material->GetCachedHLSLTree();
-	}
-	return nullptr;
+	return CachedHLSLTree.Get();
 }
 
 bool FMatExpressionPreview::IsUsingControlFlow() const
diff --git a/Engine/Source/Editor/MaterialEditor/Private/MaterialEditor.h b/Engine/Source/Editor/MaterialEditor/Private/MaterialEditor.h
--- a/Engine/Source/Editor/MaterialEditor/Private/MaterialEditor.h
+++ b/Engine/Source/Editor/MaterialEditor/Private/MaterialEditor.h
@@ -41,6 +41,7 @@ class UMaterialInstance;
 class UMaterialGraphNode;
 struct FGraphAppearanceInfo;
 class UMaterialFunctionInstance;
+class FMaterialCachedHLSLTree;
 
 /**
  * Class for rendering previews of material expressions in the material editor's linked object viewport.
@@ -48,33 +49,10 @@ class UMaterialFunctionInstance;
 class FMatExpressionPreview : public FMaterial, public FMaterialRenderProxy
 {
 public:
-	FMatExpressionPreview()
-	: FMaterial()
-	, FMaterialRenderProxy(TEXT("FMatExpressionPreview"))
-	, UnrelatedNodesOpacity(1.0f)
-	{
-		// Register this FMaterial derivative with AddEditorLoadedMaterialResource since it does not have a corresponding UMaterialInterface
-		FMaterial::AddEditorLoadedMaterialResource(this);
-		SetQualityLevelProperties(GMaxRHIFeatureLevel);
-	}
-
-	FMatExpressionPreview(UMaterialExpression* InExpression)
-	: FMaterial()
-	, FMaterialRenderProxy(GetPathNameSafe(InExpression->Material))
-	, UnrelatedNodesOpacity(1.0f)
-	, Expression(InExpression)
-	{
-		FMaterial::AddEditorLoadedMaterialResource(this);
-		FPlatformMisc::CreateGuid(Id);
+	FMatExpressionPreview();
+	FMatExpressionPreview(UMaterialExpression* InExpression);
 
-		check(InExpression->Material && InExpression->Material->Expressions.Contains(InExpression));
-		ReferencedTextures = InExpression->Material->GetReferencedTextures();
-		SetQualityLevelProperties(GMaxRHIFeatureLevel);
-	}
-
-	virtual ~FMatExpressionPreview()
-	{
-	}
+	virtual ~FMatExpressionPreview();
 
 	virtual bool PrepareDestroy_GameThread() override
 	{
@@ -191,6 +169,7 @@ public:
 	float UnrelatedNodesOpacity;
 
 private:
+	TUniquePtr<FMaterialCachedHLSLTree> CachedHLSLTree;
 	TWeakObjectPtr<UMaterialExpression> Expression;
 	TArray<TObjectPtr<UObject>> ReferencedTextures;
 	FGuid Id;
diff --git a/Engine/Source/Runtime/Engine/Private/Materials/Material.cpp b/Engine/Source/Runtime/Engine/Private/Materials/Material.cpp
--- a/Engine/Source/Runtime/Engine/Private/Materials/Material.cpp
+++ b/Engine/Source/Runtime/Engine/Private/Materials/Material.cpp
@@ -1887,7 +1887,7 @@ void UMaterial::UpdateCachedExpressionData()
 	if (IsUsingNewHLSLGenerator())
 	{
 		LocalCachedTree = new FMaterialCachedHLSLTree();
-		LocalCachedTree->GenerateTree(this, nullptr);
+		LocalCachedTree->GenerateTree(this, nullptr, nullptr);
 		LocalCachedExpressionData->UpdateForCachedHLSLTree(*LocalCachedTree, nullptr);
 	}
 	else
diff --git a/Engine/Source/Runtime/Engine/Private/Materials/MaterialCachedHLSLTree.cpp b/Engine/Source/Runtime/Engine/Private/Materials/MaterialCachedHLSLTree.cpp
--- a/Engine/Source/Runtime/Engine/Private/Materials/MaterialCachedHLSLTree.cpp
+++ b/Engine/Source/Runtime/Engine/Private/Materials/MaterialCachedHLSLTree.cpp
@@ -49,7 +49,7 @@ static UE::Shader::EValueType GetShaderType(EMaterialValueType MaterialType)
 	}
 }
 
-bool FMaterialCachedHLSLTree::GenerateTree(UMaterial* Material, const FMaterialLayersFunctions* LayerOverrides)
+bool FMaterialCachedHLSLTree::GenerateTree(UMaterial* Material, const FMaterialLayersFunctions* LayerOverrides, UMaterialExpression* PreviewExpression)
 {
 	const EMaterialShadingModel DefaultShadingModel = Material->GetShadingModels().GetFirstShadingModel();
 
@@ -116,7 +116,7 @@ bool FMaterialCachedHLSLTree::GenerateTree(UMaterial* Material, const FMaterialL
 
 	HLSLTree = UE::HLSLTree::FTree::Create(Allocator);
 
-	FMaterialHLSLGenerator Generator(Material, LayerOverrides, *this);
+	FMaterialHLSLGenerator Generator(Material, LayerOverrides, PreviewExpression, *this);
 	const bool bResult = Generator.Generate();
 	if (bResult)
 	{
diff --git a/Engine/Source/Runtime/Engine/Private/Materials/MaterialExpressions.cpp b/Engine/Source/Runtime/Engine/Private/Materials/MaterialExpressions.cpp
--- a/Engine/Source/Runtime/Engine/Private/Materials/MaterialExpressions.cpp
+++ b/Engine/Source/Runtime/Engine/Private/Materials/MaterialExpressions.cpp
@@ -14512,7 +14512,7 @@ void UMaterialExpressionMaterialFunctionCall::GetConnectorToolTip(int32 InputInd
 
 				ConvertToMultilineToolTip(InputExpression->Description, 40, OutToolTip);
 
-				if (InputExpression->bUsePreviewValueAsDefault)
+				if (InputExpression && InputExpression->bUsePreviewValueAsDefault)
 				{
 					// Can't build a tooltip of an arbitrary expression chain
 					if (InputExpression->Preview.Expression)
@@ -14537,7 +14537,11 @@ void UMaterialExpressionMaterialFunctionCall::GetConnectorToolTip(int32 InputInd
 		{
 			if (FunctionOutputs.IsValidIndex(OutputIndex))
 			{
-				ConvertToMultilineToolTip(FunctionOutputs[OutputIndex].ExpressionOutput->Description, 40, OutToolTip);
+				UMaterialExpressionFunctionOutput* OutputExpression = FunctionOutputs[OutputIndex].ExpressionOutput;
+				if (OutputExpression)
+				{
+					ConvertToMultilineToolTip(OutputExpression->Description, 40, OutToolTip);
+				}
 			}
 		}
 	}
diff --git a/Engine/Source/Runtime/Engine/Private/Materials/MaterialHLSLGenerator.cpp b/Engine/Source/Runtime/Engine/Private/Materials/MaterialHLSLGenerator.cpp
--- a/Engine/Source/Runtime/Engine/Private/Materials/MaterialHLSLGenerator.cpp
+++ b/Engine/Source/Runtime/Engine/Private/Materials/MaterialHLSLGenerator.cpp
@@ -21,9 +21,11 @@
 
 FMaterialHLSLGenerator::FMaterialHLSLGenerator(UMaterial* Material,
 	const FMaterialLayersFunctions* InLayerOverrides,
+	UMaterialExpression* InPreviewExpression,
 	FMaterialCachedHLSLTree& OutCachedTree)
 	: TargetMaterial(Material)
 	, LayerOverrides(InLayerOverrides)
+	, PreviewExpression(InPreviewExpression)
 	, CachedTree(OutCachedTree)
 	, bGeneratedResult(false)
 {
@@ -169,11 +171,10 @@ bool FMaterialHLSLGenerator::GenerateResult(UE::HLSLTree::FScope& Scope)
 		check(!CachedTree.ResultStatement);
 		check(!CachedTree.ResultExpression);
 
+		FExpression* AttributesExpression = nullptr;
 		if (TargetMaterial)
 		{
 			const FStructField* PrevWPOField = CachedTree.GetMaterialAttributesType()->FindFieldByName(TEXT("PrevWorldPositionOffset"));
-
-			FExpression* AttributesExpression = nullptr;
 			if (TargetMaterial->bUseMaterialAttributes)
 			{
 				FMaterialInputDescription InputDescription;
@@ -242,20 +243,36 @@ bool FMaterialHLSLGenerator::GenerateResult(UE::HLSLTree::FScope& Scope)
 					AttributesExpression = GetTree().NewExpression<FExpressionSetStructField>(CachedTree.GetMaterialAttributesType(), CustomOutputField, AttributesExpression, CustomOutputExpression);
 				}
 			}
+		}
 
-			if (AttributesExpression)
+		if (PreviewExpression)
+		{
+			if (!PreviewExpressionResult)
 			{
-				FStatementReturn* ReturnStatement = GetTree().NewStatement<FStatementReturn>(Scope);
-				ReturnStatement->Expression = AttributesExpression;
-				CachedTree.ResultExpression = AttributesExpression;
-				CachedTree.ResultStatement = ReturnStatement;
-				bResult = true;
+				// If we didn't hit the preview expression while generating the material normally, then generate it now
+				// Hardcoding output 0 as we don't have the UI to specify any other output
+				const int32 OutputIndex = 0;
+				PreviewExpressionResult = AcquireExpression(Scope, PreviewExpression, OutputIndex);
 			}
+			const FString& EmissiveColorName = FMaterialAttributeDefinitionMap::GetAttributeName(MP_EmissiveColor);
+			const FStructField* EmissiveColorField = CachedTree.GetMaterialAttributesType()->FindFieldByName(*EmissiveColorName);
+
+			// Get back into gamma corrected space, as DrawTile does not do this adjustment.
+			FExpression* ExpressionEmissive = GetTree().NewPowClamped(PreviewExpressionResult, NewConstant(1.f / 2.2f));
+
+			AttributesExpression = GetTree().NewExpression<FExpressionConstant>(CachedTree.GetMaterialAttributesDefaultValue());
+			AttributesExpression = GetTree().NewExpression<FExpressionSetStructField>(CachedTree.GetMaterialAttributesType(), EmissiveColorField, AttributesExpression, ExpressionEmissive);
 		}
-		else
+
+		if (AttributesExpression)
 		{
-			check(false);
+			FStatementReturn* ReturnStatement = GetTree().NewStatement<FStatementReturn>(Scope);
+			ReturnStatement->Expression = AttributesExpression;
+			CachedTree.ResultExpression = AttributesExpression;
+			CachedTree.ResultStatement = ReturnStatement;
+			bResult = true;
 		}
+
 		bGeneratedResult = true;
 	}
 	return bResult;
@@ -351,6 +368,11 @@ UE::HLSLTree::FExpression* FMaterialHLSLGenerator::AcquireExpression(UE::HLSLTre
 	FExpression* Expression = nullptr;
 	if (MaterialExpression->GenerateHLSLExpression(*this, Scope, OutputIndex, Expression))
 	{
+		if (MaterialExpression == PreviewExpression &&
+			!PreviewExpressionResult)
+		{
+			PreviewExpressionResult = Expression;
+		}
 		return Expression;
 	}
 	
diff --git a/Engine/Source/Runtime/Engine/Private/Materials/MaterialInstanceConstant.cpp b/Engine/Source/Runtime/Engine/Private/Materials/MaterialInstanceConstant.cpp
--- a/Engine/Source/Runtime/Engine/Private/Materials/MaterialInstanceConstant.cpp
+++ b/Engine/Source/Runtime/Engine/Private/Materials/MaterialInstanceConstant.cpp
@@ -214,7 +214,7 @@ void UMaterialInstanceConstant::UpdateCachedData()
 			if (bUsingNewHLSLGenerator)
 			{
 				LocalCachedTree = new FMaterialCachedHLSLTree();
-				LocalCachedTree->GenerateTree(BaseMaterial, &LocalStaticParameters.MaterialLayers);
+				LocalCachedTree->GenerateTree(BaseMaterial, &LocalStaticParameters.MaterialLayers, nullptr);
 			}
 			else
 			{
diff --git a/Engine/Source/Runtime/Engine/Public/HLSLTree/HLSLTree.h b/Engine/Source/Runtime/Engine/Public/HLSLTree/HLSLTree.h
--- a/Engine/Source/Runtime/Engine/Public/HLSLTree/HLSLTree.h
+++ b/Engine/Source/Runtime/Engine/Public/HLSLTree/HLSLTree.h
@@ -531,6 +531,7 @@ public:
 	FExpression* NewMul(FExpression* Lhs, FExpression* Rhs) { return NewBinaryOp(EOperation::Mul, Lhs, Rhs); }
 	FExpression* NewDiv(FExpression* Lhs, FExpression* Rhs) { return NewBinaryOp(EOperation::Div, Lhs, Rhs); }
 	FExpression* NewFmod(FExpression* Lhs, FExpression* Rhs) { return NewBinaryOp(EOperation::Fmod, Lhs, Rhs); }
+	FExpression* NewPowClamped(FExpression* Lhs, FExpression* Rhs) { return NewBinaryOp(EOperation::PowPositiveClamped, Lhs, Rhs); }
 	FExpression* NewMin(FExpression* Lhs, FExpression* Rhs) { return NewBinaryOp(EOperation::Min, Lhs, Rhs); }
 	FExpression* NewMax(FExpression* Lhs, FExpression* Rhs) { return NewBinaryOp(EOperation::Max, Lhs, Rhs); }
 	FExpression* NewLess(FExpression* Lhs, FExpression* Rhs) { return NewBinaryOp(EOperation::Less, Lhs, Rhs); }
diff --git a/Engine/Source/Runtime/Engine/Public/MaterialCachedHLSLTree.h b/Engine/Source/Runtime/Engine/Public/MaterialCachedHLSLTree.h
--- a/Engine/Source/Runtime/Engine/Public/MaterialCachedHLSLTree.h
+++ b/Engine/Source/Runtime/Engine/Public/MaterialCachedHLSLTree.h
@@ -18,12 +18,12 @@ class FMaterialCachedHLSLTree
 public:
 	static const FMaterialCachedHLSLTree EmptyTree;
 
-	FMaterialCachedHLSLTree();
-	~FMaterialCachedHLSLTree();
+	ENGINE_API FMaterialCachedHLSLTree();
+	ENGINE_API ~FMaterialCachedHLSLTree();
 
 	SIZE_T GetAllocatedSize() const;
 
-	bool GenerateTree(UMaterial* Material, const FMaterialLayersFunctions* LayerOverrides);
+	ENGINE_API bool GenerateTree(UMaterial* Material, const FMaterialLayersFunctions* LayerOverrides, UMaterialExpression* PreviewExpression);
 
 	UE::Shader::FStructTypeRegistry& GetTypeRegistry() { return TypeRegistry; }
 	//UE::HLSLTree::FTree& GetTree() { return *HLSLTree; }
diff --git a/Engine/Source/Runtime/Engine/Public/MaterialHLSLGenerator.h b/Engine/Source/Runtime/Engine/Public/MaterialHLSLGenerator.h
--- a/Engine/Source/Runtime/Engine/Public/MaterialHLSLGenerator.h
+++ b/Engine/Source/Runtime/Engine/Public/MaterialHLSLGenerator.h
@@ -63,6 +63,7 @@ class FMaterialHLSLGenerator
 public:
 	FMaterialHLSLGenerator(UMaterial* Material,
 		const FMaterialLayersFunctions* InLayerOverrides,
+		UMaterialExpression* InPreviewExpression,
 		FMaterialCachedHLSLTree& OutCachedTree);
 
 	const FMaterialLayersFunctions* GetLayerOverrides() const { return LayerOverrides; }
@@ -199,6 +200,7 @@ private:
 
 	UMaterial* TargetMaterial;
 	const FMaterialLayersFunctions* LayerOverrides;
+	UMaterialExpression* PreviewExpression;
 	FMaterialCachedHLSLTree& CachedTree;
 	FString CurrentErrorMessage;
 
@@ -210,6 +212,7 @@ private:
 	TMap<UMaterialFunctionInterface*, UE::HLSLTree::FFunction*> FunctionMap;
 	TMap<UMaterialExpression*, FStatementEntry> StatementMap;
 	TMap<FExpressionDataKey, void*> ExpressionDataMap;
+	UE::HLSLTree::FExpression* PreviewExpressionResult = nullptr;
 	bool bGeneratedResult;
 };
 
