Fix actor spawned with bHideFromSceneOutliner still being visible in the Scene Outliner

diff --git a/Engine/Source/Runtime/Engine/Private/LevelActor.cpp b/Engine/Source/Runtime/Engine/Private/LevelActor.cpp
--- a/Engine/Source/Runtime/Engine/Private/LevelActor.cpp
+++ b/Engine/Source/Runtime/Engine/Private/LevelActor.cpp
@@ -660,6 +660,20 @@ AActor* UWorld::SpawnActor( UClass* Class, FTransform const* UserTransformPtr, c
 	check(Actor);
 	check(Actor->GetLevel() == LevelToSpawnIn);
 
+#if WITH_EDITOR
+	// UE5-Release: bHideFromSceneOutliner must be set before anything tries
+	// to create an FActorTreeItem in the Scene Outliner. Otherwise, the tree item will
+	// be created and will be visible in the Scene Outliner. 
+	// AActor::ClearActorLabel, called below, currently does that via FCoreDelegates::OnActorLabelChanged.
+	// A better fix should prevent the FActorTreeItem creation before the end of the spawning sequence.
+	if (SpawnParameters.bHideFromSceneOutliner)
+	{
+		FSetActorHiddenInSceneOutliner SetActorHidden(Actor);
+	}
+	Actor->bIsEditorPreviewActor = SpawnParameters.bTemporaryEditorActor;
+#endif //WITH_EDITOR
+
+
 #if ENABLE_SPAWNACTORTIMER
 	SpawnTimer.SetActorName(Actor->GetFName());
 #endif
@@ -694,14 +708,6 @@ AActor* UWorld::SpawnActor( UClass* Class, FTransform const* UserTransformPtr, c
 	// tell the actor what method to use, in case it was overridden
 	Actor->SpawnCollisionHandlingMethod = CollisionHandlingMethod;
 
-#if WITH_EDITOR
-	if (SpawnParameters.bHideFromSceneOutliner)
-	{
-		FSetActorHiddenInSceneOutliner SetActorHidden(Actor);
-	}
-	Actor->bIsEditorPreviewActor = SpawnParameters.bTemporaryEditorActor;
-#endif //WITH_EDITOR
-
 	// Broadcast delegate before the actor and its contained components are initialized
 	OnActorPreSpawnInitialization.Broadcast(Actor);
 
