Initialize channel data for child connections

diff --git a/Engine/Source/Runtime/Engine/Classes/Engine/ChildConnection.h b/Engine/Source/Runtime/Engine/Classes/Engine/ChildConnection.h
--- a/Engine/Source/Runtime/Engine/Classes/Engine/ChildConnection.h
+++ b/Engine/Source/Runtime/Engine/Classes/Engine/ChildConnection.h
@@ -79,4 +79,6 @@ public:
 
 	virtual void InitRemoteConnection(UNetDriver* InDriver, class FSocket* InSocket, const FURL& InURL, const class FInternetAddr& InRemoteAddr, EConnectionState InState, int32 InMaxPacket = 0, int32 InPacketOverhead = 0) override {}
 	virtual void InitLocalConnection(UNetDriver* InDriver, class FSocket* InSocket, const FURL& InURL, EConnectionState InState, int32 InMaxPacket = 0, int32 InPacketOverhead = 0) override {}
+
+	virtual void InitChildConnection(UNetDriver* InDriver, UNetConnection* InParent);
 };
diff --git a/Engine/Source/Runtime/Engine/Private/NetConnection.cpp b/Engine/Source/Runtime/Engine/Private/NetConnection.cpp
--- a/Engine/Source/Runtime/Engine/Private/NetConnection.cpp
+++ b/Engine/Source/Runtime/Engine/Private/NetConnection.cpp
@@ -1208,6 +1208,19 @@ void UChildConnection::CleanUp()
 	Driver = NULL;
 }
 
+void UChildConnection::InitChildConnection(UNetDriver* InDriver, UNetConnection* InParent)
+{
+	Driver = InDriver;
+	URL = FURL();
+	SetConnectionState(InParent->GetConnectionState());
+	URL.Host = InParent->URL.Host;
+	Parent = InParent;
+	PackageMap = InParent->PackageMap;
+	CurrentNetSpeed = InParent->CurrentNetSpeed;
+
+	InitChannelData();
+}
+
 void UNetConnection::FinishDestroy()
 {
 	if ( !HasAnyFlags(RF_ClassDefaultObject) )
diff --git a/Engine/Source/Runtime/Engine/Private/NetDriver.cpp b/Engine/Source/Runtime/Engine/Private/NetDriver.cpp
--- a/Engine/Source/Runtime/Engine/Private/NetDriver.cpp
+++ b/Engine/Source/Runtime/Engine/Private/NetDriver.cpp
@@ -3699,15 +3699,10 @@ void UNetDriver::ForceActorRelevantNextUpdate(AActor* Actor)
 UChildConnection* UNetDriver::CreateChild(UNetConnection* Parent)
 {
 	UE_LOG(LogNet, Log, TEXT("Creating child connection with %s parent"), *Parent->GetName());
-	auto Child = NewObject<UChildConnection>();
-	Child->Driver = this;
-	Child->URL = FURL();
-	Child->SetConnectionState(Parent->GetConnectionState());
-	Child->URL.Host = Parent->URL.Host;
-	Child->Parent = Parent;
-	Child->PackageMap = Parent->PackageMap;
-	Child->CurrentNetSpeed = Parent->CurrentNetSpeed;
+	UChildConnection* Child = NewObject<UChildConnection>();
+	Child->InitChildConnection(this, Parent);
 	Parent->Children.Add(Child);
+
 	return Child;
 }
 
