Conversion of post process and light function to dedicated nodes with auto set domain.
https://horde.devtools.epicgames.com/job/6220a0e47e0217a606e837f1

diff --git a/Engine/Shaders/Private/LightFunctionCommon.ush b/Engine/Shaders/Private/LightFunctionCommon.ush
--- a/Engine/Shaders/Private/LightFunctionCommon.ush
+++ b/Engine/Shaders/Private/LightFunctionCommon.ush
@@ -55,7 +55,7 @@ float3 GetLightFunctionColor(float3 LightVector, float3 TranslatedWorldSpace)
 #if STRATA_ENABLED
 
   #if COMPUTESHADER
-	return 0.f; // STRATA_TODO: This path doesn't seem to be used
+	return StrataGetBSDFEmissive(MaterialParameters.StrataTree.BSDFs[0]); // Unused code path?
   #else
 	return StrataGetBSDFEmissive(PixelMaterialInputs.FrontMaterial.InlinedBSDF);
   #endif
diff --git a/Engine/Shaders/Private/PostProcessMaterialShaders.usf b/Engine/Shaders/Private/PostProcessMaterialShaders.usf
--- a/Engine/Shaders/Private/PostProcessMaterialShaders.usf
+++ b/Engine/Shaders/Private/PostProcessMaterialShaders.usf
@@ -203,17 +203,23 @@ void MainPS(
 	CalcMaterialParametersPost(Parameters, PixelMaterialInputs, SvPosition, true);
 
 	float Alpha = 1.0f;
-#if STRATA_ENABLED   
+#if STRATA_ENABLED
+
 	// Unlit forces a single BSDF
 	const FStrataBSDF UnlitBSDF = PixelMaterialInputs.FrontMaterial.InlinedBSDF;
 
 	// Grab emissive colour as output
   #if MATERIAL_OUTPUT_OPACITY_AS_ALPHA
-	Alpha = UnlitBSDF.Coverage;
+	// Transmittance = 1 - opacity. See UMaterialExpressionStrataPostProcess.
+	const float UnlitBSDFTransmittance = UNLIT_TRANSMITTANCE(UnlitBSDF).x;
+	const float UnlitBSDFOpacity = 1.0f - UnlitBSDFTransmittance;
+	Alpha = UnlitBSDFOpacity;
   #elif (FEATURE_LEVEL <= FEATURE_LEVEL_ES3_1)
 	Alpha = Parameters.BackupSceneColorAlpha;
   #endif
+
 	const float3 Emissive = BSDF_GETEMISSIVE(UnlitBSDF);
+
 #else
 	// Grab emissive colour as output
   #if MATERIAL_OUTPUT_OPACITY_AS_ALPHA
diff --git a/Engine/Source/Runtime/Engine/Classes/Engine/EngineTypes.h b/Engine/Source/Runtime/Engine/Classes/Engine/EngineTypes.h
--- a/Engine/Source/Runtime/Engine/Classes/Engine/EngineTypes.h
+++ b/Engine/Source/Runtime/Engine/Classes/Engine/EngineTypes.h
@@ -667,10 +667,12 @@ enum EStrataShadingModel
 	SSM_VolumetricFogCloud		UMETA(DisplayName = "VolumetricFogCloud"),
 	SSM_Hair					UMETA(DisplayName = "Hair"),
 	SSM_SingleLayerWater		UMETA(DisplayName = "SingleLayerWater"),
+	SSM_LightFunction			UMETA(DisplayName = "LightFunction"),
+	SSM_PostProcess				UMETA(DisplayName = "PostProcess"),
 	/** Number of unique shading models. */
 	SSM_NUM						UMETA(Hidden),
 };
-static_assert(SSM_NUM <= 8, "Do not exceed 16 shading models without expanding FStrataMaterialShadingModelField to support uint32 instead of uint16!");
+static_assert(SSM_NUM <= 16, "Do not exceed 16 shading models without expanding FStrataMaterialShadingModelField to support uint32 instead of uint16!");
 
 /** Gather information from the Strata material graph to setup material for runtime. */
 USTRUCT()
diff --git a/Engine/Source/Runtime/Engine/Classes/Materials/MaterialExpressionStrata.h b/Engine/Source/Runtime/Engine/Classes/Materials/MaterialExpressionStrata.h
--- a/Engine/Source/Runtime/Engine/Classes/Materials/MaterialExpressionStrata.h
+++ b/Engine/Source/Runtime/Engine/Classes/Materials/MaterialExpressionStrata.h
@@ -552,6 +552,60 @@ class UMaterialExpressionStrataSingleLayerWaterBSDF : public UMaterialExpression
 	//~ End UMaterialExpression Interface
 };
 
+UCLASS(MinimalAPI, collapsecategories, hidecategories = Object, DisplayName = "Strata Light Function")
+class UMaterialExpressionStrataLightFunction : public UMaterialExpressionStrataBSDF
+{
+	GENERATED_UCLASS_BODY()
+
+	/**
+	 * The output color of the light function
+	 */
+	UPROPERTY()
+	FExpressionInput Color;
+
+	//~ Begin UMaterialExpression Interface
+#if WITH_EDITOR
+	virtual int32 Compile(class FMaterialCompiler* Compiler, int32 OutputIndex) override;
+	virtual void GetCaption(TArray<FString>& OutCaptions) const override;
+	virtual uint32 GetOutputType(int32 OutputIndex) override;
+	virtual uint32 GetInputType(int32 InputIndex) override;
+	virtual bool IsResultStrataMaterial(int32 OutputIndex) override;
+	virtual void GatherStrataMaterialInfo(FStrataMaterialInfo& StrataMaterialInfo, int32 OutputIndex) override;
+	virtual FStrataOperator* StrataGenerateMaterialTopologyTree(class FMaterialCompiler* Compiler, class UMaterialExpression* Parent, int32 OutputIndex) override;
+#endif
+	//~ End UMaterialExpression Interface
+};
+
+UCLASS(MinimalAPI, collapsecategories, hidecategories = Object, DisplayName = "Strata Post Process")
+class UMaterialExpressionStrataPostProcess : public UMaterialExpressionStrataBSDF
+{
+	GENERATED_UCLASS_BODY()
+
+	/**
+	 * The output color of the post process
+	 */
+	UPROPERTY()
+	FExpressionInput Color;
+
+	/**
+	 * The opacity of the post process. Only writable if "Output Alpha" is enabled.
+	 */
+	UPROPERTY()
+	FExpressionInput Opacity;
+
+	//~ Begin UMaterialExpression Interface
+#if WITH_EDITOR
+	virtual int32 Compile(class FMaterialCompiler* Compiler, int32 OutputIndex) override;
+	virtual void GetCaption(TArray<FString>& OutCaptions) const override;
+	virtual uint32 GetOutputType(int32 OutputIndex) override;
+	virtual uint32 GetInputType(int32 InputIndex) override;
+	virtual bool IsResultStrataMaterial(int32 OutputIndex) override;
+	virtual void GatherStrataMaterialInfo(FStrataMaterialInfo& StrataMaterialInfo, int32 OutputIndex) override;
+	virtual FStrataOperator* StrataGenerateMaterialTopologyTree(class FMaterialCompiler* Compiler, class UMaterialExpression* Parent, int32 OutputIndex) override;
+#endif
+	//~ End UMaterialExpression Interface
+};
+
 
 
 ///////////////////////////////////////////////////////////////////////////////
diff --git a/Engine/Source/Runtime/Engine/Private/Materials/HLSLMaterialTranslator.cpp b/Engine/Source/Runtime/Engine/Private/Materials/HLSLMaterialTranslator.cpp
--- a/Engine/Source/Runtime/Engine/Private/Materials/HLSLMaterialTranslator.cpp
+++ b/Engine/Source/Runtime/Engine/Private/Materials/HLSLMaterialTranslator.cpp
@@ -1351,10 +1351,11 @@ bool FHLSLMaterialTranslator::Translate()
 			// Unlit must be used in isolation
 			if (StrataMaterialContainsAnyBSDF(this, StrataCompilationInfo, STRATA_BSDF_TYPE_UNLIT))
 			{
-				// If the unlit node is used, it must be the only one used
+				FString MaterialDomain = MaterialDomainString(Material->GetMaterialDomain());
+				// If the unlit node is used, it must be the only one used. This is true for all domain (surface, light function, post process).
 				if (!StrataIsUnlitOnly(this, StrataCompilationInfo))
 				{
-					FString ErrorMsg = FString::Printf(TEXT("Material %s contains Unlit BSDF but it is not the single BSDF representing the material asset: %s. It must be the single BSDF.\r\n"), *Material->GetDebugName(), *Material->GetAssetPath().ToString());
+					FString ErrorMsg = FString::Printf(TEXT("%s Material %s contains Unlit BSDF but it is not the single BSDF representing the material asset: %s. It must be the single BSDF.\r\n"), *MaterialDomain, *Material->GetDebugName(), *Material->GetAssetPath().ToString());
 					Error(*ErrorMsg);
 				}
 			}
diff --git a/Engine/Source/Runtime/Engine/Private/Materials/Material.cpp b/Engine/Source/Runtime/Engine/Private/Materials/Material.cpp
--- a/Engine/Source/Runtime/Engine/Private/Materials/Material.cpp
+++ b/Engine/Source/Runtime/Engine/Private/Materials/Material.cpp
@@ -3033,26 +3033,31 @@ void UMaterial::ConvertMaterialToStrataMaterial()
 			FrontMaterial.Connect(0, VolBSDF);
 			bInvalidateShader = true;
 		}
-		else if (MaterialDomain == MD_PostProcess || MaterialDomain == MD_LightFunction)
+		else if (MaterialDomain == MD_LightFunction)
 		{
-			// Some post-process material don't have their shading mode set correctly to Unlit. Since only Unlit is supported, forcing it here.
+			// Some materials don't have their shading mode set correctly to Unlit. Since only Unlit is supported, forcing it here.
 			ShadingModel = MSM_Unlit;
 			ShadingModels.ClearShadingModels();
 			ShadingModels.AddShadingModel(MSM_Unlit);
 
 			// Only Emissive & Opacity are valid input for PostProcess material
-			UMaterialExpressionStrataLegacyConversion* ConvertNode = NewObject<UMaterialExpressionStrataLegacyConversion>(this);
-			MoveConnectionTo(EmissiveColor, ConvertNode, 5);
-			MoveConnectionTo(Opacity, ConvertNode, 11);
+			UMaterialExpressionStrataLightFunction* ConvertNode = NewObject<UMaterialExpressionStrataLightFunction>(this);
+			MoveConnectionTo(EmissiveColor, ConvertNode, 0);
 
-			// Add constant for the Unlit shading model
-			UMaterialExpressionConstant* ShadingModelNode = NewObject<UMaterialExpressionConstant>(this);
-			ShadingModelNode->SetParameterName(FName(TEXT("ConstantShadingModel")));
-			ShadingModelNode->R = ShadingModel;
-			ConvertNode->ShadingModel.Connect(0, ShadingModelNode);
+			FrontMaterial.Connect(0, ConvertNode);
+			bInvalidateShader = true;
+		}
+		else if (MaterialDomain == MD_PostProcess)
+		{
+			// Some materials don't have their shading mode set correctly to Unlit. Since only Unlit is supported, forcing it here.
+			ShadingModel = MSM_Unlit;
+			ShadingModels.ClearShadingModels();
+			ShadingModels.AddShadingModel(MSM_Unlit);
 
-			AddStrataShadingModelFromMaterialShadingModel(ConvertNode->ConvertedStrataMaterialInfo, ShadingModels);
-			check(ConvertNode->ConvertedStrataMaterialInfo.CountShadingModels() == 1);
+			// Only Emissive & Opacity are valid input for PostProcess material
+			UMaterialExpressionStrataPostProcess* ConvertNode = NewObject<UMaterialExpressionStrataPostProcess>(this);
+			MoveConnectionTo(EmissiveColor, ConvertNode, 0);
+			MoveConnectionTo(Opacity, ConvertNode, 1);
 
 			FrontMaterial.Connect(0, ConvertNode);
 			bInvalidateShader = true;
@@ -3670,7 +3675,7 @@ bool UMaterial::CanEditChange(const FProperty* InProperty) const
 
 		if (PropertyName == GET_MEMBER_NAME_STRING_CHECKED(UMaterial, StrataBlendMode))
 		{
-			return Engine_IsStrataEnabled();
+			return Engine_IsStrataEnabled() && ((MaterialDomain != MD_PostProcess && MaterialDomain != MD_LightFunction && MaterialDomain != MD_Volume) || (MaterialDomain == MD_PostProcess && BlendableOutputAlpha));
 		}
 	
 		if (PropertyName == GET_MEMBER_NAME_STRING_CHECKED(UMaterial, ShadingModel))
@@ -4205,6 +4210,18 @@ void UMaterial::RebuildShadingModelField()
 				ShadingModel = MSM_SingleLayerWater;
 				BlendMode = EBlendMode::BLEND_Opaque; // STRATA_TODO water can also be masked: check Mask input from the main node to automatically enabled that?
 			}
+			else if (StrataMaterialInfo.HasOnlyShadingModel(SSM_LightFunction))
+			{
+				MaterialDomain = EMaterialDomain::MD_LightFunction;
+				ShadingModel = MSM_Unlit;
+				BlendMode = EBlendMode::BLEND_Opaque;
+			}
+			else if (StrataMaterialInfo.HasOnlyShadingModel(SSM_PostProcess))
+			{
+				MaterialDomain = EMaterialDomain::MD_PostProcess;
+				ShadingModel = MSM_Unlit;
+				BlendMode = EBlendMode::BLEND_Opaque;
+			}
 
 			// Also update the ShadingModels for remaining pipeline operation
 			ShadingModels.AddShadingModel(ShadingModel);
diff --git a/Engine/Source/Runtime/Engine/Private/Materials/MaterialExpressions.cpp b/Engine/Source/Runtime/Engine/Private/Materials/MaterialExpressions.cpp
--- a/Engine/Source/Runtime/Engine/Private/Materials/MaterialExpressions.cpp
+++ b/Engine/Source/Runtime/Engine/Private/Materials/MaterialExpressions.cpp
@@ -21339,6 +21339,146 @@ FStrataOperator* UMaterialExpressionStrataVolumetricFogCloudBSDF::StrataGenerate
 
 
 
+UMaterialExpressionStrataLightFunction::UMaterialExpressionStrataLightFunction(const FObjectInitializer& ObjectInitializer)
+	: Super(ObjectInitializer)
+{
+	struct FConstructorStatics
+	{
+		FText NAME_Strata;
+		FConstructorStatics() : NAME_Strata(LOCTEXT("Strata Others", "Strata Others")) { }
+	};
+	static FConstructorStatics ConstructorStatics;
+#if WITH_EDITORONLY_DATA
+	MenuCategories.Add(ConstructorStatics.NAME_Strata);
+#endif
+}
+
+#if WITH_EDITOR
+int32 UMaterialExpressionStrataLightFunction::Compile(class FMaterialCompiler* Compiler, int32 OutputIndex)
+{
+	int32 OutputCodeChunk = Compiler->StrataUnlitBSDF(
+		CompileWithDefaultFloat3(Compiler, Color, 0.0f, 0.0f, 0.0f),
+		Compiler->Constant(1.0f));	// Opacity / Transmittance is ignored by light functions.
+
+	StrataCompilationInfoCreateSingleBSDFMaterial(Compiler, OutputCodeChunk, StrataCompilationInfoCreateNullSharedLocalBasis(), STRATA_BSDF_TYPE_UNLIT);
+
+	return OutputCodeChunk;
+}
+
+void UMaterialExpressionStrataLightFunction::GetCaption(TArray<FString>& OutCaptions) const
+{
+	OutCaptions.Add(TEXT("Strata Light Function"));
+}
+
+uint32 UMaterialExpressionStrataLightFunction::GetOutputType(int32 OutputIndex)
+{
+	return MCT_Strata;
+}
+
+uint32 UMaterialExpressionStrataLightFunction::GetInputType(int32 InputIndex)
+{
+	switch (InputIndex)
+	{
+	case 0:
+		return MCT_Float3;
+		break;
+	}
+
+	check(false);
+	return MCT_Float1;
+}
+
+bool UMaterialExpressionStrataLightFunction::IsResultStrataMaterial(int32 OutputIndex)
+{
+	return true;
+}
+
+void UMaterialExpressionStrataLightFunction::GatherStrataMaterialInfo(FStrataMaterialInfo& StrataMaterialInfo, int32 OutputIndex)
+{
+	StrataMaterialInfo.AddShadingModel(SSM_LightFunction);
+}
+
+FStrataOperator* UMaterialExpressionStrataLightFunction::StrataGenerateMaterialTopologyTree(class FMaterialCompiler* Compiler, class UMaterialExpression* Parent, int32 OutputIndex)
+{
+	return &Compiler->StrataCompilationRegisterOperator(STRATA_OPERATOR_BSDF, this, Parent);
+}
+#endif // WITH_EDITOR
+
+
+
+UMaterialExpressionStrataPostProcess::UMaterialExpressionStrataPostProcess(const FObjectInitializer& ObjectInitializer)
+	: Super(ObjectInitializer)
+{
+	struct FConstructorStatics
+	{
+		FText NAME_Strata;
+		FConstructorStatics() : NAME_Strata(LOCTEXT("Strata Others", "Strata Others")) { }
+	};
+	static FConstructorStatics ConstructorStatics;
+#if WITH_EDITORONLY_DATA
+	MenuCategories.Add(ConstructorStatics.NAME_Strata);
+#endif
+}
+
+#if WITH_EDITOR
+int32 UMaterialExpressionStrataPostProcess::Compile(class FMaterialCompiler* Compiler, int32 OutputIndex)
+{
+	int OpacityCodeChunk = CompileWithDefaultFloat1(Compiler, Opacity, 0.0f);
+	int TransmittanceCodeChunk = Compiler->Saturate(Compiler->Sub(Compiler->Constant(1.0f), OpacityCodeChunk));
+
+	int32 OutputCodeChunk = Compiler->StrataUnlitBSDF(
+		CompileWithDefaultFloat3(Compiler, Color, 0.0f, 0.0f, 0.0f),
+		TransmittanceCodeChunk);
+
+	StrataCompilationInfoCreateSingleBSDFMaterial(Compiler, OutputCodeChunk, StrataCompilationInfoCreateNullSharedLocalBasis(), STRATA_BSDF_TYPE_UNLIT);
+
+	return OutputCodeChunk;
+}
+
+void UMaterialExpressionStrataPostProcess::GetCaption(TArray<FString>& OutCaptions) const
+{
+	OutCaptions.Add(TEXT("Strata Post Process"));
+}
+
+uint32 UMaterialExpressionStrataPostProcess::GetOutputType(int32 OutputIndex)
+{
+	return MCT_Strata;
+}
+
+uint32 UMaterialExpressionStrataPostProcess::GetInputType(int32 InputIndex)
+{
+	switch (InputIndex)
+	{
+	case 0:
+		return MCT_Float3;
+		break;
+	case 1:
+		return MCT_Float1;
+		break;
+	}
+
+	check(false);
+	return MCT_Float1;
+}
+
+bool UMaterialExpressionStrataPostProcess::IsResultStrataMaterial(int32 OutputIndex)
+{
+	return true;
+}
+
+void UMaterialExpressionStrataPostProcess::GatherStrataMaterialInfo(FStrataMaterialInfo& StrataMaterialInfo, int32 OutputIndex)
+{
+	StrataMaterialInfo.AddShadingModel(SSM_PostProcess);
+}
+
+FStrataOperator* UMaterialExpressionStrataPostProcess::StrataGenerateMaterialTopologyTree(class FMaterialCompiler* Compiler, class UMaterialExpression* Parent, int32 OutputIndex)
+{
+	return &Compiler->StrataCompilationRegisterOperator(STRATA_OPERATOR_BSDF, this, Parent);
+}
+#endif // WITH_EDITOR
+
+
+
 UMaterialExpressionStrataUnlitBSDF::UMaterialExpressionStrataUnlitBSDF(const FObjectInitializer& ObjectInitializer)
 	: Super(ObjectInitializer)
 {
