Prevent crashes when texture platform data was serialized with a bad pixel format name.

#rnx

diff --git a/Engine/Source/Runtime/Engine/Private/TextureDerivedData.cpp b/Engine/Source/Runtime/Engine/Private/TextureDerivedData.cpp
--- a/Engine/Source/Runtime/Engine/Private/TextureDerivedData.cpp
+++ b/Engine/Source/Runtime/Engine/Private/TextureDerivedData.cpp
@@ -2227,7 +2227,16 @@ static void SerializePlatformData(
 	{
 		FString PixelFormatString;
 		Ar << PixelFormatString;
-		PlatformData->PixelFormat = (EPixelFormat)PixelFormatEnum->GetValueByName(*PixelFormatString);
+		const int64 PixelFormatValue = PixelFormatEnum->GetValueByName(*PixelFormatString);
+		if (PixelFormatValue != INDEX_NONE && PixelFormatValue < PF_MAX)
+		{
+			PlatformData->PixelFormat = (EPixelFormat)PixelFormatValue;
+		}
+		else
+		{
+			UE_LOG(LogTexture, Warning, TEXT("Invalid pixel format '%s' for texture '%s'."), *PixelFormatString, Texture ? *Texture->GetPathName() : TEXT(""));
+			PlatformData->PixelFormat = PF_Unknown;
+		}
 	}
 	else if (Ar.IsSaving())
 	{
@@ -3253,7 +3262,9 @@ void UTexture::SerializeCookedPlatformData(FArchive& Ar)
 		Ar << PixelFormatName;
 		while (PixelFormatName != NAME_None)
 		{
-			EPixelFormat PixelFormat = (EPixelFormat)PixelFormatEnum->GetValueByName(PixelFormatName);
+			const int64 PixelFormatValue = PixelFormatEnum->GetValueByName(PixelFormatName);
+			const EPixelFormat PixelFormat = (PixelFormatValue != INDEX_NONE && PixelFormatValue < PF_MAX) ? (EPixelFormat)PixelFormatValue : PF_Unknown;
+
 			const int64 SkipOffsetLoc = Ar.Tell();
 			int64 SkipOffset = 0;
 			Ar << SkipOffset;
