Added support for strata post process blend mode.

diff --git a/Engine/Source/Runtime/Engine/Classes/Materials/MaterialExpressionStrata.h b/Engine/Source/Runtime/Engine/Classes/Materials/MaterialExpressionStrata.h
--- a/Engine/Source/Runtime/Engine/Classes/Materials/MaterialExpressionStrata.h
+++ b/Engine/Source/Runtime/Engine/Classes/Materials/MaterialExpressionStrata.h
@@ -582,13 +582,13 @@ class UMaterialExpressionStrataPostProcess : public UMaterialExpressionStrataBSD
 	GENERATED_UCLASS_BODY()
 
 	/**
-	 * The output color of the post process
+	 * The output color of the post process: it represents a color added over the back buffer, or a color multiplied if the Strata blend mode is TransmittanceOnly.
 	 */
 	UPROPERTY()
 	FExpressionInput Color;
 
 	/**
-	 * The opacity of the post process. Only writable if "Output Alpha" is enabled.
+	 * The coverage of the post process: the more the value is high, the less the back buffer will be visible. Only used if "Output Alpha" is enabled on the root node.
 	 */
 	UPROPERTY()
 	FExpressionInput Opacity;
diff --git a/Engine/Source/Runtime/Engine/Private/Materials/Material.cpp b/Engine/Source/Runtime/Engine/Private/Materials/Material.cpp
--- a/Engine/Source/Runtime/Engine/Private/Materials/Material.cpp
+++ b/Engine/Source/Runtime/Engine/Private/Materials/Material.cpp
@@ -3041,10 +3041,10 @@ void UMaterial::ConvertMaterialToStrataMaterial()
 			ShadingModels.AddShadingModel(MSM_Unlit);
 
 			// Only Emissive & Opacity are valid input for PostProcess material
-			UMaterialExpressionStrataLightFunction* ConvertNode = NewObject<UMaterialExpressionStrataLightFunction>(this);
-			MoveConnectionTo(EmissiveColor, ConvertNode, 0);
+			UMaterialExpressionStrataLightFunction* LightFunctionNode = NewObject<UMaterialExpressionStrataLightFunction>(this);
+			MoveConnectionTo(EmissiveColor, LightFunctionNode, 0);
 
-			FrontMaterial.Connect(0, ConvertNode);
+			FrontMaterial.Connect(0, LightFunctionNode);
 			bInvalidateShader = true;
 		}
 		else if (MaterialDomain == MD_PostProcess)
@@ -3054,12 +3054,33 @@ void UMaterial::ConvertMaterialToStrataMaterial()
 			ShadingModels.ClearShadingModels();
 			ShadingModels.AddShadingModel(MSM_Unlit);
 
-			// Only Emissive & Opacity are valid input for PostProcess material
-			UMaterialExpressionStrataPostProcess* ConvertNode = NewObject<UMaterialExpressionStrataPostProcess>(this);
-			MoveConnectionTo(EmissiveColor, ConvertNode, 0);
-			MoveConnectionTo(Opacity, ConvertNode, 1);
+			if (MaterialDomain == MD_PostProcess && !BlendableOutputAlpha)
+			{
+				BlendMode = BLEND_Opaque;
+			}
 
-			FrontMaterial.Connect(0, ConvertNode);
+			UMaterialExpressionStrataPostProcess* PostProcNode = NewObject<UMaterialExpressionStrataPostProcess>(this);
+			if (GetBlendMode() == BLEND_Translucent)
+			{
+				// If the blending mode was translucent, we must multiply the color by the opacity before the output node
+				UMaterialExpressionMultiply* ColorMultiplyOpacityNode = NewObject<UMaterialExpressionMultiply>(this);
+				CopyConnectionTo(Opacity, ColorMultiplyOpacityNode, 0);
+				MoveConnectionTo(EmissiveColor, ColorMultiplyOpacityNode, 1);
+
+				PostProcNode->Color.Connect(0, ColorMultiplyOpacityNode);
+				MoveConnectionTo(Opacity, PostProcNode, 1);
+			}
+			else
+			{
+				// Only Emissive & Opacity are valid input for PostProcess material
+				MoveConnectionTo(EmissiveColor, PostProcNode, 0);
+				if (GetBlendMode() != BLEND_Additive) // Legacy material additive was ignoring opacity
+				{
+					MoveConnectionTo(Opacity, PostProcNode, 1);
+				}
+			}
+
+			FrontMaterial.Connect(0, PostProcNode);
 			bInvalidateShader = true;
 		}
 		else if (MaterialDomain == MD_DeferredDecal)
@@ -4220,7 +4241,7 @@ void UMaterial::RebuildShadingModelField()
 			{
 				MaterialDomain = EMaterialDomain::MD_PostProcess;
 				ShadingModel = MSM_Unlit;
-				BlendMode = EBlendMode::BLEND_Opaque;
+				// We keep the blend mode resulting from ConvertLegacyToStrataBlendMode because post processes can be translucent.
 			}
 
 			// Also update the ShadingModels for remaining pipeline operation
diff --git a/Engine/Source/Runtime/Renderer/Private/PostProcess/PostProcessMaterial.cpp b/Engine/Source/Runtime/Renderer/Private/PostProcess/PostProcessMaterial.cpp
--- a/Engine/Source/Runtime/Renderer/Private/PostProcess/PostProcessMaterial.cpp
+++ b/Engine/Source/Runtime/Renderer/Private/PostProcess/PostProcessMaterial.cpp
@@ -21,6 +21,7 @@
 #include "BufferVisualizationData.h"
 #include "SceneTextureParameters.h"
 #include "SystemTextures.h"
+#include "Strata/Strata.h"
 
 namespace
 {
@@ -125,6 +126,26 @@ FRHIBlendState* GetMaterialBlendState(const FMaterial* Material)
 
 	check(Material);
 
+	if (Strata::IsStrataEnabled())
+	{
+		switch (Material->GetStrataBlendMode())
+		{
+		case EStrataBlendMode::SBM_Opaque:
+		case EStrataBlendMode::SBM_Masked:
+			return TStaticBlendState<>::GetRHI();
+		case EStrataBlendMode::SBM_TranslucentColoredTransmittance: // A platform may not support dual source blending so we always only use grey scale transmittance
+		case EStrataBlendMode::SBM_TranslucentGreyTransmittance:
+			return TStaticBlendState<CW_RGBA, BO_Add, BF_One, BF_InverseSourceAlpha, BO_Add, BF_One, BF_InverseSourceAlpha>::GetRHI();
+		case EStrataBlendMode::SBM_ColoredTransmittanceOnly:
+			return TStaticBlendState<CW_RGB, BO_Add, BF_DestColor, BF_Zero>::GetRHI();
+		case EStrataBlendMode::SBM_AlphaHoldout:
+			return TStaticBlendState<CW_RGBA, BO_Add, BF_Zero, BF_InverseSourceAlpha, BO_Add, BF_Zero, BF_InverseSourceAlpha>::GetRHI();
+		default:
+			check(false);
+			return TStaticBlendState<>::GetRHI();
+		}
+	}
+
 	return BlendStates[Material->GetBlendMode()];
 }
 
