Strata post process:
- removed no longer necessary test
- fall back post process material blend mode to opaque when write alpoha is disabled.

diff --git a/Engine/Source/Runtime/Engine/Private/Materials/Material.cpp b/Engine/Source/Runtime/Engine/Private/Materials/Material.cpp
--- a/Engine/Source/Runtime/Engine/Private/Materials/Material.cpp
+++ b/Engine/Source/Runtime/Engine/Private/Materials/Material.cpp
@@ -4189,10 +4189,7 @@ void UMaterial::RebuildShadingModelField()
 			// Now derive some properties from the material
 			if (StrataMaterialInfo.HasOnlyShadingModel(SSM_Unlit))
 			{
-				if (MaterialDomain != EMaterialDomain::MD_Surface && MaterialDomain != EMaterialDomain::MD_PostProcess && MaterialDomain != EMaterialDomain::MD_LightFunction)
-				{
-					MaterialDomain = EMaterialDomain::MD_Surface;
-				}
+				MaterialDomain = EMaterialDomain::MD_Surface;
 				ShadingModel = MSM_Unlit;
 			}
 			else if (StrataMaterialInfo.HasOnlyShadingModel(SSM_SubsurfaceLit))
@@ -4242,6 +4239,11 @@ void UMaterial::RebuildShadingModelField()
 				MaterialDomain = EMaterialDomain::MD_PostProcess;
 				ShadingModel = MSM_Unlit;
 				// We keep the blend mode resulting from ConvertLegacyToStrataBlendMode because post processes can be translucent.
+				// However, we do force opaque mode if blending has been disabled via the post-process specific BlendableOutputAlpha option.
+				if (!BlendableOutputAlpha)
+				{
+					StrataBlendMode = SBM_Opaque;
+				}
 			}
 
 			// Also update the ShadingModels for remaining pipeline operation
