Add missing forward declaration, fix non-unity compile

diff --git a/Engine/Source/Runtime/Engine/Public/MaterialCachedData.h b/Engine/Source/Runtime/Engine/Public/MaterialCachedData.h
--- a/Engine/Source/Runtime/Engine/Public/MaterialCachedData.h
+++ b/Engine/Source/Runtime/Engine/Public/MaterialCachedData.h
@@ -21,6 +21,7 @@ class ULandscapeGrassType;
 class UMaterialFunctionInterface;
 class UMaterialInterface;
 class FMaterialCachedHLSLTree;
+struct FStaticParameterSet;
 
 /** Stores information about a function that this material references, used to know when the material needs to be recompiled. */
 USTRUCT()
