Fix nullptr check reported by static analysis

diff --git a/Engine/Source/Runtime/Engine/Private/Materials/MaterialExpressions.cpp b/Engine/Source/Runtime/Engine/Private/Materials/MaterialExpressions.cpp
--- a/Engine/Source/Runtime/Engine/Private/Materials/MaterialExpressions.cpp
+++ b/Engine/Source/Runtime/Engine/Private/Materials/MaterialExpressions.cpp
@@ -14509,26 +14509,27 @@ void UMaterialExpressionMaterialFunctionCall::GetConnectorToolTip(int32 InputInd
 			if (FunctionInputs.IsValidIndex(InputIndex))
 			{
 				UMaterialExpressionFunctionInput* InputExpression = FunctionInputs[InputIndex].ExpressionInput;
-
-				ConvertToMultilineToolTip(InputExpression->Description, 40, OutToolTip);
-
-				if (InputExpression && InputExpression->bUsePreviewValueAsDefault)
+				if (InputExpression)
 				{
-					// Can't build a tooltip of an arbitrary expression chain
-					if (InputExpression->Preview.Expression)
+					ConvertToMultilineToolTip(InputExpression->Description, 40, OutToolTip);
+					if (InputExpression->bUsePreviewValueAsDefault)
 					{
-						OutToolTip.Insert(FString(TEXT("DefaultValue = Custom expressions")), 0);
+						// Can't build a tooltip of an arbitrary expression chain
+						if (InputExpression->Preview.Expression)
+						{
+							OutToolTip.Insert(FString(TEXT("DefaultValue = Custom expressions")), 0);
 
-						// Add a line after the default value string
-						OutToolTip.Insert(FString(TEXT("")), 1);
-					}
-					else if (InputExpression->InputType <= FunctionInput_Vector4)
-					{
-						// Add a string for the default value at the top
-						OutToolTip.Insert(GetInputDefaultValueString((EFunctionInputType)InputExpression->InputType, InputExpression->PreviewValue), 0);
+							// Add a line after the default value string
+							OutToolTip.Insert(FString(TEXT("")), 1);
+						}
+						else if (InputExpression->InputType <= FunctionInput_Vector4)
+						{
+							// Add a string for the default value at the top
+							OutToolTip.Insert(GetInputDefaultValueString((EFunctionInputType)InputExpression->InputType, InputExpression->PreviewValue), 0);
 
-						// Add a line after the default value string
-						OutToolTip.Insert(FString(TEXT("")), 1);
+							// Add a line after the default value string
+							OutToolTip.Insert(FString(TEXT("")), 1);
+						}
 					}
 				}
 			}
