Add descriptions for our various `stat *` EngineStat commands.

- Add the stat description to the output of `stat help`.
- `stat detailed` will now display the build string version (which includes changelist).

diff --git a/Engine/Source/Runtime/Engine/Private/UnrealEngine.cpp b/Engine/Source/Runtime/Engine/Private/UnrealEngine.cpp
--- a/Engine/Source/Runtime/Engine/Private/UnrealEngine.cpp
+++ b/Engine/Source/Runtime/Engine/Private/UnrealEngine.cpp
@@ -2034,34 +2034,24 @@ void UEngine::Init(IEngineLoop* InEngineLoop)
 
 	// Add the stats to the list, note this is also the order that they get rendered in if active.
 #if !UE_BUILD_SHIPPING
-	EngineStats.Add(FEngineStatFuncs(TEXT("STAT_Version"), TEXT("STATCAT_Engine"), FText::GetEmpty(), FEngineStatRender::CreateUObject(this, &UEngine::RenderStatVersion), FEngineStatToggle(), bIsRHS));
+	EngineStats.Add(FEngineStatFuncs(TEXT("STAT_Version"), TEXT("STATCAT_Engine"), FText::FromString(TEXT("Display build version string which includes the Changelist number.")), FEngineStatRender::CreateUObject(this, &UEngine::RenderStatVersion), FEngineStatToggle(), bIsRHS));
 #endif // !UE_BUILD_SHIPPING
 	EngineStats.Add(FEngineStatFuncs(TEXT("STAT_NamedEvents"), TEXT("STATCAT_Engine"), FText::GetEmpty(), FEngineStatRender::CreateUObject(this, &UEngine::RenderStatNamedEvents), FEngineStatToggle::CreateUObject(this, &UEngine::ToggleStatNamedEvents), bIsRHS));
 	EngineStats.Add(FEngineStatFuncs(TEXT("STAT_VerboseNamedEvents"), TEXT("STATCAT_Engine"), FText::GetEmpty(), FEngineStatRender(), FEngineStatToggle::CreateUObject(this, &UEngine::ToggleStatVerboseNamedEvents), bIsRHS));
-	EngineStats.Add(FEngineStatFuncs(TEXT("STAT_FPS"), TEXT("STATCAT_Engine"), FText::GetEmpty(), FEngineStatRender::CreateUObject(this, &UEngine::RenderStatFPS), FEngineStatToggle::CreateUObject(this, &UEngine::ToggleStatFPS), bIsRHS));
-	EngineStats.Add(FEngineStatFuncs(TEXT("STAT_Summary"), TEXT("STATCAT_Engine"), FText::GetEmpty(), FEngineStatRender::CreateUObject(this, &UEngine::RenderStatSummary), FEngineStatToggle(), bIsRHS));
-	EngineStats.Add(FEngineStatFuncs(TEXT("STAT_Unit"), TEXT("STATCAT_Engine"), FText::GetEmpty(), FEngineStatRender::CreateUObject(this, &UEngine::RenderStatUnit), FEngineStatToggle::CreateUObject(this, &UEngine::ToggleStatUnit), bIsRHS));
-	EngineStats.Add(FEngineStatFuncs(TEXT("STAT_DrawCount"), TEXT("STATCAT_Engine"), FText::GetEmpty(), FEngineStatRender::CreateUObject(this, &UEngine::RenderStatDrawCount), FEngineStatToggle(), bIsRHS));
-	/* @todo Slate Rendering
-	#if STATS
-	EngineStats.Add(FEngineStatFuncs(TEXT("STAT_SlateBatches"), TEXT("STATCAT_Engine"), FText::GetEmpty(), &UEngine::RenderStatSlateBatches, NULL, true));
-	#endif
-	*/
-	EngineStats.Add(FEngineStatFuncs(TEXT("STAT_Hitches"), TEXT("STATCAT_Engine"), FText::GetEmpty(), FEngineStatRender::CreateUObject(this, &UEngine::RenderStatHitches), FEngineStatToggle::CreateUObject(this, &UEngine::ToggleStatHitches), bIsRHS));
-	EngineStats.Add(FEngineStatFuncs(TEXT("STAT_AI"), TEXT("STATCAT_Engine"), FText::GetEmpty(), FEngineStatRender::CreateUObject(this, &UEngine::RenderStatAI), FEngineStatToggle(), bIsRHS));
-	EngineStats.Add(FEngineStatFuncs(TEXT("STAT_Timecode"), TEXT("STATCAT_Engine"), FText::GetEmpty(), FEngineStatRender::CreateUObject(this, &UEngine::RenderStatTimecode), FEngineStatToggle(), bIsRHS));
-	EngineStats.Add(FEngineStatFuncs(TEXT("STAT_FrameCounter"), TEXT("STATCAT_Engine"), FText::GetEmpty(), FEngineStatRender::CreateUObject(this, &UEngine::RenderStatFrameCounter), FEngineStatToggle(), bIsRHS));
-
-	EngineStats.Add(FEngineStatFuncs(TEXT("STAT_ColorList"), TEXT("STATCAT_Engine"), FText::GetEmpty(), FEngineStatRender::CreateUObject(this, &UEngine::RenderStatColorList), FEngineStatToggle()));
-	EngineStats.Add(FEngineStatFuncs(TEXT("STAT_Levels"), TEXT("STATCAT_Engine"), FText::GetEmpty(), FEngineStatRender::CreateUObject(this, &UEngine::RenderStatLevels), FEngineStatToggle()));
-
-	/* @todo physx fix this once we have convexelem drawing again
-	EngineStats.Add(FEngineStatFuncs(TEXT("STAT_LevelMap"), TEXT("STATCAT_Engine"), FText::GetEmpty(), &UEngine::RenderStatLevelMap, NULL));
-	*/
-	EngineStats.Add(FEngineStatFuncs(TEXT("STAT_Detailed"), TEXT("STATCAT_Engine"), FText::GetEmpty(), FEngineStatRender(), FEngineStatToggle::CreateUObject(this, &UEngine::ToggleStatDetailed)));
+	EngineStats.Add(FEngineStatFuncs(TEXT("STAT_FPS"), TEXT("STATCAT_Engine"), FText::FromString(TEXT("Display average FPS.")), FEngineStatRender::CreateUObject(this, &UEngine::RenderStatFPS), FEngineStatToggle::CreateUObject(this, &UEngine::ToggleStatFPS), bIsRHS));
+	EngineStats.Add(FEngineStatFuncs(TEXT("STAT_Summary"), TEXT("STATCAT_Engine"), FText::FromString(TEXT("Display UsedPhysical memory stat.")), FEngineStatRender::CreateUObject(this, &UEngine::RenderStatSummary), FEngineStatToggle(), bIsRHS));
+	EngineStats.Add(FEngineStatFuncs(TEXT("STAT_Unit"), TEXT("STATCAT_Engine"), FText::FromString(TEXT("Display frame times for Game Thread, Render Thread and GPU.  Also displays draw calls and primitive counts.")), FEngineStatRender::CreateUObject(this, &UEngine::RenderStatUnit), FEngineStatToggle::CreateUObject(this, &UEngine::ToggleStatUnit), bIsRHS));
+	EngineStats.Add(FEngineStatFuncs(TEXT("STAT_DrawCount"), TEXT("STATCAT_Engine"), FText::FromString(TEXT("Display draw call counts per render pass.")), FEngineStatRender::CreateUObject(this, &UEngine::RenderStatDrawCount), FEngineStatToggle(), bIsRHS));
+	EngineStats.Add(FEngineStatFuncs(TEXT("STAT_Hitches"), TEXT("STATCAT_Engine"), FText::FromString(TEXT("Display a scrolling hitch time when frame hitches occur.")), FEngineStatRender::CreateUObject(this, &UEngine::RenderStatHitches), FEngineStatToggle::CreateUObject(this, &UEngine::ToggleStatHitches), bIsRHS));
+	EngineStats.Add(FEngineStatFuncs(TEXT("STAT_AI"), TEXT("STATCAT_Engine"), FText::FromString(TEXT("Display number of APlayerControllers and the number that were rendered to screen.")), FEngineStatRender::CreateUObject(this, &UEngine::RenderStatAI), FEngineStatToggle(), bIsRHS));
+	EngineStats.Add(FEngineStatFuncs(TEXT("STAT_Timecode"), TEXT("STATCAT_Engine"), FText::FromString(TEXT("Display a time from a timecode provider.")), FEngineStatRender::CreateUObject(this, &UEngine::RenderStatTimecode), FEngineStatToggle(), bIsRHS));
+	EngineStats.Add(FEngineStatFuncs(TEXT("STAT_FrameCounter"), TEXT("STATCAT_Engine"), FText::FromString(TEXT("Display the global, steadily increasing frame counter.")), FEngineStatRender::CreateUObject(this, &UEngine::RenderStatFrameCounter), FEngineStatToggle(), bIsRHS));
+	EngineStats.Add(FEngineStatFuncs(TEXT("STAT_ColorList"), TEXT("STATCAT_Engine"), FText::FromString(TEXT("Display a list of common FColors and their names.")), FEngineStatRender::CreateUObject(this, &UEngine::RenderStatColorList), FEngineStatToggle()));
+	EngineStats.Add(FEngineStatFuncs(TEXT("STAT_Levels"), TEXT("STATCAT_Engine"), FText::FromString(TEXT("Display a list of names of all the loaded levels.")), FEngineStatRender::CreateUObject(this, &UEngine::RenderStatLevels), FEngineStatToggle()));
+	EngineStats.Add(FEngineStatFuncs(TEXT("STAT_Detailed"), TEXT("STATCAT_Engine"), FText::FromString(TEXT("Display detailed frame and fps timings.")), FEngineStatRender(), FEngineStatToggle::CreateUObject(this, &UEngine::ToggleStatDetailed)));
 #if !UE_BUILD_SHIPPING
-	EngineStats.Add(FEngineStatFuncs(TEXT("STAT_UnitMax"), TEXT("STATCAT_Engine"), FText::GetEmpty(), FEngineStatRender(), FEngineStatToggle::CreateUObject(this, &UEngine::ToggleStatUnitMax)));
-	EngineStats.Add(FEngineStatFuncs(TEXT("STAT_UnitGraph"), TEXT("STATCAT_Engine"), FText::GetEmpty(), FEngineStatRender(), FEngineStatToggle::CreateUObject(this, &UEngine::ToggleStatUnitGraph)));
+	EngineStats.Add(FEngineStatFuncs(TEXT("STAT_UnitMax"), TEXT("STATCAT_Engine"), FText::FromString(TEXT("Same as stat unit, but with max values also displayed.")), FEngineStatRender(), FEngineStatToggle::CreateUObject(this, &UEngine::ToggleStatUnitMax)));
+	EngineStats.Add(FEngineStatFuncs(TEXT("STAT_UnitGraph"), TEXT("STATCAT_Engine"), FText::FromString(TEXT("Displays a frame time stats graphed over time.")), FEngineStatRender(), FEngineStatToggle::CreateUObject(this, &UEngine::ToggleStatUnitGraph)));
 	EngineStats.Add(FEngineStatFuncs(TEXT("STAT_UnitTime"), TEXT("STATCAT_Engine"), FText::GetEmpty(), FEngineStatRender(), FEngineStatToggle::CreateUObject(this, &UEngine::ToggleStatUnitTime)));
 	EngineStats.Add(FEngineStatFuncs(TEXT("STAT_Raw"), TEXT("STATCAT_Engine"), FText::GetEmpty(), FEngineStatRender(), FEngineStatToggle::CreateUObject(this, &UEngine::ToggleStatRaw)));
 	EngineStats.Add(FEngineStatFuncs(TEXT("STAT_ParticlePerf"), TEXT("STATCAT_Engine"), FText::GetEmpty(), FEngineStatRender::CreateUObject(this, &UEngine::RenderStatParticlePerf), FEngineStatToggle::CreateUObject(this, &UEngine::ToggleStatParticlePerf), bIsRHS));
@@ -4949,7 +4939,8 @@ bool UEngine::HandleStatCommand( UWorld* World, FCommonViewportClient* ViewportC
 		for (int32 StatIdx = 0; StatIdx < EngineStats.Num(); StatIdx++)
 		{
 			const FEngineStatFuncs& EngineStat = EngineStats[StatIdx];
-			StatNames.Add_GetRef( EngineStat.CommandNameString);
+			const FString EngineStatString = FString::Printf(TEXT("%s: %s"), *EngineStat.CommandNameString, *EngineStat.DescriptionString.ToString());
+			StatNames.Add_GetRef(EngineStatString);
 		}
 
 		StatNames.Sort();
@@ -16440,6 +16431,7 @@ bool UEngine::ToggleStatDetailed(UWorld* World, FCommonViewportClient* ViewportC
 		DetailedStats.Add(TEXT("UnitMax"));
 		DetailedStats.Add(TEXT("UnitGraph"));
 		DetailedStats.Add(TEXT("Raw"));
+		DetailedStats.Add(TEXT("Version"));
 	}
 
 	// If any of the detailed stats are inactive, take this as enabling all, unless 'Skip' is specifically specified
