Adding a cvar, `r.SkinCache.SkipCompilingGPUSkinVF` that allows you to skip compiling _all_ GPU Skin Vertex Factory variants (except the pass through) if you are using the GPU Skin Cache.

# Results
* EngineTest: Saves 8,207 shaders or 8.88% of all the shaders on the project.
* MedievalGame: Saves 3,640 shaders or 7.84% of all the shaders on the project.
* QAGame: Saves 8,957 shaders or 10.43% of all the shaders on the project.

# Description
The idea is that if we are using the GPU Skin Cache and all our skeletal meshes can fit within the cache at runtime we only need to use the global compute shaders to perform GPU skinning.  This means we will never need any GPU skin factory variants for any material so we can skip compiling these shaders for a massive savings.

# Caveats
Despite all these caveats this is still a worthwhile feature.  Depending on the project these tradeoffs listed below might be worth it for reducing shader counts by so much (which will reduce runtime memory too!)
- If enabled, the project will have no GPU Skin shaders so all the skeletal meshes at runtime must fit within the GPU Skin Cache's memory limits.
- If a mesh doesn't fit within the cache it will just render with the pass through VF and ref pose.  This is preferable to crashing.  An ensure will trigger with a warning message letting you know to disable this feature or increase the Skin Cache size.
- The ability to run-time toggle the GPU Skin cache is lost because there are no GPU skin shaders to fallback to.
- The ability to toggle per-mesh the GPU Skin Cache is lost because there are no GPU skin shaders to fallback to.

# Changes
- Enabling this for the EngineTest, MedievalGame and QAGame project.
- Expose r.SkinCache.SkipCompilingGPUSkinVF as a project setting and prevent the user from enabling this if the skin cache isn't enabled.  If the user ever disables the skin cache this CVar should be disabled too.
- Ensure that we render with a ref posed version of the skel mesh in `FSkeletalMeshObjectGPUSkin::GetSkinVertexFactory` if we ever don't have shaders and the skin cache isn't available (for any reason)
- Added a helper function `ShouldWeCompileGPUSkinVFShaders` to help determine if we can skip compiling GPU Skin VF shaders.
- Added an exec command `list skincacheusage` that will list all meshes and whether or not they have the skin cache enabled/disabled.
- Modify `FGPUSkinCache::CVarSinkFunction` to prevent the user from turning of the skin cache if this feature is enabled.

diff --git a/Engine/Source/Runtime/Engine/Classes/Engine/RendererSettings.h b/Engine/Source/Runtime/Engine/Classes/Engine/RendererSettings.h
--- a/Engine/Source/Runtime/Engine/Classes/Engine/RendererSettings.h
+++ b/Engine/Source/Runtime/Engine/Classes/Engine/RendererSettings.h
@@ -921,6 +921,15 @@ class ENGINE_API URendererSettings : public UDeveloperSettings
 		ConfigRestartRequired = true))
 	uint32 bSupportSkinCacheShaders : 1;
 
+	/**
+	"When enabled this will skip compiling GPU skin vertex factory shader variants with the assumption that all skinning work will be done via the skin cache."
+	*/
+	UPROPERTY(config, EditAnywhere, Category = Optimizations, meta = (
+		ConsoleVariable = "r.SkinCache.SkipCompilingGPUSkinVF", DisplayName = "Reduce GPU Skin Vertex Factory shader permutations",
+		ToolTip = "Cannot be enabled while the skin cache is turned off.",
+		ConfigRestartRequired = true))
+	uint32 bSkipCompilingGPUSkinVF : 1;
+
 	UPROPERTY(config, EditAnywhere, Category = Optimizations, meta = (
 		ConsoleVariable = "r.SkinCache.DefaultBehavior", DisplayName = "Default Skin Cache Behavior",
 		ToolTip = "Default behavior if all skeletal meshes are included/excluded from the skin cache. If Support Ray Tracing is enabled on a mesh, will force inclusive behavior on that mesh."))
diff --git a/Engine/Source/Runtime/Engine/Private/GPUSkinCache.cpp b/Engine/Source/Runtime/Engine/Private/GPUSkinCache.cpp
--- a/Engine/Source/Runtime/Engine/Private/GPUSkinCache.cpp
+++ b/Engine/Source/Runtime/Engine/Private/GPUSkinCache.cpp
@@ -20,6 +20,7 @@ GPUSkinCache.cpp: Performs skinning on a compute shader into a buffer to avoid v
 #include "Algo/Unique.h"
 #include "HAL/IConsoleManager.h"
 #include "RayTracingSkinnedGeometry.h"
+#include "Internationalization/Internationalization.h"
 
 DEFINE_STAT(STAT_GPUSkinCache_TotalNumChunks);
 DEFINE_STAT(STAT_GPUSkinCache_TotalNumVertices);
@@ -31,6 +32,36 @@ DEFINE_STAT(STAT_GPUSkinCache_NumSetVertexStreams);
 DEFINE_STAT(STAT_GPUSkinCache_NumPreGDME);
 DEFINE_LOG_CATEGORY_STATIC(LogSkinCache, Log, All);
 
+/** Exec helper to handle GPU Skin Cache related commands. */
+class FSkinCacheExecHelper : public FSelfRegisteringExec
+{
+	virtual bool Exec(class UWorld* InWorld, const TCHAR* Cmd, FOutputDevice& Ar)
+	{
+		/** Command to list all skeletal mesh lods which have the skin cache disabled. */
+		if (FParse::Command(&Cmd, TEXT("list skincacheusage")))
+		{
+			UE_LOG(LogTemp, Display, TEXT("Name, Lod Index, Skin Cache Usage"));
+
+			for (TObjectIterator<USkeletalMesh> It; It; ++It)
+			{
+				if (USkeletalMesh* SkeletalMesh = *It)
+				{
+					for (int32 LODIndex = 0; LODIndex < SkeletalMesh->GetLODNum(); ++LODIndex)
+					{
+						if (FSkeletalMeshLODInfo* LODInfo = SkeletalMesh->GetLODInfo(LODIndex))
+						{
+							UE_LOG(LogTemp, Display, TEXT("%s, %d, %d"), *SkeletalMesh->GetFullName(), LODIndex, LODInfo->SkinCacheUsage);
+						}
+					}
+				}
+			}
+			return true;
+		}
+		return false;
+	}
+};
+static FSkinCacheExecHelper GSkelMeshExecHelper;
+
 static int32 GEnableGPUSkinCacheShaders = 0;
 
 static FAutoConsoleVariableRef CVarEnableGPUSkinCacheShaders(
@@ -43,6 +74,15 @@ static FAutoConsoleVariableRef CVarEnableGPUSkinCacheShaders(
 	ECVF_RenderThreadSafe | ECVF_ReadOnly
 );
 
+static TAutoConsoleVariable<bool> CVarSkipCompilingGPUSkinVF(
+	TEXT("r.SkinCache.SkipCompilingGPUSkinVF"),
+	false,
+	TEXT("Reduce GPU Skin Vertex Factory shader permutations. Cannot be disabled while the skin cache is turned off.\n")
+	TEXT(" False ( 0): Compile all GPU Skin Vertex factory variants.\n")
+	TEXT(" True  ( 1): Don't compile all GPU Skin Vertex factory variants."),
+    ECVF_RenderThreadSafe | ECVF_ReadOnly
+);
+
 // 0/1
 int32 GEnableGPUSkinCache = 1;
 static TAutoConsoleVariable<int32> CVarEnableGPUSkinCache(
@@ -158,6 +198,43 @@ ENGINE_API bool IsGPUSkinCacheAvailable(EShaderPlatform Platform)
 	return AreSkinCacheShadersEnabled(Platform) != 0 && DoesPlatformSupportGPUSkinCache(Platform);
 }
 
+static inline bool IsGPUSkinCacheEnable(EShaderPlatform Platform)
+{
+	static FShaderPlatformCachedIniValue<int32> PerPlatformCVar(TEXT("r.SkinCache.Mode"));
+	return (PerPlatformCVar.Get(Platform) != 0);
+}
+
+static inline bool IsGPUSkinCacheInclusive(EShaderPlatform Platform)
+{
+	static FShaderPlatformCachedIniValue<int32> PerPlatformCVar(TEXT("r.SkinCache.DefaultBehavior"));
+	return (PerPlatformCVar.Get(Platform) != 0);
+}
+
+bool ShouldWeCompileGPUSkinVFShaders(EShaderPlatform Platform)
+{
+	// If the skin cache is not available on this platform we need to compile GPU Skin VF shaders.
+	if (IsGPUSkinCacheAvailable(Platform) == false)
+	{
+		return true;
+	}
+
+	// If the skin cache is not available on this platform we need to compile GPU Skin VF Shaders.
+	if (IsGPUSkinCacheEnable(Platform) == false)
+	{
+		return true;
+	}
+
+	// If the skin cache has been globally disabled for all skeletal meshes we need to compile GPU Skin VF Shaders.
+	if (IsGPUSkinCacheInclusive(Platform) == false)
+	{
+		return true;
+	}
+
+	// If the skin cache is enabled and we've been asked to skip GPU Skin VF shaders.
+	static FShaderPlatformCachedIniValue<bool> PerPlatformCVar(TEXT("r.SkinCache.SkipCompilingGPUSkinVF"));
+	return (PerPlatformCVar.Get(Platform) == false);
+}
+
 ENGINE_API bool GPUSkinCacheNeedsDuplicatedVertices()
 {
 #if WITH_EDITOR // Duplicated vertices are used in the editor when merging meshes
@@ -2083,8 +2160,9 @@ void FGPUSkinCache::CVarSinkFunction()
 {
 	int32 NewGPUSkinCacheValue = CVarEnableGPUSkinCache.GetValueOnAnyThread() != 0;
 	int32 NewRecomputeTangentsValue = CVarGPUSkinCacheRecomputeTangents.GetValueOnAnyThread();
-	float NewSceneMaxSizeInMb = CVarGPUSkinCacheSceneMemoryLimitInMB.GetValueOnAnyThread();
-	int32 NewNumTangentIntermediateBuffers = CVarGPUSkinNumTangentIntermediateBuffers.GetValueOnAnyThread();
+	const float NewSceneMaxSizeInMb = CVarGPUSkinCacheSceneMemoryLimitInMB.GetValueOnAnyThread();
+	const int32 NewNumTangentIntermediateBuffers = CVarGPUSkinNumTangentIntermediateBuffers.GetValueOnAnyThread();
+	const bool NewSkipCompilingGPUSkinVF = CVarSkipCompilingGPUSkinVF.GetValueOnAnyThread();
 
 	if (GEnableGPUSkinCacheShaders)
 	{
@@ -2100,6 +2178,17 @@ void FGPUSkinCache::CVarSinkFunction()
 		NewRecomputeTangentsValue = 0;
 	}
 
+	// We don't have GPU Skin VF shaders at all so we can't fallback to using GPU Skinning.
+	if (NewSkipCompilingGPUSkinVF)
+	{
+		// If we had the skin cache enabled and we are turning it off.
+		if (GEnableGPUSkinCache && (NewGPUSkinCacheValue == 0))
+		{
+			NewGPUSkinCacheValue = 1;
+			UE_LOG(LogSkinCache, Warning, TEXT("Attemping to turn off the GPU Skin Cache, but we don't have GPU Skin VF shaders to fallback to (r.SkinCache.SkipCompilingGPUSkinVF=1).  Leaving skin cache turned on."));
+		}
+	}
+
 	if (NewGPUSkinCacheValue != GEnableGPUSkinCache || NewRecomputeTangentsValue != GSkinCacheRecomputeTangents
 		|| NewSceneMaxSizeInMb != GSkinCacheSceneMemoryLimitInMB || NewNumTangentIntermediateBuffers != GNumTangentIntermediateBuffers)
 	{
diff --git a/Engine/Source/Runtime/Engine/Private/GPUSkinVertexFactory.cpp b/Engine/Source/Runtime/Engine/Private/GPUSkinVertexFactory.cpp
--- a/Engine/Source/Runtime/Engine/Private/GPUSkinVertexFactory.cpp
+++ b/Engine/Source/Runtime/Engine/Private/GPUSkinVertexFactory.cpp
@@ -499,8 +499,10 @@ TGlobalResource<FBoneBufferPool> FGPUBaseSkinVertexFactory::BoneBufferPool;
 template <GPUSkinBoneInfluenceType BoneInfluenceType>
 bool TGPUSkinVertexFactory<BoneInfluenceType>::ShouldCompilePermutation(const FVertexFactoryShaderPermutationParameters& Parameters)
 {
-	bool bUnlimitedBoneInfluences = (BoneInfluenceType == UnlimitedBoneInfluence && GCVarUnlimitedBoneInfluences);
-	return ((Parameters.MaterialParameters.bIsUsedWithSkeletalMesh && (BoneInfluenceType != UnlimitedBoneInfluence || bUnlimitedBoneInfluences)) || Parameters.MaterialParameters.bIsSpecialEngineMaterial);
+	const bool bUnlimitedBoneInfluences = (BoneInfluenceType == UnlimitedBoneInfluence && GCVarUnlimitedBoneInfluences);
+	const bool bShouldCache = (ShouldWeCompileGPUSkinVFShaders(Parameters.Platform) && Parameters.MaterialParameters.bIsUsedWithSkeletalMesh && (BoneInfluenceType != UnlimitedBoneInfluence || bUnlimitedBoneInfluences));
+	const bool bIsSpecialEngineMaterial = Parameters.MaterialParameters.bIsSpecialEngineMaterial;
+	return (bShouldCache || bIsSpecialEngineMaterial);
 }
 
 template <GPUSkinBoneInfluenceType BoneInfluenceType>
diff --git a/Engine/Source/Runtime/Engine/Private/RendererSettings.cpp b/Engine/Source/Runtime/Engine/Private/RendererSettings.cpp
--- a/Engine/Source/Runtime/Engine/Private/RendererSettings.cpp
+++ b/Engine/Source/Runtime/Engine/Private/RendererSettings.cpp
@@ -70,6 +70,7 @@ URendererSettings::URendererSettings(const FObjectInitializer& ObjectInitializer
 	bSupportPointLightWholeSceneShadows = true;
 	bSupportSkyAtmosphere = true;
 	bSupportSkinCacheShaders = false;
+	bSkipCompilingGPUSkinVF = false;
 	GPUSimulationTextureSizeX = 1024;
 	GPUSimulationTextureSizeY = 1024;
 	bEnableRayTracing = 0;
@@ -181,6 +182,17 @@ void URendererSettings::PostEditChangeProperty(FPropertyChangedEvent& PropertyCh
 			}
 		}
 
+		if (PropertyChangedEvent.Property->GetFName() == GET_MEMBER_NAME_CHECKED(URendererSettings, bSupportSkinCacheShaders)
+			&& !bSupportSkinCacheShaders
+			&& bSkipCompilingGPUSkinVF)
+		{
+			FString FullPath = FPaths::ConvertRelativePathToFull(GetDefaultConfigFilename());
+			FPlatformFileManager::Get().GetPlatformFile().SetReadOnly(*FullPath, false);
+
+			bSkipCompilingGPUSkinVF = 0;
+			UpdateDependentPropertyInConfigFile(this, GET_MEMBER_NAME_CHECKED(URendererSettings, bSkipCompilingGPUSkinVF));
+		}
+
 		if (PropertyChangedEvent.Property->GetFName() == GET_MEMBER_NAME_CHECKED(URendererSettings, DynamicGlobalIllumination) 
 			&& DynamicGlobalIllumination == EDynamicGlobalIlluminationMethod::Lumen)
 		{
@@ -289,6 +301,12 @@ bool URendererSettings::CanEditChange(const FProperty* InProperty) const
 		return !bSupportSkinCacheShaders || !bEnableRayTracing;
 	}
 
+	// the bSkipCompilingGPUSkinVF setting can only be edited if the skin cache is on.
+	if ((InProperty->GetFName() == GET_MEMBER_NAME_CHECKED(URendererSettings, bSkipCompilingGPUSkinVF)))
+	{
+		return bSupportSkinCacheShaders;
+	}
+
 	// the following settings can only be edited if ray tracing is enabled
 	if ((InProperty->GetFName() == GET_MEMBER_NAME_CHECKED(URendererSettings, bEnablePathTracing)) ||
 		(InProperty->GetFName() == GET_MEMBER_NAME_CHECKED(URendererSettings, bEnableRayTracingShadows)) ||
diff --git a/Engine/Source/Runtime/Engine/Private/SkeletalRenderGPUSkin.cpp b/Engine/Source/Runtime/Engine/Private/SkeletalRenderGPUSkin.cpp
--- a/Engine/Source/Runtime/Engine/Private/SkeletalRenderGPUSkin.cpp
+++ b/Engine/Source/Runtime/Engine/Private/SkeletalRenderGPUSkin.cpp
@@ -1372,6 +1372,14 @@ const FVertexFactory* FSkeletalMeshObjectGPUSkin::GetSkinVertexFactory(const FSc
 		return LOD.GPUSkinVertexFactories.PassthroughVertexFactories[ChunkIdx].Get();
 	}
 
+	// If we have not compiled GPU Skin vertex factory variants
+	const auto CVar = IConsoleManager::Get().FindTConsoleVariableDataInt(TEXT("r.SkinCache.SkipCompilingGPUSkinVF"));
+	if (CVar && CVar->GetValueOnAnyThread() != 0)
+	{
+		ensureMsgf(false, TEXT("We are attempting to render with a GPU Skin Vertex Factory, but r.SkinCache.SkipCompilingGPUSkinVF=1 so we don't have shaders.  Skeletal meshes will draw in ref pose.  Either disable r.SkinCache.SkipCompilingGPUSkinVF or increase the r.SkinCache.SceneMemoryLimitInMB size."));
+		return LOD.GPUSkinVertexFactories.PassthroughVertexFactories[ChunkIdx].Get();
+	}
+
 	// No passthrough usage so return the base skin vertex factory.
 	return GetBaseSkinVertexFactory(LODIndex, ChunkIdx);
 }
diff --git a/Engine/Source/Runtime/Engine/Public/GPUSkinCache.h b/Engine/Source/Runtime/Engine/Public/GPUSkinCache.h
--- a/Engine/Source/Runtime/Engine/Public/GPUSkinCache.h
+++ b/Engine/Source/Runtime/Engine/Public/GPUSkinCache.h
@@ -57,6 +57,8 @@ struct FRayTracingGeometrySegment;
 // Can the skin cache be used (ie shaders added, etc)
 extern ENGINE_API bool IsGPUSkinCacheAvailable(EShaderPlatform Platform);
 
+extern bool ShouldWeCompileGPUSkinVFShaders(EShaderPlatform Platform);
+
 extern ENGINE_API bool GPUSkinCacheNeedsDuplicatedVertices();
 
 // Is it actually enabled?
