Use Owner instead of AttachmentRoot Actor for AttachParent when replicating attachment.

diff --git a/Engine/Source/Runtime/Engine/Classes/Components/SceneComponent.h b/Engine/Source/Runtime/Engine/Classes/Components/SceneComponent.h
--- a/Engine/Source/Runtime/Engine/Classes/Components/SceneComponent.h
+++ b/Engine/Source/Runtime/Engine/Classes/Components/SceneComponent.h
@@ -1258,6 +1258,9 @@ public:
 	/** Walks up the attachment chain from this SceneComponent and returns the top-level actor it's attached to.  Returns Owner if unattached. */
 	AActor* GetAttachmentRootActor() const;
 
+	/** Gets the owner of the attach parent */
+	AActor* GetAttachParentActor() const;
+
 	/** Walks up the attachment chain to see if this component is attached to the supplied component. If TestComp == this, returns false.*/
 	bool IsAttachedTo(const USceneComponent* TestComp) const;
 
diff --git a/Engine/Source/Runtime/Engine/Private/ActorReplication.cpp b/Engine/Source/Runtime/Engine/Private/ActorReplication.cpp
--- a/Engine/Source/Runtime/Engine/Private/ActorReplication.cpp
+++ b/Engine/Source/Runtime/Engine/Private/ActorReplication.cpp
@@ -395,7 +395,7 @@ void AActor::GatherCurrentMovement()
 			{
 				// Networking for attachments assumes the RootComponent of the AttachParent actor. 
 				// If that's not the case, we can't update this, as the client wouldn't be able to resolve the Component and would detach as a result.
-				AttachmentReplication.AttachParent = RootComponent->GetAttachParent()->GetAttachmentRootActor();
+				AttachmentReplication.AttachParent = RootComponent->GetAttachParentActor();
 				if (AttachmentReplication.AttachParent != nullptr)
 				{
 					AttachmentReplication.LocationOffset = RootComponent->GetRelativeLocation();
diff --git a/Engine/Source/Runtime/Engine/Private/Components/SceneComponent.cpp b/Engine/Source/Runtime/Engine/Private/Components/SceneComponent.cpp
--- a/Engine/Source/Runtime/Engine/Private/Components/SceneComponent.cpp
+++ b/Engine/Source/Runtime/Engine/Private/Components/SceneComponent.cpp
@@ -2285,6 +2285,12 @@ AActor* USceneComponent::GetAttachmentRootActor() const
 	return AttachmentRootComponent ? AttachmentRootComponent->GetOwner() : nullptr;
 }
 
+AActor* USceneComponent::GetAttachParentActor() const
+{
+	const USceneComponent* const AttachParentComponent = GetAttachParent();
+	return AttachParentComponent ? AttachParentComponent->GetOwner() : nullptr;
+}
+
 bool USceneComponent::IsAttachedTo(const USceneComponent* TestComp) const
 {
 	if(TestComp != nullptr)
