ImageUtils.h deprecate CompressImageArray, use ThumbnailCompressImageArray instead

diff --git a/Engine/Source/Editor/UnrealEd/Private/ThumbnailManager.cpp b/Engine/Source/Editor/UnrealEd/Private/ThumbnailManager.cpp
--- a/Engine/Source/Editor/UnrealEd/Private/ThumbnailManager.cpp
+++ b/Engine/Source/Editor/UnrealEd/Private/ThumbnailManager.cpp
@@ -319,7 +319,7 @@ bool UThumbnailManager::CaptureProjectThumbnail(FViewport* Viewport, const FStri
 
 		// Compress the scaled image
 		TArray<uint8> ScaledPng;
-		FImageUtils::CompressImageArray(ScaledSize, ScaledSize, ScaledBitmap, ScaledPng);
+		FImageUtils::ThumbnailCompressImageArray(ScaledSize, ScaledSize, ScaledBitmap, ScaledPng);
 
 		// Save to file
 		const FString ScreenShotPath = FPaths::GetPath(OutputFilename);
diff --git a/Engine/Source/Runtime/Engine/Private/ImageUtils.cpp b/Engine/Source/Runtime/Engine/Private/ImageUtils.cpp
--- a/Engine/Source/Runtime/Engine/Private/ImageUtils.cpp
+++ b/Engine/Source/Runtime/Engine/Private/ImageUtils.cpp
@@ -89,6 +89,8 @@ void FImageUtils::ImageResize(int32 SrcWidth, int32 SrcHeight, const TArray<FCol
  */
 void FImageUtils::ImageResize(int32 SrcWidth, int32 SrcHeight, const TArrayView<const FColor> &SrcData, int32 DstWidth, int32 DstHeight, const TArrayView<FColor> &DstData, bool bLinearSpace, bool bForceOpaqueOutput)
 {
+	//@todo Oodle : this could be better ; investigate how widely it's used
+
 	check(SrcData.Num() >= SrcWidth * SrcHeight);
 	check(DstData.Num() >= DstWidth * DstHeight);
 
@@ -211,6 +213,8 @@ void FImageUtils::ImageResize(int32 SrcWidth, int32 SrcHeight, const TArray64<FL
  */
 void FImageUtils::ImageResize(int32 SrcWidth, int32 SrcHeight, const TArrayView64<const FLinearColor>& SrcData, int32 DstWidth, int32 DstHeight, const TArrayView64<FLinearColor>& DstData)
 {
+	//@todo Oodle : this could be better ; investigate how widely it's used
+
 	check(SrcData.Num() >= SrcWidth * SrcHeight);
 	check(DstData.Num() >= DstWidth * DstHeight);
 
@@ -293,7 +297,7 @@ UTexture2D* FImageUtils::CreateTexture2D(int32 SrcWidth, int32 SrcHeight, const
 	for( int32 y=0; y<SrcHeight; y++ )
 	{
 		uint8* DestPtr = &MipData[(SrcHeight - 1 - y) * SrcWidth * sizeof(FColor)];
-		FColor* SrcPtr = const_cast<FColor*>(&SrcData[(SrcHeight - 1 - y) * SrcWidth]);
+		const FColor* SrcPtr = &SrcData[(SrcHeight - 1 - y) * SrcWidth];
 		for( int32 x=0; x<SrcWidth; x++ )
 		{
 			*DestPtr++ = SrcPtr->B;
@@ -391,28 +395,28 @@ void FImageUtils::CropAndScaleImage( int32 SrcWidth, int32 SrcHeight, int32 Desi
 	FImageUtils::ImageResize( MaxWidth, MaxHeight, CroppedData, DesiredWidth, DesiredHeight, DstData, true );
 }
 
-void FImageUtils::CompressImageArray( int32 ImageWidth, int32 ImageHeight, const TArray<FColor> &SrcData, TArray<uint8> &DstData )
+void FImageUtils::ThumbnailCompressImageArray( int32 ImageWidth, int32 ImageHeight, const TArray<FColor> &SrcData, TArray<uint8> &DstData )
 {
-	TArray<FColor> MutableSrcData = SrcData;
+	FObjectThumbnail TempThumbnail;
+	TempThumbnail.SetImageSize( ImageWidth, ImageHeight );
+	TArray<uint8>& ThumbnailByteArray = TempThumbnail.AccessImageData();
+
+	// Copy image into destination thumb
+	int32 MemorySize = ImageWidth*ImageHeight*sizeof(FColor);
+	ThumbnailByteArray.AddUninitialized(MemorySize);
+	FMemory::Memcpy(ThumbnailByteArray.GetData(), SrcData.GetData(), MemorySize);
+	
+	FColor * MutableSrcData = (FColor *)ThumbnailByteArray.GetData();
 
 	// Thumbnails are saved as RGBA but FColors are stored as BGRA. An option to swap the order upon compression may be added at 
 	// some point. At the moment, manually swapping Red and Blue 
-	for ( int32 Index = 0; Index < ImageWidth*ImageHeight; Index++ )
+	for (int32 Index = 0; Index < ImageWidth*ImageHeight; Index++ )
 	{
 		uint8 TempRed = MutableSrcData[Index].R;
 		MutableSrcData[Index].R = MutableSrcData[Index].B;
 		MutableSrcData[Index].B = TempRed;
 	}
 
-	FObjectThumbnail TempThumbnail;
-	TempThumbnail.SetImageSize( ImageWidth, ImageHeight );
-	TArray<uint8>& ThumbnailByteArray = TempThumbnail.AccessImageData();
-
-	// Copy scaled image into destination thumb
-	int32 MemorySize = ImageWidth*ImageHeight*sizeof(FColor);
-	ThumbnailByteArray.AddUninitialized(MemorySize);
-	FMemory::Memcpy(ThumbnailByteArray.GetData(), MutableSrcData.GetData(), MemorySize);
-
 	// Compress data - convert into thumbnail current format
 	TempThumbnail.CompressImageData();
 	DstData = TempThumbnail.AccessCompressedImageData();
@@ -430,14 +434,11 @@ void FImageUtils::PNGCompressImageArray(int32 ImageWidth, int32 ImageHeight, con
 
 	if (SrcData.Num() > 0 && ImageWidth > 0 && ImageHeight > 0)
 	{
-		if (DstData.Num() == 0)
+		IImageWrapperModule& ImageWrapperModule = FModuleManager::LoadModuleChecked<IImageWrapperModule>(FName("ImageWrapper"));
+		TSharedPtr<IImageWrapper> ImageWrapper = ImageWrapperModule.CreateImageWrapper(EImageFormat::PNG);
+		if (ImageWrapper.IsValid() && ImageWrapper->SetRaw(SrcFirstByte, MemorySize, ImageWidth, ImageHeight, ERGBFormat::BGRA, 8))
 		{
-			IImageWrapperModule& ImageWrapperModule = FModuleManager::LoadModuleChecked<IImageWrapperModule>(FName("ImageWrapper"));
-			TSharedPtr<IImageWrapper> ImageWrapper = ImageWrapperModule.CreateImageWrapper(EImageFormat::PNG);
-			if (ImageWrapper.IsValid() && ImageWrapper->SetRaw(SrcFirstByte, MemorySize, ImageWidth, ImageHeight, ERGBFormat::BGRA, 8))
-			{
-				DstData = ImageWrapper->GetCompressed();
-			}
+			DstData = ImageWrapper->GetCompressed();
 		}
 	}
 }
@@ -521,6 +522,8 @@ UTextureCube* FImageUtils::CreateCheckerboardCubeTexture(FColor ColorOne, FColor
 /*------------------------------------------------------------------------------
 HDR file format helper.
 ------------------------------------------------------------------------------*/
+// @todo Oodle : can we get rid of this HDR writer?
+//  try to remove uses of HDR as an export format
 class FHDRExportHelper
 {
 public:
@@ -966,13 +969,10 @@ UTexture2D* FImageUtils::ImportBufferAsTexture2D(TArrayView64<const uint8> Buffe
 		{
 			PixelFormat = PF_Unknown;
 			
+			//@todo Oodle : guessing format from bit depth looks wrong
 			ERGBFormat RGBFormat = ERGBFormat::Invalid;
-			
 			BitDepth = ImageWrapper->GetBitDepth();
-			
-			Width = ImageWrapper->GetWidth();
-			Height = ImageWrapper->GetHeight();
-			
+						
 			if (BitDepth == 16)
 			{
 				PixelFormat = PF_FloatRGBA;
@@ -989,6 +989,9 @@ UTexture2D* FImageUtils::ImportBufferAsTexture2D(TArrayView64<const uint8> Buffe
 				return nullptr;
 			}
 			
+			Width = ImageWrapper->GetWidth();
+			Height = ImageWrapper->GetHeight();
+
 			TArray64<uint8> UncompressedData;
 			ImageWrapper->GetRaw(RGBFormat, BitDepth, UncompressedData);
 			
diff --git a/Engine/Source/Runtime/Engine/Public/ImageUtils.h b/Engine/Source/Runtime/Engine/Public/ImageUtils.h
--- a/Engine/Source/Runtime/Engine/Public/ImageUtils.h
+++ b/Engine/Source/Runtime/Engine/Public/ImageUtils.h
@@ -153,7 +153,13 @@ public:
 	 * @param DstData			compressed image array.
 	 *
 	 */
-	ENGINE_API static void CompressImageArray( int32 ImageWidth, int32 ImageHeight, const TArray<FColor> &SrcData, TArray<uint8> &DstData );
+	ENGINE_API static void ThumbnailCompressImageArray( int32 ImageWidth, int32 ImageHeight, const TArray<FColor> &SrcData, TArray<uint8> &DstData );
+
+	UE_DEPRECATED(5.1, "Please use PNGCompressImageArray or ThumbnailCompressImageArray")
+	ENGINE_API static void CompressImageArray( int32 ImageWidth, int32 ImageHeight, const TArray<FColor> &SrcData, TArray<uint8> &DstData )
+	{
+		ThumbnailCompressImageArray(ImageWidth,ImageHeight,SrcData,DstData);
+	}
 
 	/**
 	 * Compress image to PNG format uint8 array.
