Remove invisible character in comment
#rnx

diff --git a/Engine/Source/Runtime/Engine/Classes/Kismet/KismetTextLibrary.h b/Engine/Source/Runtime/Engine/Classes/Kismet/KismetTextLibrary.h
--- a/Engine/Source/Runtime/Engine/Classes/Kismet/KismetTextLibrary.h
+++ b/Engine/Source/Runtime/Engine/Classes/Kismet/KismetTextLibrary.h
@@ -275,7 +275,7 @@ class ENGINE_API UKismetTextLibrary : public UBlueprintFunctionLibrary
 
 	/**
 	 * Attempts to create a text instance from a string table ID and key.
-	 * @note This exists to allow programmatic ‎look-up of a string table entry from dynamic content - you should favor setting your string table reference on a text property or pin wherever possible as it is significantly more robust (see "Make Literal Text").
+	 * @note This exists to allow programmatic look-up of a string table entry from dynamic content - you should favor setting your string table reference on a text property or pin wherever possible as it is significantly more robust (see "Make Literal Text").
 	 * @return The found text, or a dummy text if the entry could not be found.
 	 */
 	UFUNCTION(BlueprintPure, Category="Utilities|Text", meta=(DisplayName="Make Text from String Table (Advanced)"))
