Strata - added strata material to decal conversion node.

diff --git a/Engine/Shaders/Private/Strata/StrataLegacyConversion.ush b/Engine/Shaders/Private/Strata/StrataLegacyConversion.ush
--- a/Engine/Shaders/Private/Strata/StrataLegacyConversion.ush
+++ b/Engine/Shaders/Private/Strata/StrataLegacyConversion.ush
@@ -508,6 +508,16 @@ FStrataData StrataConvertLegacyMaterialDynamic(
 	Out.PreMultipliedAlphaOverrideCoverage = Opacity;
 	#endif
 
+	// We also evaluated the inlined BSDF in case it is used (for instance, for by decals)
+	const float DummyNoV = 1.0f;
+	uint NewNormalIndex = BSDF_GETSHAREDLOCALBASISID(Top_Slab.InlinedBSDF);
+	FStrataData ParamBlendVertical = StrataVerticalLayeringParameterBlending(StrataWeightParameterBlending(Top_Slab, Top.Weight), StrataWeightParameterBlending(Bot_Slab, Bot.Weight), NewNormalIndex, DummyNoV);
+#if STRATA_LEGACY_MATERIAL_APPLIES_FINAL_WEIGHT
+	Out.InlinedBSDF = StrataWeightParameterBlending(ParamBlendVertical, FinalWeight).InlinedBSDF;
+#else
+	Out.InlinedBSDF = ParamBlendVertical.InlinedBSDF;
+#endif
+
 	return Out;
 }
 
diff --git a/Engine/Source/Runtime/Engine/Classes/Engine/EngineTypes.h b/Engine/Source/Runtime/Engine/Classes/Engine/EngineTypes.h
--- a/Engine/Source/Runtime/Engine/Classes/Engine/EngineTypes.h
+++ b/Engine/Source/Runtime/Engine/Classes/Engine/EngineTypes.h
@@ -669,6 +669,7 @@ enum EStrataShadingModel
 	SSM_SingleLayerWater		UMETA(DisplayName = "SingleLayerWater"),
 	SSM_LightFunction			UMETA(DisplayName = "LightFunction"),
 	SSM_PostProcess				UMETA(DisplayName = "PostProcess"),
+	SSM_Decal					UMETA(DisplayName = "Decal"),
 	/** Number of unique shading models. */
 	SSM_NUM						UMETA(Hidden),
 };
@@ -712,7 +713,7 @@ public:
 
 private:
 	UPROPERTY()
-	uint8 ShadingModelField = 0;
+	uint16 ShadingModelField = 0;
 
 	/* Indicates if the shading model is constant or data-driven from the shader graph */
 	UPROPERTY()
diff --git a/Engine/Source/Runtime/Engine/Classes/Materials/MaterialExpressionStrata.h b/Engine/Source/Runtime/Engine/Classes/Materials/MaterialExpressionStrata.h
--- a/Engine/Source/Runtime/Engine/Classes/Materials/MaterialExpressionStrata.h
+++ b/Engine/Source/Runtime/Engine/Classes/Materials/MaterialExpressionStrata.h
@@ -606,6 +606,33 @@ class UMaterialExpressionStrataPostProcess : public UMaterialExpressionStrataBSD
 	//~ End UMaterialExpression Interface
 };
 
+UCLASS(MinimalAPI, collapsecategories, hidecategories = Object, DisplayName = "Strata Convert To Decal")
+class UMaterialExpressionStrataConvertToDecal : public UMaterialExpressionStrataBSDF
+{
+	GENERATED_UCLASS_BODY()
+
+	/**
+	 * The Strata material to convert to a decal.
+	 */
+	UPROPERTY()
+	FExpressionInput DecalMaterial;
+
+	// When a material is converted to decals, we use parameter blending to ensure that a single slab of matter is provided (current restriction).
+	const bool bUseParameterBlending = true;
+
+	//~ Begin UMaterialExpression Interface
+#if WITH_EDITOR
+	virtual int32 Compile(class FMaterialCompiler* Compiler, int32 OutputIndex) override;
+	virtual void GetCaption(TArray<FString>& OutCaptions) const override;
+	virtual uint32 GetOutputType(int32 OutputIndex) override;
+	virtual uint32 GetInputType(int32 InputIndex) override;
+	virtual bool IsResultStrataMaterial(int32 OutputIndex) override;
+	virtual void GatherStrataMaterialInfo(FStrataMaterialInfo& StrataMaterialInfo, int32 OutputIndex) override;
+	virtual FStrataOperator* StrataGenerateMaterialTopologyTree(class FMaterialCompiler* Compiler, class UMaterialExpression* Parent, int32 OutputIndex) override;
+#endif
+	//~ End UMaterialExpression Interface
+};
+
 
 
 ///////////////////////////////////////////////////////////////////////////////
diff --git a/Engine/Source/Runtime/Engine/Private/Materials/Material.cpp b/Engine/Source/Runtime/Engine/Private/Materials/Material.cpp
--- a/Engine/Source/Runtime/Engine/Private/Materials/Material.cpp
+++ b/Engine/Source/Runtime/Engine/Private/Materials/Material.cpp
@@ -3113,7 +3113,11 @@ void UMaterial::ConvertMaterialToStrataMaterial()
 			AddStrataShadingModelFromMaterialShadingModel(ConvertNode->ConvertedStrataMaterialInfo, ShadingModels);
 			check(ConvertNode->ConvertedStrataMaterialInfo.CountShadingModels() == 1);
 
-			FrontMaterial.Connect(0, ConvertNode);
+			// Now pas through the convert to decal node
+			UMaterialExpressionStrataConvertToDecal* ConvertToDecalNode= NewObject<UMaterialExpressionStrataConvertToDecal>(this);
+			ConvertToDecalNode->DecalMaterial.Connect(0, ConvertNode);
+
+			FrontMaterial.Connect(0, ConvertToDecalNode);
 			bInvalidateShader = true;
 		}
 
@@ -4109,7 +4113,13 @@ void UMaterial::RebuildShadingModelField()
 		}
 		if (StrataMaterialInfo.CountShadingModels() > 1)
 		{
-			if (StrataMaterialInfo.HasShadingModelFromExpression())
+			if (StrataMaterialInfo.HasShadingModel(SSM_Decal))
+			{
+				// Keep the decals information and remove the shading from expression part 
+				// since we are going to bake down the material to a single slab using parameter blending.
+				StrataMaterialInfo.SetShadingModelFromExpression(false);
+			}
+			else if (StrataMaterialInfo.HasShadingModelFromExpression())
 			{
 				if (BlendMode == EBlendMode::BLEND_Opaque || BlendMode == EBlendMode::BLEND_Masked)
 				{
@@ -4174,7 +4184,7 @@ void UMaterial::RebuildShadingModelField()
 				}
 			}
 
-			// Subsurface profil
+			// Subsurface profile
 			if (StrataMaterialInfo.HasShadingModel(SSM_SubsurfaceLit) && StrataMaterialInfo.CountSubsurfaceProfiles() > 0)
 			{
 				if (StrataMaterialInfo.CountSubsurfaceProfiles() > 1)
@@ -4204,10 +4214,7 @@ void UMaterial::RebuildShadingModelField()
 			}
 			else if (StrataMaterialInfo.HasOnlyShadingModel(SSM_DefaultLit))
 			{
-				if (MaterialDomain != EMaterialDomain::MD_Surface && MaterialDomain != EMaterialDomain::MD_DeferredDecal)
-				{
-					MaterialDomain = EMaterialDomain::MD_Surface;
-				}
+				MaterialDomain = EMaterialDomain::MD_Surface;
 				ShadingModel = MSM_DefaultLit;
 			}
 			else if (StrataMaterialInfo.HasOnlyShadingModel(SSM_VolumetricFogCloud))
@@ -4245,6 +4252,13 @@ void UMaterial::RebuildShadingModelField()
 					StrataBlendMode = SBM_Opaque;
 				}
 			}
+			else if (StrataMaterialInfo.HasShadingModel(SSM_Decal))
+			{
+				// Decal can have multiple shading model
+				MaterialDomain = EMaterialDomain::MD_DeferredDecal;
+				ShadingModel = MSM_DefaultLit;
+				BlendMode = EBlendMode::BLEND_Translucent;
+			}
 
 			// Also update the ShadingModels for remaining pipeline operation
 			ShadingModels.AddShadingModel(ShadingModel);
diff --git a/Engine/Source/Runtime/Engine/Private/Materials/MaterialExpressions.cpp b/Engine/Source/Runtime/Engine/Private/Materials/MaterialExpressions.cpp
--- a/Engine/Source/Runtime/Engine/Private/Materials/MaterialExpressions.cpp
+++ b/Engine/Source/Runtime/Engine/Private/Materials/MaterialExpressions.cpp
@@ -21346,7 +21346,7 @@ UMaterialExpressionStrataLightFunction::UMaterialExpressionStrataLightFunction(c
 	struct FConstructorStatics
 	{
 		FText NAME_Strata;
-		FConstructorStatics() : NAME_Strata(LOCTEXT("Strata Others", "Strata Others")) { }
+		FConstructorStatics() : NAME_Strata(LOCTEXT("Strata Extras", "Strata Extras")) { }
 	};
 	static FConstructorStatics ConstructorStatics;
 #if WITH_EDITORONLY_DATA
@@ -21413,7 +21413,7 @@ UMaterialExpressionStrataPostProcess::UMaterialExpressionStrataPostProcess(const
 	struct FConstructorStatics
 	{
 		FText NAME_Strata;
-		FConstructorStatics() : NAME_Strata(LOCTEXT("Strata Others", "Strata Others")) { }
+		FConstructorStatics() : NAME_Strata(LOCTEXT("Strata Extras", "Strata Extras")) { }
 	};
 	static FConstructorStatics ConstructorStatics;
 #if WITH_EDITORONLY_DATA
@@ -21480,6 +21480,101 @@ FStrataOperator* UMaterialExpressionStrataPostProcess::StrataGenerateMaterialTop
 
 
 
+UMaterialExpressionStrataConvertToDecal::UMaterialExpressionStrataConvertToDecal(const FObjectInitializer& ObjectInitializer)
+	: Super(ObjectInitializer)
+{
+	struct FConstructorStatics
+	{
+		FText NAME_Strata;
+		FConstructorStatics() : NAME_Strata(LOCTEXT("Strata Extras", "Strata Extras")) { }
+	};
+	static FConstructorStatics ConstructorStatics;
+#if WITH_EDITORONLY_DATA
+	MenuCategories.Add(ConstructorStatics.NAME_Strata);
+#endif
+}
+
+#if WITH_EDITOR
+int32 UMaterialExpressionStrataConvertToDecal::Compile(class FMaterialCompiler* Compiler, int32 OutputIndex)
+{
+	if (!DecalMaterial.GetTracedInput().Expression)
+	{
+		return Compiler->Errorf(TEXT("Missing DecalMaterial input"));
+	}
+
+	int32 DecalMaterialCodeChunk = DecalMaterial.Compile(Compiler);
+	int32 WeightOfOneCodeChunk = Compiler->Constant(1.0f);
+
+	FStrataOperator& StrataOperator = Compiler->StrataCompilationGetOperator(this);
+
+	if (!StrataOperator.bUseParameterBlending)
+	{
+		return Compiler->Errorf(TEXT("Strata Convert To Decal node is not receiveing a parameter blended strata material tree."));
+	}
+	else if (!StrataOperator.bRootOfParameterBlendingSubTree)
+	{
+		return Compiler->Errorf(TEXT("Strata Convert To Decal node must be the root of a parameter blending sun tree."));
+	}
+
+	int32 OutputCodeChunk = INDEX_NONE;
+	if (StrataOperator.bUseParameterBlending)
+	{
+		OutputCodeChunk = Compiler->StrataWeightParameterBlending(
+			DecalMaterialCodeChunk, WeightOfOneCodeChunk,
+			StrataOperator.bRootOfParameterBlendingSubTree ? &StrataOperator : nullptr);
+
+		FStrataMaterialCompilationInfo StrataInfo = StrataCompilationInfoWeightParamBlend(Compiler, Compiler->GetStrataCompilationInfo(DecalMaterialCodeChunk));
+
+		Compiler->StrataCompilationInfoRegisterCodeChunk(OutputCodeChunk, StrataInfo);
+	}
+	return OutputCodeChunk;
+}
+
+void UMaterialExpressionStrataConvertToDecal::GetCaption(TArray<FString>& OutCaptions) const
+{
+	OutCaptions.Add(TEXT("Strata Convert To Decal"));
+}
+
+uint32 UMaterialExpressionStrataConvertToDecal::GetOutputType(int32 OutputIndex)
+{
+	return MCT_Strata;
+}
+
+uint32 UMaterialExpressionStrataConvertToDecal::GetInputType(int32 InputIndex)
+{
+	return MCT_Strata;
+}
+
+bool UMaterialExpressionStrataConvertToDecal::IsResultStrataMaterial(int32 OutputIndex)
+{
+	return true;
+}
+
+void UMaterialExpressionStrataConvertToDecal::GatherStrataMaterialInfo(FStrataMaterialInfo& StrataMaterialInfo, int32 OutputIndex)
+{
+	if (DecalMaterial.GetTracedInput().Expression)
+	{
+		DecalMaterial.GetTracedInput().Expression->GatherStrataMaterialInfo(StrataMaterialInfo, DecalMaterial.OutputIndex);
+	}
+	StrataMaterialInfo.AddShadingModel(SSM_Decal);
+}
+
+FStrataOperator* UMaterialExpressionStrataConvertToDecal::StrataGenerateMaterialTopologyTree(class FMaterialCompiler* Compiler, class UMaterialExpression* Parent, int32 OutputIndex)
+{
+	FStrataOperator& StrataOperator = Compiler->StrataCompilationRegisterOperator(STRATA_OPERATOR_WEIGHT, this, Parent, bUseParameterBlending);
+
+	UMaterialExpression* ChildDecalMaterialExpression = DecalMaterial.GetTracedInput().Expression;
+	if (ChildDecalMaterialExpression)
+	{
+		StrataOperator.LeftIndex = ChildDecalMaterialExpression->StrataGenerateMaterialTopologyTree(Compiler, this, DecalMaterial.OutputIndex)->Index;
+	}
+
+	return &StrataOperator;
+}
+#endif // WITH_EDITOR
+
+
+
 UMaterialExpressionStrataUnlitBSDF::UMaterialExpressionStrataUnlitBSDF(const FObjectInitializer& ObjectInitializer)
 	: Super(ObjectInitializer)
 {
