PR #8928: Add SetLineThickness() to UBoxComponent (Contributed by [REDACTED_PERSON])

#note - had to add a push to the render state scene proxy to get updates working on tick from the original PR.

diff --git a/Engine/Source/Runtime/Engine/Classes/Components/BoxComponent.h b/Engine/Source/Runtime/Engine/Classes/Components/BoxComponent.h
--- a/Engine/Source/Runtime/Engine/Classes/Components/BoxComponent.h
+++ b/Engine/Source/Runtime/Engine/Classes/Components/BoxComponent.h
@@ -36,6 +36,10 @@ public:
 	UFUNCTION(BlueprintCallable, Category="Components|Box")
 	void SetBoxExtent(FVector InBoxExtent, bool bUpdateOverlaps=true);
 
+	// Set the LineThickness
+	UFUNCTION(BlueprintCallable, Category="Components|Box")
+	void SetLineThickness(float Thickness);
+
 	// @return the box extent, scaled by the component scale.
 	UFUNCTION(BlueprintCallable, Category="Components|Box")
 	FVector GetScaledBoxExtent() const;
diff --git a/Engine/Source/Runtime/Engine/Private/Components/BoxComponent.cpp b/Engine/Source/Runtime/Engine/Private/Components/BoxComponent.cpp
--- a/Engine/Source/Runtime/Engine/Private/Components/BoxComponent.cpp
+++ b/Engine/Source/Runtime/Engine/Private/Components/BoxComponent.cpp
@@ -40,6 +40,11 @@ void UBoxComponent::SetBoxExtent(FVector NewBoxExtent, bool bUpdateOverlaps)
 	}
 }
 
+void UBoxComponent::SetLineThickness(float Thickness)
+{
+	LineThickness = Thickness;
+	MarkRenderStateDirty();
+}
 
 template <EShapeBodySetupHelper UpdateBodySetupAction, typename BodySetupType>
 bool InvalidateOrUpdateBoxBodySetup(BodySetupType& ShapeBodySetup, bool bUseArchetypeBodySetup, FVector BoxExtent)
