Fix replicated references to class default objects.
They're considered stably named for networking again, so they can be referred to by path. This had changed as unintended fallout of CL 18803336 updating UObjectBaseUtility::IsDefaultSubobject.

diff --git a/Engine/Source/Runtime/CoreUObject/Private/UObject/Obj.cpp b/Engine/Source/Runtime/CoreUObject/Private/UObject/Obj.cpp
--- a/Engine/Source/Runtime/CoreUObject/Private/UObject/Obj.cpp
+++ b/Engine/Source/Runtime/CoreUObject/Private/UObject/Obj.cpp
@@ -4880,7 +4880,7 @@ EDataValidationResult UObject::IsDataValid(TArray<FText>& ValidationErrors)
 /** IsNameStableForNetworking means an object can be referred to its path name (relative to outer) over the network */
 bool UObject::IsNameStableForNetworking() const
 {
-	return HasAnyFlags(RF_WasLoaded | RF_DefaultSubObject) || IsNative() || IsDefaultSubobject();
+	return HasAnyFlags(RF_WasLoaded | RF_DefaultSubObject | RF_ClassDefaultObject) || IsNative() || IsDefaultSubobject();
 }
 
 /** IsFullNameStableForNetworking means an object can be referred to its full path name over the network */
diff --git a/Engine/Source/Runtime/Engine/Private/RepLayout.cpp b/Engine/Source/Runtime/Engine/Private/RepLayout.cpp
--- a/Engine/Source/Runtime/Engine/Private/RepLayout.cpp
+++ b/Engine/Source/Runtime/Engine/Private/RepLayout.cpp
@@ -5163,8 +5163,8 @@ static bool DiffStableProperties_r(FDiffStablePropertiesSharedParams& Params, TD
 						const bool bNetStartupActor = bIsActor && ObjValue && Cast<AActor>(ObjValue)->IsNetStartupActor();
 						const bool bIsReplicatedActor = bIsActor && ObjValue && Cast<AActor>(ObjValue)->GetIsReplicated();
 
-						// This is similar to AActor::IsNameStableForNetworking but we want to ignore forced net addressable objects for now
-						const bool bStableForNetworking = bNetStartupActor || (ObjValue && (ObjValue->HasAnyFlags(RF_WasLoaded | RF_DefaultSubObject) || ObjValue->IsNative() || ObjValue->IsDefaultSubobject()));
+						// This is similar to AActor::IsNameStableForNetworking but we want to ignore forced net addressable objects for now. Explicitly call the UObject implementation.
+						const bool bStableForNetworking = bNetStartupActor || (ObjValue && (ObjValue->UObject::IsNameStableForNetworking()));
 
 						if (bIsReplicatedActor || bIsActorComponent)
 						{
@@ -5199,8 +5199,8 @@ static bool DiffStableProperties_r(FDiffStablePropertiesSharedParams& Params, TD
 								continue;
 							}
 
-							const bool bStableForNetworking = (InterfaceObjValue->HasAnyFlags(RF_WasLoaded | RF_DefaultSubObject) || InterfaceObjValue->IsNative() || InterfaceObjValue->IsDefaultSubobject());
-							if (!bStableForNetworking)
+							// Explicitly call the UObject implementation to mirror object property behavior above.
+							if (!InterfaceObjValue->UObject::IsNameStableForNetworking())
 							{
 								// skip object references without a stable name
 								continue;
