Add ensure to catch morph target weights count mismatching morph info count in FSkeletalMeshObjectGPUSkin::Update on game thread.
Also added morph target weights array size checks before indexing into it.

diff --git a/Engine/Source/Runtime/Engine/Private/SkeletalRenderGPUSkin.cpp b/Engine/Source/Runtime/Engine/Private/SkeletalRenderGPUSkin.cpp
--- a/Engine/Source/Runtime/Engine/Private/SkeletalRenderGPUSkin.cpp
+++ b/Engine/Source/Runtime/Engine/Private/SkeletalRenderGPUSkin.cpp
@@ -420,6 +420,12 @@ void FSkeletalMeshObjectGPUSkin::Update(int32 LODIndex,USkinnedMeshComponent* In
 		InitMorphResources(InMeshComponent->bPerBoneMotionBlur, MorphTargetWeights);		
 	}
 
+	const FSkeletalMeshLODRenderData& LODData = SkeletalMeshRenderData->LODRenderData[LODIndex];
+	ensureAlwaysMsgf(MorphTargetWeights.Num() == LODData.MorphTargetVertexInfoBuffers.GetNumMorphs(),
+				TEXT("Skeletal mesh '%s' morph data mismatching: MorphTargetWeights=%d, [LOD%d]MorphTargetVertexInfoBuffers=%d, MorphTargets=%d"),
+				InMeshComponent->SkeletalMesh ? *InMeshComponent->SkeletalMesh->GetName() : TEXT(""), MorphTargetWeights.Num(), LODIndex, LODData.MorphTargetVertexInfoBuffers.GetNumMorphs(),
+				InMeshComponent->SkeletalMesh ? InMeshComponent->SkeletalMesh->GetMorphTargets().Num() : 0);
+
 	// create the new dynamic data for use by the rendering thread
 	// this data is only deleted when another update is sent
 	FDynamicSkelMeshObjectDataGPUSkin* NewDynamicData = FDynamicSkelMeshObjectDataGPUSkin::AllocDynamicSkelMeshObjectDataGPUSkin();		
@@ -1005,8 +1011,11 @@ static void CalculateMorphDeltaBounds(const TArray<float>& MorphTargetWeights, c
 
 		for (uint32 j = 0; j < 4; j++)
 		{
-			MinAccumScale[j] += MorphTargetWeights[i] * MinMorphScale[j];
-			MaxAccumScale[j] += MorphTargetWeights[i] * MaxMorphScale[j];
+			if (MorphTargetWeights.IsValidIndex(i))
+			{
+				MinAccumScale[j] += MorphTargetWeights[i] * MinMorphScale[j];
+				MaxAccumScale[j] += MorphTargetWeights[i] * MaxMorphScale[j];
+			}
 
 			double AbsMorphScale = FMath::Max<double>(FMath::Abs(MinMorphScale[j]), FMath::Abs(MaxMorphScale[j]));
 			double AbsAccumScale = FMath::Max<double>(FMath::Abs(MinAccumScale[j]), FMath::Abs(MaxAccumScale[j]));
@@ -1096,7 +1105,7 @@ void FSkeletalMeshObjectGPUSkin::FSkeletalMeshObjectLOD::UpdateMorphVertexBuffer
 					uint32 NumOutputMorphs = 0;
 					while (InputMorphStartIndex < MorphTargetVertexInfoBuffers.GetNumMorphs() && NumOutputMorphs < FGPUMorphUpdateCS::MorphTargetDispatchBatchSize)
 					{
-						if (MorphTargetWeights[InputMorphStartIndex] != 0.0f) 	// Omit morphs with zero weight
+						if (MorphTargetWeights.IsValidIndex(InputMorphStartIndex) && MorphTargetWeights[InputMorphStartIndex] != 0.0f) 	// Omit morphs with zero weight
 						{
 							BatchOffsets[NumOutputMorphs] = MorphTargetVertexInfoBuffers.GetBatchStartOffset(InputMorphStartIndex);
 							GroupOffsets[NumOutputMorphs] = NumBatches;
@@ -1270,7 +1279,7 @@ void FSkeletalMeshObjectGPUSkin::FSkeletalMeshObjectLOD::UpdateMorphVertexBuffer
 				const FActiveMorphTarget& MorphTarget = ActiveMorphTargets[AnimIdx];
 				checkSlow(MorphTarget.MorphTarget != NULL);
 				checkSlow(MorphTarget.MorphTarget->HasDataForLOD(LODIndex));
-				const float MorphTargetWeight = MorphTargetWeights[MorphTarget.WeightIndex];
+				const float MorphTargetWeight = MorphTargetWeights.IsValidIndex(MorphTarget.WeightIndex) ? MorphTargetWeights[MorphTarget.WeightIndex] : 0.0f;
 				const float MorphAbsWeight = FMath::Abs(MorphTargetWeight);
 				checkSlow(MorphAbsWeight >= MinMorphTargetBlendWeight && MorphAbsWeight <= MaxMorphTargetBlendWeight);
 
