Convert non-shipping assert to rate-limited log warning.

diff --git a/Engine/Source/Runtime/Engine/Private/NetConnection.cpp b/Engine/Source/Runtime/Engine/Private/NetConnection.cpp
--- a/Engine/Source/Runtime/Engine/Private/NetConnection.cpp
+++ b/Engine/Source/Runtime/Engine/Private/NetConnection.cpp
@@ -2524,8 +2524,16 @@ void UNetConnection::ReceivedPacket( FBitReader& Reader, bool bIsReinjectedPacke
 				Reader.SerializeInt(PacketJitterClockTimeMS, UE_NetConnectionPrivate::MaxJitterClockTimeValue + 1);
 
 #if !UE_BUILD_SHIPPING
-				checkf(Reader.GetPosBits() - BitsReadPreJitterClock == NetConnectionHelper::NumBitsForJitterClockTimeInHeader, TEXT("JitterClockTime did not read the expected nb of bits. Read %d, Expected %d"),
-					   Reader.GetPosBits() - BitsReadPreJitterClock, NetConnectionHelper::NumBitsForJitterClockTimeInHeader);
+				static double LastJitterLogTime = 0.0;
+
+				if (((Reader.GetPosBits() - BitsReadPreJitterClock) != NetConnectionHelper::NumBitsForJitterClockTimeInHeader) &&
+					((CurrentReceiveTimeInS - LastJitterLogTime) > 5.0))
+				{
+					UE_LOG(LogNet, Warning, TEXT("JitterClockTime did not read the expected nb of bits. Read %d, Expected %d"),
+							Reader.GetPosBits() - BitsReadPreJitterClock, NetConnectionHelper::NumBitsForJitterClockTimeInHeader);
+
+					LastJitterLogTime = CurrentReceiveTimeInS;
+				}
 #endif
 
 				if (!bIsReinjectedPacket)
