Update streaming manager when USkinnedMeshComponent bounds change

diff --git a/Engine/Source/Runtime/Engine/Classes/Components/SkinnedMeshComponent.h b/Engine/Source/Runtime/Engine/Classes/Components/SkinnedMeshComponent.h
--- a/Engine/Source/Runtime/Engine/Classes/Components/SkinnedMeshComponent.h
+++ b/Engine/Source/Runtime/Engine/Classes/Components/SkinnedMeshComponent.h
@@ -875,6 +875,7 @@ public:
 	virtual bool GetMaterialPropertyPath(int32 ElementIndex, UObject*& OutOwner, FString& OutPropertyPath, FProperty*& OutProperty) override;
 #endif // WITH_EDITOR
 	virtual FBoxSphereBounds CalcBounds(const FTransform& LocalToWorld) const override;
+	virtual void UpdateBounds() override;
 	virtual FTransform GetSocketTransform(FName InSocketName, ERelativeTransformSpace TransformSpace = RTS_World) const override;
 	virtual bool DoesSocketExist(FName InSocketName) const override;
 	virtual bool HasAnySockets() const override;
diff --git a/Engine/Source/Runtime/Engine/Private/Components/SkinnedMeshComponent.cpp b/Engine/Source/Runtime/Engine/Private/Components/SkinnedMeshComponent.cpp
--- a/Engine/Source/Runtime/Engine/Private/Components/SkinnedMeshComponent.cpp
+++ b/Engine/Source/Runtime/Engine/Private/Components/SkinnedMeshComponent.cpp
@@ -1167,6 +1167,18 @@ FBoxSphereBounds USkinnedMeshComponent::CalcBounds(const FTransform& LocalToWorl
 	return CalcMeshBound( FVector3f::ZeroVector, false, LocalToWorld );
 }
 
+void USkinnedMeshComponent::UpdateBounds()
+{
+	const auto OldRadius = Bounds.SphereRadius;
+
+	Super::UpdateBounds();
+
+	// Avoid updating the streamer for small changes in size
+	if (!FMath::IsNearlyEqual(OldRadius, Bounds.SphereRadius, 0.1f) && GetForcedLOD() == 0)
+	{
+		IStreamingManager::Get().NotifyPrimitiveUpdated_Concurrent(this);
+	}
+}
 
 class UPhysicsAsset* USkinnedMeshComponent::GetPhysicsAsset() const
 {
