Fix for camera stutter when following a physics object when the frame rate is hitching longer than the physics max delta time setting. New flag for Let spring arm which forces it to clam delta time to the same value as the physics system.

diff --git a/Engine/Source/Runtime/Engine/Classes/GameFramework/SpringArmComponent.h b/Engine/Source/Runtime/Engine/Classes/GameFramework/SpringArmComponent.h
--- a/Engine/Source/Runtime/Engine/Classes/GameFramework/SpringArmComponent.h
+++ b/Engine/Source/Runtime/Engine/Classes/GameFramework/SpringArmComponent.h
@@ -112,6 +112,10 @@ class ENGINE_API USpringArmComponent : public USceneComponent
 	UPROPERTY(EditAnywhere, BlueprintReadWrite, Category=Lag, meta=(editcondition="bEnableCameraLag", ClampMin="0.0", UIMin = "0.0"))
 	float CameraLagMaxDistance;
 
+	/** If true AND the view target is simulating using physics then use the same max timestep cap as the physics system. Prevents camera jitter when delta time is clamped within Chaos Physics. */
+	UPROPERTY(EditAnywhere, BlueprintReadWrite, Category = Lag)
+	uint32 bClampToMaxPhysicsDeltaTime : 1;
+
 	/**
 	 * Get the target rotation we inherit, used as the base target for the boom rotation.
 	 * This is derived from attachment to our parent and considering the UsePawnControlRotation and absolute rotation flags.
diff --git a/Engine/Source/Runtime/Engine/Private/GameFramework/SpringArmComponent.cpp b/Engine/Source/Runtime/Engine/Private/GameFramework/SpringArmComponent.cpp
--- a/Engine/Source/Runtime/Engine/Private/GameFramework/SpringArmComponent.cpp
+++ b/Engine/Source/Runtime/Engine/Private/GameFramework/SpringArmComponent.cpp
@@ -6,6 +6,7 @@
 #include "WorldCollision.h"
 #include "Engine/World.h"
 #include "DrawDebugHelpers.h"
+#include "PhysicsEngine/PhysicsSettings.h"
 
 //////////////////////////////////////////////////////////////////////////
 // USpringArmComponent
@@ -38,6 +39,7 @@ USpringArmComponent::USpringArmComponent(const FObjectInitializer& ObjectInitial
 	CameraRotationLagSpeed = 10.f;
 	CameraLagMaxTimeStep = 1.f / 60.f;
 	CameraLagMaxDistance = 0.f;
+	bClampToMaxPhysicsDeltaTime = false;
 
  	UnfixedCameraPosition = FVector::ZeroVector;
 }
@@ -90,6 +92,13 @@ void USpringArmComponent::UpdateDesiredArmLocation(bool bDoTrace, bool bDoLocati
 {
 	FRotator DesiredRot = GetTargetRotation();
 
+	// If our viewtarget is simulating using physics, we may need to clamp deltatime
+	if (bClampToMaxPhysicsDeltaTime)
+	{
+		// Use the same max timestep cap as the physics system to avoid camera jitter when the viewtarget simulates less time than the camera
+		DeltaTime = FMath::Min(DeltaTime, UPhysicsSettings::Get()->MaxPhysicsDeltaTime);
+	}
+
 	// Apply 'lag' to rotation if desired
 	if(bDoRotationLag)
 	{
