Fixed hit proxy gathering on asset editors that have no world

diff --git a/Engine/Source/Runtime/Engine/Private/UnrealClient.cpp b/Engine/Source/Runtime/Engine/Private/UnrealClient.cpp
--- a/Engine/Source/Runtime/Engine/Private/UnrealClient.cpp
+++ b/Engine/Source/Runtime/Engine/Private/UnrealClient.cpp
@@ -1677,9 +1677,8 @@ const TArray<FColor>& FViewport::GetRawHitProxyData(FIntRect InRect)
 	const bool bIsRenderingStereo = GEngine->IsStereoscopic3D( this ) && this->IsStereoRenderingAllowed();
 
 	bool bFetchHitProxyBytes = !bIsRenderingStereo && ( !bHitProxiesCached || (SizeY*SizeX) != CachedHitProxyData.Num() );
-	UWorld* World = ViewportClient->GetWorld();
 
-	if( bIsRenderingStereo || (World == nullptr))
+	if (bIsRenderingStereo)
 	{
 		// Stereo viewports don't support hit proxies, and we don't want to update them because it will adversely
 		// affect performance.
@@ -1704,8 +1703,7 @@ const TArray<FColor>& FViewport::GetRawHitProxyData(FIntRect InRect)
 			});
 
 		// Let the viewport client draw its hit proxies.
-		check(World != nullptr);
-		FCanvas Canvas(&HitProxyMap, &HitProxyMap, World, World->FeatureLevel.GetValue(), FCanvas::CDM_DeferDrawing, ViewportClient->ShouldDPIScaleSceneCanvas() ? ViewportClient->GetDPIScale() : 1.0f);
+		FCanvas Canvas(&HitProxyMap, &HitProxyMap, ViewportClient->GetWorld(), GetFeatureLevel(), FCanvas::CDM_DeferDrawing, ViewportClient->ShouldDPIScaleSceneCanvas() ? ViewportClient->GetDPIScale() : 1.0f);
 		{
 			ViewportClient->Draw(this, &Canvas);
 		}
@@ -1772,8 +1770,8 @@ const TArray<FColor>& FViewport::GetRawHitProxyData(FIntRect InRect)
 	}
 
 	return CachedHitProxyData;
-
 }
+
 void FViewport::GetHitProxyMap(FIntRect InRect,TArray<HHitProxy*>& OutMap)
 {
 	const TArray<FColor>& CachedData = GetRawHitProxyData(InRect);
