Fix build errors from removing implict conversion to FVector.

diff --git a/Engine/Plugins/AI/MLAdapter/Source/MLAdapter/Private/Sensors/MLAdapterSensor_AIPerception.cpp b/Engine/Plugins/AI/MLAdapter/Source/MLAdapter/Private/Sensors/MLAdapterSensor_AIPerception.cpp
--- a/Engine/Plugins/AI/MLAdapter/Source/MLAdapter/Private/Sensors/MLAdapterSensor_AIPerception.cpp
+++ b/Engine/Plugins/AI/MLAdapter/Source/MLAdapter/Private/Sensors/MLAdapterSensor_AIPerception.cpp
@@ -217,7 +217,7 @@ void UMLAdapterSensor_AIPerception::GetObservations(FMLAdapterMemoryWriter& Ar)
 		Ar.Serialize(&TargetData.Distance, sizeof(float));
 		if (bVectorMode)
 		{
-			FVector3f HeadingVector3f = TargetData.HeadingVector;
+			FVector3f HeadingVector3f = (FVector3f)TargetData.HeadingVector;
 			Ar << HeadingVector3f;
 		}
 		else
diff --git a/Engine/Plugins/AI/MLAdapter/Source/MLAdapter/Private/Sensors/MLAdapterSensor_Movement.cpp b/Engine/Plugins/AI/MLAdapter/Source/MLAdapter/Private/Sensors/MLAdapterSensor_Movement.cpp
--- a/Engine/Plugins/AI/MLAdapter/Source/MLAdapter/Private/Sensors/MLAdapterSensor_Movement.cpp
+++ b/Engine/Plugins/AI/MLAdapter/Source/MLAdapter/Private/Sensors/MLAdapterSensor_Movement.cpp
@@ -74,8 +74,8 @@ void UMLAdapterSensor_Movement::OnAvatarSet(AActor* Avatar)
 void UMLAdapterSensor_Movement::GetObservations(FMLAdapterMemoryWriter& Ar)
 {
 	FScopeLock Lock(&ObservationCS);
-	FVector3f Location = bAbsoluteLocation ? CurrentLocation : (CurrentLocation - RefLocation);
-	FVector3f Velocity = bAbsoluteVelocity ? CurrentVelocity : (CurrentVelocity - RefVelocity);
+	FVector3f Location = FVector3f(bAbsoluteLocation ? CurrentLocation : (CurrentLocation - RefLocation));
+	FVector3f Velocity = FVector3f(bAbsoluteVelocity ? CurrentVelocity : (CurrentVelocity - RefVelocity));
 	
 	FMLAdapter::FSpaceSerializeGuard SerializeGuard(SpaceDef, Ar);	
 	Ar << Location << Velocity;
diff --git a/Engine/Plugins/Animation/ControlRig/Source/ControlRig/Private/Rigs/RigHierarchy.cpp b/Engine/Plugins/Animation/ControlRig/Source/ControlRig/Private/Rigs/RigHierarchy.cpp
--- a/Engine/Plugins/Animation/ControlRig/Source/ControlRig/Private/Rigs/RigHierarchy.cpp
+++ b/Engine/Plugins/Animation/ControlRig/Source/ControlRig/Private/Rigs/RigHierarchy.cpp
@@ -942,7 +942,7 @@ FString URigHierarchy::GetControlPinDefaultValue(FRigControlElement* InControlEl
 				if(bForEdGraph)
 				{
 					// NOTE: se explanations above for Position/Scale
-					const FRotator Rotator = FRotator::MakeFromEuler(Value.GetRef<FVector3f>());
+					const FRotator Rotator = FRotator::MakeFromEuler((FVector)Value.GetRef<FVector3f>());
 					return FString::Printf(TEXT("%3.3f,%3.3f,%3.3f"), Rotator.Pitch, Rotator.Yaw, Rotator.Roll);
 				}
 				return Value.ToString<FRotator>();
diff --git a/Engine/Source/Runtime/Engine/Private/PrimitiveDrawingUtils.cpp b/Engine/Source/Runtime/Engine/Private/PrimitiveDrawingUtils.cpp
--- a/Engine/Source/Runtime/Engine/Private/PrimitiveDrawingUtils.cpp
+++ b/Engine/Source/Runtime/Engine/Private/PrimitiveDrawingUtils.cpp
@@ -795,7 +795,7 @@ void DrawTorus(FPrimitiveDrawInterface* PDI, const FMatrix& Transform, const FVe
 			FVector TangentX = -UnitDir * FMath::Sin(InnerAngle) + ZAxis * FMath::Cos(InnerAngle);
 			FVector TangentY = Dir ^ TangentX;
 			FVector TangentZ = Dir;
-			MeshVertex.SetTangents(TangentX, TangentY, TangentZ);
+			MeshVertex.SetTangents((FVector3f)TangentX, (FVector3f)TangentY, (FVector3f)TangentZ);
 
 			MeshBuilder.AddVertex(MeshVertex); //Add bottom vertex
 		}
@@ -1074,11 +1074,11 @@ void DrawRectangleMesh(FPrimitiveDrawInterface* PDI, const FVector& Center, cons
 	for (int32 VertexIndex = 0; VertexIndex < 4; VertexIndex++)
 	{
 		VertexIndices[VertexIndex] = MeshBuilder.AddVertex(
-			Positions[VertexIndex],
+			(FVector3f)Positions[VertexIndex],
 			UVs[VertexIndex],
-			FVector(1.0f, 0.0f, 0.0f),
-			FVector(0.0f, 1.0f, 0.0f),
-			FVector(0.0f, 0.0f, 1.0f),
+			FVector3f(1.0f, 0.0f, 0.0f),
+			FVector3f(0.0f, 1.0f, 0.0f),
+			FVector3f(0.0f, 0.0f, 1.0f),
 			Color
 		);
 	}
diff --git a/Engine/Source/Runtime/Renderer/Private/DistanceFieldObjectManagement.cpp b/Engine/Source/Runtime/Renderer/Private/DistanceFieldObjectManagement.cpp
--- a/Engine/Source/Runtime/Renderer/Private/DistanceFieldObjectManagement.cpp
+++ b/Engine/Source/Runtime/Renderer/Private/DistanceFieldObjectManagement.cpp
@@ -797,7 +797,7 @@ void FDistanceFieldSceneData::UpdateDistanceFieldObjectBuffers(
 									UploadObjectData[7] = *(FVector4f*)&VolumeToWorldT.M[1];
 									UploadObjectData[8] = *(FVector4f*)&VolumeToWorldT.M[2];
 
-									FVector4f FloatVector8(VolumeToWorld.GetScaleVector(), 0.0f);
+									FVector4f FloatVector8(FVector3f(VolumeToWorld.GetScaleVector()), 0.0f);
 
 									// Bypass NaN checks in FVector4f ctor
 									FSetElementId AssetStateSetId = AssetStateArray.FindId(DistanceFieldData);
diff --git a/Engine/Source/Runtime/Renderer/Private/RayTracing/RayTracingScene.cpp b/Engine/Source/Runtime/Renderer/Private/RayTracing/RayTracingScene.cpp
--- a/Engine/Source/Runtime/Renderer/Private/RayTracing/RayTracingScene.cpp
+++ b/Engine/Source/Runtime/Renderer/Private/RayTracing/RayTracingScene.cpp
@@ -212,8 +212,6 @@ void FRayTracingScene::Create(FRDGBuilder& GraphBuilder, const FGPUScene& GPUSce
 				BuildRayTracingInstanceBuffer(
 					RHICmdList,
 					GPUScene,
-					ViewTilePosition,
-					RelativePreViewTranslation,
 					PassParams->InstanceBuffer->GetRHI(),
 					InstanceUploadSRV,
 					AccelerationStructureAddressesBuffer.SRV,
