Dynamic RT geometries are always set to Rendering type because RT geometry is build dynamically when needed

diff --git a/Engine/Source/Runtime/Engine/Private/StaticMeshRender.cpp b/Engine/Source/Runtime/Engine/Private/StaticMeshRender.cpp
--- a/Engine/Source/Runtime/Engine/Private/StaticMeshRender.cpp
+++ b/Engine/Source/Runtime/Engine/Private/StaticMeshRender.cpp
@@ -431,6 +431,7 @@ void FStaticMeshSceneProxy::SetEvaluateWorldPositionOffsetInRayTracing(bool NewV
 				}
 				Initializer.bAllowUpdate = true;
 				Initializer.bFastBuild = true;
+				Initializer.Type = ERayTracingGeometryInitializerType::Rendering;
 			}
 
 			for (int32 i = 0; i < DynamicRayTracingGeometries.Num(); i++)
@@ -742,6 +743,7 @@ void FStaticMeshSceneProxy::CreateRenderThreadResources()
 				}
 				Initializer.bAllowUpdate = true;
 				Initializer.bFastBuild = true;
+				Initializer.Type = ERayTracingGeometryInitializerType::Rendering;
 			}
 		}
 
