Add an Input Setting to toggle motion controls. This can be useful if you want to enable/disable any logic inside of UPlayerInput::InputMotion.
Adding this setting in the Input Settings will make it a config variable that you can simply set per-platform via the DefaultInput.ini, which will be useful when shipping on multiple platforms.

Also add this flag on the player controller so that you enable/disable motion controls per player if desired.

diff --git a/Engine/Source/Runtime/Engine/Classes/GameFramework/InputSettings.h b/Engine/Source/Runtime/Engine/Classes/GameFramework/InputSettings.h
--- a/Engine/Source/Runtime/Engine/Classes/GameFramework/InputSettings.h
+++ b/Engine/Source/Runtime/Engine/Classes/GameFramework/InputSettings.h
@@ -57,6 +57,16 @@ class ENGINE_API UInputSettings
 	UPROPERTY(config, EditAnywhere, Category = "Input")
 	uint8 bEnableLegacyInputScales:1;
 	
+	/**
+	 * If set to false, then the player controller's InputMotion function will never be called.
+	 * This will effectively disable any motion input (tilt, rotation, acceleration, etc) on
+	 * the GameViewportClient.
+	 * 
+	 * @see GameViewportClient::InputMotion
+	 */
+	UPROPERTY(config, EditAnywhere, Category = "Input")
+	uint8 bEnableMotionControls:1;
+	
 	/** Should the touch input interface be shown always, or only when the platform has a touch screen? */
 	UPROPERTY(config, EditAnywhere, Category="Mobile")
 	uint8 bAlwaysShowTouchInterface:1;
diff --git a/Engine/Source/Runtime/Engine/Classes/GameFramework/PlayerController.h b/Engine/Source/Runtime/Engine/Classes/GameFramework/PlayerController.h
--- a/Engine/Source/Runtime/Engine/Classes/GameFramework/PlayerController.h
+++ b/Engine/Source/Runtime/Engine/Classes/GameFramework/PlayerController.h
@@ -505,6 +505,13 @@ public:
 	UPROPERTY(EditAnywhere, BlueprintReadWrite, Category="Game|Feedback")
 	uint32 bForceFeedbackEnabled:1;
 
+	/** Whether or not to consider input from motion sources (tilt, acceleration, etc) */
+	UPROPERTY(EditAnywhere, BlueprintReadWrite, BlueprintSetter=SetMotionControlsEnabled, Category="Input")
+	uint32 bEnableMotionControls:1;
+
+	UFUNCTION(BlueprintSetter)
+	void SetMotionControlsEnabled(bool bEnabled);
+
 	/** Whether the PlayerController should be used as a World Partiton streaming source. */
 	UPROPERTY(EditAnywhere, BlueprintReadWrite, Category = WorldPartition)
 	uint32 bEnableStreamingSource:1;
diff --git a/Engine/Source/Runtime/Engine/Private/GameViewportClient.cpp b/Engine/Source/Runtime/Engine/Private/GameViewportClient.cpp
--- a/Engine/Source/Runtime/Engine/Private/GameViewportClient.cpp
+++ b/Engine/Source/Runtime/Engine/Private/GameViewportClient.cpp
@@ -801,7 +801,7 @@ bool UGameViewportClient::InputTouch(FViewport* InViewport, int32 ControllerId,
 
 bool UGameViewportClient::InputMotion(FViewport* InViewport, int32 ControllerId, const FVector& Tilt, const FVector& RotationRate, const FVector& Gravity, const FVector& Acceleration)
 {
-	if (IgnoreInput())
+	if (IgnoreInput() || !GetDefault<UInputSettings>()->bEnableMotionControls)
 	{
 		return false;
 	}
diff --git a/Engine/Source/Runtime/Engine/Private/PlayerController.cpp b/Engine/Source/Runtime/Engine/Private/PlayerController.cpp
--- a/Engine/Source/Runtime/Engine/Private/PlayerController.cpp
+++ b/Engine/Source/Runtime/Engine/Private/PlayerController.cpp
@@ -2478,7 +2478,7 @@ bool APlayerController::InputMotion(const FVector& Tilt, const FVector& Rotation
 {
 	bool bResult = false;
 
-	if (PlayerInput)
+	if (PlayerInput && bEnableMotionControls)
 	{
 		bResult = PlayerInput->InputMotion(Tilt, RotationRate, Gravity, Acceleration);
 	}
@@ -2486,6 +2486,15 @@ bool APlayerController::InputMotion(const FVector& Tilt, const FVector& Rotation
 	return bResult;
 }
 
+void APlayerController::SetMotionControlsEnabled(bool bEnabled)
+{
+	bEnableMotionControls = bEnabled;
+	if (bEnableMotionControls && !GetDefault<UInputSettings>()->bEnableMotionControls)
+	{
+		UE_LOG(LogPlayerController, Warning, TEXT("Player bEnableMotionControls has been set to true, but motion controls are disabled in the input settings! See UInputSettings::bEnableMotionControls"));
+	}
+}
+
 bool APlayerController::ShouldShowMouseCursor() const
 {
 	return bShowMouseCursor;
diff --git a/Engine/Source/Runtime/Engine/Private/UserInterface/InputSettings.cpp b/Engine/Source/Runtime/Engine/Private/UserInterface/InputSettings.cpp
--- a/Engine/Source/Runtime/Engine/Private/UserInterface/InputSettings.cpp
+++ b/Engine/Source/Runtime/Engine/Private/UserInterface/InputSettings.cpp
@@ -8,6 +8,7 @@
 #include "Misc/CommandLine.h"
 #include "UObject/UObjectHash.h"
 #include "UObject/UObjectIterator.h"
+#include "HAL/PlatformApplicationMisc.h"
 
 #if WITH_EDITOR
 #include "Editor.h"
@@ -22,6 +23,7 @@ UInputSettings::UInputSettings(const FObjectInitializer& ObjectInitializer)
 	, bCaptureMouseOnLaunch(true)
 	, bDefaultViewportMouseLock_DEPRECATED(false)
 	, bEnableLegacyInputScales(true)
+	, bEnableMotionControls(true)
 	, DefaultViewportMouseCaptureMode(EMouseCaptureMode::CapturePermanently_IncludingInitialMouseDown)
 	, DefaultViewportMouseLockMode(EMouseLockMode::LockOnCapture)
 	, DefaultPlayerInputClass(UPlayerInput::StaticClass())
@@ -134,6 +136,8 @@ void UInputSettings::PostInitProperties()
 			UE_LOG(LogInput, Warning, TEXT("Axis %s uses deprecated key %s."), *KeyMapping.AxisName.ToString(), *KeyMapping.Key.ToString());
 		}
 	}
+
+	FPlatformApplicationMisc::EnableMotionData(bEnableMotionControls);
 }
 
 void UInputSettings::PopulateAxisConfigs()
@@ -188,11 +192,14 @@ void UInputSettings::PostEditChangeChainProperty(FPropertyChangedChainEvent& Pro
 		ForceRebuildKeymaps();
 		FEditorDelegates::OnActionAxisMappingsChanged.Broadcast();
 	}
-
-	if (MemberPropertyName == GET_MEMBER_NAME_CHECKED(UInputSettings, bEnableGestureRecognizer))
+	else if (MemberPropertyName == GET_MEMBER_NAME_CHECKED(UInputSettings, bEnableGestureRecognizer))
 	{
 		FEditorDelegates::OnEnableGestureRecognizerChanged.Broadcast();
 	}
+	else if (MemberPropertyName == GET_MEMBER_NAME_CHECKED(UInputSettings, bEnableMotionControls))
+	{
+		FPlatformApplicationMisc::EnableMotionData(bEnableMotionControls);
+	}
 }
 
 #endif
