Incorrect CONDITIONAL_BRACKET behaviour can cause massive performance/load regression
#fix ensure that bShouldTransact is forwarded to Open/CloseBracket rather than being used as a conditional for calling the functionality

diff --git a/Engine/Source/Developer/AnimationDataController/Private/AnimDataController.cpp b/Engine/Source/Developer/AnimationDataController/Private/AnimDataController.cpp
--- a/Engine/Source/Developer/AnimationDataController/Private/AnimDataController.cpp
+++ b/Engine/Source/Developer/AnimationDataController/Private/AnimDataController.cpp
@@ -58,12 +58,7 @@ namespace Anim {
 		Transaction = MakeUnique<UE::Anim::FScopedCompoundTransaction>(ChangeTransactor, Text); \
 	}
 
-#define CONDITIONAL_BRACKET(Text) \
-	TUniquePtr<IAnimationDataController::FScopedBracket> Transaction; \
-	if (UE::Anim::CanTransactChanges() && bShouldTransact) \
-	{ \
-		Transaction = MakeUnique<IAnimationDataController::FScopedBracket>(this, Text); \
-	}
+#define CONDITIONAL_BRACKET(Text) IAnimationDataController::FScopedBracket Transaction(this, Text, UE::Anim::CanTransactChanges() && bShouldTransact);
 
 #define CONDITIONAL_ACTION(ActionClass, ...) \
 	if (UE::Anim::CanTransactChanges() && bShouldTransact) \
diff --git a/Engine/Source/Runtime/Engine/Classes/Animation/AnimData/IAnimationDataController.h b/Engine/Source/Runtime/Engine/Classes/Animation/AnimData/IAnimationDataController.h
--- a/Engine/Source/Runtime/Engine/Classes/Animation/AnimData/IAnimationDataController.h
+++ b/Engine/Source/Runtime/Engine/Classes/Animation/AnimData/IAnimationDataController.h
@@ -41,30 +41,31 @@ public:
 	/** RAII helper to define a scoped-based bracket, opens and closes a controller bracket automatically */
     struct FScopedBracket
 	{
-		FScopedBracket(IAnimationDataController* InController, const FText& InDescription)
-            : Controller(*InController)
+		FScopedBracket(IAnimationDataController* InController, const FText& InDescription, bool bInShouldTransact=true)
+            : Controller(*InController), bShouldTransact(bInShouldTransact)
 		{
-			Controller.OpenBracket(InDescription);
+			Controller.OpenBracket(InDescription, bShouldTransact);
 		}
 
-		FScopedBracket(IAnimationDataController& InController, const FText& InDescription)
-            : Controller(InController)
+		FScopedBracket(IAnimationDataController& InController, const FText& InDescription, bool bInShouldTransact=true)
+            : Controller(InController), bShouldTransact(bInShouldTransact)
 		{
-			Controller.OpenBracket(InDescription);
+			Controller.OpenBracket(InDescription, bShouldTransact);
 		}
 		
-		FScopedBracket(TScriptInterface<IAnimationDataController>& InController, const FText& InDescription)
-            : Controller(*InController)
+		FScopedBracket(TScriptInterface<IAnimationDataController>& InController, const FText& InDescription, bool bInShouldTransact=true)
+            : Controller(*InController), bShouldTransact(bInShouldTransact)
 		{
-			Controller.OpenBracket(InDescription);
+			Controller.OpenBracket(InDescription, bShouldTransact);
 		}
 		
 		~FScopedBracket()
 		{
-			Controller.CloseBracket();
+			Controller.CloseBracket(bShouldTransact);
 		}
 	private:
 		IAnimationDataController& Controller;
+		bool bShouldTransact;
 	};
 #endif // WITH_EDITOR
 
