AnimBP: Guard against base class of AnimInstanceProxy::Evaluate being called

diff --git a/Engine/Source/Runtime/Engine/Private/Animation/AnimInstanceProxy.cpp b/Engine/Source/Runtime/Engine/Private/Animation/AnimInstanceProxy.cpp
--- a/Engine/Source/Runtime/Engine/Private/Animation/AnimInstanceProxy.cpp
+++ b/Engine/Source/Runtime/Engine/Private/Animation/AnimInstanceProxy.cpp
@@ -3283,7 +3283,10 @@ void FAnimInstanceProxy::EvaluateInputProxy(FAnimInstanceProxy* InputProxy, FPos
 {
 	if (InputProxy)
 	{
-		InputProxy->Evaluate(Output);
+		if(!InputProxy->Evaluate(Output))
+		{
+			Output.ResetToRefPose();
+		}
 	}
 }
 
