Add Insights Trace for HISM tree building

diff --git a/Engine/Source/Runtime/Engine/Private/HierarchicalInstancedStaticMesh.cpp b/Engine/Source/Runtime/Engine/Private/HierarchicalInstancedStaticMesh.cpp
--- a/Engine/Source/Runtime/Engine/Private/HierarchicalInstancedStaticMesh.cpp
+++ b/Engine/Source/Runtime/Engine/Private/HierarchicalInstancedStaticMesh.cpp
@@ -2628,6 +2628,8 @@ void UHierarchicalInstancedStaticMeshComponent::PostBuildStats()
 
 void UHierarchicalInstancedStaticMeshComponent::BuildTree()
 {
+	TRACE_CPUPROFILER_EVENT_SCOPE(UHierarchicalInstancedStaticMeshComponent::BuildTree);
+
 	check(!HasAnyFlags(RF_ClassDefaultObject | RF_ArchetypeObject));
 	checkSlow(IsInGameThread());
 
@@ -2849,6 +2851,8 @@ bool UHierarchicalInstancedStaticMeshComponent::BuildTreeIfOutdated(bool Async,
 		return false;
 	}
 
+	TRACE_CPUPROFILER_EVENT_SCOPE(UHierarchicalInstancedStaticMeshComponent::BuildTreeIfOutdated);
+
 	if (ForceUpdate 
 		|| bIsOutOfDate
 		|| InstanceUpdateCmdBuffer.NumTotalCommands() != 0
@@ -2936,6 +2940,8 @@ void UHierarchicalInstancedStaticMeshComponent::InitializeInstancingRandomSeed()
 
 void UHierarchicalInstancedStaticMeshComponent::BuildTreeAsync()
 {
+	TRACE_CPUPROFILER_EVENT_SCOPE(UHierarchicalInstancedStaticMeshComponent::BuildTreeAsync);
+
 	check(!HasAnyFlags(RF_ClassDefaultObject | RF_ArchetypeObject));
 	check(IsInGameThread());
 
