Don't compile GPU Skin VF for special materials when `r.SkinCache.SkipCompilingGPUSkinVF` is enabled.

- This accounts for roughly a couple thousand shaders across all our special engine materials.
- These shaders will not be used, even for special engine, so it should be safe to skip them.
- Addition to CL# 19244928

diff --git a/Engine/Source/Runtime/Engine/Private/GPUSkinVertexFactory.cpp b/Engine/Source/Runtime/Engine/Private/GPUSkinVertexFactory.cpp
--- a/Engine/Source/Runtime/Engine/Private/GPUSkinVertexFactory.cpp
+++ b/Engine/Source/Runtime/Engine/Private/GPUSkinVertexFactory.cpp
@@ -499,10 +499,9 @@ TGlobalResource<FBoneBufferPool> FGPUBaseSkinVertexFactory::BoneBufferPool;
 template <GPUSkinBoneInfluenceType BoneInfluenceType>
 bool TGPUSkinVertexFactory<BoneInfluenceType>::ShouldCompilePermutation(const FVertexFactoryShaderPermutationParameters& Parameters)
 {
-	const bool bUnlimitedBoneInfluences = (BoneInfluenceType == UnlimitedBoneInfluence && GCVarUnlimitedBoneInfluences);
-	const bool bShouldCache = (ShouldWeCompileGPUSkinVFShaders(Parameters.Platform) && Parameters.MaterialParameters.bIsUsedWithSkeletalMesh && (BoneInfluenceType != UnlimitedBoneInfluence || bUnlimitedBoneInfluences));
-	const bool bIsSpecialEngineMaterial = Parameters.MaterialParameters.bIsSpecialEngineMaterial;
-	return (bShouldCache || bIsSpecialEngineMaterial);
+	bool bUnlimitedBoneInfluences = (BoneInfluenceType == UnlimitedBoneInfluence && GCVarUnlimitedBoneInfluences);
+	return ShouldWeCompileGPUSkinVFShaders(Parameters.Platform) && 
+		  ((Parameters.MaterialParameters.bIsUsedWithSkeletalMesh && (BoneInfluenceType != UnlimitedBoneInfluence || bUnlimitedBoneInfluences)) || Parameters.MaterialParameters.bIsSpecialEngineMaterial);
 }
 
 template <GPUSkinBoneInfluenceType BoneInfluenceType>
