Implement LightmassReplace for GPULM

This allows static lighting to be baked with a simpler material, or to tweak the color of indirect lighting.

Ensure the simplified material gets picked up by GPULM

diff --git a/Engine/Plugins/Experimental/GPULightmass/Source/GPULightmass/Private/LightmapRayTracing.cpp b/Engine/Plugins/Experimental/GPULightmass/Source/GPULightmass/Private/LightmapRayTracing.cpp
--- a/Engine/Plugins/Experimental/GPULightmass/Source/GPULightmass/Private/LightmapRayTracing.cpp
+++ b/Engine/Plugins/Experimental/GPULightmass/Source/GPULightmass/Private/LightmapRayTracing.cpp
@@ -42,7 +42,7 @@ public:
 	static void ModifyCompilationEnvironment(const FMaterialShaderPermutationParameters& Parameters, FShaderCompilerEnvironment& OutEnvironment)
 	{
 		OutEnvironment.SetDefine(TEXT("USE_MATERIAL_CLOSEST_HIT_SHADER"), 1);
-		OutEnvironment.SetDefine(TEXT("USE_MATERIAL_ANY_HIT_SHADER"), 1);
+		OutEnvironment.SetDefine(TEXT("USE_MATERIAL_ANY_HIT_SHADER"), UseAnyHitShader);
 		OutEnvironment.SetDefine(TEXT("USE_RAYTRACED_TEXTURE_RAYCONE_LOD"), 0);
 		OutEnvironment.SetDefine(TEXT("SCENE_TEXTURES_DISABLED"), 1);
 		OutEnvironment.SetDefine(TEXT("SIMPLIFIED_MATERIAL_SHADER"), 1);
diff --git a/Engine/Plugins/Experimental/GPULightmass/Source/GPULightmass/Private/LightmapRayTracing.h b/Engine/Plugins/Experimental/GPULightmass/Source/GPULightmass/Private/LightmapRayTracing.h
--- a/Engine/Plugins/Experimental/GPULightmass/Source/GPULightmass/Private/LightmapRayTracing.h
+++ b/Engine/Plugins/Experimental/GPULightmass/Source/GPULightmass/Private/LightmapRayTracing.h
@@ -28,7 +28,7 @@ class FLightmapRayTracingMeshProcessor : public FRayTracingMeshProcessor
 {
 public:
 	FLightmapRayTracingMeshProcessor(FRayTracingMeshCommandContext* InCommandContext, FMeshPassProcessorRenderState InPassDrawRenderState)
-		: FRayTracingMeshProcessor(InCommandContext, nullptr, nullptr, InPassDrawRenderState, ERayTracingMeshCommandsMode::PATH_TRACING) // NOTE: RayTracingMeshCommandsMode does not really matter here since Process is overridden
+		: FRayTracingMeshProcessor(InCommandContext, nullptr, nullptr, InPassDrawRenderState, ERayTracingMeshCommandsMode::RAY_TRACING) // NOTE: Must pass RAY_TRACING here so that Process can be called
 	{}
 
 	virtual bool Process(
diff --git a/Engine/Shaders/Private/Common.ush b/Engine/Shaders/Private/Common.ush
--- a/Engine/Shaders/Private/Common.ush
+++ b/Engine/Shaders/Private/Common.ush
@@ -1795,7 +1795,16 @@ bool GetRayTracingQualitySwitch()
 
 bool GetPathTracingQualitySwitch()
 {
-#if RAYHITGROUPSHADER && PATH_TRACING
+#if RAYHITGROUPSHADER && PATH_TRACING && !SIMPLIFIED_MATERIAL_SHADER
+	return true;
+#else
+	return false;
+#endif
+}
+
+bool GetLightmassReplaceState()
+{
+#if RAYHITGROUPSHADER && PATH_TRACING && SIMPLIFIED_MATERIAL_SHADER
 	return true;
 #else
 	return false;
diff --git a/Engine/Source/Runtime/Engine/Private/Materials/HLSLMaterialTranslator.cpp b/Engine/Source/Runtime/Engine/Private/Materials/HLSLMaterialTranslator.cpp
--- a/Engine/Source/Runtime/Engine/Private/Materials/HLSLMaterialTranslator.cpp
+++ b/Engine/Source/Runtime/Engine/Private/Materials/HLSLMaterialTranslator.cpp
@@ -8786,6 +8786,10 @@ int32 FHLSLMaterialTranslator::PathTracingQualitySwitchReplace(int32 Normal, int
 	return GenericSwitch(TEXT("GetPathTracingQualitySwitch()"), PathTraced, Normal);
 }
 
+int32 FHLSLMaterialTranslator::LightmassReplace(int32 Realtime, int32 Lightmass)
+{
+	return GenericSwitch(TEXT("GetLightmassReplaceState()"), Lightmass, Realtime);
+}
 
 int32 FHLSLMaterialTranslator::VirtualTextureOutputReplace(int32 Default, int32 VirtualTexture)
 {
diff --git a/Engine/Source/Runtime/Engine/Private/Materials/HLSLMaterialTranslator.h b/Engine/Source/Runtime/Engine/Private/Materials/HLSLMaterialTranslator.h
--- a/Engine/Source/Runtime/Engine/Private/Materials/HLSLMaterialTranslator.h
+++ b/Engine/Source/Runtime/Engine/Private/Materials/HLSLMaterialTranslator.h
@@ -918,6 +918,7 @@ protected:
 
 	virtual int32 RayTracingQualitySwitchReplace(int32 Normal, int32 RayTraced) override;
 	virtual int32 PathTracingQualitySwitchReplace(int32 Normal, int32 PathTraced) override;
+	virtual int32 LightmassReplace(int32 Realtime, int32 Lightmass) override;
 
 	virtual int32 VirtualTextureOutputReplace(int32 Default, int32 VirtualTexture) override;
 
diff --git a/Engine/Source/Runtime/Engine/Private/Materials/MaterialExpressions.cpp b/Engine/Source/Runtime/Engine/Private/Materials/MaterialExpressions.cpp
--- a/Engine/Source/Runtime/Engine/Private/Materials/MaterialExpressions.cpp
+++ b/Engine/Source/Runtime/Engine/Private/Materials/MaterialExpressions.cpp
@@ -15674,7 +15674,27 @@ int32 UMaterialExpressionLightmassReplace::Compile(class FMaterialCompiler* Comp
 	}
 	else
 	{
-		return Compiler->IsLightmassCompiler() ? Lightmass.Compile(Compiler) : Realtime.Compile(Compiler);
+		const int32 Arg2 = Lightmass.Compile(Compiler);
+		if (Compiler->IsLightmassCompiler())
+		{
+			return Arg2;
+		}
+		const int32 Arg1 = Realtime.Compile(Compiler);
+		//only when both of these are real expressions do the actual code.  otherwise various output pins will
+		//end up considered 'set' when really we just want a default.  This can cause us to force depth output when we don't want it for example.
+		if (Arg1 != INDEX_NONE && Arg2 != INDEX_NONE)
+		{
+			return Compiler->LightmassReplace(Arg1, Arg2);
+		}
+		else if (Arg1 != INDEX_NONE)
+		{
+			return Arg1;
+		}
+		else if (Arg2 != INDEX_NONE)
+		{
+			return Arg2;
+		}
+		return INDEX_NONE;
 	}
 }
 
diff --git a/Engine/Source/Runtime/Engine/Public/MaterialCompiler.h b/Engine/Source/Runtime/Engine/Public/MaterialCompiler.h
--- a/Engine/Source/Runtime/Engine/Public/MaterialCompiler.h
+++ b/Engine/Source/Runtime/Engine/Public/MaterialCompiler.h
@@ -398,6 +398,7 @@ public:
 	virtual int32 ShadowReplace(int32 Default, int32 Shadow) = 0;
 	virtual int32 RayTracingQualitySwitchReplace(int32 Normal, int32 RayTraced) = 0;
 	virtual int32 PathTracingQualitySwitchReplace(int32 Normal, int32 PathTraced) = 0;
+	virtual int32 LightmassReplace(int32 Realtime, int32 Lightmass) = 0;
 	virtual int32 VirtualTextureOutputReplace(int32 Default, int32 VirtualTexture) = 0;
 	virtual int32 ReflectionCapturePassSwitch(int32 Default, int32 Reflection) = 0;
 
@@ -779,6 +780,7 @@ public:
 	virtual int32 ShadowReplace(int32 Default, int32 Shadow) override { return Compiler->ShadowReplace(Default, Shadow); }
 	virtual int32 RayTracingQualitySwitchReplace(int32 Normal, int32 RayTraced) override { return Compiler->RayTracingQualitySwitchReplace(Normal, RayTraced); }
 	virtual int32 PathTracingQualitySwitchReplace(int32 Normal, int32 PathTraced) override { return Compiler->PathTracingQualitySwitchReplace(Normal, PathTraced); }
+	virtual int32 LightmassReplace(int32 Realtime, int32 Lightmass) override { return Compiler->LightmassReplace(Realtime, Lightmass); }
 	virtual int32 VirtualTextureOutputReplace(int32 Default, int32 VirtualTexture) override { return Compiler->VirtualTextureOutputReplace(Default, VirtualTexture); }
 	virtual int32 ReflectionCapturePassSwitch(int32 Default, int32 Reflection) override { return Compiler->ReflectionCapturePassSwitch(Default, Reflection); }
 
