Added a StreamingGenerationErrorHandler to accumulate all error messages
Moved FTokenizedMessage creation code to ITokenizedMessageErrorHandler
FStreamingGenerationMapCheckErrorHandler & FTokenizedMessageAccumulatorErrorHandler now inherit from ITokenizedMessageErrorHandle

diff --git a/Engine/Source/Runtime/Engine/Private/WorldPartition/ErrorHandling/WorldPartitionStreamingGenerationMapCheckErrorHandler.cpp b/Engine/Source/Runtime/Engine/Private/WorldPartition/ErrorHandling/WorldPartitionStreamingGenerationMapCheckErrorHandler.cpp
--- a/Engine/Source/Runtime/Engine/Private/WorldPartition/ErrorHandling/WorldPartitionStreamingGenerationMapCheckErrorHandler.cpp
+++ b/Engine/Source/Runtime/Engine/Private/WorldPartition/ErrorHandling/WorldPartitionStreamingGenerationMapCheckErrorHandler.cpp
@@ -10,65 +10,14 @@
 
 #define LOCTEXT_NAMESPACE "WorldPartition"
 
-void FStreamingGenerationMapCheckErrorHandler::OnInvalidReference(const FWorldPartitionActorDescView& ActorDescView, const FGuid& ReferenceGuid)
+void FStreamingGenerationMapCheckErrorHandler::AddAdditionalNameToken(TSharedRef<FTokenizedMessage>& InMessage, const FName& InErrorName)
 {
-	TSharedRef<FTokenizedMessage> Error = FMessageLog("MapCheck").Warning()
-		->AddToken(FTextToken::Create(LOCTEXT("MapCheck_WorldPartition_Actor", "Actor")))
-		->AddToken(FActorToken::Create(ActorDescView.GetActorPath().ToString(), ActorDescView.GetGuid(), FText::FromName(ActorDescView.GetActorLabelOrName())))
-		->AddToken(FTextToken::Create(LOCTEXT("MapCheck_WorldPartition_HaveMissingRefsTo", "have missing references to")))
-		->AddToken(FTextToken::Create(FText::FromString(ReferenceGuid.ToString())))
-		->AddToken(FMapErrorToken::Create(FName(TEXT("WorldPartition_MissingActorReference_CheckForErrors"))));
+	InMessage->AddToken(FMapErrorToken::Create(InErrorName));
 }
 
-void FStreamingGenerationMapCheckErrorHandler::OnInvalidReferenceGridPlacement(const FWorldPartitionActorDescView& ActorDescView, const FWorldPartitionActorDescView& ReferenceActorDescView)
+void FStreamingGenerationMapCheckErrorHandler::HandleTokenizedMessage(TSharedRef<FTokenizedMessage>&& ErrorMessage)
 {
-	const FText SpatiallyLoadedActor(LOCTEXT("MapCheck_WorldPartition_SpatiallyLoadedActor", "Spatially loaded actor"));
-	const FText NonSpatiallyLoadedActor(LOCTEXT("MapCheck_WorldPartition_NonSpatiallyLoadedActor", "Non-spatially loaded actor"));
-
-	TSharedRef<FTokenizedMessage> Error = FMessageLog("MapCheck").Error()
-		->AddToken(FTextToken::Create(ActorDescView.GetIsSpatiallyLoaded() ? SpatiallyLoadedActor : NonSpatiallyLoadedActor))
-		->AddToken(FActorToken::Create(ActorDescView.GetActorPath().ToString(), ActorDescView.GetGuid(), FText::FromName(ActorDescView.GetActorLabelOrName())))
-		->AddToken(FTextToken::Create(LOCTEXT("MapCheck_WorldPartition_References", "references")))
-		->AddToken(FTextToken::Create(ReferenceActorDescView.GetIsSpatiallyLoaded() ? SpatiallyLoadedActor : NonSpatiallyLoadedActor))
-		->AddToken(FActorToken::Create(ReferenceActorDescView.GetActorPath().ToString(), ReferenceActorDescView.GetGuid(), FText::FromName(ReferenceActorDescView.GetActorLabelOrName())))
-		->AddToken(FMapErrorToken::Create(FName(TEXT("WorldPartition_StreamedActorReferenceAlwaysLoadedActor_CheckForErrors"))));
-}
-
-void FStreamingGenerationMapCheckErrorHandler::OnInvalidReferenceDataLayers(const FWorldPartitionActorDescView& ActorDescView, const FWorldPartitionActorDescView& ReferenceActorDescView)
-{
-	TSharedRef<FTokenizedMessage> Error = FMessageLog("MapCheck").Error()
-		->AddToken(FTextToken::Create(LOCTEXT("MapCheck_WorldPartition_Actor", "Actor")))
-		->AddToken(FActorToken::Create(ActorDescView.GetActorPath().ToString(), ActorDescView.GetGuid(), FText::FromName(ActorDescView.GetActorLabelOrName())))
-		->AddToken(FTextToken::Create(LOCTEXT("MapCheck_WorldPartition_ReferenceActorInOtherDataLayers", "references an actor in a different set of data layers")))
-		->AddToken(FActorToken::Create(ReferenceActorDescView.GetActorPath().ToString(), ReferenceActorDescView.GetGuid(), FText::FromName(ReferenceActorDescView.GetActorLabelOrName())))
-		->AddToken(FMapErrorToken::Create(FName(TEXT("WorldPartition_ActorReferenceActorInAnotherDataLayer_CheckForErrors"))));
-}
-
-void FStreamingGenerationMapCheckErrorHandler::OnInvalidReferenceRuntimeGrid(const FWorldPartitionActorDescView& ActorDescView, const FWorldPartitionActorDescView& ReferenceActorDescView)
-{
-	TSharedRef<FTokenizedMessage> Error = FMessageLog("MapCheck").Error()
-		->AddToken(FTextToken::Create(LOCTEXT("MapCheck_WorldPartition_Actor", "Actor")))
-		->AddToken(FActorToken::Create(ActorDescView.GetActorPath().ToString(), ActorDescView.GetGuid(), FText::FromName(ActorDescView.GetActorLabelOrName())))
-		->AddToken(FTextToken::Create(LOCTEXT("MapCheck_WorldPartition_ReferenceActorInOtherRuntimeGrid", "references an actor in a different runtime grid")))
-		->AddToken(FActorToken::Create(ReferenceActorDescView.GetActorPath().ToString(), ReferenceActorDescView.GetGuid(), FText::FromName(ReferenceActorDescView.GetActorLabelOrName())))
-		->AddToken(FMapErrorToken::Create(FName(TEXT("WorldPartition_ActorReferenceActorInAnotherRuntimeGrid_CheckForErrors"))));
-}
-
-void FStreamingGenerationMapCheckErrorHandler::OnInvalidReferenceLevelScriptStreamed(const FWorldPartitionActorDescView& ActorDescView)
-{
-	TSharedRef<FTokenizedMessage> Error = FMessageLog("MapCheck").Error()
-		->AddToken(FTextToken::Create(LOCTEXT("MapCheck_WorldPartition_LevelScriptBlueprintStreamedActorReference", "Level Script Blueprint references streamed actor")))
-		->AddToken(FActorToken::Create(ActorDescView.GetActorPath().ToString(), ActorDescView.GetGuid(), FText::FromName(ActorDescView.GetActorLabelOrName())))
-		->AddToken(FMapErrorToken::Create(FName(TEXT("WorldPartition_LevelScriptBlueprintRefefenceStreamed_CheckForErrors"))));
-}
-
-void FStreamingGenerationMapCheckErrorHandler::OnInvalidReferenceLevelScriptDataLayers(const FWorldPartitionActorDescView& ActorDescView)
-{
-	TSharedRef<FTokenizedMessage> Error = FMessageLog("MapCheck").Error()
-		->AddToken(FTextToken::Create(LOCTEXT("MapCheck_WorldPartition_LevelScriptBlueprintActorReference", "Level Script Blueprint references actor")))
-		->AddToken(FActorToken::Create(ActorDescView.GetActorPath().ToString(), ActorDescView.GetGuid(), FText::FromName(ActorDescView.GetActorLabelOrName())))
-		->AddToken(FTextToken::Create(LOCTEXT("MapCheck_WorldPartition_LevelScriptBlueprintDataLayerReference", "with a non empty set of data layers")))
-		->AddToken(FMapErrorToken::Create(FName(TEXT("WorldPartition_LevelScriptBlueprintRefefenceDataLayer_CheckForErrors"))));
+	FMessageLog("MapCheck").AddMessage(ErrorMessage);
 }
 
 #undef LOCTEXT_NAMESPACE
diff --git a/Engine/Source/Runtime/Engine/Private/WorldPartition/ErrorHandling/WorldPartitionStreamingGenerationTokenizedMessageErrorHandler.cpp b/Engine/Source/Runtime/Engine/Private/WorldPartition/ErrorHandling/WorldPartitionStreamingGenerationTokenizedMessageErrorHandler.cpp
new file mode 100644
--- /dev/null
+++ b/Engine/Source/Runtime/Engine/Private/WorldPartition/ErrorHandling/WorldPartitionStreamingGenerationTokenizedMessageErrorHandler.cpp
@@ -0,0 +1,94 @@
+// Copyright Epic Games, Inc. All Rights Reserved.
+
+#if WITH_EDITOR
+#include "WorldPartition/ErrorHandling/WorldPartitionStreamingGenerationTokenizedMessageErrorHandler.h"
+
+#include "GameFramework/Actor.h"
+#include "Logging/MessageLog.h"
+#include "Misc/UObjectToken.h"
+
+#define LOCTEXT_NAMESPACE "WorldPartition"
+
+void ITokenizedMessageErrorHandler::OnInvalidReference(const FWorldPartitionActorDescView& ActorDescView, const FGuid& ReferenceGuid)
+{
+	TSharedRef<FTokenizedMessage> Message = FTokenizedMessage::Create(EMessageSeverity::Warning);
+	Message->AddToken(FTextToken::Create(LOCTEXT("TokenMessage_WorldPartition_Actor", "Actor")))
+		->AddToken(FActorToken::Create(ActorDescView.GetActorPath().ToString(), ActorDescView.GetGuid(), FText::FromName(ActorDescView.GetActorLabelOrName())))
+		->AddToken(FTextToken::Create(LOCTEXT("TokenMessage_WorldPartition_HaveMissingRefsTo", "have missing references to")))
+		->AddToken(FTextToken::Create(FText::FromString(ReferenceGuid.ToString())));
+
+	AddAdditionalNameToken(Message, FName(TEXT("WorldPartition_MissingActorReference_CheckForErrors")));
+
+	HandleTokenizedMessage(MoveTemp(Message));
+}
+
+void ITokenizedMessageErrorHandler::OnInvalidReferenceGridPlacement(const FWorldPartitionActorDescView& ActorDescView, const FWorldPartitionActorDescView& ReferenceActorDescView)
+{
+	const FText SpatiallyLoadedActor(LOCTEXT("TokenMessage_WorldPartition_SpatiallyLoadedActor", "Spatially loaded actor"));
+	const FText NonSpatiallyLoadedActor(LOCTEXT("TokenMessage_WorldPartition_NonSpatiallyLoadedActor", "Non-spatially loaded actor"));
+
+	TSharedRef<FTokenizedMessage> Message = FTokenizedMessage::Create(EMessageSeverity::Error);
+	Message->AddToken(FTextToken::Create(ActorDescView.GetIsSpatiallyLoaded() ? SpatiallyLoadedActor : NonSpatiallyLoadedActor))
+		->AddToken(FActorToken::Create(ActorDescView.GetActorPath().ToString(), ActorDescView.GetGuid(), FText::FromName(ActorDescView.GetActorLabelOrName())))
+		->AddToken(FTextToken::Create(LOCTEXT("TokenMessage_WorldPartition_References", "references")))
+		->AddToken(FTextToken::Create(ReferenceActorDescView.GetIsSpatiallyLoaded() ? SpatiallyLoadedActor : NonSpatiallyLoadedActor))
+		->AddToken(FActorToken::Create(ReferenceActorDescView.GetActorPath().ToString(), ReferenceActorDescView.GetGuid(), FText::FromName(ReferenceActorDescView.GetActorLabelOrName())));
+
+	AddAdditionalNameToken(Message, FName(TEXT("WorldPartition_StreamedActorReferenceAlwaysLoadedActor_CheckForErrors")));
+
+	HandleTokenizedMessage(MoveTemp(Message));
+}
+
+void ITokenizedMessageErrorHandler::OnInvalidReferenceDataLayers(const FWorldPartitionActorDescView& ActorDescView, const FWorldPartitionActorDescView& ReferenceActorDescView)
+{
+	TSharedRef<FTokenizedMessage> Message = FTokenizedMessage::Create(EMessageSeverity::Error);
+	Message->AddToken(FTextToken::Create(LOCTEXT("TokenMessage_WorldPartition_Actor", "Actor")))
+		->AddToken(FActorToken::Create(ActorDescView.GetActorPath().ToString(), ActorDescView.GetGuid(), FText::FromName(ActorDescView.GetActorLabelOrName())))
+		->AddToken(FTextToken::Create(LOCTEXT("TokenMessage_WorldPartition_ReferenceActorInOtherDataLayers", "references an actor in a different set of data layers")))
+		->AddToken(FActorToken::Create(ReferenceActorDescView.GetActorPath().ToString(), ReferenceActorDescView.GetGuid(), FText::FromName(ReferenceActorDescView.GetActorLabelOrName())));
+
+	AddAdditionalNameToken(Message, FName(TEXT("WorldPartition_ActorReferenceActorInAnotherDataLayer_CheckForErrors")));
+
+	HandleTokenizedMessage(MoveTemp(Message));
+}
+
+void ITokenizedMessageErrorHandler::OnInvalidReferenceRuntimeGrid(const FWorldPartitionActorDescView& ActorDescView, const FWorldPartitionActorDescView& ReferenceActorDescView)
+{
+	TSharedRef<FTokenizedMessage> Message = FTokenizedMessage::Create(EMessageSeverity::Error);
+	Message->AddToken(FTextToken::Create(LOCTEXT("TokenMessage_WorldPartition_Actor", "Actor")))
+		->AddToken(FActorToken::Create(ActorDescView.GetActorPath().ToString(), ActorDescView.GetGuid(), FText::FromName(ActorDescView.GetActorLabelOrName())))
+		->AddToken(FTextToken::Create(LOCTEXT("TokenMessage_WorldPartition_ReferenceActorInOtherRuntimeGrid", "references an actor in a different runtime grid")))
+		->AddToken(FActorToken::Create(ReferenceActorDescView.GetActorPath().ToString(), ReferenceActorDescView.GetGuid(), FText::FromName(ReferenceActorDescView.GetActorLabelOrName())));
+
+	AddAdditionalNameToken(Message, FName(TEXT("WorldPartition_ActorReferenceActorInAnotherRuntimeGrid_CheckForErrors")));
+
+	HandleTokenizedMessage(MoveTemp(Message));
+}
+
+void ITokenizedMessageErrorHandler::OnInvalidReferenceLevelScriptStreamed(const FWorldPartitionActorDescView& ActorDescView)
+{
+	TSharedRef<FTokenizedMessage> Message = FTokenizedMessage::Create(EMessageSeverity::Error);
+	Message->AddToken(FTextToken::Create(LOCTEXT("TokenMessage_WorldPartition_LevelScriptBlueprintStreamedActorReference", "Level Script Blueprint references streamed actor")))
+		->AddToken(FActorToken::Create(ActorDescView.GetActorPath().ToString(), ActorDescView.GetGuid(), FText::FromName(ActorDescView.GetActorLabelOrName())));
+
+	AddAdditionalNameToken(Message, FName(TEXT("WorldPartition_LevelScriptBlueprintRefefenceStreamed_CheckForErrors")));
+
+	HandleTokenizedMessage(MoveTemp(Message));
+	
+}
+
+void ITokenizedMessageErrorHandler::OnInvalidReferenceLevelScriptDataLayers(const FWorldPartitionActorDescView& ActorDescView)
+{
+	TSharedRef<FTokenizedMessage> Message = FTokenizedMessage::Create(EMessageSeverity::Error);
+	Message->AddToken(FTextToken::Create(LOCTEXT("TokenMessage_WorldPartition_LevelScriptBlueprintActorReference", "Level Script Blueprint references actor")))
+		->AddToken(FActorToken::Create(ActorDescView.GetActorPath().ToString(), ActorDescView.GetGuid(), FText::FromName(ActorDescView.GetActorLabelOrName())))
+		->AddToken(FTextToken::Create(LOCTEXT("TokenMessage_WorldPartition_LevelScriptBlueprintDataLayerReference", "with a non empty set of data layers")));
+
+	AddAdditionalNameToken(Message, FName(TEXT("WorldPartition_LevelScriptBlueprintRefefenceDataLayer_CheckForErrors")));
+
+	HandleTokenizedMessage(MoveTemp(Message));
+}
+
+#undef LOCTEXT_NAMESPACE
+
+#endif
diff --git a/Engine/Source/Runtime/Engine/Public/WorldPartition/ErrorHandling/WorldPartitionStreamingGenerationMapCheckErrorHandler.h b/Engine/Source/Runtime/Engine/Public/WorldPartition/ErrorHandling/WorldPartitionStreamingGenerationMapCheckErrorHandler.h
--- a/Engine/Source/Runtime/Engine/Public/WorldPartition/ErrorHandling/WorldPartitionStreamingGenerationMapCheckErrorHandler.h
+++ b/Engine/Source/Runtime/Engine/Public/WorldPartition/ErrorHandling/WorldPartitionStreamingGenerationMapCheckErrorHandler.h
@@ -2,16 +2,13 @@
 #pragma once
 
 #if WITH_EDITOR
-#include "WorldPartition/ErrorHandling/WorldPartitionStreamingGenerationErrorHandler.h"
+#include "WorldPartition/ErrorHandling/WorldPartitionStreamingGenerationTokenizedMessageErrorHandler.h"
 
-class ENGINE_API FStreamingGenerationMapCheckErrorHandler : public IStreamingGenerationErrorHandler
+class ENGINE_API FStreamingGenerationMapCheckErrorHandler : public ITokenizedMessageErrorHandler
 {
-public:
-	virtual void OnInvalidReference(const FWorldPartitionActorDescView& ActorDescView, const FGuid& ReferenceGuid) override;
-	virtual void OnInvalidReferenceGridPlacement(const FWorldPartitionActorDescView& ActorDescView, const FWorldPartitionActorDescView& ReferenceActorDescView) override;
-	virtual void OnInvalidReferenceDataLayers(const FWorldPartitionActorDescView& ActorDescView, const FWorldPartitionActorDescView& ReferenceActorDescView) override;
-	virtual void OnInvalidReferenceLevelScriptStreamed(const FWorldPartitionActorDescView& ActorDescView) override;
-	virtual void OnInvalidReferenceLevelScriptDataLayers(const FWorldPartitionActorDescView& ActorDescView) override;
-	virtual void OnInvalidReferenceRuntimeGrid(const FWorldPartitionActorDescView& ActorDescView, const FWorldPartitionActorDescView& ReferenceActorDescView) override;
+protected:
+	virtual void AddAdditionalNameToken(TSharedRef<FTokenizedMessage>& InMessage, const FName& InErrorName) override;
+
+	virtual void HandleTokenizedMessage(TSharedRef<FTokenizedMessage>&& ErrorMessage) override;
 };
 #endif
diff --git a/Engine/Source/Runtime/Engine/Public/WorldPartition/ErrorHandling/WorldPartitionStreamingGenerationTokenizedMessageErrorHandler.h b/Engine/Source/Runtime/Engine/Public/WorldPartition/ErrorHandling/WorldPartitionStreamingGenerationTokenizedMessageErrorHandler.h
new file mode 100644
--- /dev/null
+++ b/Engine/Source/Runtime/Engine/Public/WorldPartition/ErrorHandling/WorldPartitionStreamingGenerationTokenizedMessageErrorHandler.h
@@ -0,0 +1,35 @@
+// Copyright Epic Games, Inc. All Rights Reserved.
+#pragma once
+
+#if WITH_EDITOR
+#include "WorldPartition/ErrorHandling/WorldPartitionStreamingGenerationErrorHandler.h"
+
+class FTokenizedMessage;
+
+class ENGINE_API ITokenizedMessageErrorHandler : public IStreamingGenerationErrorHandler
+{
+	virtual void OnInvalidReference(const FWorldPartitionActorDescView& ActorDescView, const FGuid& ReferenceGuid) override;
+	virtual void OnInvalidReferenceGridPlacement(const FWorldPartitionActorDescView& ActorDescView, const FWorldPartitionActorDescView& ReferenceActorDescView) override;
+	virtual void OnInvalidReferenceDataLayers(const FWorldPartitionActorDescView& ActorDescView, const FWorldPartitionActorDescView& ReferenceActorDescView) override;
+	virtual void OnInvalidReferenceLevelScriptStreamed(const FWorldPartitionActorDescView& ActorDescView) override;
+	virtual void OnInvalidReferenceLevelScriptDataLayers(const FWorldPartitionActorDescView& ActorDescView) override;
+	virtual void OnInvalidReferenceRuntimeGrid(const FWorldPartitionActorDescView& ActorDescView, const FWorldPartitionActorDescView& ReferenceActorDescView) override;
+
+protected:
+	virtual void AddAdditionalNameToken(TSharedRef<FTokenizedMessage>& InMessage, const FName& InErrorName) {}
+
+	virtual void HandleTokenizedMessage(TSharedRef<FTokenizedMessage>&& ErrorMessage) = 0;
+};
+
+class ENGINE_API FTokenizedMessageAccumulatorErrorHandler : public ITokenizedMessageErrorHandler
+{
+protected:
+	virtual void HandleTokenizedMessage(TSharedRef<FTokenizedMessage>&& ErrorMessage) override { ErrorMessages.Add(MoveTemp(ErrorMessage)); }
+
+	const TArray<TSharedRef<FTokenizedMessage>>& GetErrorMessages() const { return ErrorMessages; }
+
+private:
+	TArray<TSharedRef<FTokenizedMessage>> ErrorMessages;
+};
+
+#endif
