Some fixes for rules around promoting scalar values in new HLSL translator
Support some new node types

diff --git a/Engine/Source/Runtime/Engine/Classes/Materials/MaterialExpressionBumpOffset.h b/Engine/Source/Runtime/Engine/Classes/Materials/MaterialExpressionBumpOffset.h
--- a/Engine/Source/Runtime/Engine/Classes/Materials/MaterialExpressionBumpOffset.h
+++ b/Engine/Source/Runtime/Engine/Classes/Materials/MaterialExpressionBumpOffset.h
@@ -38,6 +38,7 @@ class UMaterialExpressionBumpOffset : public UMaterialExpression
 #if WITH_EDITOR
 	virtual int32 Compile(class FMaterialCompiler* Compiler, int32 OutputIndex) override;
 	virtual void GetCaption(TArray<FString>& OutCaptions) const override;
+	virtual bool GenerateHLSLExpression(FMaterialHLSLGenerator& Generator, UE::HLSLTree::FScope& Scope, int32 OutputIndex, UE::HLSLTree::FExpression*& OutExpression) override;
 #endif
 	//~ End UMaterialExpression Interface
 };
diff --git a/Engine/Source/Runtime/Engine/Classes/Materials/MaterialExpressionCurveAtlasRowParameter.h b/Engine/Source/Runtime/Engine/Classes/Materials/MaterialExpressionCurveAtlasRowParameter.h
--- a/Engine/Source/Runtime/Engine/Classes/Materials/MaterialExpressionCurveAtlasRowParameter.h
+++ b/Engine/Source/Runtime/Engine/Classes/Materials/MaterialExpressionCurveAtlasRowParameter.h
@@ -44,7 +44,8 @@ class UMaterialExpressionCurveAtlasRowParameter : public UMaterialExpressionScal
 	{
 		OutCaptions.Empty();
 		OutCaptions.Add(TEXT(""));
-	};
+	}
+	virtual bool GenerateHLSLExpression(FMaterialHLSLGenerator& Generator, UE::HLSLTree::FScope& Scope, int32 OutputIndex, UE::HLSLTree::FExpression*& OutExpression) override;
 //~ Begin UMaterialExpression Interface
 	virtual FName GetInputName(int32 InputIndex) const override
 	{
diff --git a/Engine/Source/Runtime/Engine/Classes/Materials/MaterialExpressionEyeAdaptation.h b/Engine/Source/Runtime/Engine/Classes/Materials/MaterialExpressionEyeAdaptation.h
--- a/Engine/Source/Runtime/Engine/Classes/Materials/MaterialExpressionEyeAdaptation.h
+++ b/Engine/Source/Runtime/Engine/Classes/Materials/MaterialExpressionEyeAdaptation.h
@@ -21,6 +21,7 @@ class UMaterialExpressionEyeAdaptation : public UMaterialExpression
 #if WITH_EDITOR
 	virtual int32 Compile(class FMaterialCompiler* Compiler, int32 OutputIndex) override;
 	virtual void GetCaption(TArray<FString>& OutCaptions) const override;
+	virtual bool GenerateHLSLExpression(FMaterialHLSLGenerator& Generator, UE::HLSLTree::FScope& Scope, int32 OutputIndex, UE::HLSLTree::FExpression*& OutExpression) override;
 #endif
 	//~ End UMaterialExpression Interface
 };
diff --git a/Engine/Source/Runtime/Engine/Classes/Materials/MaterialExpressionEyeAdaptationInverse.h b/Engine/Source/Runtime/Engine/Classes/Materials/MaterialExpressionEyeAdaptationInverse.h
--- a/Engine/Source/Runtime/Engine/Classes/Materials/MaterialExpressionEyeAdaptationInverse.h
+++ b/Engine/Source/Runtime/Engine/Classes/Materials/MaterialExpressionEyeAdaptationInverse.h
@@ -21,6 +21,7 @@ class UMaterialExpressionEyeAdaptationInverse : public UMaterialExpression
 #if WITH_EDITOR
 	virtual int32 Compile(class FMaterialCompiler* Compiler, int32 OutputIndex) override;
 	virtual void GetCaption(TArray<FString>& OutCaptions) const override;
+	virtual bool GenerateHLSLExpression(FMaterialHLSLGenerator& Generator, UE::HLSLTree::FScope& Scope, int32 OutputIndex, UE::HLSLTree::FExpression*& OutExpression) override;
 #endif
 	//~ End UMaterialExpression Interface
 
diff --git a/Engine/Source/Runtime/Engine/Classes/Materials/MaterialExpressionPreSkinnedLocalBounds.h b/Engine/Source/Runtime/Engine/Classes/Materials/MaterialExpressionPreSkinnedLocalBounds.h
--- a/Engine/Source/Runtime/Engine/Classes/Materials/MaterialExpressionPreSkinnedLocalBounds.h
+++ b/Engine/Source/Runtime/Engine/Classes/Materials/MaterialExpressionPreSkinnedLocalBounds.h
@@ -18,6 +18,7 @@ class UMaterialExpressionPreSkinnedLocalBounds : public UMaterialExpression
 	virtual void GetCaption(TArray<FString>& OutCaptions) const override;
 	virtual void GetConnectorToolTip(int32 InputIndex, int32 OutputIndex, TArray<FString>& OutToolTip) override;
 	virtual void GetExpressionToolTip(TArray<FString>& OutToolTip) override;
+	virtual bool GenerateHLSLExpression(FMaterialHLSLGenerator& Generator, UE::HLSLTree::FScope& Scope, int32 OutputIndex, UE::HLSLTree::FExpression*& OutExpression) override;
 	//~ End UMaterialExpression Interface
 
 #if WITH_EDITORONLY_DATA
diff --git a/Engine/Source/Runtime/Engine/Classes/Materials/MaterialExpressionPreSkinnedNormal.h b/Engine/Source/Runtime/Engine/Classes/Materials/MaterialExpressionPreSkinnedNormal.h
--- a/Engine/Source/Runtime/Engine/Classes/Materials/MaterialExpressionPreSkinnedNormal.h
+++ b/Engine/Source/Runtime/Engine/Classes/Materials/MaterialExpressionPreSkinnedNormal.h
@@ -18,6 +18,7 @@ class UMaterialExpressionPreSkinnedNormal : public UMaterialExpression
 	virtual void GetCaption(TArray<FString>& OutCaptions) const override;
 	virtual void GetExpressionToolTip(TArray<FString>& OutToolTip) override;
 	virtual FText GetKeywords() const override {return FText::FromString(TEXT("normal"));}
+	virtual bool GenerateHLSLExpression(FMaterialHLSLGenerator& Generator, UE::HLSLTree::FScope& Scope, int32 OutputIndex, UE::HLSLTree::FExpression*& OutExpression) override;
 #endif
 	//~ End UMaterialExpression Interface
 };
diff --git a/Engine/Source/Runtime/Engine/Classes/Materials/MaterialExpressionPreSkinnedPosition.h b/Engine/Source/Runtime/Engine/Classes/Materials/MaterialExpressionPreSkinnedPosition.h
--- a/Engine/Source/Runtime/Engine/Classes/Materials/MaterialExpressionPreSkinnedPosition.h
+++ b/Engine/Source/Runtime/Engine/Classes/Materials/MaterialExpressionPreSkinnedPosition.h
@@ -18,6 +18,7 @@ class UMaterialExpressionPreSkinnedPosition : public UMaterialExpression
 	virtual void GetCaption(TArray<FString>& OutCaptions) const override;
 	virtual void GetExpressionToolTip(TArray<FString>& OutToolTip) override;
 	virtual FText GetKeywords() const override {return FText::FromString(TEXT("position"));}
+	virtual bool GenerateHLSLExpression(FMaterialHLSLGenerator& Generator, UE::HLSLTree::FScope& Scope, int32 OutputIndex, UE::HLSLTree::FExpression*& OutExpression) override;
 #endif
 	//~ End UMaterialExpression Interface
 };
diff --git a/Engine/Source/Runtime/Engine/Private/HLSLTree/HLSLTree.cpp b/Engine/Source/Runtime/Engine/Private/HLSLTree/HLSLTree.cpp
--- a/Engine/Source/Runtime/Engine/Private/HLSLTree/HLSLTree.cpp
+++ b/Engine/Source/Runtime/Engine/Private/HLSLTree/HLSLTree.cpp
@@ -710,31 +710,9 @@ FRequestedType FPreparedType::GetRequestedType() const
 	return Result;
 }
 
-FPreparedType FPreparedType::GetTypeForRequest(const FRequestedType& RequestedType) const
-{
-	const int32 NumComponents = FMath::Min(PreparedComponents.Num(), RequestedType.RequestedComponents.Num());
-
-	FPreparedType Result;
-	Result.StructType = StructType;
-	Result.ValueComponentType = ValueComponentType;
-	Result.PreparedComponents.Reserve(NumComponents);
-	for (int32 Index = 0; Index < NumComponents; ++Index)
-	{
-		if (RequestedType.IsComponentRequested(Index))
-		{
-			Result.SetComponent(Index, GetComponent(Index));
-		}
-		else
-		{
-			Result.SetComponent(Index, EExpressionEvaluation::ConstantZero);
-		}
-	}
-	return Result;
-}
-
 EExpressionEvaluation FPreparedType::GetEvaluation(const FEmitScope& Scope) const
 {
-	EExpressionEvaluation Result = EExpressionEvaluation::None;
+	EExpressionEvaluation Result = EExpressionEvaluation::ConstantZero;
 	for (int32 Index = 0; Index < PreparedComponents.Num(); ++Index)
 	{
 		Result = CombineEvaluations(Result, PreparedComponents[Index].GetEvaluation(Scope));
@@ -744,16 +722,13 @@ EExpressionEvaluation FPreparedType::GetEvaluation(const FEmitScope& Scope) cons
 
 EExpressionEvaluation FPreparedType::GetEvaluation(const FEmitScope& Scope, const FRequestedType& RequestedType) const
 {
-	EExpressionEvaluation Result = EExpressionEvaluation::None;
-	for (int32 Index = 0; Index < PreparedComponents.Num(); ++Index)
+	const int32 NumComponents = RequestedType.GetNumComponents();
+	EExpressionEvaluation Result = EExpressionEvaluation::ConstantZero;
+	for (int32 Index = 0; Index < NumComponents; ++Index)
 	{
 		if (RequestedType.IsComponentRequested(Index))
 		{
-			Result = CombineEvaluations(Result, PreparedComponents[Index].GetEvaluation(Scope));
-		}
-		else
-		{
-			Result = CombineEvaluations(Result, EExpressionEvaluation::ConstantZero);
+			Result = CombineEvaluations(Result, GetComponent(Index).GetEvaluation(Scope));
 		}
 	}
 	return Result;
@@ -761,7 +736,7 @@ EExpressionEvaluation FPreparedType::GetEvaluation(const FEmitScope& Scope, cons
 
 EExpressionEvaluation FPreparedType::GetFieldEvaluation(const FEmitScope& Scope, const FRequestedType& RequestedType, int32 ComponentIndex, int32 NumComponents) const
 {
-	EExpressionEvaluation Result = EExpressionEvaluation::None;
+	EExpressionEvaluation Result = EExpressionEvaluation::ConstantZero;
 	for (int32 Index = 0; Index < NumComponents; ++Index)
 	{
 		const FPreparedComponent Component = GetComponent(ComponentIndex + Index);
@@ -769,10 +744,6 @@ EExpressionEvaluation FPreparedType::GetFieldEvaluation(const FEmitScope& Scope,
 		{
 			Result = CombineEvaluations(Result, Component.GetEvaluation(Scope));
 		}
-		else if(!Component.IsNone())
-		{
-			Result = CombineEvaluations(Result, EExpressionEvaluation::ConstantZero);
-		}
 	}
 	return Result;
 }
@@ -793,9 +764,11 @@ FPreparedComponent FPreparedType::GetComponent(int32 Index) const
 	{
 		return PreparedComponents[Index];
 	}
+
 	// Return 'ConstantZero' for unprepared components
 	// TODO - do we ever want to return 'None' instead?
-	return EExpressionEvaluation::ConstantZero;
+	const int32 NumComponents = GetNumComponents();
+	return (NumComponents == 1) ? PreparedComponents[0] : FPreparedComponent(EExpressionEvaluation::ConstantZero);
 }
 
 void FPreparedType::EnsureNumComponents(int32 NumComponents)
@@ -1070,11 +1043,12 @@ bool FPrepareValueResult::SetType(FEmitContext& Context, const FRequestedType& R
 		const int32 NumComponents = Type.GetNumComponents();
 		for (int32 Index = 0; Index < NumComponents; ++Index)
 		{
+			const FPreparedComponent& Component = Type.GetComponent(Index);
 			if (RequestedType.IsComponentRequested(Index))
 			{
-				PreparedType.MergeComponent(Index, Type.GetComponent(Index));
+				PreparedType.MergeComponent(Index, Component);
 			}
-			else
+			else if(!Component.IsNone())
 			{
 				// If component wasn't requested, it can be replaced with constant-0
 				PreparedType.MergeComponent(Index, EExpressionEvaluation::ConstantZero);
diff --git a/Engine/Source/Runtime/Engine/Private/HLSLTree/HLSLTreeCommon.cpp b/Engine/Source/Runtime/Engine/Private/HLSLTree/HLSLTreeCommon.cpp
--- a/Engine/Source/Runtime/Engine/Private/HLSLTree/HLSLTreeCommon.cpp
+++ b/Engine/Source/Runtime/Engine/Private/HLSLTree/HLSLTreeCommon.cpp
@@ -133,7 +133,23 @@ void FExpressionConstant::ComputeAnalyticDerivatives(FTree& Tree, FExpressionDer
 
 bool FExpressionConstant::PrepareValue(FEmitContext& Context, FEmitScope& Scope, const FRequestedType& RequestedType, FPrepareValueResult& OutResult) const
 {
-	return OutResult.SetType(Context, RequestedType, EExpressionEvaluation::Constant, Value.Type);
+	const int32 NumComponents = Value.Type.GetNumComponents();
+	FPreparedType ResultType(Value.Type);
+	ResultType.PreparedComponents.Reserve(NumComponents);
+	for (int32 Index = 0; Index < NumComponents; ++Index)
+	{
+		EExpressionEvaluation Evaluation = EExpressionEvaluation::ConstantZero;
+		if (RequestedType.IsComponentRequested(Index))
+		{
+			const Shader::FValueComponent Component = Value.GetComponent(Index);
+			if (Component.Packed != 0u)
+			{
+				Evaluation = EExpressionEvaluation::Constant;
+			}
+		}
+		ResultType.PreparedComponents.Add(Evaluation);
+	}
+	return OutResult.SetType(Context, RequestedType, ResultType);
 }
 
 void FExpressionConstant::EmitValuePreshader(FEmitContext& Context, FEmitScope& Scope, const FRequestedType& RequestedType, FEmitValuePreshaderResult& OutResult) const
@@ -406,35 +422,25 @@ void FExpressionSetStructField::EmitValueShader(FEmitContext& Context, FEmitScop
 	FRequestedType RequestedStructType(RequestedType);
 	RequestedStructType.ClearFieldRequested(Field);
 	const EExpressionEvaluation StructEvaluation = Context.GetEvaluation(StructExpression, Scope, RequestedStructType);
+	check(StructEvaluation != EExpressionEvaluation::None);
 
 	const FRequestedType RequestedFieldType = RequestedType.GetField(Field);
-	const EExpressionEvaluation FieldEvaluation = Context.GetEvaluation(FieldExpression, Scope, RequestedStructType);
-
-	FEmitShaderExpression* StructValue = (StructEvaluation != EExpressionEvaluation::None) ? StructExpression->GetValueShader(Context, Scope, RequestedStructType) : nullptr;
-	FEmitShaderExpression* FieldValue = (FieldEvaluation != EExpressionEvaluation::None) ? FieldExpression->GetValueShader(Context, Scope, RequestedFieldType, Field->Type) : nullptr;
+	const EExpressionEvaluation FieldEvaluation = Context.GetEvaluation(FieldExpression, Scope, RequestedFieldType);
+	check(FieldEvaluation != EExpressionEvaluation::None);
 
-	if (StructEvaluation == EExpressionEvaluation::None && FieldEvaluation == EExpressionEvaluation::None)
+	if (StructEvaluation == EExpressionEvaluation::ConstantZero && FieldEvaluation == EExpressionEvaluation::ConstantZero)
 	{
-		OutResult.Code = Context.EmitExpression(Scope, StructType, TEXT("((%)0)"), StructType->Name);
-	}
-	else if (StructEvaluation == EExpressionEvaluation::None)
-	{
-		// StructExpression is not used, so default to a zero-initialized struct
-		// This will happen if all the accessed struct fields are explicitly defined
-		OutResult.Code = Context.EmitExpression(Scope, StructType, TEXT("%_Set%((%)0, %)"),
-			StructType->Name,
-			Field->Name,
-			StructType->Name,
-			FieldValue);
-	}
-	else if (FieldEvaluation == EExpressionEvaluation::None)
-	{
-		// Don't need field, can just forward the struct value
-		OutResult.Code = StructValue;
+		OutResult.Code = Context.EmitConstantZero(Scope, StructType);
 	}
 	else
 	{
-		check(StructValue->Type.StructType == StructType);
+		FEmitShaderExpression* StructValue = (StructEvaluation != EExpressionEvaluation::ConstantZero)
+			? StructExpression->GetValueShader(Context, Scope, RequestedStructType)
+			: Context.EmitConstantZero(Scope, StructType);
+		FEmitShaderExpression* FieldValue = (FieldEvaluation != EExpressionEvaluation::ConstantZero)
+			? FieldExpression->GetValueShader(Context, Scope, RequestedFieldType, Field->Type)
+			: Context.EmitConstantZero(Scope, Field->Type);
+
 		OutResult.Code = Context.EmitExpression(Scope, StructType, TEXT("%_Set%(%, %)"),
 			StructType->Name,
 			Field->Name,
@@ -450,20 +456,17 @@ void FExpressionSetStructField::EmitValuePreshader(FEmitContext& Context, FEmitS
 	const EExpressionEvaluation StructEvaluation = Context.GetEvaluation(StructExpression, Scope, RequestedStructType);
 
 	const FRequestedType RequestedFieldType = RequestedType.GetField(Field);
-	const EExpressionEvaluation FieldEvaluation = Context.GetEvaluation(FieldExpression, Scope, RequestedStructType);
+	const EExpressionEvaluation FieldEvaluation = Context.GetEvaluation(FieldExpression, Scope, RequestedFieldType);
 
-	if (StructEvaluation != EExpressionEvaluation::None)
-	{
-		StructExpression->GetValuePreshader(Context, Scope, RequestedStructType, OutResult.Preshader);
-	}
-	else
+	OutResult.Type = StructType;
+	if (StructEvaluation == EExpressionEvaluation::ConstantZero && FieldEvaluation == EExpressionEvaluation::ConstantZero)
 	{
 		Context.PreshaderStackPosition++;
 		OutResult.Preshader.WriteOpcode(Shader::EPreshaderOpcode::ConstantZero).Write(Shader::FType(StructType));
 	}
-
-	if (FieldEvaluation != EExpressionEvaluation::None)
+	else
 	{
+		StructExpression->GetValuePreshader(Context, Scope, RequestedStructType, OutResult.Preshader);
 		FieldExpression->GetValuePreshader(Context, Scope, RequestedFieldType, OutResult.Preshader);
 
 		check(Context.PreshaderStackPosition > 0);
@@ -471,7 +474,6 @@ void FExpressionSetStructField::EmitValuePreshader(FEmitContext& Context, FEmitS
 
 		OutResult.Preshader.WriteOpcode(Shader::EPreshaderOpcode::SetField).Write(Field->ComponentIndex).Write(Field->GetNumComponents());
 	}
-	OutResult.Type = StructType;
 }
 
 void FExpressionSelect::ComputeAnalyticDerivatives(FTree& Tree, FExpressionDerivatives& OutResult) const
@@ -494,7 +496,7 @@ bool FExpressionSelect::PrepareValue(FEmitContext& Context, FEmitScope& Scope, c
 {
 	const FPreparedType& ConditionType = Context.PrepareExpression(ConditionExpression, Scope, ERequestedType::Scalar);
 	const EExpressionEvaluation ConditionEvaluation = ConditionType.GetEvaluation(Scope, ERequestedType::Scalar);
-	if (ConditionEvaluation == EExpressionEvaluation::Constant)
+	if (IsConstantEvaluation(ConditionEvaluation))
 	{
 		const bool bCondition = ConditionExpression->GetValueConstant(Context, Scope, Shader::EValueType::Bool1).AsBoolScalar();
 		FPreparedType ResultType = Context.PrepareExpression(bCondition ? TrueExpression : FalseExpression, Scope, RequestedType);
@@ -520,7 +522,7 @@ bool FExpressionSelect::PrepareValue(FEmitContext& Context, FEmitScope& Scope, c
 void FExpressionSelect::EmitValueShader(FEmitContext& Context, FEmitScope& Scope, const FRequestedType& RequestedType, FEmitValueShaderResult& OutResult) const
 {
 	const EExpressionEvaluation ConditionEvaluation = Context.GetEvaluation(ConditionExpression, Scope, ERequestedType::Scalar);
-	if (ConditionEvaluation == EExpressionEvaluation::Constant)
+	if (IsConstantEvaluation(ConditionEvaluation))
 	{
 		const bool bCondition = ConditionExpression->GetValueConstant(Context, Scope, Shader::EValueType::Bool1).AsBoolScalar();
 		FExpression* InputExpression = bCondition ? TrueExpression : FalseExpression;
@@ -686,8 +688,6 @@ void FExpressionSwizzle::EmitValueShader(FEmitContext& Context, FEmitScope& Scop
 	{
 		const Shader::EValueType ResultType = Shader::MakeValueType(InputTypeDesc.ComponentType, NumComponents);
 		const int32 NumRequestedComponents = RequestedInputType.GetNumComponents();
-		check(NumRequestedComponents > 0);
-	
 		if (NumRequestedComponents > InputTypeDesc.NumComponents)
 		{
 			// Zero-extend our input if needed, so we can acccess all the given components
@@ -1041,18 +1041,18 @@ bool FStatementIf::Prepare(FEmitContext& Context, FEmitScope& Scope) const
 
 	const EExpressionEvaluation ConditionEvaluation = ConditionType.GetEvaluation(Scope, ERequestedType::Scalar);
 	check(ConditionEvaluation != EExpressionEvaluation::None);
-	if (ConditionEvaluation == EExpressionEvaluation::Constant)
+	if (IsConstantEvaluation(ConditionEvaluation))
 	{
 		const bool bCondition = ConditionExpression->GetValueConstant(Context, Scope, Shader::EValueType::Bool1).AsBoolScalar();
 		if (bCondition)
 		{
-			Context.MarkScopeEvaluation(Scope, ThenScope, EExpressionEvaluation::Constant);
+			Context.MarkScopeEvaluation(Scope, ThenScope, ConditionEvaluation);
 			Context.MarkScopeDead(Scope, ElseScope);
 		}
 		else
 		{
 			Context.MarkScopeDead(Scope, ThenScope);
-			Context.MarkScopeEvaluation(Scope, ElseScope, EExpressionEvaluation::Constant);
+			Context.MarkScopeEvaluation(Scope, ElseScope, ConditionEvaluation);
 		}
 	}
 	else
@@ -1068,7 +1068,7 @@ void FStatementIf::EmitShader(FEmitContext& Context, FEmitScope& Scope) const
 {
 	FEmitShaderNode* Dependency = nullptr;
 	const EExpressionEvaluation ConditionEvaluation = Context.GetEvaluation(ConditionExpression, Scope, ERequestedType::Scalar);
-	if (ConditionEvaluation == EExpressionEvaluation::Constant)
+	if (IsConstantEvaluation(ConditionEvaluation))
 	{
 		const bool bCondition = ConditionExpression->GetValueConstant(Context, Scope, Shader::EValueType::Bool1).AsBoolScalar();
 		if (bCondition)
diff --git a/Engine/Source/Runtime/Engine/Private/HLSLTree/HLSLTreeOperation.cpp b/Engine/Source/Runtime/Engine/Private/HLSLTree/HLSLTreeOperation.cpp
--- a/Engine/Source/Runtime/Engine/Private/HLSLTree/HLSLTreeOperation.cpp
+++ b/Engine/Source/Runtime/Engine/Private/HLSLTree/HLSLTreeOperation.cpp
@@ -103,6 +103,12 @@ FExpression* FTree::NewBinaryOp(EOperation Op, FExpression* Lhs, FExpression* Rh
 	return NewExpression<FExpressionOperation>(Op, Inputs);
 }
 
+FExpression* FTree::NewLog(FExpression* Input)
+{
+	const float Log2ToLog = 1.0f / FMath::Log2(EULERS_NUMBER);
+	return NewMul(NewLog2(Input), NewConstant(Log2ToLog));
+}
+
 FExpression* FTree::NewCross(FExpression* Lhs, FExpression* Rhs)
 {
 	//c_P[0] = v_A[1] * v_B[2] - v_A[2] * v_B[1];
diff --git a/Engine/Source/Runtime/Engine/Private/Materials/MaterialExpressionHLSL.cpp b/Engine/Source/Runtime/Engine/Private/Materials/MaterialExpressionHLSL.cpp
--- a/Engine/Source/Runtime/Engine/Private/Materials/MaterialExpressionHLSL.cpp
+++ b/Engine/Source/Runtime/Engine/Private/Materials/MaterialExpressionHLSL.cpp
@@ -16,6 +16,7 @@
 #include "Materials/MaterialExpressionConstantBiasScale.h"
 #include "Materials/MaterialExpressionShadingModel.h"
 #include "Materials/MaterialExpressionParameter.h"
+#include "Materials/MaterialExpressionCurveAtlasRowParameter.h"
 #include "Materials/MaterialExpressionStaticSwitchParameter.h"
 #include "Materials/MaterialExpressionPixelDepth.h"
 #include "Materials/MaterialExpressionWorldPosition.h"
@@ -31,11 +32,16 @@
 #include "Materials/MaterialExpressionPixelNormalWS.h"
 #include "Materials/MaterialExpressionReflectionVectorWS.h"
 #include "Materials/MaterialExpressionActorPositionWS.h"
+#include "Materials/MaterialExpressionPreSkinnedPosition.h"
+#include "Materials/MaterialExpressionPreSkinnedNormal.h"
+#include "Materials/MaterialExpressionPreSkinnedLocalBounds.h"
 #include "Materials/MaterialExpressionVertexColor.h"
 #include "Materials/MaterialExpressionVertexNormalWS.h"
 #include "Materials/MaterialExpressionVertexTangentWS.h"
 #include "Materials/MaterialExpressionTextureCoordinate.h"
 #include "Materials/MaterialExpressionLightmapUVs.h"
+#include "Materials/MaterialExpressionEyeAdaptation.h"
+#include "Materials/MaterialExpressionEyeAdaptationInverse.h"
 #include "Materials/MaterialExpressionParticleColor.h"
 #include "Materials/MaterialExpressionParticlePositionWS.h"
 #include "Materials/MaterialExpressionParticleRadius.h"
@@ -110,6 +116,7 @@
 #include "Materials/MaterialExpressionDesaturation.h"
 #include "Materials/MaterialExpressionSphereMask.h"
 #include "Materials/MaterialExpressionRotateAboutAxis.h"
+#include "Materials/MaterialExpressionBumpOffset.h"
 #include "Materials/MaterialExpressionPanner.h"
 #include "Materials/MaterialExpressionRotator.h"
 #include "Materials/MaterialExpressionVolumetricAdvancedMaterialInput.h"
@@ -134,6 +141,8 @@
 #include "Engine/Engine.h"
 #include "Engine/Texture2D.h"
 #include "Engine/Font.h"
+#include "Curves/CurveLinearColorAtlas.h"
+#include "Curves/CurveLinearColor.h"
 #include "RenderUtils.h"
 #include "Misc/MemStackUtility.h"
 #include "MaterialHLSLGenerator.h"
@@ -307,6 +316,44 @@ bool UMaterialExpressionParameter::GenerateHLSLExpression(FMaterialHLSLGenerator
 	return true;
 }
 
+bool UMaterialExpressionCurveAtlasRowParameter::GenerateHLSLExpression(FMaterialHLSLGenerator& Generator, UE::HLSLTree::FScope& Scope, int32 OutputIndex, UE::HLSLTree::FExpression*& OutExpression)
+{
+	using namespace UE::HLSLTree;
+	using namespace UE::Shader;
+	if (!Atlas)
+	{
+		return Generator.Error(TEXT("Missing atlas"));
+	}
+	if (!Curve)
+	{
+		return Generator.Error(TEXT("Missing curve"));
+	}
+
+	FMaterialParameterMetadata ParameterMeta;
+	if (!GetParameterValue(ParameterMeta))
+	{
+		return Generator.Error(TEXT("Failed to get parameter value"));
+	}
+	FExpression* ExpressionSlot = Generator.GenerateMaterialParameter(ParameterName, ParameterMeta, ParameterMeta.Value.AsShaderValue());
+
+	FExpression* ExpressionTexture = nullptr;
+	{
+		const FMaterialParameterMetadata TextureParameterMeta(Atlas.Get());
+		const FTextureValue* TextureValue = Generator.AcquireTextureValue(FTextureValue(Atlas.Get(), SAMPLERTYPE_LinearColor));
+		ExpressionTexture = Generator.GenerateMaterialParameter(FName(), TextureParameterMeta, TextureValue);
+	}
+
+	FTree& Tree = Generator.GetTree();
+	FExpression* ExpressionTextureSize = Tree.NewExpression<Material::FExpressionTextureSize>(FMaterialParameterInfo(), Atlas.Get());
+	FExpression* ExpressionTextureHeight = Tree.NewExpression<FExpressionSwizzle>(MakeSwizzleMask(false, true, false, false), ExpressionTextureSize);
+	FExpression* ExpressionCoordU = InputTime.AcquireHLSLExpressionOrConstant(Generator, Scope, 0.0f);
+	FExpression* ExpressionCoordV = Tree.NewMul(Tree.NewAdd(ExpressionSlot, Tree.NewConstant(0.5f)), Tree.NewRcp(ExpressionTextureHeight));
+	FExpression* ExpressionUV = Tree.NewExpression<FExpressionAppend>(ExpressionCoordU, ExpressionCoordV);
+
+	OutExpression = Generator.GetTree().NewExpression<FExpressionTextureSample>(ExpressionTexture, ExpressionUV, nullptr, FExpressionDerivatives(), SSM_Clamp_WorldGroupSettings, TMVM_None);
+	return true;
+}
+
 bool UMaterialExpressionStaticSwitchParameter::GenerateHLSLExpression(FMaterialHLSLGenerator& Generator, UE::HLSLTree::FScope& Scope, int32 OutputIndex, UE::HLSLTree::FExpression*& OutExpression)
 {
 	using namespace UE::HLSLTree;
@@ -487,6 +534,43 @@ bool UMaterialExpressionActorPositionWS::GenerateHLSLExpression(FMaterialHLSLGen
 	return true;
 }
 
+bool UMaterialExpressionPreSkinnedPosition::GenerateHLSLExpression(FMaterialHLSLGenerator& Generator, UE::HLSLTree::FScope& Scope, int32 OutputIndex, UE::HLSLTree::FExpression*& OutExpression)
+{
+	using namespace UE::HLSLTree;
+	OutExpression = Generator.NewExternalInput(Material::EExternalInput::PreSkinnedPosition);
+	return true;
+}
+
+bool UMaterialExpressionPreSkinnedNormal::GenerateHLSLExpression(FMaterialHLSLGenerator& Generator, UE::HLSLTree::FScope& Scope, int32 OutputIndex, UE::HLSLTree::FExpression*& OutExpression)
+{
+	using namespace UE::HLSLTree;
+	OutExpression = Generator.NewExternalInput(Material::EExternalInput::PreSkinnedNormal);
+	return true;
+}
+
+bool UMaterialExpressionPreSkinnedLocalBounds::GenerateHLSLExpression(FMaterialHLSLGenerator& Generator, UE::HLSLTree::FScope& Scope, int32 OutputIndex, UE::HLSLTree::FExpression*& OutExpression)
+{
+	using namespace UE::HLSLTree;
+	FTree& Tree = Generator.GetTree();
+	switch (OutputIndex)
+	{
+	case 0: // Half extents
+		OutExpression = Tree.NewMul(Tree.NewSub(Generator.NewExternalInput(Material::EExternalInput::PreSkinnedLocalBoundsMax), Generator.NewExternalInput(Material::EExternalInput::PreSkinnedLocalBoundsMin)), Tree.NewConstant(0.5f));
+		return true;
+	case 1: // Full extents
+		OutExpression = Tree.NewSub(Generator.NewExternalInput(Material::EExternalInput::PreSkinnedLocalBoundsMax), Generator.NewExternalInput(Material::EExternalInput::PreSkinnedLocalBoundsMin));
+		return true;
+	case 2: // Min point
+		OutExpression = Generator.NewExternalInput(Material::EExternalInput::PreSkinnedLocalBoundsMin);
+		return true;
+	case 3: // Max point
+		OutExpression = Generator.NewExternalInput(Material::EExternalInput::PreSkinnedLocalBoundsMax);
+		return true;
+	default:
+		return Generator.Error(TEXT("Invalid Output"));
+	}
+}
+
 bool UMaterialExpressionPixelNormalWS::GenerateHLSLExpression(FMaterialHLSLGenerator& Generator, UE::HLSLTree::FScope& Scope, int32 OutputIndex, UE::HLSLTree::FExpression*& OutExpression)
 {
 	using namespace UE::HLSLTree::Material;
@@ -545,6 +629,31 @@ bool UMaterialExpressionLightmapUVs::GenerateHLSLExpression(FMaterialHLSLGenerat
 	return true;
 }
 
+bool UMaterialExpressionEyeAdaptation::GenerateHLSLExpression(FMaterialHLSLGenerator& Generator, UE::HLSLTree::FScope& Scope, int32 OutputIndex, UE::HLSLTree::FExpression*& OutExpression)
+{
+	using namespace UE::HLSLTree;
+	OutExpression = Generator.GetTree().NewExpression<Material::FExpressionExternalInput>(Material::EExternalInput::EyeAdaptation);
+	return true;
+}
+
+bool UMaterialExpressionEyeAdaptationInverse::GenerateHLSLExpression(FMaterialHLSLGenerator& Generator, UE::HLSLTree::FScope& Scope, int32 OutputIndex, UE::HLSLTree::FExpression*& OutExpression)
+{
+	using namespace UE::HLSLTree;
+	FExpression* ExpressionLightValue = LightValueInput.AcquireHLSLExpressionOrConstant(Generator, Scope, FVector3f(1.0f, 1.0f, 1.0f));
+	FExpression* ExpressionAlpha = AlphaInput.AcquireHLSLExpressionOrConstant(Generator, Scope, 1.0f);
+	if (!ExpressionLightValue || !ExpressionAlpha)
+	{
+		return false;
+	}
+
+	FTree& Tree = Generator.GetTree();
+	FExpression* ExpressionAdaptation = Tree.NewExpression<Material::FExpressionExternalInput>(Material::EExternalInput::EyeAdaptation);
+	FExpression* LerpLogScale = Tree.NewMul(Tree.NewNeg(ExpressionAlpha), Tree.NewLog2(ExpressionAdaptation));
+	FExpression* Scale = Tree.NewExp2(LerpLogScale);
+	OutExpression = Tree.NewMul(ExpressionLightValue, Scale);
+	return true;
+}
+
 bool UMaterialExpressionParticleColor::GenerateHLSLExpression(FMaterialHLSLGenerator& Generator, UE::HLSLTree::FScope& Scope, int32 OutputIndex, UE::HLSLTree::FExpression*& OutExpression)
 {
 	using namespace UE::HLSLTree::Material;
@@ -1629,7 +1738,7 @@ static UE::HLSLTree::FExpression* TransformBase(UE::HLSLTree::FTree& Tree,
 
 	FExpression* Result = nullptr;
 	EMaterialCommonBasis IntermediaryBasis = MCB_World;
-	EOperation Op = bWComponent ? EOperation::VecMulMatrix4 : EOperation::VecMulMatrix3;
+	const EOperation Op = bWComponent ? EOperation::VecMulMatrix4 : EOperation::VecMulMatrix3;
 	switch (SourceCoordBasis)
 	{
 	case MCB_Tangent:
@@ -1685,7 +1794,8 @@ static UE::HLSLTree::FExpression* TransformBase(UE::HLSLTree::FTree& Tree,
 	{
 		if (DestCoordBasis == MCB_Tangent)
 		{
-			Result = Tree.NewBinaryOp(Op, Tree.NewExpression<FExpressionExternalInput>(EExternalInput::TangentToWorld), Input);
+			const EOperation InvOp = bWComponent ? EOperation::Matrix4MulVec : EOperation::Matrix3MulVec;
+			Result = Tree.NewBinaryOp(InvOp, Tree.NewExpression<FExpressionExternalInput>(EExternalInput::TangentToWorld), Input);
 			//CodeStr = MultiplyTransposeMatrix(TEXT("Parameters.TangentToWorld"), TEXT("<A>"), AWComponent);
 		}
 		else if (DestCoordBasis == MCB_Local)
@@ -1985,6 +2095,31 @@ bool UMaterialExpressionRotateAboutAxis::GenerateHLSLExpression(FMaterialHLSLGen
 	return true;
 }
 
+bool UMaterialExpressionBumpOffset::GenerateHLSLExpression(FMaterialHLSLGenerator& Generator, UE::HLSLTree::FScope& Scope, int32 OutputIndex, UE::HLSLTree::FExpression*& OutExpression)
+{
+	using namespace UE::HLSLTree;
+
+	FExpression* ExpressionHeightRatio = HeightRatioInput.AcquireHLSLExpressionOrConstant(Generator, Scope, HeightRatio);
+	FExpression* ExpressionHeight = Height.AcquireHLSLExpression(Generator, Scope);
+	FExpression* ExpressionCoordinate = Coordinate.AcquireHLSLExpression(Generator, Scope);
+	if (!ExpressionCoordinate)
+	{
+		ExpressionCoordinate = Generator.NewTexCoord(ConstCoordinate);
+	}
+	if (!ExpressionHeightRatio || !ExpressionHeight || !ExpressionCoordinate)
+	{
+		return false;
+	}
+
+	FTree& Tree = Generator.GetTree();
+	FExpression* ExpressionCamera = TransformBase(Tree, MCB_World, MCB_Tangent, Generator.NewExternalInput(Material::EExternalInput::CameraVector), false);
+	ExpressionCamera = Generator.NewSwizzle(MakeSwizzleMask(true, true, false, false), ExpressionCamera);
+
+	FExpression* ExpressionHeightOffset = Tree.NewAdd(Tree.NewMul(ExpressionHeightRatio, ExpressionHeight), Tree.NewMul(ExpressionHeightRatio, Tree.NewConstant(-ReferencePlane)));
+	OutExpression = Tree.NewAdd(Tree.NewMul(ExpressionHeightOffset, ExpressionCamera), ExpressionCoordinate);
+	return true;
+}
+
 bool UMaterialExpressionPanner::GenerateHLSLExpression(FMaterialHLSLGenerator& Generator, UE::HLSLTree::FScope& Scope, int32 OutputIndex, UE::HLSLTree::FExpression*& OutExpression)
 {
 	using namespace UE::HLSLTree;
diff --git a/Engine/Source/Runtime/Engine/Private/Materials/MaterialHLSLGenerator.cpp b/Engine/Source/Runtime/Engine/Private/Materials/MaterialHLSLGenerator.cpp
--- a/Engine/Source/Runtime/Engine/Private/Materials/MaterialHLSLGenerator.cpp
+++ b/Engine/Source/Runtime/Engine/Private/Materials/MaterialHLSLGenerator.cpp
@@ -310,8 +310,14 @@ UE::HLSLTree::FExpression* FMaterialHLSLGenerator::NewConstant(const UE::Shader:
 
 UE::HLSLTree::FExpression* FMaterialHLSLGenerator::NewTexCoord(int32 Index)
 {
-	using namespace UE::HLSLTree::Material;
-	return GetTree().NewExpression<FExpressionExternalInput>(MakeInputTexCoord(Index));
+	using namespace UE::HLSLTree;
+	return NewExternalInput(Material::MakeInputTexCoord(Index));
+}
+
+UE::HLSLTree::FExpression* FMaterialHLSLGenerator::NewExternalInput(UE::HLSLTree::Material::EExternalInput Input)
+{
+	using namespace UE::HLSLTree;
+	return GetTree().NewExpression<Material::FExpressionExternalInput>(Input);
 }
 
 UE::HLSLTree::FExpression* FMaterialHLSLGenerator::NewSwizzle(const UE::HLSLTree::FSwizzleParameters& Params, UE::HLSLTree::FExpression* Input)
diff --git a/Engine/Source/Runtime/Engine/Private/Materials/MaterialHLSLTree.cpp b/Engine/Source/Runtime/Engine/Private/Materials/MaterialHLSLTree.cpp
--- a/Engine/Source/Runtime/Engine/Private/Materials/MaterialHLSLTree.cpp
+++ b/Engine/Source/Runtime/Engine/Private/Materials/MaterialHLSLTree.cpp
@@ -77,6 +77,11 @@ FExternalInputDescription GetExternalInputDescription(EExternalInput Input)
 	case EExternalInput::WorldVertexNormal: return FExternalInputDescription(TEXT("WorldVertexNormal"), Shader::EValueType::Float3);
 	case EExternalInput::WorldVertexTangent: return FExternalInputDescription(TEXT("WorldVertexTangent"), Shader::EValueType::Float3);
 
+	case EExternalInput::PreSkinnedPosition: return FExternalInputDescription(TEXT("PreSkinnedPosition"), Shader::EValueType::Float3);
+	case EExternalInput::PreSkinnedNormal: return FExternalInputDescription(TEXT("PreSkinnedNormal"), Shader::EValueType::Float3);
+	case EExternalInput::PreSkinnedLocalBoundsMin: return FExternalInputDescription(TEXT("PreSkinnedLocalBoundsMin"), Shader::EValueType::Float3);
+	case EExternalInput::PreSkinnedLocalBoundsMax: return FExternalInputDescription(TEXT("PreSkinnedLocalBoundsMax"), Shader::EValueType::Float3);
+
 	case EExternalInput::ViewportUV: return FExternalInputDescription(TEXT("ViewportUV"), Shader::EValueType::Float2);
 	case EExternalInput::PixelPosition: return FExternalInputDescription(TEXT("PixelPosition"), Shader::EValueType::Float2);
 	case EExternalInput::ViewSize: return FExternalInputDescription(TEXT("ViewSize"), Shader::EValueType::Float2);
@@ -89,6 +94,7 @@ FExternalInputDescription GetExternalInputDescription(EExternalInput Input)
 	case EExternalInput::TemporalSampleOffset: return FExternalInputDescription(TEXT("TemporalSampleOffset"), Shader::EValueType::Float2);
 	case EExternalInput::PreExposure: return FExternalInputDescription(TEXT("PreExposure"), Shader::EValueType::Float1);
 	case EExternalInput::RcpPreExposure: return FExternalInputDescription(TEXT("RcpPreExposure"), Shader::EValueType::Float1);
+	case EExternalInput::EyeAdaptation: return FExternalInputDescription(TEXT("EyeAdaptation"), Shader::EValueType::Float1);
 	case EExternalInput::RuntimeVirtualTextureOutputLevel: return FExternalInputDescription(TEXT("RuntimeVirtualTextureOutputLevel"), Shader::EValueType::Float1);
 	case EExternalInput::RuntimeVirtualTextureOutputDerivative: return FExternalInputDescription(TEXT("RuntimeVirtualTextureOutputDerivative"), Shader::EValueType::Float2);
 	case EExternalInput::RuntimeVirtualTextureMaxLevel: return FExternalInputDescription(TEXT("RuntimeVirtualTextureMaxLevel"), Shader::EValueType::Float1);
@@ -261,6 +267,11 @@ void FExpressionExternalInput::EmitValueShader(FEmitContext& Context, FEmitScope
 		case EExternalInput::WorldVertexNormal: Code = TEXT("Parameters.TangentToWorld[2]"); break;
 		case EExternalInput::WorldVertexTangent: Code = TEXT("Parameters.TangentToWorld[0]"); break;
 
+		case EExternalInput::PreSkinnedPosition: Code = TEXT("Parameters.PreSkinnedPosition"); break;
+		case EExternalInput::PreSkinnedNormal: Code = TEXT("Parameters.PreSkinnedNormal"); break;
+		case EExternalInput::PreSkinnedLocalBoundsMin: Code = TEXT("GetPrimitiveData(Parameters).PreSkinnedLocalBoundsMin"); break;
+		case EExternalInput::PreSkinnedLocalBoundsMax: Code = TEXT("GetPrimitiveData(Parameters).PreSkinnedLocalBoundsMax"); break;
+
 		case EExternalInput::ViewportUV: Code = TEXT("GetViewportUV(Parameters)"); break;
 		case EExternalInput::PixelPosition: Code = TEXT("GetPixelPosition(Parameters)"); break;
 		case EExternalInput::ViewSize: Code = TEXT("View.ViewSizeAndInvSize.xy"); break;
@@ -274,6 +285,7 @@ void FExpressionExternalInput::EmitValueShader(FEmitContext& Context, FEmitScope
 		case EExternalInput::TemporalSampleOffset: Code = TEXT("View.TemporalAAParams.zw"); break;
 		case EExternalInput::PreExposure: Code = TEXT("View.PreExposure.x"); break;
 		case EExternalInput::RcpPreExposure: Code = TEXT("View.OneOverPreExposure.x"); break;
+		case EExternalInput::EyeAdaptation: Code = TEXT("EyeAdaptationLookup()"); Context.MaterialCompilationOutput->bUsesEyeAdaptation = true; break;
 		case EExternalInput::RuntimeVirtualTextureOutputLevel:  Code = TEXT("View.RuntimeVirtualTextureMipLevel.x"); break;
 		case EExternalInput::RuntimeVirtualTextureOutputDerivative: Code = TEXT("View.RuntimeVirtualTextureMipLevel.zw"); break;
 		case EExternalInput::RuntimeVirtualTextureMaxLevel:  Code = TEXT("View.RuntimeVirtualTextureMipLevel.y"); break;
@@ -529,6 +541,27 @@ void FExpressionParameter::EmitValuePreshader(FEmitContext& Context, FEmitScope&
 	}
 }
 
+bool FExpressionTextureSize::PrepareValue(FEmitContext& Context, FEmitScope& Scope, const FRequestedType& RequestedType, FPrepareValueResult& OutResult) const
+{
+	FEmitData& EmitData = Context.FindData<FEmitData>();
+	if (EmitData.CachedExpressionData)
+	{
+		EmitData.CachedExpressionData->ReferencedTextures.AddUnique(Texture);
+	}
+
+	return OutResult.SetType(Context, RequestedType, EExpressionEvaluation::Preshader, Shader::EValueType::Float2);
+}
+
+void FExpressionTextureSize::EmitValuePreshader(FEmitContext& Context, FEmitScope& Scope, const FRequestedType& RequestedType, FEmitValuePreshaderResult& OutResult) const
+{
+	const int32 TextureIndex = Context.Material->GetReferencedTextures().Find(Texture);
+	check(TextureIndex != INDEX_NONE);
+	
+	Context.PreshaderStackPosition++;
+	OutResult.Type = Shader::EValueType::Float2;
+	OutResult.Preshader.WriteOpcode(Shader::EPreshaderOpcode::TextureSize).Write<FMemoryImageMaterialParameterInfo>(ParameterInfo).Write(TextureIndex);
+}
+
 bool FExpressionFunctionCall::PrepareValue(FEmitContext& Context, FEmitScope& Scope, const FRequestedType& RequestedType, FPrepareValueResult& OutResult) const
 {
 	FEmitData& EmitMaterialData = Context.FindData<FEmitData>();
diff --git a/Engine/Source/Runtime/Engine/Private/Materials/MaterialHLSLTree.h b/Engine/Source/Runtime/Engine/Private/Materials/MaterialHLSLTree.h
--- a/Engine/Source/Runtime/Engine/Private/Materials/MaterialHLSLTree.h
+++ b/Engine/Source/Runtime/Engine/Private/Materials/MaterialHLSLTree.h
@@ -74,6 +74,11 @@ enum class EExternalInput : uint8
 	WorldNormal,
 	WorldReflection,
 
+	PreSkinnedPosition,
+	PreSkinnedNormal,
+	PreSkinnedLocalBoundsMin,
+	PreSkinnedLocalBoundsMax,
+
 	ViewportUV,
 	PixelPosition,
 	ViewSize,
@@ -86,6 +91,7 @@ enum class EExternalInput : uint8
 	TemporalSampleOffset,
 	PreExposure,
 	RcpPreExposure,
+	EyeAdaptation,
 	RuntimeVirtualTextureOutputLevel,
 	RuntimeVirtualTextureOutputDerivative,
 	RuntimeVirtualTextureMaxLevel,
@@ -243,6 +249,20 @@ public:
 	virtual void EmitValuePreshader(FEmitContext& Context, FEmitScope& Scope, const FRequestedType& RequestedType, FEmitValuePreshaderResult& OutResult) const override;
 };
 
+class FExpressionTextureSize : public FExpression
+{
+public:
+	explicit FExpressionTextureSize(const FMaterialParameterInfo& InParameterInfo, UTexture* InTexture)
+		: ParameterInfo(InParameterInfo), Texture(InTexture)
+	{}
+
+	FMaterialParameterInfo ParameterInfo;
+	UTexture* Texture;
+
+	virtual bool PrepareValue(FEmitContext& Context, FEmitScope& Scope, const FRequestedType& RequestedType, FPrepareValueResult& OutResult) const override;
+	virtual void EmitValuePreshader(FEmitContext& Context, FEmitScope& Scope, const FRequestedType& RequestedType, FEmitValuePreshaderResult& OutResult) const override;
+};
+
 class FExpressionFunctionCall : public FExpressionForward
 {
 public:
diff --git a/Engine/Source/Runtime/Engine/Public/HLSLTree/HLSLTree.h b/Engine/Source/Runtime/Engine/Public/HLSLTree/HLSLTree.h
--- a/Engine/Source/Runtime/Engine/Public/HLSLTree/HLSLTree.h
+++ b/Engine/Source/Runtime/Engine/Public/HLSLTree/HLSLTree.h
@@ -280,7 +280,6 @@ public:
 
 	int32 GetNumComponents() const;
 	FRequestedType GetRequestedType() const;
-	FPreparedType GetTypeForRequest(const FRequestedType& RequestedType) const;
 	Shader::FType GetType() const;
 	const TCHAR* GetName() const { return GetType().GetName(); }
 	bool IsStruct() const { return !IsVoid() && StructType != nullptr; }
@@ -295,7 +294,6 @@ public:
 	FPreparedComponent GetMergedComponent() const;
 
 	FPreparedComponent GetComponent(int32 Index) const;
-	EExpressionEvaluation GetComponentEvaluation(int32 Index) const { return GetComponent(Index).Evaluation; }
 	Shader::FComponentBounds GetComponentBounds(int32 Index) const;
 
 	void SetComponent(int32 Index, const FPreparedComponent& InComponent);
@@ -539,6 +537,7 @@ public:
 	FExpression* NewLessEqual(FExpression* Lhs, FExpression* Rhs) { return NewBinaryOp(EOperation::LessEqual, Lhs, Rhs); }
 	FExpression* NewGreaterEqual(FExpression* Lhs, FExpression* Rhs) { return NewBinaryOp(EOperation::GreaterEqual, Lhs, Rhs); }
 
+	FExpression* NewLog(FExpression* Input);
 	FExpression* NewPow2(FExpression* Input) { return NewMul(Input, Input); }
 	FExpression* NewCross(FExpression* Lhs, FExpression* Rhs);
 	FExpression* NewDot(FExpression* Lhs, FExpression* Rhs) { return NewSum(NewMul(Lhs, Rhs)); }
diff --git a/Engine/Source/Runtime/Engine/Public/HLSLTree/HLSLTreeTypes.h b/Engine/Source/Runtime/Engine/Public/HLSLTree/HLSLTreeTypes.h
--- a/Engine/Source/Runtime/Engine/Public/HLSLTree/HLSLTreeTypes.h
+++ b/Engine/Source/Runtime/Engine/Public/HLSLTree/HLSLTreeTypes.h
@@ -55,6 +55,11 @@ EExpressionEvaluation CombineEvaluations(EExpressionEvaluation Lhs, EExpressionE
 EExpressionEvaluation MakeLoopEvaluation(EExpressionEvaluation Evaluation);
 EExpressionEvaluation MakeNonLoopEvaluation(EExpressionEvaluation Evaluation);
 
+inline bool IsConstantEvaluation(EExpressionEvaluation Evaluation)
+{
+	return Evaluation == EExpressionEvaluation::Constant || Evaluation == EExpressionEvaluation::ConstantZero;
+}
+
 inline bool IsLoopEvaluation(EExpressionEvaluation Evaluation)
 {
 	return Evaluation == EExpressionEvaluation::PreshaderLoop || Evaluation == EExpressionEvaluation::ConstantLoop;
diff --git a/Engine/Source/Runtime/Engine/Public/MaterialHLSLGenerator.h b/Engine/Source/Runtime/Engine/Public/MaterialHLSLGenerator.h
--- a/Engine/Source/Runtime/Engine/Public/MaterialHLSLGenerator.h
+++ b/Engine/Source/Runtime/Engine/Public/MaterialHLSLGenerator.h
@@ -30,16 +30,17 @@ struct FFunctionExpressionInput;
 class ITargetPlatform;
 enum class EMaterialParameterType : uint8;
 
-namespace UE
-{
-namespace HLSLTree
+namespace UE::HLSLTree
 {
 struct FSwizzleParameters;
 }
-namespace Shader
+namespace UE::HLSLTree::Material
 {
-struct FStructType;
+enum class EExternalInput : uint8;
 }
+namespace UE::Shader
+{
+struct FStructType;
 }
 
 enum class EMaterialNewScopeFlag : uint8
@@ -98,6 +99,7 @@ public:
 
 	UE::HLSLTree::FExpression* NewConstant(const UE::Shader::FValue& Value);
 	UE::HLSLTree::FExpression* NewTexCoord(int32 Index);
+	UE::HLSLTree::FExpression* NewExternalInput(UE::HLSLTree::Material::EExternalInput Input);
 	UE::HLSLTree::FExpression* NewSwizzle(const UE::HLSLTree::FSwizzleParameters& Params, UE::HLSLTree::FExpression* Input);
 
 	const UE::Shader::FTextureValue* AcquireTextureValue(const UE::Shader::FTextureValue& Value);
