Fix missing serialization backwards compat path

diff --git a/Engine/Source/Runtime/Core/Public/Misc/NetworkVersion.h b/Engine/Source/Runtime/Core/Public/Misc/NetworkVersion.h
--- a/Engine/Source/Runtime/Core/Public/Misc/NetworkVersion.h
+++ b/Engine/Source/Runtime/Core/Public/Misc/NetworkVersion.h
@@ -56,6 +56,7 @@ enum EEngineNetworkVersionHistory
 	HISTORY_SERIALIZE_DOUBLE_VECTORS_AS_DOUBLES	= 22,// Bump version to support net serialization of double vector types
 	HISTORY_PACKED_VECTOR_LWC_SUPPORT = 23,			// Bump version to support quantized LWC FVector net serialization
 	HISTORY_PAWN_REMOTEVIEWPITCH = 24,				// Bump version to support serialization changes to RemoteViewPitch
+	HISTORY_REPMOVE_SERVERFRAME_AND_HANDLE = 25,	// Bump version to support serialization changes to RepMove so we can get the serverframe and physics handle associated with the object
 	// New history items go above here.
 
 	HISTORY_ENGINENETVERSION_PLUS_ONE,
diff --git a/Engine/Source/Runtime/Engine/Private/EngineTypes.cpp b/Engine/Source/Runtime/Engine/Private/EngineTypes.cpp
--- a/Engine/Source/Runtime/Engine/Private/EngineTypes.cpp
+++ b/Engine/Source/Runtime/Engine/Private/EngineTypes.cpp
@@ -608,12 +608,13 @@ FRepMovement::FRepMovement()
 bool FRepMovement::NetSerialize(FArchive& Ar, class UPackageMap* Map, bool& bOutSuccess)
 {
 	// pack bitfield with flags
+	const bool bServerFrameAndHandleSupported = Ar.EngineNetVer() >= HISTORY_REPMOVE_SERVERFRAME_AND_HANDLE;
 	uint8 Flags = (bSimulatedPhysicSleep << 0) | (bRepPhysics << 1) | ((ServerFrame > 0) << 2) | ((ServerPhysicsHandle != INDEX_NONE) << 3);
-	Ar.SerializeBits(&Flags, 4);
+	Ar.SerializeBits(&Flags, bServerFrameAndHandleSupported ? 4 : 2);
 	bSimulatedPhysicSleep = ( Flags & ( 1 << 0 ) ) ? 1 : 0;
 	bRepPhysics = ( Flags & ( 1 << 1 ) ) ? 1 : 0;
-	const bool bRepServerFrame = (Flags & (1 << 2)) ? 1 : 0;
-	const bool bRepServerHandle = (Flags & (1 << 3)) ? 1 : 0;
+	const bool bRepServerFrame = (Flags & (1 << 2) && bServerFrameAndHandleSupported) ? 1 : 0;
+	const bool bRepServerHandle = (Flags & (1 << 3) && bServerFrameAndHandleSupported) ? 1 : 0;
 
 	bOutSuccess = true;
 
