Fix asan heap-use-after-free in GameViewportClient

diff --git a/Engine/Source/Runtime/Engine/Private/GameViewportClient.cpp b/Engine/Source/Runtime/Engine/Private/GameViewportClient.cpp
--- a/Engine/Source/Runtime/Engine/Private/GameViewportClient.cpp
+++ b/Engine/Source/Runtime/Engine/Private/GameViewportClient.cpp
@@ -339,6 +339,9 @@ UGameViewportClient::~UGameViewportClient()
 {
 	if (EngineShowFlags.Collision)
 	{
+		// Clear ref to world as it may be GC'd & we don't want to use it in toggle below.
+		World = nullptr;
+
 		EngineShowFlags.SetCollision(false);
 		ToggleShowCollision();
 	}
