Prevent replay from recording multicast rpcs forwarded from only relevant to owner actors

diff --git a/Engine/Source/Runtime/Engine/Private/DemoNetDriver.cpp b/Engine/Source/Runtime/Engine/Private/DemoNetDriver.cpp
--- a/Engine/Source/Runtime/Engine/Private/DemoNetDriver.cpp
+++ b/Engine/Source/Runtime/Engine/Private/DemoNetDriver.cpp
@@ -1309,10 +1309,15 @@ void UDemoNetDriver::ProcessRemoteFunction(AActor* Actor, UFunction* Function, v
 
 			if (bRecordRPC)
 			{
-				// Handle role swapping if this is a client-recorded replay.
-				FScopedActorRoleSwap RoleSwap(Actor);
+				const bool bIsRelevant = !Actor->bOnlyRelevantToOwner || (Actor->GetNetDriverName() == NetDriverName);
 
-				InternalProcessRemoteFunction(Actor, SubObject, ClientConnections[0], Function, Parameters, OutParms, Stack, IsServer());
+				if (bIsRelevant)
+				{
+					// Handle role swapping if this is a client-recorded replay.
+					FScopedActorRoleSwap RoleSwap(Actor);
+
+					InternalProcessRemoteFunction(Actor, SubObject, ClientConnections[0], Function, Parameters, OutParms, Stack, IsServer());
+				}
 			}
 		}
 	}
