SimpleElementShader: Remove DepthMinValue clamping

diff --git a/Engine/Shaders/Private/SimpleElementVertexShader.usf b/Engine/Shaders/Private/SimpleElementVertexShader.usf
--- a/Engine/Shaders/Private/SimpleElementVertexShader.usf
+++ b/Engine/Shaders/Private/SimpleElementVertexShader.usf
@@ -20,7 +20,6 @@ float3 TransformTilePosition;
 #endif
 
 float SwitchVerticalAxis;
-float DepthMinValue;
 
 void Main(
 	in float4 InRelativePosition : ATTRIBUTE0,
@@ -38,7 +37,7 @@ void Main(
 	FLWCVector4 InPosition = MakeLWCVector4(InTilePosition, InRelativePosition);
 	FLWCInverseMatrix WorldToClip = MakeLWCInverseMatrix(TransformTilePosition, Transform);
 	OutPosition = LWCMultiply(InPosition, WorldToClip);
-	OutPosition.z = max(OutPosition.z, DepthMinValue * OutPosition.w);
+	OutPosition.z = max(OutPosition.z, 0);
 #else
 	OutPosition = mul(InRelativePosition, Transform);
 #endif
diff --git a/Engine/Source/Runtime/Engine/Private/BatchedElements.cpp b/Engine/Source/Runtime/Engine/Private/BatchedElements.cpp
--- a/Engine/Source/Runtime/Engine/Private/BatchedElements.cpp
+++ b/Engine/Source/Runtime/Engine/Private/BatchedElements.cpp
@@ -708,8 +708,7 @@ void FBatchedElements::PrepareShaders(
 		}
 
 		// Set the simple element vertex shader parameters
-		const bool bIsPerspectiveProjection = View ? View->IsPerspectiveProjection() : true;
-		VertexShader->SetParameters(RHICmdList, ViewMatrices, bIsPerspectiveProjection, bSwitchVerticalAxis);
+		VertexShader->SetParameters(RHICmdList, ViewMatrices, bSwitchVerticalAxis);
 	}
 }
 
diff --git a/Engine/Source/Runtime/Engine/Private/SimpleElementShaders.cpp b/Engine/Source/Runtime/Engine/Private/SimpleElementShaders.cpp
--- a/Engine/Source/Runtime/Engine/Private/SimpleElementShaders.cpp
+++ b/Engine/Source/Runtime/Engine/Private/SimpleElementShaders.cpp
@@ -20,7 +20,6 @@ FSimpleElementVS::FSimpleElementVS(const ShaderMetaType::CompiledShaderInitializ
 	RelativeTransform.Bind(Initializer.ParameterMap,TEXT("Transform"), SPF_Mandatory);
 	TransformTilePosition.Bind(Initializer.ParameterMap, TEXT("TransformTilePosition"), SPF_Optional); // TransformTilePosition may be optimized out if LWC is disabled
 	SwitchVerticalAxis.Bind(Initializer.ParameterMap,TEXT("SwitchVerticalAxis"), SPF_Optional);
-	DepthMinValue.Bind(Initializer.ParameterMap, TEXT("DepthMinValue"), SPF_Optional);
 }
 
 void FSimpleElementVS::SetParameters(FRHICommandList& RHICmdList, const FMatrix& WorldToClipMatrix, bool bSwitchVerticalAxis)
@@ -28,18 +27,13 @@ void FSimpleElementVS::SetParameters(FRHICommandList& RHICmdList, const FMatrix&
 	SetShaderValue(RHICmdList, RHICmdList.GetBoundVertexShader(), RelativeTransform, FMatrix44f(WorldToClipMatrix));
 	SetShaderValue(RHICmdList, RHICmdList.GetBoundVertexShader(), TransformTilePosition, FVector3f::ZeroVector);
 	SetShaderValue(RHICmdList, RHICmdList.GetBoundVertexShader(), SwitchVerticalAxis, bSwitchVerticalAxis ? -1.0f : 1.0f);
-	SetShaderValue(RHICmdList, RHICmdList.GetBoundVertexShader(), DepthMinValue, -1.0f);
 }
 
-void FSimpleElementVS::SetParameters(FRHICommandList& RHICmdList, const FRelativeViewMatrices& Matrices, bool bPerspectiveProjection, bool bSwitchVerticalAxis)
+void FSimpleElementVS::SetParameters(FRHICommandList& RHICmdList, const FRelativeViewMatrices& Matrices, bool bSwitchVerticalAxis)
 {
 	SetShaderValue(RHICmdList, RHICmdList.GetBoundVertexShader(), RelativeTransform, Matrices.RelativeWorldToClip);
 	SetShaderValue(RHICmdList, RHICmdList.GetBoundVertexShader(), TransformTilePosition, Matrices.TilePosition);
 	SetShaderValue(RHICmdList, RHICmdList.GetBoundVertexShader(), SwitchVerticalAxis, bSwitchVerticalAxis ? -1.0f : 1.0f);
-
-	// Clamp Z-values to far plane for ortho-projection, we don't want any far-plane clipping
-	static_assert(bool(ERHIZBuffer::IsInverted), "FSimpleElementVS far plane clipping assumes inverted depth buffer");
-	SetShaderValue(RHICmdList, RHICmdList.GetBoundVertexShader(), DepthMinValue, bPerspectiveProjection ? -1.0f : 0.0f);
 }
 
 void FSimpleElementVS::ModifyCompilationEnvironment(const FGlobalShaderPermutationParameters& Parameters, FShaderCompilerEnvironment& OutEnvironment)
diff --git a/Engine/Source/Runtime/Engine/Public/SimpleElementShaders.h b/Engine/Source/Runtime/Engine/Public/SimpleElementShaders.h
--- a/Engine/Source/Runtime/Engine/Public/SimpleElementShaders.h
+++ b/Engine/Source/Runtime/Engine/Public/SimpleElementShaders.h
@@ -30,7 +30,7 @@ public:
 	FSimpleElementVS() {}
 
 	void SetParameters(FRHICommandList& RHICmdList, const FMatrix& WorldToClipMatrix, bool bSwitchVerticalAxis = false);
-	void SetParameters(FRHICommandList& RHICmdList, const FRelativeViewMatrices& Matrices, bool bPerspectiveProjection, bool bSwitchVerticalAxis = false);
+	void SetParameters(FRHICommandList& RHICmdList, const FRelativeViewMatrices& Matrices, bool bSwitchVerticalAxis = false);
 
 	//virtual bool Serialize(FArchive& Ar) override;
 
@@ -40,7 +40,6 @@ private:
 	LAYOUT_FIELD(FShaderParameter, RelativeTransform);
 	LAYOUT_FIELD(FShaderParameter, TransformTilePosition);
 	LAYOUT_FIELD(FShaderParameter, SwitchVerticalAxis);
-	LAYOUT_FIELD(FShaderParameter, DepthMinValue);
 };
 
 /**
