WorldPartition:
- Always Create UWorldPartitionSubsystem/UDataLayerSubsystem/UHLODSubsystem
- UWorldPartition:Initialize/Uninitialize moved into ULevel

#rnx

diff --git a/Engine/Source/Runtime/Engine/Classes/Engine/Level.h b/Engine/Source/Runtime/Engine/Classes/Engine/Level.h
--- a/Engine/Source/Runtime/Engine/Classes/Engine/Level.h
+++ b/Engine/Source/Runtime/Engine/Classes/Engine/Level.h
@@ -904,9 +904,10 @@ public:
 	 * Flag this level instance for destruction.
 	 * This is called by UWorld::CleanupWorld to flag the level and its owned packages for destruction.
 	 *
-	 * @param bCleanupResources Whether to also remove RF_Standalone flags in the editor on level-specific subassets so that they will be garbage collected
+	 * @param bCleanupResources Whether to uninit anything that was initialized through OnLevelLoaded()
+	 * @param bUnloadFromEditor Whether to also remove RF_Standalone flags in the editor on level-specific subassets so that they can be garbage collected
 	 */
-	ENGINE_API void CleanupLevel(bool bCleanupResources);
+	ENGINE_API void CleanupLevel(bool bCleanupResources = true, bool bUnloadFromEditor = true);
 
 	/**
 	 * Cleans all references from and to this level that may be preventing it from being Garbage Collected
diff --git a/Engine/Source/Runtime/Engine/Classes/Engine/World.h b/Engine/Source/Runtime/Engine/Classes/Engine/World.h
--- a/Engine/Source/Runtime/Engine/Classes/Engine/World.h
+++ b/Engine/Source/Runtime/Engine/Classes/Engine/World.h
@@ -1002,6 +1002,18 @@ struct ENGINE_API FScopedLoadAllExternalObjects
 	FName PackageName;
 };
 
+/**
+ * Helper class allows UWorldPartition to broadcast UWorld events 
+ */
+struct FWorldPartitionEvents
+{
+	friend UWorldPartition;
+
+private:
+	static void BroadcastWorldPartitionInitialized(UWorld* InWorld, UWorldPartition* InWorldPartition);
+	static void BroadcastWorldPartitionUninitialized(UWorld* InWorld, UWorldPartition* InWorldPartition);
+};
+
 /** 
  * The World is the top level object representing a map or a sandbox in which Actors and Components will exist and be rendered.  
  *
@@ -1653,6 +1665,28 @@ private:
 	/** Broadcasted on UWorld::BeginTearingDown */
 	FOnBeginTearingDownEvent BeginTearingDownEvent;
 
+	/** Broadcasted when WorldPartition gets initialized */
+	DECLARE_EVENT_OneParam(UWorld, FWorldPartitionInitializedEvent, UWorldPartition*);
+	
+	FWorldPartitionInitializedEvent OnWorldPartitionInitializedEvent;
+
+	void BroadcastWorldPartitionInitialized(UWorldPartition* InWorldPartition)
+	{
+		OnWorldPartitionInitializedEvent.Broadcast(InWorldPartition);
+	}
+
+	/** Broadcasted when WorldPartition gets uninitialized */
+	DECLARE_EVENT_OneParam(UWorld, FWorldPartitionUninitializedEvent, UWorldPartition*);
+
+	FWorldPartitionUninitializedEvent OnWorldPartitionUninitializedEvent;
+
+	void BroadcastWorldPartitionUninitialized(UWorldPartition* InWorldPartition)
+	{
+		OnWorldPartitionUninitializedEvent.Broadcast(InWorldPartition);
+	}
+
+	friend FWorldPartitionEvents;
+
 #if WITH_EDITOR
 
 	/** Broadcasts that selected levels have changed. */
@@ -2691,6 +2725,9 @@ public:
 
 		return false;
 	}
+
+	FWorldPartitionInitializedEvent& OnWorldPartitionInitialized() { return OnWorldPartitionInitializedEvent; }
+	FWorldPartitionUninitializedEvent& OnWorldPartitionUninitialized() { return OnWorldPartitionUninitializedEvent; }
 		
 	/**
 	 * Returns the current levels BSP model.
@@ -3904,16 +3941,6 @@ public:
 		return (OwningGameInstance ? OwningGameInstance->GetLatentActionManager() : LatentActionManager);
 	}
 
-	/**
-	 * Reinitialize all subsystems
-	 */
-	void ReinitializeSubSystems()
-	{
-		SubsystemCollection.Deinitialize();
-		SubsystemCollection.Initialize(this);
-		PostInitializeSubsystems();
-	}
-
 	/**
 	 * Get a Subsystem of specified type
 	 */
diff --git a/Engine/Source/Runtime/Engine/Private/Level.cpp b/Engine/Source/Runtime/Engine/Private/Level.cpp
--- a/Engine/Source/Runtime/Engine/Private/Level.cpp
+++ b/Engine/Source/Runtime/Engine/Private/Level.cpp
@@ -78,6 +78,7 @@ Level.cpp: Level-related functions
 #include "UObject/MetaData.h"
 #endif
 #include "WorldPartition/WorldPartition.h"
+#include "WorldPartition/WorldPartitionSubsystem.h"
 #include "Engine/LevelStreaming.h"
 #include "LevelUtils.h"
 #include "Components/ModelComponent.h"
@@ -419,11 +420,14 @@ void ULevel::AddReferencedObjects(UObject* InThis, FReferenceCollector& Collecto
 	Super::AddReferencedObjects( This, Collector );
 }
 
-void ULevel::CleanupLevel(bool bCleanupResources)
+void ULevel::CleanupLevel(bool bCleanupResources, bool bUnloadFromEditor)
 {
-	if (bCleanupResources)
+	OnCleanupLevel.Broadcast();
+
+#if WITH_EDITOR
+	if (bUnloadFromEditor)
 	{
-		OnCleanupLevel.Broadcast();
+		check(bCleanupResources);
 		// if the level contains any actor with an external package, clear their metadata standalone flag so that the packages can be properly unloaded.
 		// Do so for any actors outered to the level and not just from the level actors array
 		ForEachObjectWithOuter(this, [](UObject* InObject)
@@ -439,6 +443,15 @@ void ULevel::CleanupLevel(bool bCleanupResources)
 			}
 		}, false);
 	}
+#endif
+
+	if (bCleanupResources)
+	{
+		if (UWorldPartition* WorldPartition = GetWorldPartition(); WorldPartition && WorldPartition->IsInitialized())
+		{
+			WorldPartition->Uninitialize();
+		}
+	}
 }
 
 void ULevel::CleanupReferences()
@@ -2399,6 +2412,16 @@ void ULevel::OnLevelLoaded()
 #if WITH_EDITOR
 	FixupActorFolders();
 #endif
+
+	// 1. Cook commandlet does it's own UWorldPartition::Initialize call in FWorldPartitionCookPackageSplitter::GetGenerateList
+	// 2. Do not Initialize if World doesn't have a UWorldPartitionSubsystem (Known case is when WorldType == EWorldType::Inactive)
+	if (!IsRunningCookCommandlet() && IsPersistentLevel() && GetWorld()->HasSubsystem<UWorldPartitionSubsystem>())
+	{
+		if (UWorldPartition* WorldPartition = GetWorldPartition())
+		{
+			WorldPartition->Initialize(GetWorld(), FTransform::Identity);
+		}
+	}
 }
 
 void ULevel::BuildStreamingData(UWorld* World, ULevel* TargetLevel/*=NULL*/, UTexture2D* UpdateSpecificTextureOnly/*=NULL*/)
diff --git a/Engine/Source/Runtime/Engine/Private/World.cpp b/Engine/Source/Runtime/Engine/Private/World.cpp
--- a/Engine/Source/Runtime/Engine/Private/World.cpp
+++ b/Engine/Source/Runtime/Engine/Private/World.cpp
@@ -336,6 +336,18 @@ FScopedLoadAllExternalObjects::~FScopedLoadAllExternalObjects()
 	}
 }
 
+void FWorldPartitionEvents::BroadcastWorldPartitionInitialized(UWorld* InWorld, UWorldPartition* InWorldPartition)
+{
+	check(InWorld);
+	InWorld->BroadcastWorldPartitionInitialized(InWorldPartition);
+}
+
+void FWorldPartitionEvents::BroadcastWorldPartitionUninitialized(UWorld* InWorld, UWorldPartition* InWorldPartition)
+{
+	check(InWorld);
+	InWorld->BroadcastWorldPartitionUninitialized(InWorldPartition);
+}
+
 
 FAudioDeviceWorldDelegates::FOnWorldRegisteredToAudioDevice FAudioDeviceWorldDelegates::OnWorldRegisteredToAudioDevice;
 
@@ -3291,7 +3303,7 @@ void FLevelStreamingGCHelper::PrepareStreamedOutLevelsForGC()
 			// Make sure that this package has been unloaded after GC pass.
 			LevelPackageNames.Add( LevelPackage->GetFName() );
 
-			Level->CleanupLevel(true /* bCleanupResources */);
+			Level->CleanupLevel();
 
 			// Mark world and all other package subobjects as pending kill
 			// This will destroy metadata objects and any other objects left behind
@@ -4931,7 +4943,7 @@ void UWorld::CleanupWorldInternal(bool bSessionEnded, bool bCleanupResources, UW
 	{
 		return;
 	}
-	bool bWorldChanged = NewWorld != this;
+	const bool bWorldChanged = NewWorld != this;
 	CleanupWorldTag = CleanupWorldGlobalTag;
 
 	UE_LOG(LogWorld, Log, TEXT("UWorld::CleanupWorld for %s, bSessionEnded=%s, bCleanupResources=%s"), *GetName(), bSessionEnded ? TEXT("true") : TEXT("false"), bCleanupResources ? TEXT("true") : TEXT("false"));
@@ -4998,61 +5010,67 @@ void UWorld::CleanupWorldInternal(bool bSessionEnded, bool bCleanupResources, UW
 	}
 
 #if WITH_EDITOR
+	const bool bUnloadFromEditor = GIsEditor && !IsTemplate() && bWorldChanged && bCleanupResources;
+
 	// Clear standalone flag when switching maps in the Editor. This causes resources placed in the map
 	// package to be garbage collected together with the world.
-	if( GIsEditor && !IsTemplate() && this != NewWorld )
+	if (bUnloadFromEditor)
 	{
-		if (bCleanupResources)
+		// Iterate over all objects to find ones that reside in the same package as the world.
+		ForEachObjectWithPackage(GetOutermost(), [this](UObject* CurrentObject)
 		{
-			// Iterate over all objects to find ones that reside in the same package as the world.
-			ForEachObjectWithPackage(GetOutermost(), [this](UObject* CurrentObject)
+			if (CurrentObject != this)
 			{
-				if ( CurrentObject != this )
-				{
-					CurrentObject->ClearFlags( RF_Standalone );
-				}
-				return true;
-			});
+				CurrentObject->ClearFlags(RF_Standalone);
+			}
+			return true;
+		});
 
-			if (WorldType != EWorldType::PIE)
+		if (WorldType != EWorldType::PIE)
+		{
+			if (PersistentLevel && PersistentLevel->MapBuildData)
 			{
-				if (PersistentLevel && PersistentLevel->MapBuildData)
-				{
-					PersistentLevel->MapBuildData->ClearFlags(RF_Standalone);
+				PersistentLevel->MapBuildData->ClearFlags(RF_Standalone);
 
-					// Iterate over all objects to find ones that reside in the same package as the MapBuildData.
-					// Specifically the PackageMetaData
-					ForEachObjectWithPackage(PersistentLevel->MapBuildData->GetOutermost(), [this](UObject* CurrentObject)
+				// Iterate over all objects to find ones that reside in the same package as the MapBuildData.
+				// Specifically the PackageMetaData
+				ForEachObjectWithPackage(PersistentLevel->MapBuildData->GetOutermost(), [this](UObject* CurrentObject)
+				{
+					if (CurrentObject != this)
 					{
-						if (CurrentObject != this)
-						{
-							CurrentObject->ClearFlags(RF_Standalone);
-						}
-						return true;
-					});
-				}
+						CurrentObject->ClearFlags(RF_Standalone);
+					}
+					return true;
+				});
 			}
 		}
+	}
 
-		// Cleanup Persistent level outside of following loop because uninitialized worlds don't have a valid Levels array
-		// StreamingLevels are not initialized.
-		if (PersistentLevel)
-		{
-			PersistentLevel->CleanupLevel(bCleanupResources);
-			PersistentLevel->CleanupReferences();
-		}
+	// Cleanup Persistent level outside of following loop because uninitialized worlds don't have a valid Levels array
+	// StreamingLevels are not initialized.
+	if (PersistentLevel)
+	{
+		PersistentLevel->CleanupLevel(bCleanupResources, bUnloadFromEditor);
+		PersistentLevel->CleanupReferences();
+	}
 
-		if (GetNumLevels() > 1)
+	if (GetNumLevels() > 1)
+	{
+		check(GetLevel(0) == PersistentLevel);
+		for (int32 LevelIndex = 1; LevelIndex < GetNumLevels(); ++LevelIndex)
 		{
-			check(GetLevel(0) == PersistentLevel);
-			for (int32 LevelIndex = 1; LevelIndex < GetNumLevels(); ++LevelIndex)
-			{
-				ULevel* Level = GetLevel(LevelIndex);
-				Level->CleanupLevel(bCleanupResources);
-				Level->CleanupReferences();
-			}
+			ULevel* Level = GetLevel(LevelIndex);
+			Level->CleanupLevel(bCleanupResources, bUnloadFromEditor);
+			Level->CleanupReferences();
 		}
 	}
+
+#else
+	if (PersistentLevel)
+	{
+		check(bCleanupResources);
+		PersistentLevel->CleanupLevel();
+	}
 #endif //WITH_EDITOR
 
 	for (int32 LevelIndex = 0; LevelIndex < GetNumLevels(); ++LevelIndex)
diff --git a/Engine/Source/Runtime/Engine/Private/WorldPartition/DataLayer/DataLayerSubsystem.cpp b/Engine/Source/Runtime/Engine/Private/WorldPartition/DataLayer/DataLayerSubsystem.cpp
--- a/Engine/Source/Runtime/Engine/Private/WorldPartition/DataLayer/DataLayerSubsystem.cpp
+++ b/Engine/Source/Runtime/Engine/Private/WorldPartition/DataLayer/DataLayerSubsystem.cpp
@@ -41,21 +41,6 @@ static FAutoConsoleCommandWithOutputDevice GDumpDataLayersCmd(
 UDataLayerSubsystem::UDataLayerSubsystem()
 {}
 
-bool UDataLayerSubsystem::ShouldCreateSubsystem(UObject* Outer) const
-{
-	if (!Super::ShouldCreateSubsystem(Outer))
-	{
-		return false;
-	}
-
-	if (UWorld* WorldOuter = Cast<UWorld>(Outer))
-	{
-		return WorldOuter->IsPartitionedWorld();
-	}
-
-	return false;
-}
-
 void UDataLayerSubsystem::Initialize(FSubsystemCollectionBase& Collection)
 {
 	Super::Initialize(Collection);
diff --git a/Engine/Source/Runtime/Engine/Private/WorldPartition/HLOD/HLODSubsystem.cpp b/Engine/Source/Runtime/Engine/Private/WorldPartition/HLOD/HLODSubsystem.cpp
--- a/Engine/Source/Runtime/Engine/Private/WorldPartition/HLOD/HLODSubsystem.cpp
+++ b/Engine/Source/Runtime/Engine/Private/WorldPartition/HLOD/HLODSubsystem.cpp
@@ -98,20 +98,6 @@ bool UHLODSubsystem::IsHLODEnabled()
 	return UHLODSubsystem::WorldPartitionHLODEnabled;
 }
 
-bool UHLODSubsystem::ShouldCreateSubsystem(UObject* Outer) const
-{
-	if (!Super::ShouldCreateSubsystem(Outer))
-	{
-		return false;
-	}
-
-	if (UWorld* WorldOuter = Cast<UWorld>(Outer))
-	{
-		return WorldOuter->IsPartitionedWorld();
-	}
-	return false;
-}
-
 void UHLODSubsystem::Initialize(FSubsystemCollectionBase& Collection)
 {
 	// Ensure the WorldPartitionSubsystem gets created before the HLODSubsystem
@@ -123,16 +109,24 @@ void UHLODSubsystem::Initialize(FSubsystemCollectionBase& Collection)
 
 	if (World->IsGameWorld())
 	{
-		UWorldPartition* WorldPartition = GetWorld()->GetWorldPartition();
-		check(WorldPartition);
-
-		WorldPartition->OnWorldPartitionInitialized.AddUObject(this, &UHLODSubsystem::OnWorldPartitionInitialized);
-		WorldPartition->OnWorldPartitionUninitialized.AddUObject(this, &UHLODSubsystem::OnWorldPartitionUninitialized);
+		if (UWorldPartition* WorldPartition = GetWorld()->GetWorldPartition())
+		{
+			GetWorld()->OnWorldPartitionInitialized().AddUObject(this, &UHLODSubsystem::OnWorldPartitionInitialized);
+			GetWorld()->OnWorldPartitionUninitialized().AddUObject(this, &UHLODSubsystem::OnWorldPartitionUninitialized);
 
-		SceneViewExtension = FSceneViewExtensions::NewExtension<FHLODResourcesResidencySceneViewExtension>(World);
+			SceneViewExtension = FSceneViewExtensions::NewExtension<FHLODResourcesResidencySceneViewExtension>(World);
+		}
 	}
 }
 
+void UHLODSubsystem::Deinitialize()
+{
+	Super::Deinitialize();
+
+	GetWorld()->OnWorldPartitionInitialized().RemoveAll(this);
+	GetWorld()->OnWorldPartitionUninitialized().RemoveAll(this);
+}
+
 void UHLODSubsystem::OnWorldPartitionInitialized(UWorldPartition* InWorldPartition)
 {
 	check(InWorldPartition == GetWorld()->GetWorldPartition());
diff --git a/Engine/Source/Runtime/Engine/Private/WorldPartition/WorldPartition.cpp b/Engine/Source/Runtime/Engine/Private/WorldPartition/WorldPartition.cpp
--- a/Engine/Source/Runtime/Engine/Private/WorldPartition/WorldPartition.cpp
+++ b/Engine/Source/Runtime/Engine/Private/WorldPartition/WorldPartition.cpp
@@ -545,7 +545,10 @@ void UWorldPartition::Initialize(UWorld* InWorld, const FTransform& InTransform)
 	}
 #endif
 
+	FWorldPartitionEvents::BroadcastWorldPartitionInitialized(World, this);
+PRAGMA_DISABLE_DEPRECATION_WARNINGS
 	OnWorldPartitionInitialized.Broadcast(this);
+PRAGMA_ENABLE_DEPRECATION_WARNINGS
 }
 
 void UWorldPartition::RegisterStreamingSourceProvider(IWorldPartitionStreamingSourceProvider* StreamingSource)
@@ -612,7 +615,10 @@ void UWorldPartition::Uninitialize()
 
 		InitState = EWorldPartitionInitState::Uninitialized;
 
+		FWorldPartitionEvents::BroadcastWorldPartitionUninitialized(World, this);
+PRAGMA_DISABLE_DEPRECATION_WARNINGS
 		OnWorldPartitionUninitialized.Broadcast(this);
+PRAGMA_ENABLE_DEPRECATION_WARNINGS
 	}
 
 	Super::Uninitialize();
diff --git a/Engine/Source/Runtime/Engine/Private/WorldPartition/WorldPartitionConverter.cpp b/Engine/Source/Runtime/Engine/Private/WorldPartition/WorldPartitionConverter.cpp
--- a/Engine/Source/Runtime/Engine/Private/WorldPartition/WorldPartitionConverter.cpp
+++ b/Engine/Source/Runtime/Engine/Private/WorldPartition/WorldPartitionConverter.cpp
@@ -166,22 +166,7 @@ bool FWorldPartitionConverter::Convert()
 
 	if (bCreatedWorldPartition)
 	{
-		// Manually initialize World Partition subsystems as ReinitializeSubSystems is not fully supported by all subsystems.
-		FSubsystemCollectionBase::ActivateExternalSubsystem(UWorldPartitionSubsystem::StaticClass());
-		FSubsystemCollectionBase::ActivateExternalSubsystem(UDataLayerSubsystem::StaticClass());
-		FSubsystemCollectionBase::ActivateExternalSubsystem(UHLODSubsystem::StaticClass());
-		if (UWorldPartitionSubsystem* WorldPartitionSubsystem = World->GetSubsystem<UWorldPartitionSubsystem>())
-		{
-			WorldPartitionSubsystem->PostInitialize();
-		}
-		if (UDataLayerSubsystem* DataLayerSubsystem = World->GetSubsystem<UDataLayerSubsystem>())
-		{
-			DataLayerSubsystem->PostInitialize();
-		}
-		if (UHLODSubsystem* HLODSubsystem = World->GetSubsystem<UHLODSubsystem>())
-		{
-			HLODSubsystem->PostInitialize();
-		}
+		WorldPartition->Initialize(World, FTransform::Identity);
 	}
 
 	return true;
diff --git a/Engine/Source/Runtime/Engine/Private/WorldPartition/WorldPartitionSubsystem.cpp b/Engine/Source/Runtime/Engine/Private/WorldPartition/WorldPartitionSubsystem.cpp
--- a/Engine/Source/Runtime/Engine/Private/WorldPartition/WorldPartitionSubsystem.cpp
+++ b/Engine/Source/Runtime/Engine/Private/WorldPartition/WorldPartitionSubsystem.cpp
@@ -93,21 +93,6 @@ static FAutoConsoleCommandWithOutputDevice GDumpStreamingSourcesCmd(
 UWorldPartitionSubsystem::UWorldPartitionSubsystem()
 {}
 
-bool UWorldPartitionSubsystem::ShouldCreateSubsystem(UObject* Outer) const
-{
-	if (!Super::ShouldCreateSubsystem(Outer))
-	{
-		return false;
-	}
-
-	if (UWorld* WorldOuter = Cast<UWorld>(Outer))
-	{
-		return WorldOuter->IsPartitionedWorld();
-	}
-
-	return false;
-}
-
 UWorldPartition* UWorldPartitionSubsystem::GetWorldPartition()
 {
 	return GetWorld()->GetWorldPartition();
@@ -128,9 +113,9 @@ bool UWorldPartitionSubsystem::IsRunningConvertWorldPartitionCommandlet() const
 }
 #endif
 
-void UWorldPartitionSubsystem::PostInitialize()
+void UWorldPartitionSubsystem::Initialize(FSubsystemCollectionBase& Collection)
 {
-	Super::PostInitialize();
+	Super::Initialize(Collection);
 
 #if WITH_EDITOR
 	if (IsRunningConvertWorldPartitionCommandlet())
@@ -139,25 +124,8 @@ void UWorldPartitionSubsystem::PostInitialize()
 	}
 #endif
 
-	if (UWorldPartition* WorldPartition = GetWorldPartition())
-	{
-		WorldPartition->Initialize(GetWorld(), FTransform::Identity);
-
-		if (WorldPartition->CanDrawRuntimeHash() && (GetWorld()->GetNetMode() != NM_DedicatedServer))
-		{
-			DrawHandle = UDebugDrawService::Register(TEXT("Game"), FDebugDrawDelegate::CreateUObject(this, &UWorldPartitionSubsystem::Draw));
-		}
-
-		// Enforce some GC settings when using World Partition
-		if (GetWorld()->IsGameWorld())
-		{
-			LevelStreamingContinuouslyIncrementalGCWhileLevelsPendingPurge = GLevelStreamingContinuouslyIncrementalGCWhileLevelsPendingPurge;
-			LevelStreamingForceGCAfterLevelStreamedOut = GLevelStreamingForceGCAfterLevelStreamedOut;
-
-			GLevelStreamingContinuouslyIncrementalGCWhileLevelsPendingPurge = GLevelStreamingContinuouslyIncrementalGCWhileLevelsPendingPurgeForWP;
-			GLevelStreamingForceGCAfterLevelStreamedOut = 0;
-		}
-	}
+	GetWorld()->OnWorldPartitionInitialized().AddUObject(this, &UWorldPartitionSubsystem::OnWorldPartitionInitialized);
+	GetWorld()->OnWorldPartitionUninitialized().AddUObject(this, &UWorldPartitionSubsystem::OnWorldPartitionUninitialized);
 }
 
 void UWorldPartitionSubsystem::Deinitialize()
@@ -170,6 +138,37 @@ void UWorldPartitionSubsystem::Deinitialize()
 	}
 #endif 
 
+	GetWorld()->OnWorldPartitionInitialized().RemoveAll(this);
+	GetWorld()->OnWorldPartitionUninitialized().RemoveAll(this);
+
+	// At this point World Partition should be uninitialized
+	check(!GetWorldPartition() || !GetWorldPartition()->IsInitialized());
+
+	Super::Deinitialize();
+}
+
+void UWorldPartitionSubsystem::OnWorldPartitionInitialized(UWorldPartition* InWorldPartition)
+{
+	check(InWorldPartition == GetWorldPartition());
+	if (InWorldPartition->CanDrawRuntimeHash() && (GetWorld()->GetNetMode() != NM_DedicatedServer))
+	{
+		DrawHandle = UDebugDrawService::Register(TEXT("Game"), FDebugDrawDelegate::CreateUObject(this, &UWorldPartitionSubsystem::Draw));
+	}
+
+	// Enforce some GC settings when using World Partition
+	if (GetWorld()->IsGameWorld())
+	{
+		LevelStreamingContinuouslyIncrementalGCWhileLevelsPendingPurge = GLevelStreamingContinuouslyIncrementalGCWhileLevelsPendingPurge;
+		LevelStreamingForceGCAfterLevelStreamedOut = GLevelStreamingForceGCAfterLevelStreamedOut;
+
+		GLevelStreamingContinuouslyIncrementalGCWhileLevelsPendingPurge = GLevelStreamingContinuouslyIncrementalGCWhileLevelsPendingPurgeForWP;
+		GLevelStreamingForceGCAfterLevelStreamedOut = 0;
+	}
+}
+
+void UWorldPartitionSubsystem::OnWorldPartitionUninitialized(UWorldPartition* InWorldPartition)
+{
+	check(InWorldPartition == GetWorldPartition());
 	if (GetWorld()->IsGameWorld())
 	{
 		GLevelStreamingContinuouslyIncrementalGCWhileLevelsPendingPurge = LevelStreamingContinuouslyIncrementalGCWhileLevelsPendingPurge;
@@ -181,35 +180,28 @@ void UWorldPartitionSubsystem::Deinitialize()
 		UDebugDrawService::Unregister(DrawHandle);
 		DrawHandle.Reset();
 	}
-
-	// During garbage collection, the world partition object can be uninitialized before the subsystem due to unordered calls to BeginDestroy
-	if (UWorldPartition* WorldPartition = GetWorldPartition(); WorldPartition && WorldPartition->IsInitialized())
-	{
-		// Uninitialize registered world partitions
-		WorldPartition->Uninitialize();
-	}
-
-	Super::Deinitialize();
 }
 
 void UWorldPartitionSubsystem::Tick(float DeltaSeconds)
 {
 	Super::Tick(DeltaSeconds);
 
-	UWorldPartition* Partition = GetWorldPartition();
-	Partition->Tick(DeltaSeconds);
-		
-	if (GDrawRuntimeHash3D && Partition->CanDrawRuntimeHash())
+	if (UWorldPartition* Partition = GetWorldPartition())
 	{
-		Partition->DrawRuntimeHash3D();
-	}
+		Partition->Tick(DeltaSeconds);
+
+		if (GDrawRuntimeHash3D && Partition->CanDrawRuntimeHash())
+		{
+			Partition->DrawRuntimeHash3D();
+		}
 
 #if WITH_EDITOR
-	if (!GetWorld()->IsGameWorld())
-	{
-		Partition->DrawRuntimeHashPreview();
-	}
+		if (!GetWorld()->IsGameWorld())
+		{
+			Partition->DrawRuntimeHashPreview();
+		}
 #endif
+	}
 }
 
 ETickableTickType UWorldPartitionSubsystem::GetTickableTickType() const
@@ -224,11 +216,9 @@ TStatId UWorldPartitionSubsystem::GetStatId() const
 
 bool UWorldPartitionSubsystem::IsStreamingCompleted(EWorldPartitionRuntimeCellState QueryState, const TArray<FWorldPartitionStreamingQuerySource>& QuerySources, bool bExactState) const
 {
-	const UWorldPartition* Partition = GetWorldPartition();
-
-	if (!Partition->IsStreamingCompleted(QueryState, QuerySources, bExactState))
+	if (const UWorldPartition* Partition = GetWorldPartition())
 	{
-		return false;
+		return Partition->IsStreamingCompleted(QueryState, QuerySources, bExactState);
 	}
 
 	return true;
@@ -236,14 +226,16 @@ bool UWorldPartitionSubsystem::IsStreamingCompleted(EWorldPartitionRuntimeCellSt
 
 void UWorldPartitionSubsystem::DumpStreamingSources(FOutputDevice& OutputDevice) const
 {
-	const UWorldPartition* WorldPartition = GetWorldPartition();
-	const TArray<FWorldPartitionStreamingSource>* StreamingSources = WorldPartition ? &WorldPartition->GetStreamingSources() : nullptr;
-	if (StreamingSources && (StreamingSources->Num() > 0))
+	if (const UWorldPartition* WorldPartition = GetWorldPartition())
 	{
-		OutputDevice.Logf(TEXT("Streaming Sources:"));
-		for (const FWorldPartitionStreamingSource& StreamingSource : *StreamingSources)
+		const TArray<FWorldPartitionStreamingSource>* StreamingSources = WorldPartition ? &WorldPartition->GetStreamingSources() : nullptr;
+		if (StreamingSources && (StreamingSources->Num() > 0))
 		{
-			OutputDevice.Logf(TEXT("  - %s: %s"), *StreamingSource.Name.ToString(), *StreamingSource.ToString());
+			OutputDevice.Logf(TEXT("Streaming Sources:"));
+			for (const FWorldPartitionStreamingSource& StreamingSource : *StreamingSources)
+			{
+				OutputDevice.Logf(TEXT("  - %s: %s"), *StreamingSource.Name.ToString(), *StreamingSource.ToString());
+			}
 		}
 	}
 }
@@ -258,8 +250,10 @@ void UWorldPartitionSubsystem::UpdateStreamingState()
 	}
 #endif
 
-	UWorldPartition* Partition = GetWorldPartition();
-	Partition->UpdateStreamingState();
+	if (UWorldPartition* Partition = GetWorldPartition())
+	{
+		Partition->UpdateStreamingState();
+	}
 }
 
 void UWorldPartitionSubsystem::Draw(UCanvas* Canvas, class APlayerController* PC)
@@ -288,10 +282,12 @@ void UWorldPartitionSubsystem::Draw(UCanvas* Canvas, class APlayerController* PC
 		const FVector2D CanvasMinimumSize(100.f, 100.f);
 		const FVector2D CanvasMaxScreenSize = FVector2D::Max(MaxScreenRatio*FVector2D(Canvas->ClipX, Canvas->ClipY) - CanvasBottomRightPadding - CurrentOffset, CanvasMinimumSize);
 
-		UWorldPartition* Partition = GetWorldPartition();
-		FVector2D PartitionCanvasSize = FVector2D(CanvasMaxScreenSize.X, CanvasMaxScreenSize.Y);
-		Partition->DrawRuntimeHash2D(Canvas, PartitionCanvasSize, CurrentOffset);
-		CurrentOffset.X = CanvasBottomRightPadding.X;
+		if (UWorldPartition* Partition = GetWorldPartition())
+		{
+			FVector2D PartitionCanvasSize = FVector2D(CanvasMaxScreenSize.X, CanvasMaxScreenSize.Y);
+			Partition->DrawRuntimeHash2D(Canvas, PartitionCanvasSize, CurrentOffset);
+			CurrentOffset.X = CanvasBottomRightPadding.X;
+		}
 	}
 	
 	if (GDrawStreamingPerfs || GDrawRuntimeHash2D)
@@ -315,25 +311,27 @@ void UWorldPartitionSubsystem::Draw(UCanvas* Canvas, class APlayerController* PC
 
 		{
 			FString StatusText;
-			const UWorldPartition* WorldPartition = GetWorldPartition();
-			EWorldPartitionStreamingPerformance StreamingPerformance = WorldPartition->GetStreamingPerformance();
-			switch (StreamingPerformance)
+			if (const UWorldPartition* WorldPartition = GetWorldPartition())
 			{
-			case EWorldPartitionStreamingPerformance::Good:
-				StatusText = TEXT("Good");
-				break;
-			case EWorldPartitionStreamingPerformance::Slow:
-				StatusText = TEXT("Slow");
-				break;
-			case EWorldPartitionStreamingPerformance::Critical:
-				StatusText = TEXT("Critical");
-				break;
-			default:
-				StatusText = TEXT("Unknown");
-				break;
+				EWorldPartitionStreamingPerformance StreamingPerformance = WorldPartition->GetStreamingPerformance();
+				switch (StreamingPerformance)
+				{
+				case EWorldPartitionStreamingPerformance::Good:
+					StatusText = TEXT("Good");
+					break;
+				case EWorldPartitionStreamingPerformance::Slow:
+					StatusText = TEXT("Slow");
+					break;
+				case EWorldPartitionStreamingPerformance::Critical:
+					StatusText = TEXT("Critical");
+					break;
+				default:
+					StatusText = TEXT("Unknown");
+					break;
+				}
+				const FString Text = FString::Printf(TEXT("Streaming Performance: %s (Blocking %s)"), *StatusText, GBlockOnSlowStreaming ? TEXT("Enabled") : TEXT("Disabled"));
+				FWorldPartitionDebugHelper::DrawText(Canvas, Text, GEngine->GetSmallFont(), FColor::White, CurrentOffset);
 			}
-			const FString Text = FString::Printf(TEXT("Streaming Performance: %s (Blocking %s)"), *StatusText, GBlockOnSlowStreaming ? TEXT("Enabled") : TEXT("Disabled"));
-			FWorldPartitionDebugHelper::DrawText(Canvas, Text, GEngine->GetSmallFont(), FColor::White, CurrentOffset);
 		}
 	}
 
@@ -341,26 +339,28 @@ void UWorldPartitionSubsystem::Draw(UCanvas* Canvas, class APlayerController* PC
 	{
 		TRACE_CPUPROFILER_EVENT_SCOPE(UWorldPartitionSubsystem::DrawStreamingSources);
 
-		const UWorldPartition* WorldPartition = GetWorldPartition();
-		const TArray<FWorldPartitionStreamingSource>* StreamingSources = WorldPartition ? &WorldPartition->GetStreamingSources() : nullptr;
-		if (StreamingSources && (StreamingSources->Num() > 0))
+		if (const UWorldPartition* WorldPartition = GetWorldPartition())
 		{
-			FString Title(TEXT("Streaming Sources"));
-			FWorldPartitionDebugHelper::DrawText(Canvas, Title, GEngine->GetSmallFont(), FColor::Yellow, CurrentOffset);
-
-			FVector2D Pos = CurrentOffset;
-			float MaxTextWidth = 0;
-			for (const FWorldPartitionStreamingSource& StreamingSource : *StreamingSources)
-			{
-				FString StreamingSourceDisplay = StreamingSource.Name.ToString();
-				FWorldPartitionDebugHelper::DrawText(Canvas, StreamingSourceDisplay, GEngine->GetSmallFont(), StreamingSource.GetDebugColor(), Pos, &MaxTextWidth);
-			}
-			Pos = CurrentOffset + FVector2D(MaxTextWidth + 10, 0.f);
-			for (const FWorldPartitionStreamingSource& StreamingSource : *StreamingSources)
+			const TArray<FWorldPartitionStreamingSource>& StreamingSources = WorldPartition->GetStreamingSources();
+			if (StreamingSources.Num() > 0)
 			{
-				FWorldPartitionDebugHelper::DrawText(Canvas, *StreamingSource.ToString(), GEngine->GetSmallFont(), FColor::White, Pos);
+				FString Title(TEXT("Streaming Sources"));
+				FWorldPartitionDebugHelper::DrawText(Canvas, Title, GEngine->GetSmallFont(), FColor::Yellow, CurrentOffset);
+
+				FVector2D Pos = CurrentOffset;
+				float MaxTextWidth = 0;
+				for (const FWorldPartitionStreamingSource& StreamingSource : StreamingSources)
+				{
+					FString StreamingSourceDisplay = StreamingSource.Name.ToString();
+					FWorldPartitionDebugHelper::DrawText(Canvas, StreamingSourceDisplay, GEngine->GetSmallFont(), StreamingSource.GetDebugColor(), Pos, &MaxTextWidth);
+				}
+				Pos = CurrentOffset + FVector2D(MaxTextWidth + 10, 0.f);
+				for (const FWorldPartitionStreamingSource& StreamingSource : StreamingSources)
+				{
+					FWorldPartitionDebugHelper::DrawText(Canvas, *StreamingSource.ToString(), GEngine->GetSmallFont(), FColor::White, Pos);
+				}
+				CurrentOffset.Y = Pos.Y;
 			}
-			CurrentOffset.Y = Pos.Y;
 		}
 	}
 
diff --git a/Engine/Source/Runtime/Engine/Public/WorldPartition/DataLayer/DataLayerSubsystem.h b/Engine/Source/Runtime/Engine/Public/WorldPartition/DataLayer/DataLayerSubsystem.h
--- a/Engine/Source/Runtime/Engine/Public/WorldPartition/DataLayer/DataLayerSubsystem.h
+++ b/Engine/Source/Runtime/Engine/Public/WorldPartition/DataLayer/DataLayerSubsystem.h
@@ -25,7 +25,6 @@ public:
 	UDataLayerSubsystem();
 
 	//~ Begin USubsystem Interface.
-	virtual bool ShouldCreateSubsystem(UObject* Outer) const override;
 	virtual void Initialize(FSubsystemCollectionBase& Collection) override;
 	//~ End USubsystem Interface.
 
diff --git a/Engine/Source/Runtime/Engine/Public/WorldPartition/HLOD/HLODSubsystem.h b/Engine/Source/Runtime/Engine/Public/WorldPartition/HLOD/HLODSubsystem.h
--- a/Engine/Source/Runtime/Engine/Public/WorldPartition/HLOD/HLODSubsystem.h
+++ b/Engine/Source/Runtime/Engine/Public/WorldPartition/HLOD/HLODSubsystem.h
@@ -40,8 +40,8 @@ public:
 	virtual ~UHLODSubsystem();
 
 	//~ Begin USubsystem Interface.
-	virtual bool ShouldCreateSubsystem(UObject* Outer) const override;
 	virtual void Initialize(FSubsystemCollectionBase& Collection) override;
+	virtual void Deinitialize() override;
 	//~ End USubsystem Interface.
 
 	void RegisterHLODActor(AWorldPartitionHLOD* InWorldPartitionHLOD);
diff --git a/Engine/Source/Runtime/Engine/Public/WorldPartition/WorldPartition.h b/Engine/Source/Runtime/Engine/Public/WorldPartition/WorldPartition.h
--- a/Engine/Source/Runtime/Engine/Public/WorldPartition/WorldPartition.h
+++ b/Engine/Source/Runtime/Engine/Public/WorldPartition/WorldPartition.h
@@ -87,7 +87,10 @@ public:
 #endif
 
 	DECLARE_MULTICAST_DELEGATE_OneParam(FWorldPartitionInitializeDelegate, UWorldPartition*);
+	
+	UE_DEPRECATED(5.1, "Please use FWorldPartitionInitializedEvent& UWorld::OnWorldPartitionInitialized() instead.")
 	FWorldPartitionInitializeDelegate OnWorldPartitionInitialized;
+	UE_DEPRECATED(5.1, "Please use FWorldPartitionInitializedEvent& UWorld::OnWorldPartitionUninitialized() instead.")
 	FWorldPartitionInitializeDelegate OnWorldPartitionUninitialized;
 
 #if WITH_EDITOR
diff --git a/Engine/Source/Runtime/Engine/Public/WorldPartition/WorldPartitionSubsystem.h b/Engine/Source/Runtime/Engine/Public/WorldPartition/WorldPartitionSubsystem.h
--- a/Engine/Source/Runtime/Engine/Public/WorldPartition/WorldPartitionSubsystem.h
+++ b/Engine/Source/Runtime/Engine/Public/WorldPartition/WorldPartitionSubsystem.h
@@ -23,12 +23,11 @@ public:
 	UWorldPartitionSubsystem();
 
 	//~ Begin USubsystem Interface.
-	virtual bool ShouldCreateSubsystem(UObject* Outer) const override;
+	virtual void Initialize(FSubsystemCollectionBase& Collection) override;
 	virtual void Deinitialize() override;
 	//~ End USubsystem Interface.
 
 	//~ Begin UWorldSubsystem Interface.
-	virtual void PostInitialize() override;
 	virtual void UpdateStreamingState() override;
 	//~ End UWorldSubsystem Interface.
 
@@ -49,6 +48,9 @@ private:
 	bool IsRunningConvertWorldPartitionCommandlet() const;
 #endif
 
+	void OnWorldPartitionInitialized(UWorldPartition* InWorldPartition);
+	void OnWorldPartitionUninitialized(UWorldPartition* InWorldPartition);
+
 	UWorldPartition* GetWorldPartition();
 	const UWorldPartition* GetWorldPartition() const;
 	void Draw(class UCanvas* Canvas, class APlayerController* PC);
