Have UActorChannel::KeyNeedsToReplicate check if the subobject is pending dormancy in order to prevent cases where an unchanged key prevents the channel from going dormant.

diff --git a/Engine/Source/Runtime/Engine/Private/DataChannel.cpp b/Engine/Source/Runtime/Engine/Private/DataChannel.cpp
--- a/Engine/Source/Runtime/Engine/Private/DataChannel.cpp
+++ b/Engine/Source/Runtime/Engine/Private/DataChannel.cpp
@@ -4546,9 +4546,17 @@ bool UActorChannel::ObjectHasReplicator(const TWeakObjectPtr<UObject>& Obj) cons
 #if NET_ENABLE_SUBOBJECT_REPKEYS
 bool UActorChannel::KeyNeedsToReplicate(int32 ObjID, int32 RepKey)
 {
+	if (bPendingDormancy)
+	{
+		// Return true to keep the channel from being stuck unable to go dormant until the key changes
+		return true;
+	}
+
 	int32 &MapKey = SubobjectRepKeyMap.FindOrAdd(ObjID);
 	if (MapKey == RepKey)
+	{ 
 		return false;
+	}
 
 	MapKey = RepKey;
 	PendingObjKeys.Add(ObjID);
