Add shader compilation editor analytics.

# Analytics
- ShadersCompiled: The number of shaders compiled
- ShaderDDCMisses: How many times we missed shader maps in the DDC.  This is a measure of how many times we didn't find Global and DefaultMaterial shaders.
- ShaderDDCHits: How many times did we successfully find shaders in the DDC.
- TimeShaderCompilationWasActive: Wall clock time of how long shader compilation was active across your whole editor session.

# Changes
- Refactor some code into a function GetTotalShadersCompiled to make it easier to access how many shaders we've compiled.
- Added a trace event scoped marker to GetTotalShadersCompiled since we are now calling this every tick.  On average it costs 1.35us so the currently isn't a performance issue.
- Refactored TotalTimeAtLeastOneJobWasInFlight calculation into a function.

diff --git a/Engine/Source/Runtime/Engine/Private/Analytics/EditorAnalyticsSessionSummary.cpp b/Engine/Source/Runtime/Engine/Private/Analytics/EditorAnalyticsSessionSummary.cpp
--- a/Engine/Source/Runtime/Engine/Private/Analytics/EditorAnalyticsSessionSummary.cpp
+++ b/Engine/Source/Runtime/Engine/Private/Analytics/EditorAnalyticsSessionSummary.cpp
@@ -12,6 +12,7 @@
 #include "Kismet2/DebuggerCommands.h"
 #include "ProfilingDebugging/StallDetector.h"
 #include "UObject/Package.h"
+#include "ShaderCompiler.h"
 
 namespace EditorAnalyticsProperties
 {
@@ -32,6 +33,11 @@ namespace EditorAnalyticsProperties
 	static const TAnalyticsProperty<double>  TopStallOverageSeconds = TEXT("TopStallOverageSeconds");
 	static const TAnalyticsProperty<uint32>  ProcessDiagnostics     = TEXT("ProcessDiagnostics"); // Whether some profiling/diagnostic tools are enabled, which could slow down the Editor.
 	static const TAnalyticsProperty<uint32>  SummaryEventVersion    = TEXT("SummaryEventVersion"); // A version number used identify the key/set used. Can be used to compare before/after some changes too since comparing engine version is not always straingth foward.
+
+	static const TAnalyticsProperty<uint32> ShadersCompiled			= TEXT("ShadersCompiled");
+	static const TAnalyticsProperty<uint32> ShaderDDCMisses			= TEXT("ShaderDDCMisses");
+	static const TAnalyticsProperty<uint32> ShaderDDCHits			= TEXT("ShaderDDCHits");
+	static const TAnalyticsProperty<double> TimeShaderCompilationWasActive = TEXT("TimeShaderCompilationWasActive ");
 }
 
 FEditorAnalyticsSessionSummary::FEditorAnalyticsSessionSummary(TSharedPtr<IAnalyticsPropertyStore> Store, uint32 MonitorProcessId)
@@ -60,7 +66,13 @@ FEditorAnalyticsSessionSummary::FEditorAnalyticsSessionSummary(TSharedPtr<IAnaly
 	//    - V3 -> Windows optimization for stall/ensure -> The engine only captures the responsible thread so CRC walks 1 thread rather than all threads.
 	//    - V4 -> Added DirtyPackageCount property.
 	//    - V5 -> Added CPUInfo to the base class (EngineAnalyticsSessionSummary)
-	EditorAnalyticsProperties::SummaryEventVersion.Set(GetStore(), 5);
+	//    - V6 -> Added shader compilation statistics.
+	EditorAnalyticsProperties::SummaryEventVersion.Set(GetStore(), 6);
+
+	EditorAnalyticsProperties::ShadersCompiled.Set(GetStore(), 0);
+	EditorAnalyticsProperties::ShaderDDCMisses.Set(GetStore(), 0);
+	EditorAnalyticsProperties::ShaderDDCHits.Set(GetStore(), 0);
+	EditorAnalyticsProperties::TimeShaderCompilationWasActive.Set(GetStore(), 0.0);
 
 	// Persist the session to disk.
 	GetStore()->Flush();
@@ -111,6 +123,12 @@ bool FEditorAnalyticsSessionSummary::UpdateSessionProgressInternal(bool bCrashin
 	EditorAnalyticsProperties::TotalStallReported.Set(GetStore(), UE::FStallDetectorStats::TotalReportedCount.Get());
 #endif // STALL_DETECTOR
 
+	// Shader stats.
+	EditorAnalyticsProperties::ShadersCompiled.Set(GetStore(), GShaderCompilerStats->GetTotalShadersCompiled());
+	EditorAnalyticsProperties::ShaderDDCMisses.Set(GetStore(), GShaderCompilerStats->GetDDCMisses());
+	EditorAnalyticsProperties::ShaderDDCHits.Set(GetStore(), GShaderCompilerStats->GetDDCHits());
+	EditorAnalyticsProperties::TimeShaderCompilationWasActive.Set(GetStore(), GShaderCompilerStats->GetTimeShaderCompilationWasActive());
+
 	return bShouldPersist;
 }
 
diff --git a/Engine/Source/Runtime/Engine/Private/Materials/MaterialShader.cpp b/Engine/Source/Runtime/Engine/Private/Materials/MaterialShader.cpp
--- a/Engine/Source/Runtime/Engine/Private/Materials/MaterialShader.cpp
+++ b/Engine/Source/Runtime/Engine/Private/Materials/MaterialShader.cpp
@@ -1399,6 +1399,8 @@ void FMaterialShaderMap::LoadFromDerivedDataCache(const FMaterial* Material, con
 
 				// Register in the global map
 				InOutShaderMap->Register(InPlatform);
+
+				GShaderCompilerStats->AddDDCHit(1);
 			}
 			else
 			{
@@ -1406,6 +1408,8 @@ void FMaterialShaderMap::LoadFromDerivedDataCache(const FMaterial* Material, con
 				// We should be build the data later, and we can track that the resource was built there when we push it to the DDC.
 				COOK_STAT(Timer.TrackCyclesOnly());
 				InOutShaderMap = nullptr;
+
+				GShaderCompilerStats->AddDDCMiss(1);
 			}
 
 			if (InOutShaderMap)
diff --git a/Engine/Source/Runtime/Engine/Private/ShaderCompiler/ShaderCompiler.cpp b/Engine/Source/Runtime/Engine/Private/ShaderCompiler/ShaderCompiler.cpp
--- a/Engine/Source/Runtime/Engine/Private/ShaderCompiler/ShaderCompiler.cpp
+++ b/Engine/Source/Runtime/Engine/Private/ShaderCompiler/ShaderCompiler.cpp
@@ -2931,25 +2931,8 @@ void FShaderCompilerStats::WriteStats(FOutputDevice* Ar)
 
 void FShaderCompilerStats::WriteStatSummary()
 {
-	uint32 TotalCompiled = 0;
-
-	{
-		FScopeLock Lock(&CompileStatsLock);
-		const TSparseArray<ShaderCompilerStats>& PlatformStats = GetShaderCompilerStats();
-		for (int32 Platform = 0; Platform < PlatformStats.GetMaxIndex(); ++Platform)
-		{
-			if (PlatformStats.IsValidIndex(Platform))
-			{
-				const ShaderCompilerStats& Stats = PlatformStats[Platform];
-				for (const auto& Pair : Stats)
-				{
-					const FShaderCompilerStats::FShaderStats& SingleStats = Pair.Value;
-
-					TotalCompiled += SingleStats.Compiled;
-				}
-			}
-		}
-	}
+	const uint32 TotalCompiled = GetTotalShadersCompiled();
+	const double TotalTimeAtLeastOneJobWasInFlight = GetTimeShaderCompilationWasActive();
 
 	UE_LOG(LogShaderCompilers, Display, TEXT("=== Shader Compilation stats ==="));
 	UE_LOG(LogShaderCompilers, Display, TEXT("Shaders Compiled: %u"), TotalCompiled);
@@ -2973,18 +2956,6 @@ void FShaderCompilerStats::WriteStatSummary()
 		UE_LOG(LogShaderCompilers, Display, TEXT("Job life time (pending + execution): average %.2f s, max %.2f"), AccumulatedJobLifeTime / JobsCompleted, MaxJobLifeTime);
 	}
 
-	auto SumUpInterval = [](TArray<TInterval<double>>& Accumulator) -> double
-	{
-		double Sum = 0;
-		for (int32 Idx = 0; Idx < Accumulator.Num(); ++Idx)
-		{
-			const TInterval<double>& Existing = Accumulator[Idx];
-			Sum += Existing.Size();
-		}
-		return Sum;
-	};
-
-	double TotalTimeAtLeastOneJobWasInFlight = SumUpInterval(JobLifeTimeIntervals);
 	UE_LOG(LogShaderCompilers, Display, TEXT("Time at least one job was in flight (either pending or executed): %.2f s"), TotalTimeAtLeastOneJobWasInFlight);
 
 	// print stats about the batches
@@ -3076,6 +3047,33 @@ void FShaderCompilerStats::WriteStatSummary()
 	}
 }
 
+uint32 FShaderCompilerStats::GetTotalShadersCompiled()
+{
+	TRACE_CPUPROFILER_EVENT_SCOPE(FShaderCompilerStats::GetTotalShadersCompiled);
+
+	uint32 TotalCompiled = 0;
+
+	{
+		FScopeLock Lock(&CompileStatsLock);
+		const TSparseArray<ShaderCompilerStats>& PlatformStats = GetShaderCompilerStats();
+		for (int32 Platform = 0; Platform < PlatformStats.GetMaxIndex(); ++Platform)
+		{
+			if (PlatformStats.IsValidIndex(Platform))
+			{
+				const ShaderCompilerStats& Stats = PlatformStats[Platform];
+				for (const auto& Pair : Stats)
+				{
+					const FShaderCompilerStats::FShaderStats& SingleStats = Pair.Value;
+
+					TotalCompiled += SingleStats.Compiled;
+				}
+			}
+		}
+	}
+
+	return TotalCompiled;
+}
+
 void FShaderCompilerStats::RegisterLocalWorkerIdleTime(double IdleTime)
 {
 	FScopeLock Lock(&CompileStatsLock);
@@ -3291,6 +3289,38 @@ void FShaderCompilerStats::RegisterCompiledShaders(uint32 NumCompiled, EShaderPl
 	}
 }
 
+void FShaderCompilerStats::AddDDCMiss(uint32 NumMisses)
+{
+	ShaderMapDDCMisses += NumMisses;
+}
+
+uint32 FShaderCompilerStats::GetDDCMisses() const
+{
+	return ShaderMapDDCMisses;
+}
+
+void FShaderCompilerStats::AddDDCHit(uint32 NumHits)
+{
+	ShaderMapDDCHits += NumHits;
+}
+
+uint32 FShaderCompilerStats::GetDDCHits() const
+{
+	return ShaderMapDDCHits;
+}
+
+double FShaderCompilerStats::GetTimeShaderCompilationWasActive()
+{
+	FScopeLock Lock(&CompileStatsLock);
+	double Sum = 0;
+	for (int32 Idx = 0; Idx < JobLifeTimeIntervals.Num(); ++Idx)
+	{
+		const TInterval<double>& Existing = JobLifeTimeIntervals[Idx];
+		Sum += Existing.Size();
+	}
+	return Sum;
+}
+
 FShaderCompilingManager* GShaderCompilingManager = NULL;
 
 bool FShaderCompilingManager::AllTargetPlatformSupportsRemoteShaderCompiling()
@@ -6933,6 +6963,9 @@ void CompileGlobalShaderMap(EShaderPlatform Platform, const ITargetPlatform* Tar
 
 				HandleIndex = 0;
 
+				int32 DDCHits = 0;
+				int32 DDCMisses = 0;
+
 				// Process finished DDC requests.
 				for (const auto& ShaderFilenameDependencies : ShaderMapId.GetShaderFilenameToDependeciesMap())
 				{
@@ -6946,16 +6979,22 @@ void CompileGlobalShaderMap(EShaderPlatform Platform, const ITargetPlatform* Tar
 						COOK_STAT(Timer.AddHit(CachedData.Num()));
 						FMemoryReader MemoryReader(CachedData);
 						GGlobalShaderMap[Platform]->AddSection(FGlobalShaderMapSection::CreateFromArchive(MemoryReader));
+
+						DDCHits++;
 					}
 					else
 					{
 						// it's a miss, but we haven't built anything yet. Save the counting until we actually have it built.
 						COOK_STAT(Timer.TrackCyclesOnly());
 						bShaderMapIsBeingCompiled = true;
+						DDCMisses++;
 					}
 
 					++HandleIndex;
 				}
+
+				GShaderCompilerStats->AddDDCHit(DDCHits);
+				GShaderCompilerStats->AddDDCMiss(DDCMisses);
 			}
 		}
 
diff --git a/Engine/Source/Runtime/Engine/Public/ShaderCompiler.h b/Engine/Source/Runtime/Engine/Public/ShaderCompiler.h
--- a/Engine/Source/Runtime/Engine/Public/ShaderCompiler.h
+++ b/Engine/Source/Runtime/Engine/Public/ShaderCompiler.h
@@ -718,6 +718,13 @@ public:
 	ENGINE_API const TSparseArray<ShaderCompilerStats>& GetShaderCompilerStats() { return CompileStats; }
 	ENGINE_API void WriteStats(class FOutputDevice* Ar = nullptr);
 	ENGINE_API void WriteStatSummary();
+	ENGINE_API uint32 GetTotalShadersCompiled();
+
+	void AddDDCMiss(uint32 NumMisses);
+	uint32 GetDDCMisses() const;
+	void AddDDCHit(uint32 NumHits);
+	uint32 GetDDCHits() const;
+	double GetTimeShaderCompilationWasActive();
 
 	enum class EExecutionType
 	{
@@ -794,6 +801,12 @@ private:
 
 	/** Map of shader names to their compilation timings */
 	TMap<FString, FShaderTimings> ShaderTimings;
+
+	/** Total number of DDC misses on shader maps. */
+	uint32 ShaderMapDDCMisses = 0;
+
+	/** Total number of DDC hits on shader maps. */
+	uint32 ShaderMapDDCHits = 0;
 };
 
 
