Run two passes of PrepareExpression, first pass to determine types, second pass to flag expressions that are used.
This way we can skip flagging constant/zero expressions (which we don't know about during the first pass)

diff --git a/Engine/Source/Editor/MaterialEditor/Private/MaterialEditor.cpp b/Engine/Source/Editor/MaterialEditor/Private/MaterialEditor.cpp
--- a/Engine/Source/Editor/MaterialEditor/Private/MaterialEditor.cpp
+++ b/Engine/Source/Editor/MaterialEditor/Private/MaterialEditor.cpp
@@ -194,7 +194,6 @@ FMatExpressionPreview::FMatExpressionPreview(UMaterialExpression* InExpression)
 	FPlatformMisc::CreateGuid(Id);
 
 	check(InExpression->Material && InExpression->Material->Expressions.Contains(InExpression));
-	ReferencedTextures = InExpression->Material->GetReferencedTextures();
 	SetQualityLevelProperties(GMaxRHIFeatureLevel);
 
 	UMaterial* BaseMaterial = InExpression->Material;
@@ -203,6 +202,14 @@ FMatExpressionPreview::FMatExpressionPreview(UMaterialExpression* InExpression)
 		FMaterialCachedHLSLTree* LocalTree = new FMaterialCachedHLSLTree();
 		LocalTree->GenerateTree(BaseMaterial, nullptr, InExpression);
 		CachedHLSLTree.Reset(LocalTree);
+
+		FMaterialCachedExpressionData* LocalCachedData = new FMaterialCachedExpressionData();
+		LocalCachedData->UpdateForCachedHLSLTree(*LocalTree, nullptr);
+		CachedExpressionData.Reset(LocalCachedData);
+	}
+	else
+	{
+		ReferencedTextures = InExpression->Material->GetReferencedTextures();
 	}
 }
 
@@ -210,6 +217,15 @@ FMatExpressionPreview::~FMatExpressionPreview()
 {
 }
 
+void FMatExpressionPreview::AddReferencedObjects(FReferenceCollector& Collector)
+{
+	Collector.AddReferencedObjects(ReferencedTextures);
+	if (CachedExpressionData)
+	{
+		CachedExpressionData->AddReferencedObjects(Collector);
+	}
+}
+
 bool FMatExpressionPreview::ShouldCache(EShaderPlatform Platform, const FShaderType* ShaderType, const FVertexFactoryType* VertexFactoryType) const
 {
 	if(VertexFactoryType == FindVertexFactoryType(FName(TEXT("FLocalVertexFactory"), FNAME_Find)))
@@ -310,6 +326,18 @@ void FMatExpressionPreview::NotifyCompilationFinished()
 	FMaterialRenderProxy::CacheUniformExpressions_GameThread(true);
 }
 
+TArrayView<const TObjectPtr<UObject>> FMatExpressionPreview::GetReferencedTextures() const
+{
+	if (CachedExpressionData)
+	{
+		// Path for new HLSL translator
+		return MakeArrayView(CachedExpressionData->ReferencedTextures);
+	}
+
+	// Legacy path
+	return MakeArrayView(ReferencedTextures);
+}
+
 const FMaterialCachedHLSLTree* FMatExpressionPreview::GetCachedHLSLTree() const
 {
 	return CachedHLSLTree.Get();
diff --git a/Engine/Source/Editor/MaterialEditor/Private/MaterialEditor.h b/Engine/Source/Editor/MaterialEditor/Private/MaterialEditor.h
--- a/Engine/Source/Editor/MaterialEditor/Private/MaterialEditor.h
+++ b/Engine/Source/Editor/MaterialEditor/Private/MaterialEditor.h
@@ -42,6 +42,7 @@ class UMaterialGraphNode;
 struct FGraphAppearanceInfo;
 class UMaterialFunctionInstance;
 class FMaterialCachedHLSLTree;
+struct FMaterialCachedExpressionData;
 
 /**
  * Class for rendering previews of material expressions in the material editor's linked object viewport.
@@ -66,13 +67,7 @@ public:
 		ReleaseResource();
 	}
 
-	void AddReferencedObjects( FReferenceCollector& Collector )
-	{
-		for (int32 TextureIndex = 0; TextureIndex < ReferencedTextures.Num(); TextureIndex++)
-		{
-			Collector.AddReferencedObject(ReferencedTextures[TextureIndex]);
-		}
-	}
+	void AddReferencedObjects(FReferenceCollector& Collector);
 
 	/**
 	 * Should the shader for this material with the given platform, shader type and vertex 
@@ -86,10 +81,7 @@ public:
 	 */
 	virtual bool ShouldCache(EShaderPlatform Platform, const FShaderType* ShaderType, const FVertexFactoryType* VertexFactoryType) const override;
 
-	virtual TArrayView<const TObjectPtr<UObject>> GetReferencedTextures() const override
-	{
-		return MakeArrayView(ReferencedTextures);
-	}
+	virtual TArrayView<const TObjectPtr<UObject>> GetReferencedTextures() const override;
 
 	////////////////
 	// FMaterialRenderProxy interface.
@@ -169,6 +161,7 @@ public:
 	float UnrelatedNodesOpacity;
 
 private:
+	TUniquePtr<FMaterialCachedExpressionData> CachedExpressionData;
 	TUniquePtr<FMaterialCachedHLSLTree> CachedHLSLTree;
 	TWeakObjectPtr<UMaterialExpression> Expression;
 	TArray<TObjectPtr<UObject>> ReferencedTextures;
diff --git a/Engine/Source/Runtime/Engine/Private/HLSLTree/HLSLTree.cpp b/Engine/Source/Runtime/Engine/Private/HLSLTree/HLSLTree.cpp
--- a/Engine/Source/Runtime/Engine/Private/HLSLTree/HLSLTree.cpp
+++ b/Engine/Source/Runtime/Engine/Private/HLSLTree/HLSLTree.cpp
@@ -1084,22 +1084,21 @@ void FExpression::EmitValuePreshader(FEmitContext& Context, FEmitScope& Scope, c
 	check(false);
 }
 
-FEmitShaderExpression* FExpression::GetValueShader(FEmitContext& Context, FEmitScope& Scope, const FRequestedType& RequestedType)
+FEmitShaderExpression* FExpression::GetValueShader(FEmitContext& Context, FEmitScope& Scope, const FRequestedType& RequestedType, const FPreparedType& PreparedType, const Shader::FType& ResultType)
 {
 	FOwnerScope OwnerScope(*Context.Errors, GetOwner());
-	
-	const FPreparedType PreparedType = Context.GetPreparedType(this);
+
 	const EExpressionEvaluation Evaluation = PreparedType.GetEvaluation(Scope, RequestedType);
 	check(Evaluation != EExpressionEvaluation::None && Evaluation != EExpressionEvaluation::Unknown);
 
 	FEmitShaderExpression* Value = nullptr;
 	if (Evaluation == EExpressionEvaluation::ConstantZero)
 	{
-		Value = Context.EmitConstantZero(Scope, PreparedType.GetType());
+		Value = Context.EmitConstantZero(Scope, ResultType);
 	}
 	else if (Evaluation == EExpressionEvaluation::Constant || Evaluation == EExpressionEvaluation::Preshader)
 	{
-		Value = Context.EmitPreshaderOrConstant(Scope, RequestedType, this);
+		Value = Context.EmitPreshaderOrConstant(Scope, RequestedType, ResultType, this);
 	}
 	else
 	{
@@ -1108,24 +1107,19 @@ FEmitShaderExpression* FExpression::GetValueShader(FEmitContext& Context, FEmitS
 		check(Result.Code);
 		Value = Result.Code;
 	}
-
-	if (Value->Type.IsNumeric() && PreparedType.IsNumeric())
-	{
-		const Shader::FValueTypeDescription TypeDest = Shader::GetValueTypeDescription(Value->Type);
-		if (TypeDest.ComponentType == Shader::EValueComponentType::Double && PreparedType.ValueComponentType == Shader::EValueComponentType::Float)
-		{
-			const Shader::EValueType CastType = Shader::MakeValueType(Shader::EValueComponentType::Float, TypeDest.NumComponents);
-			Value = Context.EmitCast(Scope, Value, CastType);
-		}
-	}
-
-	return Value;
+	return Context.EmitCast(Scope, Value, ResultType);
 }
 
 FEmitShaderExpression* FExpression::GetValueShader(FEmitContext& Context, FEmitScope& Scope, const FRequestedType& RequestedType, const Shader::FType& ResultType)
 {
-	FEmitShaderExpression* Value = GetValueShader(Context, Scope, RequestedType);
-	return Context.EmitCast(Scope, Value, ResultType);
+	const FPreparedType& PreparedType = Context.GetPreparedType(this);
+	return GetValueShader(Context, Scope, RequestedType, PreparedType, ResultType);
+}
+
+FEmitShaderExpression* FExpression::GetValueShader(FEmitContext& Context, FEmitScope& Scope, const FRequestedType& RequestedType)
+{
+	const FPreparedType& PreparedType = Context.GetPreparedType(this);
+	return GetValueShader(Context, Scope, RequestedType, PreparedType, PreparedType.GetType());
 }
 
 FEmitShaderExpression* FExpression::GetValueShader(FEmitContext& Context, FEmitScope& Scope, const Shader::FType& ResultType)
@@ -1140,29 +1134,29 @@ FEmitShaderExpression* FExpression::GetValueShader(FEmitContext& Context, FEmitS
 
 FEmitShaderExpression* FExpression::GetValueShader(FEmitContext& Context, FEmitScope& Scope)
 {
-	const FPreparedType PreparedType = Context.GetPreparedType(this);
-	return GetValueShader(Context, Scope, PreparedType.GetRequestedType());
+	const FPreparedType& PreparedType = Context.GetPreparedType(this);
+	return GetValueShader(Context, Scope, PreparedType.GetRequestedType(), PreparedType, PreparedType.GetType());
 }
 
-Shader::FType FExpression::GetValuePreshader(FEmitContext& Context, FEmitScope& Scope, const FRequestedType& RequestedType, Shader::FPreshaderData& OutPreshader)
+Shader::FType FExpression::GetValuePreshader(FEmitContext& Context, FEmitScope& Scope, const FRequestedType& RequestedType, const FPreparedType& PreparedType, const Shader::FType& ResultType, Shader::FPreshaderData& OutPreshader)
 {
 	FOwnerScope OwnerScope(*Context.Errors, GetOwner());
 	
 	const int32 PrevStackPosition = Context.PreshaderStackPosition;
-	const FPreparedType PreparedType = Context.GetPreparedType(this);
 	const EExpressionEvaluation Evaluation = PreparedType.GetEvaluation(Scope, RequestedType);
 	check(Evaluation != EExpressionEvaluation::None && Evaluation != EExpressionEvaluation::Unknown);
 
 	FEmitValuePreshaderResult Result(OutPreshader);
 	if (Evaluation == EExpressionEvaluation::ConstantZero)
 	{
+		check(!ResultType.IsVoid());
 		Context.PreshaderStackPosition++;
-		Result.Type = PreparedType.GetType();
-		OutPreshader.WriteOpcode(Shader::EPreshaderOpcode::ConstantZero).Write(Result.Type);
+		Result.Type = ResultType;
+		OutPreshader.WriteOpcode(Shader::EPreshaderOpcode::ConstantZero).Write(ResultType);
 	}
 	else if (Evaluation == EExpressionEvaluation::Constant)
 	{
-		const Shader::FValue ConstantValue = GetValueConstant(Context, Scope, RequestedType);
+		const Shader::FValue ConstantValue = GetValueConstant(Context, Scope, RequestedType, ResultType);
 		Context.PreshaderStackPosition++;
 		OutPreshader.WriteOpcode(Shader::EPreshaderOpcode::Constant).Write(ConstantValue);
 		Result.Type = ConstantValue.Type;
@@ -1176,15 +1170,26 @@ Shader::FType FExpression::GetValuePreshader(FEmitContext& Context, FEmitScope&
 	return Result.Type;
 }
 
-Shader::FValue FExpression::GetValueConstant(FEmitContext& Context, FEmitScope& Scope, const FRequestedType& RequestedType)
+Shader::FType FExpression::GetValuePreshader(FEmitContext& Context, FEmitScope& Scope, const FRequestedType& RequestedType, const Shader::FType& ResultType, Shader::FPreshaderData& OutPreshader)
+{
+	const FPreparedType& PreparedType = Context.GetPreparedType(this);
+	return GetValuePreshader(Context, Scope, RequestedType, PreparedType, ResultType, OutPreshader);
+}
+
+Shader::FType FExpression::GetValuePreshader(FEmitContext& Context, FEmitScope& Scope, const FRequestedType& RequestedType, Shader::FPreshaderData& OutPreshader)
+{
+	const FPreparedType& PreparedType = Context.GetPreparedType(this);
+	return GetValuePreshader(Context, Scope, RequestedType, PreparedType, PreparedType.GetType(), OutPreshader);
+}
+
+Shader::FValue FExpression::GetValueConstant(FEmitContext& Context, FEmitScope& Scope, const FRequestedType& RequestedType, const FPreparedType& PreparedType, const Shader::FType& ResultType)
 {
 	FOwnerScope OwnerScope(*Context.Errors, GetOwner());
-	
-	const FPreparedType PreparedType = Context.GetPreparedType(this);
+
 	const EExpressionEvaluation Evaluation = PreparedType.GetEvaluation(Scope, RequestedType);
 	if (Evaluation == EExpressionEvaluation::ConstantZero)
 	{
-		return Shader::FValue(PreparedType.GetType());
+		return Shader::FValue(ResultType);
 	}
 	else
 	{
@@ -1202,29 +1207,35 @@ Shader::FValue FExpression::GetValueConstant(FEmitContext& Context, FEmitScope&
 		Shader::FPreshaderStack Stack;
 		const Shader::FPreshaderValue PreshaderValue = ConstantPreshader.EvaluateConstant(*Context.Material, Stack);
 		Shader::FValue Result = PreshaderValue.AsShaderValue(Context.TypeRegistry);
-		return Result;
+		return Shader::Cast(Result, ResultType);
 	}
 }
 
+Shader::FValue FExpression::GetValueConstant(FEmitContext& Context, FEmitScope& Scope, const FRequestedType& RequestedType, const FPreparedType& PreparedType)
+{
+	return GetValueConstant(Context, Scope, RequestedType, PreparedType, PreparedType.GetType());
+}
+
 Shader::FValue FExpression::GetValueConstant(FEmitContext& Context, FEmitScope& Scope, const FRequestedType& RequestedType, const Shader::FType& ResultType)
 {
-	Shader::FValue Result = GetValueConstant(Context, Scope, RequestedType);
-	if (Result.Type.IsNumeric() && ResultType.IsNumeric())
-	{
-		Result = Shader::Cast(Result, ResultType.ValueType);
-	}
-	check(Result.Type == ResultType);
-	return Result;
+	const FPreparedType PreparedType = Context.GetPreparedType(this);
+	return GetValueConstant(Context, Scope, RequestedType, PreparedType, ResultType);
+}
+
+Shader::FValue FExpression::GetValueConstant(FEmitContext& Context, FEmitScope& Scope, const FRequestedType& RequestedType)
+{
+	const FPreparedType PreparedType = Context.GetPreparedType(this);
+	return GetValueConstant(Context, Scope, RequestedType, PreparedType, PreparedType.GetType());
 }
 
-Shader::FValue FExpression::GetValueConstant(FEmitContext& Context, FEmitScope& Scope, const Shader::FType& ResultType)
+Shader::FValue FExpression::GetValueConstant(FEmitContext& Context, FEmitScope& Scope, const FPreparedType& PreparedType, const Shader::FType& ResultType)
 {
-	return GetValueConstant(Context, Scope, FRequestedType(ResultType), ResultType);
+	return GetValueConstant(Context, Scope, FRequestedType(ResultType), PreparedType, ResultType);
 }
 
-Shader::FValue FExpression::GetValueConstant(FEmitContext& Context, FEmitScope& Scope, Shader::EValueType ResultType)
+Shader::FValue FExpression::GetValueConstant(FEmitContext& Context, FEmitScope& Scope, const FPreparedType& PreparedType, Shader::EValueType ResultType)
 {
-	return GetValueConstant(Context, Scope, FRequestedType(ResultType), ResultType);
+	return GetValueConstant(Context, Scope, FRequestedType(ResultType), PreparedType, ResultType);
 }
 
 bool FScope::HasParentScope(const FScope& InParentScope) const
diff --git a/Engine/Source/Runtime/Engine/Private/HLSLTree/HLSLTreeCommon.cpp b/Engine/Source/Runtime/Engine/Private/HLSLTree/HLSLTreeCommon.cpp
--- a/Engine/Source/Runtime/Engine/Private/HLSLTree/HLSLTreeCommon.cpp
+++ b/Engine/Source/Runtime/Engine/Private/HLSLTree/HLSLTreeCommon.cpp
@@ -91,28 +91,27 @@ bool FExpressionError::PrepareValue(FEmitContext& Context, FEmitScope&, const FR
 
 void FExpressionForward::ComputeAnalyticDerivatives(FTree& Tree, FExpressionDerivatives& OutResult) const
 {
-	OutResult = Tree.GetAnalyticDerivatives(Expression);
+	return Expression->ComputeAnalyticDerivatives(Tree, OutResult);
 }
 
 FExpression* FExpressionForward::ComputePreviousFrame(FTree& Tree, const FRequestedType& RequestedType) const
 {
-	return Tree.GetPreviousFrame(Expression, RequestedType);
+	return Expression->ComputePreviousFrame(Tree, RequestedType);
 }
 
 bool FExpressionForward::PrepareValue(FEmitContext& Context, FEmitScope& Scope, const FRequestedType& RequestedType, FPrepareValueResult& OutResult) const
 {
-	const FPreparedType& ResultType = Context.PrepareExpression(Expression, Scope, RequestedType);
-	return OutResult.SetType(Context, RequestedType, ResultType);
+	return Expression->PrepareValue(Context, Scope, RequestedType, OutResult);
 }
 
 void FExpressionForward::EmitValueShader(FEmitContext& Context, FEmitScope& Scope, const FRequestedType& RequestedType, FEmitValueShaderResult& OutResult) const
 {
-	OutResult.Code = Expression->GetValueShader(Context, Scope, RequestedType);
+	return Expression->EmitValueShader(Context, Scope, RequestedType, OutResult);
 }
 
 void FExpressionForward::EmitValuePreshader(FEmitContext& Context, FEmitScope& Scope, const FRequestedType& RequestedType, FEmitValuePreshaderResult& OutResult) const
 {
-	OutResult.Type = Expression->GetValuePreshader(Context, Scope, RequestedType, OutResult.Preshader);
+	return Expression->EmitValuePreshader(Context, Scope, RequestedType, OutResult);
 }
 
 FExpression* FTree::NewConstant(const Shader::FValue& Value)
@@ -434,13 +433,8 @@ void FExpressionSetStructField::EmitValueShader(FEmitContext& Context, FEmitScop
 	}
 	else
 	{
-		FEmitShaderExpression* StructValue = (StructEvaluation != EExpressionEvaluation::ConstantZero)
-			? StructExpression->GetValueShader(Context, Scope, RequestedStructType)
-			: Context.EmitConstantZero(Scope, StructType);
-		FEmitShaderExpression* FieldValue = (FieldEvaluation != EExpressionEvaluation::ConstantZero)
-			? FieldExpression->GetValueShader(Context, Scope, RequestedFieldType, Field->Type)
-			: Context.EmitConstantZero(Scope, Field->Type);
-
+		FEmitShaderExpression* StructValue = StructExpression->GetValueShader(Context, Scope, RequestedStructType, StructType);
+		FEmitShaderExpression* FieldValue = FieldExpression->GetValueShader(Context, Scope, RequestedFieldType, Field->Type);
 		OutResult.Code = Context.EmitExpression(Scope, StructType, TEXT("%_Set%(%, %)"),
 			StructType->Name,
 			Field->Name,
@@ -466,8 +460,8 @@ void FExpressionSetStructField::EmitValuePreshader(FEmitContext& Context, FEmitS
 	}
 	else
 	{
-		StructExpression->GetValuePreshader(Context, Scope, RequestedStructType, OutResult.Preshader);
-		FieldExpression->GetValuePreshader(Context, Scope, RequestedFieldType, OutResult.Preshader);
+		StructExpression->GetValuePreshader(Context, Scope, RequestedStructType, StructType, OutResult.Preshader);
+		FieldExpression->GetValuePreshader(Context, Scope, RequestedFieldType, Field->Type, OutResult.Preshader);
 
 		check(Context.PreshaderStackPosition > 0);
 		Context.PreshaderStackPosition--;
@@ -498,7 +492,7 @@ bool FExpressionSelect::PrepareValue(FEmitContext& Context, FEmitScope& Scope, c
 	const EExpressionEvaluation ConditionEvaluation = ConditionType.GetEvaluation(Scope, ERequestedType::Scalar);
 	if (IsConstantEvaluation(ConditionEvaluation))
 	{
-		const bool bCondition = ConditionExpression->GetValueConstant(Context, Scope, Shader::EValueType::Bool1).AsBoolScalar();
+		const bool bCondition = ConditionExpression->GetValueConstant(Context, Scope, ConditionType, Shader::EValueType::Bool1).AsBoolScalar();
 		FPreparedType ResultType = Context.PrepareExpression(bCondition ? TrueExpression : FalseExpression, Scope, RequestedType);
 		ResultType.MergeEvaluation(EExpressionEvaluation::Shader); // TODO - support preshader
 		return OutResult.SetType(Context, RequestedType, ResultType);
@@ -521,10 +515,11 @@ bool FExpressionSelect::PrepareValue(FEmitContext& Context, FEmitScope& Scope, c
 
 void FExpressionSelect::EmitValueShader(FEmitContext& Context, FEmitScope& Scope, const FRequestedType& RequestedType, FEmitValueShaderResult& OutResult) const
 {
-	const EExpressionEvaluation ConditionEvaluation = Context.GetEvaluation(ConditionExpression, Scope, ERequestedType::Scalar);
+	const FPreparedType& ConditionType = Context.GetPreparedType(ConditionExpression);
+	const EExpressionEvaluation ConditionEvaluation = ConditionType.GetEvaluation(Scope, ERequestedType::Scalar);
 	if (IsConstantEvaluation(ConditionEvaluation))
 	{
-		const bool bCondition = ConditionExpression->GetValueConstant(Context, Scope, Shader::EValueType::Bool1).AsBoolScalar();
+		const bool bCondition = ConditionExpression->GetValueConstant(Context, Scope, ConditionType, Shader::EValueType::Bool1).AsBoolScalar();
 		FExpression* InputExpression = bCondition ? TrueExpression : FalseExpression;
 		OutResult.Code = InputExpression->GetValueShader(Context, Scope, RequestedType);
 	}
@@ -1043,7 +1038,7 @@ bool FStatementIf::Prepare(FEmitContext& Context, FEmitScope& Scope) const
 	check(ConditionEvaluation != EExpressionEvaluation::None);
 	if (IsConstantEvaluation(ConditionEvaluation))
 	{
-		const bool bCondition = ConditionExpression->GetValueConstant(Context, Scope, Shader::EValueType::Bool1).AsBoolScalar();
+		const bool bCondition = ConditionExpression->GetValueConstant(Context, Scope, ConditionType, Shader::EValueType::Bool1).AsBoolScalar();
 		if (bCondition)
 		{
 			Context.MarkScopeEvaluation(Scope, ThenScope, ConditionEvaluation);
@@ -1067,10 +1062,11 @@ bool FStatementIf::Prepare(FEmitContext& Context, FEmitScope& Scope) const
 void FStatementIf::EmitShader(FEmitContext& Context, FEmitScope& Scope) const
 {
 	FEmitShaderNode* Dependency = nullptr;
-	const EExpressionEvaluation ConditionEvaluation = Context.GetEvaluation(ConditionExpression, Scope, ERequestedType::Scalar);
+	const FPreparedType& ConditionType = Context.GetPreparedType(ConditionExpression);
+	const EExpressionEvaluation ConditionEvaluation = ConditionType.GetEvaluation(Scope, ERequestedType::Scalar);
 	if (IsConstantEvaluation(ConditionEvaluation))
 	{
-		const bool bCondition = ConditionExpression->GetValueConstant(Context, Scope, Shader::EValueType::Bool1).AsBoolScalar();
+		const bool bCondition = ConditionExpression->GetValueConstant(Context, Scope, ConditionType, Shader::EValueType::Bool1).AsBoolScalar();
 		if (bCondition)
 		{
 			Dependency = Context.EmitNextScope(Scope, ThenScope);
diff --git a/Engine/Source/Runtime/Engine/Private/HLSLTree/HLSLTreeEmit.cpp b/Engine/Source/Runtime/Engine/Private/HLSLTree/HLSLTreeEmit.cpp
--- a/Engine/Source/Runtime/Engine/Private/HLSLTree/HLSLTreeEmit.cpp
+++ b/Engine/Source/Runtime/Engine/Private/HLSLTree/HLSLTreeEmit.cpp
@@ -348,8 +348,13 @@ void FEmitContext::EmitDeclarationsCode(FStringBuilderBase& OutCode)
 const FPreparedType& FEmitContext::GetPreparedType(const FExpression* Expression) const
 {
 	static const FPreparedType VoidType;
-	FPrepareValueResult* const* PrevResult = PrepareValueMap.Find(Expression);
-	return PrevResult ? (*PrevResult)->PreparedType : VoidType;
+	FPrepareValueResult const* const* PrevResult = PrepareValueMap.Find(Expression);
+	const FPrepareValueResult* Result = PrevResult ? *PrevResult : nullptr;
+	if (Result)
+	{
+		return Result->PreparedType;
+	}
+	return VoidType;
 }
 
 FRequestedType FEmitContext::GetRequestedType(const FExpression* Expression) const
@@ -378,17 +383,33 @@ FPreparedType FEmitContext::PrepareExpression(FExpression* InExpression, FEmitSc
 	FPrepareValueResult* Result = PrevResult ? *PrevResult : nullptr;
 	if (!Result)
 	{
+		check(!bMarkLiveValues); // value should already be prepared at this point
 		Result = new(*Allocator) FPrepareValueResult();
 		PrepareValueMap.Add(InExpression, Result);
 	}
 
+	if (bMarkLiveValues)
+	{
+		if (Result->bPreparingValue)
+		{
+			return Result->PreparedType;
+		}
+
+		const EExpressionEvaluation Evaluation = Result->PreparedType.GetEvaluation(Scope, RequestedType);
+		if (Evaluation == EExpressionEvaluation::None || IsConstantEvaluation(Evaluation))
+		{
+			// Don't need to prepare constant values, since we are just replacing the expression with a constant
+			// Can't set the 'bIsLive' flag here, since expression may not be constant in all contexts
+			return Result->PreparedType;
+		}
+	}
 
-	FOwnerScope OwnerScope(*Errors, InExpression->GetOwner());
 	if (Result->bPreparingValue)
 	{
 		// Valid for this to be called reentrantly
 		// Code should ensure that the type is set before the reentrant call, otherwise type will not be valid here
 		// LocalPHI nodes rely on this to break loops
+		check(!bMarkLiveValues);
 		FEmitScope* LoopScope = Scope.FindLoop();
 		check(LoopScope);
 		Result->PreparedType.SetLoopEvaluation(*LoopScope, RequestedType);
@@ -397,6 +418,7 @@ FPreparedType FEmitContext::PrepareExpression(FExpression* InExpression, FEmitSc
 
 	bool bResult = false;
 	{
+		FOwnerScope OwnerScope(*Errors, InExpression->GetOwner());
 		Result->bPreparingValue = true;
 		bResult = InExpression->PrepareValue(*this, Scope, RequestedType, *Result);
 		Result->bPreparingValue = false;
@@ -835,14 +857,14 @@ void WriteMaterialUniformAccess(Shader::EValueComponentType ComponentType, uint3
 }
 } // namespace Private
 
-FEmitShaderExpression* FEmitContext::EmitPreshaderOrConstant(FEmitScope& Scope, const FRequestedType& RequestedType, FExpression* Expression)
+FEmitShaderExpression* FEmitContext::EmitPreshaderOrConstant(FEmitScope& Scope, const FRequestedType& RequestedType, const Shader::FType& ResultType, FExpression* Expression)
 {
 	Shader::FPreshaderData LocalPreshader;
-	Shader::FType Type;
+	Shader::FType Type = ResultType;
 	{
 		FEmitValuePreshaderResult PreshaderResult(LocalPreshader);
 		Expression->EmitValuePreshader(*this, Scope, RequestedType, PreshaderResult);
-		Type = PreshaderResult.Type;
+		check(PreshaderResult.Type.IsNumeric() || PreshaderResult.Type == Type);
 	}
 
 	FXxHash64Builder Hasher;
@@ -858,8 +880,11 @@ FEmitShaderExpression* FEmitContext::EmitPreshaderOrConstant(FEmitScope& Scope,
 		return ShaderValue;
 	}
 
+	const FPreparedType& PreparedType = GetPreparedType(Expression);
+
 	Shader::FPreshaderStack Stack;
-	const Shader::FPreshaderValue ConstantValue = LocalPreshader.EvaluateConstant(*Material, Stack);
+	const Shader::FPreshaderValue PreshaderConstantValue = LocalPreshader.EvaluateConstant(*Material, Stack);
+	const Shader::FValue ConstantValue = Shader::Cast(PreshaderConstantValue.AsShaderValue(TypeRegistry), Type);
 
 	TStringBuilder<1024> FormattedCode;
 	if (Type.IsStruct())
@@ -877,10 +902,10 @@ FEmitShaderExpression* FEmitContext::EmitPreshaderOrConstant(FEmitScope& Scope,
 			FormattedCode.Append(TEXT(", "));
 		}
 
-		const Shader::EValueType FieldType = Type.GetFlatFieldType(FieldIndex);
+		const Shader::EValueType FieldType = ResultType.GetFlatFieldType(FieldIndex);
 		const Shader::FValueTypeDescription TypeDesc = Shader::GetValueTypeDescription(FieldType);
 		const int32 NumFieldComponents = TypeDesc.NumComponents;
-		const EExpressionEvaluation FieldEvaluation = GetPreparedType(Expression).GetFieldEvaluation(Scope, RequestedType, ComponentIndex, NumFieldComponents);
+		const EExpressionEvaluation FieldEvaluation = PreparedType.GetFieldEvaluation(Scope, RequestedType, ComponentIndex, NumFieldComponents);
 
 		if (FieldEvaluation == EExpressionEvaluation::Preshader)
 		{
@@ -974,7 +999,7 @@ FEmitShaderExpression* FEmitContext::EmitPreshaderOrConstant(FEmitScope& Scope,
 
 			// The type generated by the preshader might not match the expected type
 			// In the future, with new HLSLTree, preshader could potentially include explicit cast opcodes, and avoid implicit conversions
-			Shader::FValue FieldConstantValue(ConstantValue.Type.GetComponentType(ComponentIndex), NumFieldComponents);
+			Shader::FValue FieldConstantValue(Type.GetComponentType(ComponentIndex), NumFieldComponents);
 			for (int32 i = 0; i < NumFieldComponents; ++i)
 			{
 				if (ConstantValue.Component.Num() == 1)
@@ -1055,6 +1080,7 @@ FEmitShaderExpression* FEmitContext::EmitPreshaderOrConstant(FEmitScope& Scope,
 
 FEmitShaderExpression* FEmitContext::EmitConstantZero(FEmitScope& Scope, const Shader::FType& Type)
 {
+	check(!Type.IsVoid());
 	return EmitInlineExpression(Scope, Type, TEXT("((%)0)"), Type.GetName());
 }
 
diff --git a/Engine/Source/Runtime/Engine/Private/HLSLTree/HLSLTreeOperation.cpp b/Engine/Source/Runtime/Engine/Private/HLSLTree/HLSLTreeOperation.cpp
--- a/Engine/Source/Runtime/Engine/Private/HLSLTree/HLSLTreeOperation.cpp
+++ b/Engine/Source/Runtime/Engine/Private/HLSLTree/HLSLTreeOperation.cpp
@@ -186,14 +186,8 @@ FOperationRequestedTypes GetOperationRequestedTypes(EOperation Op, const FReques
 	return Types;
 }
 
-FOperationTypes GetOperationTypes(FEmitContext& Context, EOperation Op, TConstArrayView<FExpression*> Inputs)
+FOperationTypes GetOperationTypes(EOperation Op, TConstArrayView<FPreparedType> InputPreparedType)
 {
-	FPreparedType InputPreparedType[FExpressionOperation::MaxInputs];
-	for (int32 Index = 0; Index < Inputs.Num(); ++Index)
-	{
-		InputPreparedType[Index] = Context.GetPreparedType(Inputs[Index]);
-	}
-
 	FOperationTypes Types;
 	if (Op == EOperation::VecMulMatrix3 ||
 		Op == EOperation::VecMulMatrix4 ||
@@ -202,7 +196,7 @@ FOperationTypes GetOperationTypes(FEmitContext& Context, EOperation Op, TConstAr
 	{
 		FPreparedComponent IntermediateComponent;
 		Shader::EValueComponentType IntermediateComponentType = Shader::EValueComponentType::Void;
-		for (int32 Index = 0; Index < Inputs.Num(); ++Index)
+		for (int32 Index = 0; Index < InputPreparedType.Num(); ++Index)
 		{
 			IntermediateComponent = CombineComponents(IntermediateComponent, InputPreparedType[Index].GetMergedComponent());
 			IntermediateComponentType = Shader::CombineComponentTypes(IntermediateComponentType, InputPreparedType[Index].ValueComponentType);
@@ -233,24 +227,42 @@ FOperationTypes GetOperationTypes(FEmitContext& Context, EOperation Op, TConstAr
 	else
 	{
 		FPreparedType IntermediateType;
-		for (int32 Index = 0; Index < Inputs.Num(); ++Index)
+		for (int32 Index = 0; Index < InputPreparedType.Num(); ++Index)
 		{
 			IntermediateType = MergePreparedTypes(IntermediateType, InputPreparedType[Index]);
 		}
 
-		const Shader::EValueType IntermediateValueType = IntermediateType.GetType();
-		for (int32 Index = 0; Index < Inputs.Num(); ++Index)
+		const Shader::EValueComponentType IntermediateComponentType = IntermediateType.ValueComponentType;
+		const int32 NumIntermediateComponents = IntermediateType.GetNumComponents();
+		for (int32 Index = 0; Index < InputPreparedType.Num(); ++Index)
 		{
-			Types.InputType[Index] = IntermediateValueType;
+			Shader::EValueComponentType InputComponentType = InputPreparedType[Index].ValueComponentType;
+			if (InputComponentType != IntermediateComponentType)
+			{
+				switch (IntermediateComponentType)
+				{
+				case Shader::EValueComponentType::Float:
+				case Shader::EValueComponentType::Double:
+					// If the result type is either float or double, we promote input types to float
+					// Nothing should promote to double, unless it's already at double precision
+					InputComponentType = Shader::EValueComponentType::Float;
+					break;
+				default:
+					// Otherwise use ints
+					InputComponentType = Shader::EValueComponentType::Int;
+					break;
+				}
+			}
+			Types.InputType[Index] = Shader::MakeValueType(InputComponentType, NumIntermediateComponents);
 		}
 		Types.ResultType = IntermediateType;
-		Types.bIsLWC = (IntermediateType.ValueComponentType == Shader::EValueComponentType::Double);
+		Types.bIsLWC = (IntermediateComponentType == Shader::EValueComponentType::Double);
 
 		switch (Op)
 		{
 		case EOperation::Length:
 		case EOperation::Sum:
-			Types.ResultType = FPreparedType(Shader::MakeValueType(IntermediateType.ValueComponentType, 1), IntermediateType.GetMergedComponent());
+			Types.ResultType = FPreparedType(Shader::MakeValueType(IntermediateComponentType, 1), IntermediateType.GetMergedComponent());
 			break;
 		case EOperation::Normalize:
 		case EOperation::Rcp:
@@ -285,7 +297,7 @@ FOperationTypes GetOperationTypes(FEmitContext& Context, EOperation Op, TConstAr
 		case EOperation::Log2:
 		case EOperation::Exp2:
 			// No LWC support yet
-			Types.InputType[0] = Shader::MakeNonLWCType(IntermediateValueType);
+			Types.InputType[0] = Shader::MakeNonLWCType(Types.InputType[0]);
 			Types.ResultType = MakeNonLWCType(IntermediateType);
 			break;
 		case EOperation::Less:
@@ -295,14 +307,15 @@ FOperationTypes GetOperationTypes(FEmitContext& Context, EOperation Op, TConstAr
 			Types.ResultType.ValueComponentType = Shader::EValueComponentType::Bool;
 			break;
 		case EOperation::Fmod:
-			Types.InputType[1] = Shader::MakeNonLWCType(IntermediateValueType);
+			Types.InputType[1] = Shader::MakeNonLWCType(Types.InputType[1]);
 			Types.ResultType = MakeNonLWCType(IntermediateType);
 			break;
 		case EOperation::PowPositiveClamped:
 		case EOperation::Atan2:
 		case EOperation::Atan2Fast:
 			// No LWC support yet
-			Types.InputType[0] = Types.InputType[1] = Shader::MakeNonLWCType(IntermediateValueType);
+			Types.InputType[0] = Shader::MakeNonLWCType(Types.InputType[0]);
+			Types.InputType[1] = Shader::MakeNonLWCType(Types.InputType[1]);
 			Types.ResultType = MakeNonLWCType(IntermediateType);
 			break;
 		case EOperation::Min:
@@ -531,30 +544,31 @@ bool FExpressionOperation::PrepareValue(FEmitContext& Context, FEmitScope& Scope
 	const FOperationDescription OpDesc = GetOperationDescription(Op);
 	const Private::FOperationRequestedTypes RequestedTypes = Private::GetOperationRequestedTypes(Op, RequestedType);
 
+	FPreparedType InputPreparedType[MaxInputs];
 	Shader::FValue ConstantInput[MaxInputs];
 	bool bConstantZeroInput[MaxInputs] = { false };
 	for (int32 Index = 0; Index < OpDesc.NumInputs; ++Index)
 	{
-		const FPreparedType& InputType = Context.PrepareExpression(Inputs[Index], Scope, RequestedTypes.InputType[Index]);
-		if (InputType.IsVoid())
+		InputPreparedType[Index] = Context.PrepareExpression(Inputs[Index], Scope, RequestedTypes.InputType[Index]);
+		if (InputPreparedType[Index].IsVoid())
 		{
 			return false;
 		}
 
-		if (!InputType.IsNumeric())
+		if (!InputPreparedType[Index].IsNumeric())
 		{
 			return Context.Errors->AddError(TEXT("Invalid arithmetic between non-numeric types"));
 		}
 		
-		const EExpressionEvaluation InputEvaluation = InputType.GetEvaluation(Scope, RequestedType);
-		if (InputEvaluation == EExpressionEvaluation::Constant)
+		const EExpressionEvaluation InputEvaluation = InputPreparedType[Index].GetEvaluation(Scope, RequestedType);
+		if (IsConstantEvaluation(InputEvaluation))
 		{
-			ConstantInput[Index] = Inputs[Index]->GetValueConstant(Context, Scope, RequestedType);
+			ConstantInput[Index] = Inputs[Index]->GetValueConstant(Context, Scope, RequestedType, InputPreparedType[Index]);
 			bConstantZeroInput[Index] = ConstantInput[Index].IsZero();
 		}
 	}
 
-	Private::FOperationTypes Types = Private::GetOperationTypes(Context, Op, MakeArrayView(Inputs, OpDesc.NumInputs));
+	Private::FOperationTypes Types = Private::GetOperationTypes(Op, MakeArrayView(InputPreparedType, OpDesc.NumInputs));
 	if (OpDesc.PreshaderOpcode == Shader::EPreshaderOpcode::Nop)
 	{
 		// No preshader support
@@ -580,8 +594,14 @@ bool FExpressionOperation::PrepareValue(FEmitContext& Context, FEmitScope& Scope
 void FExpressionOperation::EmitValueShader(FEmitContext& Context, FEmitScope& Scope, const FRequestedType& RequestedType, FEmitValueShaderResult& OutResult) const
 {
 	const FOperationDescription OpDesc = GetOperationDescription(Op);
+	FPreparedType InputPreparedType[MaxInputs];
+	for (int32 Index = 0; Index < OpDesc.NumInputs; ++Index)
+	{
+		InputPreparedType[Index] = Context.GetPreparedType(Inputs[Index]);
+	}
+
 	const Private::FOperationRequestedTypes RequestedTypes = Private::GetOperationRequestedTypes(Op, RequestedType);
-	const Private::FOperationTypes Types = Private::GetOperationTypes(Context, Op, MakeArrayView(Inputs, OpDesc.NumInputs));
+	const Private::FOperationTypes Types = Private::GetOperationTypes(Op, MakeArrayView(InputPreparedType, OpDesc.NumInputs));
 	FEmitShaderExpression* InputValue[MaxInputs] = { nullptr };
 	for (int32 Index = 0; Index < OpDesc.NumInputs; ++Index)
 	{
@@ -680,8 +700,14 @@ void FExpressionOperation::EmitValueShader(FEmitContext& Context, FEmitScope& Sc
 void FExpressionOperation::EmitValuePreshader(FEmitContext& Context, FEmitScope& Scope, const FRequestedType& RequestedType, FEmitValuePreshaderResult& OutResult) const
 {
 	const FOperationDescription OpDesc = GetOperationDescription(Op);
+	FPreparedType InputPreparedType[MaxInputs];
+	for (int32 Index = 0; Index < OpDesc.NumInputs; ++Index)
+	{
+		InputPreparedType[Index] = Context.GetPreparedType(Inputs[Index]);
+	}
+
 	const Private::FOperationRequestedTypes RequestedTypes = Private::GetOperationRequestedTypes(Op, RequestedType);
-	const Private::FOperationTypes Types = Private::GetOperationTypes(Context, Op, MakeArrayView(Inputs, OpDesc.NumInputs));
+	const Private::FOperationTypes Types = Private::GetOperationTypes(Op, MakeArrayView(InputPreparedType, OpDesc.NumInputs));
 	check(OpDesc.PreshaderOpcode != Shader::EPreshaderOpcode::Nop);
 
 	for (int32 Index = 0; Index < OpDesc.NumInputs; ++Index)
diff --git a/Engine/Source/Runtime/Engine/Private/Materials/MaterialCachedData.cpp b/Engine/Source/Runtime/Engine/Private/Materials/MaterialCachedData.cpp
--- a/Engine/Source/Runtime/Engine/Private/Materials/MaterialCachedData.cpp
+++ b/Engine/Source/Runtime/Engine/Private/Materials/MaterialCachedData.cpp
@@ -475,6 +475,7 @@ void PrepareHLSLTree(UE::HLSLTree::FEmitContext& EmitContext,
 
 	EmitContext.ShaderFrequency = ShaderFrequency;
 	EmitContext.bUseAnalyticDerivatives = true; // We want to consider expressions used for analytic derivatives
+	EmitContext.bMarkLiveValues = false;
 	FEmitScope* EmitResultScope = EmitContext.PrepareScope(&CachedTree.GetResultStatement()->GetParentScope());
 
 	FRequestedType RequestedAttributesType;
@@ -495,15 +496,20 @@ void PrepareHLSLTree(UE::HLSLTree::FEmitContext& EmitContext,
 	}
 
 	const FPreparedType& ResultType = EmitContext.PrepareExpression(CachedTree.GetResultExpression(), *EmitResultScope, RequestedAttributesType);
-
-	for (const FGuid& AttributeID : OrderedVisibleAttributes)
+	if (!ResultType.IsVoid())
 	{
-		if (FMaterialAttributeDefinitionMap::GetShaderFrequency(AttributeID) == ShaderFrequency)
+		EmitContext.bMarkLiveValues = true;
+		EmitContext.PrepareExpression(CachedTree.GetResultExpression(), *EmitResultScope, RequestedAttributesType);
+
+		for (const FGuid& AttributeID : OrderedVisibleAttributes)
 		{
-			const EMaterialProperty Property = FMaterialAttributeDefinitionMap::GetProperty(AttributeID);
-			if (CachedTree.IsAttributeUsed(EmitContext, *EmitResultScope, ResultType, Property))
+			if (FMaterialAttributeDefinitionMap::GetShaderFrequency(AttributeID) == ShaderFrequency)
 			{
-				CachedData.SetMaterialAttributePropertyConnected(Property, true);
+				const EMaterialProperty Property = FMaterialAttributeDefinitionMap::GetProperty(AttributeID);
+				if (CachedTree.IsAttributeUsed(EmitContext, *EmitResultScope, ResultType, Property))
+				{
+					CachedData.SetMaterialAttributePropertyConnected(Property, true);
+				}
 			}
 		}
 	}
diff --git a/Engine/Source/Runtime/Engine/Private/Materials/MaterialHLSLEmitter.cpp b/Engine/Source/Runtime/Engine/Private/Materials/MaterialHLSLEmitter.cpp
--- a/Engine/Source/Runtime/Engine/Private/Materials/MaterialHLSLEmitter.cpp
+++ b/Engine/Source/Runtime/Engine/Private/Materials/MaterialHLSLEmitter.cpp
@@ -415,7 +415,10 @@ static void GetMaterialEnvironment(EShaderPlatform InPlatform,
 
 	OutEnvironment.SetDefine(TEXT("ENABLE_NEW_HLSL_GENERATOR"), 1);
 
-	if (false)//bNeedsParticlePosition || Material->ShouldGenerateSphericalParticleNormals() || bUsesSphericalParticleOpacity)
+	const bool bNeedsParticlePosition = EmitMaterialData.IsExternalInputUsed(Material::EExternalInput::ParticleTranslatedWorldPosition) ||
+		EmitMaterialData.IsExternalInputUsed(Material::EExternalInput::ParticleRadius);
+
+	if (bNeedsParticlePosition)// || Material->ShouldGenerateSphericalParticleNormals() || bUsesSphericalParticleOpacity)
 	{
 		OutEnvironment.SetDefine(TEXT("NEEDS_PARTICLE_POSITION"), 1);
 	}
@@ -500,11 +503,18 @@ static void GetMaterialEnvironment(EShaderPlatform InPlatform,
 	OutEnvironment.SetDefine(TEXT("USES_PER_INSTANCE_RANDOM"), MaterialCompilationOutput.bUsesPerInstanceRandom && InMaterial.IsUsedWithInstancedStaticMeshes());
 	OutEnvironment.SetDefine(TEXT("USES_VERTEX_INTERPOLATOR"), MaterialCompilationOutput.bUsesVertexInterpolator);
 
+	const bool bUsesVertexColor = EmitMaterialData.IsExternalInputUsed(SF_Vertex, Material::EExternalInput::VertexColor) ||
+		EmitMaterialData.IsExternalInputUsed(SF_Vertex, Material::EExternalInput::VertexColor_Ddx) ||
+		EmitMaterialData.IsExternalInputUsed(SF_Vertex, Material::EExternalInput::VertexColor_Ddy);
+
+	const bool bUsesParticleColor = EmitMaterialData.IsExternalInputUsed(SF_Pixel, Material::EExternalInput::ParticleColor);
+
+
 	// @todo MetalMRT: Remove this hack and implement proper atmospheric-fog solution for Metal MRT...
 	OutEnvironment.SetDefine(TEXT("MATERIAL_ATMOSPHERIC_FOG"), false);// !IsMetalMRTPlatform(InPlatform) ? bUsesAtmosphericFog : 0);
 	OutEnvironment.SetDefine(TEXT("MATERIAL_SKY_ATMOSPHERE"), false);// bUsesSkyAtmosphere);
-	OutEnvironment.SetDefine(TEXT("INTERPOLATE_VERTEX_COLOR"), EmitMaterialData.bUsesVertexColor);
-	OutEnvironment.SetDefine(TEXT("NEEDS_PARTICLE_COLOR"), false);// bUsesParticleColor);
+	OutEnvironment.SetDefine(TEXT("INTERPOLATE_VERTEX_COLOR"), bUsesVertexColor);
+	OutEnvironment.SetDefine(TEXT("NEEDS_PARTICLE_COLOR"), bUsesParticleColor);
 	OutEnvironment.SetDefine(TEXT("NEEDS_PARTICLE_LOCAL_TO_WORLD"), false);// bUsesParticleLocalToWorld);
 	OutEnvironment.SetDefine(TEXT("NEEDS_PARTICLE_WORLD_TO_LOCAL"), false);// bUsesParticleWorldToLocal);
 	OutEnvironment.SetDefine(TEXT("USES_TRANSFORM_VECTOR"), false);// bUsesTransformVector);
@@ -830,6 +840,7 @@ bool MaterialEmitHLSL(const FMaterialCompileTargetParameters& InCompilerTarget,
 	{
 		EmitContext.ShaderFrequency = SF_Pixel;
 		EmitContext.bUseAnalyticDerivatives = (DerivativeIndex == 1);
+		EmitContext.bMarkLiveValues = false;
 		FEmitScope* EmitResultScope = EmitContext.PrepareScope(&CachedTree->GetResultStatement()->GetParentScope());
 
 		// Prepare all fields *except* normal
@@ -854,12 +865,18 @@ bool MaterialEmitHLSL(const FMaterialCompileTargetParameters& InCompilerTarget,
 			return false;
 		}
 
+		EmitContext.bMarkLiveValues = true;
+		EmitContext.PrepareExpression(CachedTree->GetResultExpression(), *EmitResultScope, RequestedPixelAttributesType);
+		EmitContext.bMarkLiveValues = false;
+
 		if (CachedTree->IsAttributeUsed(EmitContext, *EmitResultScope, PixelResultType0, MP_PixelDepthOffset))
 		{
 			bUsesPixelDepthOffset = true;
 		}
 
-		if (!EmitMaterialData.bReadMaterialNormal)
+		const bool bUseNormal = EmitMaterialData.IsExternalInputUsed(SF_Pixel, Material::EExternalInput::WorldNormal);
+		const bool bUseReflection = EmitMaterialData.IsExternalInputUsed(SF_Pixel, Material::EExternalInput::WorldReflection);
+		if (!bUseNormal && !bUseReflection)
 		{
 			// No access to material normal, can execute everything in phase0
 			RequestedPixelAttributesType.SetFieldRequested(NormalField);
@@ -869,6 +886,9 @@ bool MaterialEmitHLSL(const FMaterialCompileTargetParameters& InCompilerTarget,
 				return false;
 			}
 
+			EmitContext.bMarkLiveValues = true;
+			EmitContext.PrepareExpression(CachedTree->GetResultExpression(), *EmitResultScope, RequestedPixelAttributesType);
+
 			PixelCodePhase0[DerivativeIndex] = new(Allocator) FStringBuilderMemstack(Allocator, 128 * 1024);
 			CachedTree->GetTree().EmitShader(EmitContext, *PixelCodePhase0[DerivativeIndex]);
 		}
@@ -889,6 +909,9 @@ bool MaterialEmitHLSL(const FMaterialCompileTargetParameters& InCompilerTarget,
 				return false;
 			}
 
+			EmitContext.bMarkLiveValues = true;
+			EmitContext.PrepareExpression(CachedTree->GetResultExpression(), *EmitResultScope, RequestedMaterialNormal);
+
 			// Execute the normal in phase0
 			PixelCodePhase0[DerivativeIndex] = new(Allocator) FStringBuilderMemstack(Allocator, 128 * 1024);
 			CachedTree->GetTree().EmitShader(EmitContext, *PixelCodePhase0[DerivativeIndex]);
@@ -916,6 +939,7 @@ bool MaterialEmitHLSL(const FMaterialCompileTargetParameters& InCompilerTarget,
 
 		EmitContext.ShaderFrequency = SF_Vertex;
 		EmitContext.bUseAnalyticDerivatives = false;
+		EmitContext.bMarkLiveValues = false;
 		FEmitScope* EmitResultScope = EmitContext.PrepareScope(&CachedTree->GetResultStatement()->GetParentScope());
 
 		const FPreparedType VertexResultType = EmitContext.PrepareExpression(CachedTree->GetResultExpression(), *EmitResultScope, RequestedVertexAttributesType);
@@ -924,6 +948,9 @@ bool MaterialEmitHLSL(const FMaterialCompileTargetParameters& InCompilerTarget,
 			return false;
 		}
 
+		EmitContext.bMarkLiveValues = true;
+		EmitContext.PrepareExpression(CachedTree->GetResultExpression(), *EmitResultScope, RequestedVertexAttributesType);
+
 		bUsesWorldPositionOffset = CachedTree->IsAttributeUsed(EmitContext, *EmitResultScope, VertexResultType, MP_WorldPositionOffset);
 
 		CachedTree->GetTree().EmitShader(EmitContext, VertexCode);
@@ -937,6 +964,7 @@ bool MaterialEmitHLSL(const FMaterialCompileTargetParameters& InCompilerTarget,
 	OutCompilationOutput.bModifiesMeshPosition = bUsesPixelDepthOffset || bUsesWorldPositionOffset;
 	OutCompilationOutput.bUsesWorldPositionOffset = bUsesWorldPositionOffset;
 	OutCompilationOutput.bUsesPixelDepthOffset = bUsesPixelDepthOffset;
+	OutCompilationOutput.bUsesEyeAdaptation = EmitMaterialData.IsExternalInputUsed(SF_Pixel, Material::EExternalInput::EyeAdaptation);
 
 	FStringBuilderMemstack Declarations(Allocator, 32 * 1024);
 	CachedTree->GetTypeRegistry().EmitDeclarationsCode(Declarations);
@@ -954,8 +982,8 @@ bool MaterialEmitHLSL(const FMaterialCompileTargetParameters& InCompilerTarget,
 		};
 		const TCHAR* InputPixelCodePhase1[2] =
 		{
-			EmitMaterialData.bReadMaterialNormal ? PixelCodePhase1[0]->ToString() : nullptr,
-			EmitMaterialData.bReadMaterialNormal ? PixelCodePhase1[1]->ToString() : nullptr,
+			PixelCodePhase1[0] ? PixelCodePhase1[0]->ToString() : nullptr,
+			PixelCodePhase1[1] ? PixelCodePhase1[1]->ToString() : nullptr,
 		};
 		MaterialTemplateSource = GenerateMaterialTemplateHLSL(InCompilerTarget.ShaderPlatform,
 			InOutMaterial,
diff --git a/Engine/Source/Runtime/Engine/Private/Materials/MaterialHLSLTree.cpp b/Engine/Source/Runtime/Engine/Private/Materials/MaterialHLSLTree.cpp
--- a/Engine/Source/Runtime/Engine/Private/Materials/MaterialHLSLTree.cpp
+++ b/Engine/Source/Runtime/Engine/Private/Materials/MaterialHLSLTree.cpp
@@ -203,14 +203,11 @@ bool FExpressionExternalInput::PrepareValue(FEmitContext& Context, FEmitScope& S
 {
 	const FExternalInputDescription InputDesc = GetExternalInputDescription(InputType);
 
-	switch (InputType)
+	if (Context.bMarkLiveValues)
 	{
-	case EExternalInput::WorldNormal:
-	case EExternalInput::WorldReflection:
-		Context.FindData<FEmitData>().bReadMaterialNormal = true;
-		break;
-	default:
-		break;
+		FEmitData& EmitMaterialData = Context.FindData<FEmitData>();
+		const int32 TypeIndex = (int32)InputType;
+		EmitMaterialData.ExternalInputMask[Context.ShaderFrequency][TypeIndex] = true;
 	}
 
 	return OutResult.SetType(Context, RequestedType, EExpressionEvaluation::Shader, InputDesc.Type);
@@ -247,9 +244,9 @@ void FExpressionExternalInput::EmitValueShader(FEmitContext& Context, FEmitScope
 		case EExternalInput::LightmapTexCoord_Ddx: Code = TEXT("GetLightmapUVs_DDX(Parameters)"); break;
 		case EExternalInput::LightmapTexCoord_Ddy: Code = TEXT("GetLightmapUVs_DDY(Parameters)"); break;
 		case EExternalInput::TwoSidedSign: Code = TEXT("Parameters.TwoSidedSign"); break;
-		case EExternalInput::VertexColor: Code = TEXT("Parameters.VertexColor"); EmitMaterialData.bUsesVertexColor |= (Context.ShaderFrequency != SF_Vertex); break;
-		case EExternalInput::VertexColor_Ddx: Code = TEXT("Parameters.VertexColor_DDX"); EmitMaterialData.bUsesVertexColor |= (Context.ShaderFrequency != SF_Vertex); break;
-		case EExternalInput::VertexColor_Ddy: Code = TEXT("Parameters.VertexColor_DDY"); EmitMaterialData.bUsesVertexColor |= (Context.ShaderFrequency != SF_Vertex); break;
+		case EExternalInput::VertexColor: Code = TEXT("Parameters.VertexColor"); break;
+		case EExternalInput::VertexColor_Ddx: Code = TEXT("Parameters.VertexColor_DDX"); break;
+		case EExternalInput::VertexColor_Ddy: Code = TEXT("Parameters.VertexColor_DDY"); break;
 		case EExternalInput::WorldPosition: Code = TEXT("GetWorldPosition(Parameters)"); break;
 		case EExternalInput::WorldPosition_NoOffsets: Code = TEXT("GetWorldPosition_NoMaterialOffsets(Parameters)"); break;
 		case EExternalInput::TranslatedWorldPosition: Code = TEXT("GetTranslatedWorldPosition(Parameters)"); break;
@@ -285,7 +282,7 @@ void FExpressionExternalInput::EmitValueShader(FEmitContext& Context, FEmitScope
 		case EExternalInput::TemporalSampleOffset: Code = TEXT("View.TemporalAAParams.zw"); break;
 		case EExternalInput::PreExposure: Code = TEXT("View.PreExposure.x"); break;
 		case EExternalInput::RcpPreExposure: Code = TEXT("View.OneOverPreExposure.x"); break;
-		case EExternalInput::EyeAdaptation: Code = TEXT("EyeAdaptationLookup()"); Context.MaterialCompilationOutput->bUsesEyeAdaptation = true; break;
+		case EExternalInput::EyeAdaptation: Code = TEXT("EyeAdaptationLookup()"); break;
 		case EExternalInput::RuntimeVirtualTextureOutputLevel:  Code = TEXT("View.RuntimeVirtualTextureMipLevel.x"); break;
 		case EExternalInput::RuntimeVirtualTextureOutputDerivative: Code = TEXT("View.RuntimeVirtualTextureMipLevel.zw"); break;
 		case EExternalInput::RuntimeVirtualTextureMaxLevel:  Code = TEXT("View.RuntimeVirtualTextureMipLevel.y"); break;
@@ -328,9 +325,9 @@ void FExpressionExternalInput::EmitValueShader(FEmitContext& Context, FEmitScope
 		case EExternalInput::PrevGameTime: Code = TEXT("View.PrevFrameGameTime"); break;
 		case EExternalInput::PrevRealTime: Code = TEXT("View.PrevFrameRealTime"); break;
 
-		case EExternalInput::ParticleColor: Code = TEXT("Parameters.Particle.Color"); EmitMaterialData.bUsesParticleColor |= (Context.ShaderFrequency != SF_Vertex); break;
-		case EExternalInput::ParticleTranslatedWorldPosition: Code = TEXT("Parameters.Particle.TranslatedWorldPositionAndSize.xyz"); EmitMaterialData.bNeedsParticlePosition = true; break;
-		case EExternalInput::ParticleRadius: Code = TEXT("Parameters.Particle.TranslatedWorldPositionAndSize.w"); EmitMaterialData.bNeedsParticlePosition = true; break;
+		case EExternalInput::ParticleColor: Code = TEXT("Parameters.Particle.Color"); break;
+		case EExternalInput::ParticleTranslatedWorldPosition: Code = TEXT("Parameters.Particle.TranslatedWorldPositionAndSize.xyz"); break;
+		case EExternalInput::ParticleRadius: Code = TEXT("Parameters.Particle.TranslatedWorldPositionAndSize.w"); break;
 
 		default:
 			checkNoEntry();
@@ -349,13 +346,16 @@ void FExpressionShadingModel::ComputeAnalyticDerivatives(FTree& Tree, FExpressio
 
 bool FExpressionShadingModel::PrepareValue(FEmitContext& Context, FEmitScope& Scope, const FRequestedType& RequestedType, FPrepareValueResult& OutResult) const
 {
+	if (Context.bMarkLiveValues)
+	{
+		Context.FindData<FEmitData>().ShadingModelsFromCompilation.AddShadingModel(ShadingModel);
+	}
+
 	return OutResult.SetType(Context, RequestedType, EExpressionEvaluation::Constant, Shader::EValueType::Int1);
 }
 
 void FExpressionShadingModel::EmitValuePreshader(FEmitContext& Context, FEmitScope& Scope, const FRequestedType& RequestedType, FEmitValuePreshaderResult& OutResult) const
 {
-	Context.FindData<FEmitData>().ShadingModelsFromCompilation.AddShadingModel(ShadingModel);
-
 	Context.PreshaderStackPosition++;
 	OutResult.Type = Shader::EValueType::Int1;
 	OutResult.Preshader.WriteOpcode(Shader::EPreshaderOpcode::Constant).Write(Shader::FValue((int32)ShadingModel));
@@ -376,32 +376,34 @@ void FExpressionParameter::ComputeAnalyticDerivatives(FTree& Tree, FExpressionDe
 bool FExpressionParameter::PrepareValue(FEmitContext& Context, FEmitScope& Scope, const FRequestedType& RequestedType, FPrepareValueResult& OutResult) const
 {
 	const EMaterialParameterType ParameterType = ParameterMeta.Value.Type;
-
-	FEmitData& EmitData = Context.FindData<FEmitData>();
-	if (EmitData.CachedExpressionData)
+	if (Context.bMarkLiveValues)
 	{
-		if (!ParameterInfo.Name.IsNone())
+		FEmitData& EmitData = Context.FindData<FEmitData>();
+		if (EmitData.CachedExpressionData)
 		{
-			EmitData.CachedExpressionData->Parameters.AddParameter(ParameterInfo, ParameterMeta);
-		}
+			if (!ParameterInfo.Name.IsNone())
+			{
+				EmitData.CachedExpressionData->Parameters.AddParameter(ParameterInfo, ParameterMeta);
+			}
 
-		UObject* ReferencedTexture = nullptr;
-		switch (ParameterType)
-		{
-		case EMaterialParameterType::Texture: ReferencedTexture = ParameterMeta.Value.Texture; break;
-		case EMaterialParameterType::RuntimeVirtualTexture: ReferencedTexture = ParameterMeta.Value.RuntimeVirtualTexture; break;
-		case EMaterialParameterType::Font:
-			if (ParameterMeta.Value.Font.Value && ParameterMeta.Value.Font.Value->Textures.IsValidIndex(ParameterMeta.Value.Font.Page))
+			UObject* ReferencedTexture = nullptr;
+			switch (ParameterType)
 			{
-				ReferencedTexture = ParameterMeta.Value.Font.Value->Textures[ParameterMeta.Value.Font.Page];
+			case EMaterialParameterType::Texture: ReferencedTexture = ParameterMeta.Value.Texture; break;
+			case EMaterialParameterType::RuntimeVirtualTexture: ReferencedTexture = ParameterMeta.Value.RuntimeVirtualTexture; break;
+			case EMaterialParameterType::Font:
+				if (ParameterMeta.Value.Font.Value && ParameterMeta.Value.Font.Value->Textures.IsValidIndex(ParameterMeta.Value.Font.Page))
+				{
+					ReferencedTexture = ParameterMeta.Value.Font.Value->Textures[ParameterMeta.Value.Font.Page];
+				}
+				break;
+			default:
+				break;
+			}
+			if (ReferencedTexture)
+			{
+				EmitData.CachedExpressionData->ReferencedTextures.AddUnique(ReferencedTexture);
 			}
-			break;
-		default:
-			break;
-		}
-		if (ReferencedTexture)
-		{
-			EmitData.CachedExpressionData->ReferencedTextures.AddUnique(ReferencedTexture);
 		}
 	}
 
@@ -543,12 +545,14 @@ void FExpressionParameter::EmitValuePreshader(FEmitContext& Context, FEmitScope&
 
 bool FExpressionTextureSize::PrepareValue(FEmitContext& Context, FEmitScope& Scope, const FRequestedType& RequestedType, FPrepareValueResult& OutResult) const
 {
-	FEmitData& EmitData = Context.FindData<FEmitData>();
-	if (EmitData.CachedExpressionData)
+	if (Context.bMarkLiveValues)
 	{
-		EmitData.CachedExpressionData->ReferencedTextures.AddUnique(Texture);
+		FEmitData& EmitData = Context.FindData<FEmitData>();
+		if (EmitData.CachedExpressionData)
+		{
+			EmitData.CachedExpressionData->ReferencedTextures.AddUnique(Texture);
+		}
 	}
-
 	return OutResult.SetType(Context, RequestedType, EExpressionEvaluation::Preshader, Shader::EValueType::Float2);
 }
 
@@ -564,43 +568,48 @@ void FExpressionTextureSize::EmitValuePreshader(FEmitContext& Context, FEmitScop
 
 bool FExpressionFunctionCall::PrepareValue(FEmitContext& Context, FEmitScope& Scope, const FRequestedType& RequestedType, FPrepareValueResult& OutResult) const
 {
-	FEmitData& EmitMaterialData = Context.FindData<FEmitData>();
-	if (EmitMaterialData.CachedExpressionData)
+	if (Context.bMarkLiveValues)
 	{
-		FMaterialFunctionInfo NewFunctionInfo;
-		NewFunctionInfo.Function = MaterialFunction;
-		NewFunctionInfo.StateId = MaterialFunction->StateId;
-		EmitMaterialData.CachedExpressionData->FunctionInfos.AddUnique(NewFunctionInfo);
+		FEmitData& EmitMaterialData = Context.FindData<FEmitData>();
+		if (EmitMaterialData.CachedExpressionData)
+		{
+			FMaterialFunctionInfo NewFunctionInfo;
+			NewFunctionInfo.Function = MaterialFunction;
+			NewFunctionInfo.StateId = MaterialFunction->StateId;
+			EmitMaterialData.CachedExpressionData->FunctionInfos.AddUnique(NewFunctionInfo);
+		}
 	}
-
 	return FExpressionForward::PrepareValue(Context, Scope, RequestedType, OutResult);
 }
 
 bool FExpressionMaterialLayers::PrepareValue(FEmitContext& Context, FEmitScope& Scope, const FRequestedType& RequestedType, FPrepareValueResult& OutResult) const
 {
-	FEmitData& EmitMaterialData = Context.FindData<FEmitData>();
-	if (EmitMaterialData.CachedExpressionData)
+	if (Context.bMarkLiveValues)
 	{
-		// Only a single set of material layers are supported
-		// There should only be a single FExpressionMaterialLayers, but PrepareValue may be called multiple times, only need to capture the layers once
-		if (!EmitMaterialData.CachedExpressionData->bHasMaterialLayers)
+		FEmitData& EmitMaterialData = Context.FindData<FEmitData>();
+		if (EmitMaterialData.CachedExpressionData)
 		{
-			// TODO(?) - Layers for MIs are currently duplicated here and in FStaticParameterSet
-			EmitMaterialData.CachedExpressionData->bHasMaterialLayers = true;
-			EmitMaterialData.CachedExpressionData->MaterialLayers = *MaterialLayers;
+			// Only a single set of material layers are supported
+			// There should only be a single FExpressionMaterialLayers, but PrepareValue may be called multiple times, only need to capture the layers once
+			if (!EmitMaterialData.CachedExpressionData->bHasMaterialLayers)
+			{
+				// TODO(?) - Layers for MIs are currently duplicated here and in FStaticParameterSet
+				EmitMaterialData.CachedExpressionData->bHasMaterialLayers = true;
+				EmitMaterialData.CachedExpressionData->MaterialLayers = *MaterialLayers;
+			}
 		}
 	}
-
 	return FExpressionForward::PrepareValue(Context, Scope, RequestedType, OutResult);
 }
 
 bool FExpressionSceneTexture::PrepareValue(FEmitContext& Context, FEmitScope& Scope, const FRequestedType& RequestedType, FPrepareValueResult& OutResult) const
 {
 	Context.PrepareExpression(TexCoordExpression, Scope, ERequestedType::Vector2);
-
-	Context.MaterialCompilationOutput->bNeedsSceneTextures = true;
-	Context.MaterialCompilationOutput->SetIsSceneTextureUsed((ESceneTextureId)SceneTextureId);
-
+	if (Context.bMarkLiveValues)
+	{
+		Context.MaterialCompilationOutput->bNeedsSceneTextures = true;
+		Context.MaterialCompilationOutput->SetIsSceneTextureUsed((ESceneTextureId)SceneTextureId);
+	}
 	return OutResult.SetType(Context, RequestedType, EExpressionEvaluation::Shader, Shader::EValueType::Float4);
 }
 
diff --git a/Engine/Source/Runtime/Engine/Private/Materials/MaterialHLSLTree.h b/Engine/Source/Runtime/Engine/Private/Materials/MaterialHLSLTree.h
--- a/Engine/Source/Runtime/Engine/Private/Materials/MaterialHLSLTree.h
+++ b/Engine/Source/Runtime/Engine/Private/Materials/MaterialHLSLTree.h
@@ -140,6 +140,8 @@ enum class EExternalInput : uint8
 	ParticleColor,
 	ParticleTranslatedWorldPosition,
 	ParticleRadius,
+
+	Num,
 };
 static constexpr int32 MaxNumTexCoords = 8;
 
@@ -346,21 +348,18 @@ public:
 	{
 		for (int32 Index = 0; Index < SF_NumFrequencies; ++Index)
 		{
-			ExternalInputMask[Index].Init(false, 256);
+			ExternalInputMask[Index].Init(false, (int32)EExternalInput::Num);
 		}
 	}
 
 	const FStaticParameterSet* StaticParameters = nullptr;
 	FMaterialCachedExpressionData* CachedExpressionData = nullptr;
 	TMap<Shader::FValue, uint32> DefaultUniformValues;
-	bool bReadMaterialNormal = false;
-	bool bUsesVertexColor = false;
-	bool bNeedsParticlePosition = false;
-	bool bUsesParticleColor = false;
 	TBitArray<> ExternalInputMask[SF_NumFrequencies];
 	FMaterialShadingModelField ShadingModelsFromCompilation;
 
-	bool IsExternalInputUsed(EShaderFrequency Frequency, EExternalInput Input) const { return ExternalInputMask[Frequency][(uint8)Input]; }
+	bool IsExternalInputUsed(EShaderFrequency Frequency, EExternalInput Input) const { return ExternalInputMask[Frequency][(int32)Input]; }
+	bool IsExternalInputUsed(EExternalInput Input) const { return IsExternalInputUsed(SF_Vertex, Input) || IsExternalInputUsed(SF_Pixel, Input); }
 };
 
 } // namespace UE::HLSLTree::Material
diff --git a/Engine/Source/Runtime/Engine/Public/HLSLTree/HLSLTree.h b/Engine/Source/Runtime/Engine/Public/HLSLTree/HLSLTree.h
--- a/Engine/Source/Runtime/Engine/Public/HLSLTree/HLSLTree.h
+++ b/Engine/Source/Runtime/Engine/Public/HLSLTree/HLSLTree.h
@@ -381,18 +381,24 @@ struct FExpressionDerivatives
 class FExpression : public FNode
 {
 public:
-	FEmitShaderExpression* GetValueShader(FEmitContext& Context, FEmitScope& Scope, const FRequestedType& RequestedType);
-	Shader::FType GetValuePreshader(FEmitContext& Context, FEmitScope& Scope, const FRequestedType& RequestedType, Shader::FPreshaderData& OutPreshader);
-	Shader::FValue GetValueConstant(FEmitContext& Context, FEmitScope& Scope, const FRequestedType& RequestedType);
+	FEmitShaderExpression* GetValueShader(FEmitContext& Context, FEmitScope& Scope, const FRequestedType& RequestedType, const FPreparedType& PreparedType, const Shader::FType& ResultType);
+	Shader::FType GetValuePreshader(FEmitContext& Context, FEmitScope& Scope, const FRequestedType& RequestedType, const FPreparedType& PreparedType, const Shader::FType& ResultType, Shader::FPreshaderData& OutPreshader);
+	Shader::FValue GetValueConstant(FEmitContext& Context, FEmitScope& Scope, const FRequestedType& RequestedType, const FPreparedType& PreparedType, const Shader::FType& ResultType);
 
 	FEmitShaderExpression* GetValueShader(FEmitContext& Context, FEmitScope& Scope, const FRequestedType& RequestedType, const Shader::FType& ResultType);
+	FEmitShaderExpression* GetValueShader(FEmitContext& Context, FEmitScope& Scope, const FRequestedType& RequestedType);
 	FEmitShaderExpression* GetValueShader(FEmitContext& Context, FEmitScope& Scope, const Shader::FType& ResultType);
 	FEmitShaderExpression* GetValueShader(FEmitContext& Context, FEmitScope& Scope, Shader::EValueType ResultType);
 	FEmitShaderExpression* GetValueShader(FEmitContext& Context, FEmitScope& Scope);
 
+	Shader::FType GetValuePreshader(FEmitContext& Context, FEmitScope& Scope, const FRequestedType& RequestedType, const Shader::FType& ResultType, Shader::FPreshaderData& OutPreshader);
+	Shader::FType GetValuePreshader(FEmitContext& Context, FEmitScope& Scope, const FRequestedType& RequestedType, Shader::FPreshaderData& OutPreshader);
+
+	Shader::FValue GetValueConstant(FEmitContext& Context, FEmitScope& Scope, const FRequestedType& RequestedType, const FPreparedType& PreparedType);
 	Shader::FValue GetValueConstant(FEmitContext& Context, FEmitScope& Scope, const FRequestedType& RequestedType, const Shader::FType& ResultType);
-	Shader::FValue GetValueConstant(FEmitContext& Context, FEmitScope& Scope, const Shader::FType& ResultType);
-	Shader::FValue GetValueConstant(FEmitContext& Context, FEmitScope& Scope, Shader::EValueType ResultType);
+	Shader::FValue GetValueConstant(FEmitContext& Context, FEmitScope& Scope, const FRequestedType& RequestedType);
+	Shader::FValue GetValueConstant(FEmitContext& Context, FEmitScope& Scope, const FPreparedType& PreparedType, const Shader::FType& ResultType);
+	Shader::FValue GetValueConstant(FEmitContext& Context, FEmitScope& Scope, const FPreparedType& PreparedType, Shader::EValueType ResultType);
 
 protected:
 	virtual void ComputeAnalyticDerivatives(FTree& Tree, FExpressionDerivatives& OutResult) const;
@@ -403,6 +409,7 @@ protected:
 
 	friend class FTree;
 	friend class FEmitContext;
+	friend class FExpressionForward;
 };
 
 class FFunction final : public FNode
diff --git a/Engine/Source/Runtime/Engine/Public/HLSLTree/HLSLTreeEmit.h b/Engine/Source/Runtime/Engine/Public/HLSLTree/HLSLTreeEmit.h
--- a/Engine/Source/Runtime/Engine/Public/HLSLTree/HLSLTreeEmit.h
+++ b/Engine/Source/Runtime/Engine/Public/HLSLTree/HLSLTreeEmit.h
@@ -474,7 +474,7 @@ public:
 		return nullptr;
 	}
 
-	FEmitShaderExpression* EmitPreshaderOrConstant(FEmitScope& Scope, const FRequestedType& RequestedType, FExpression* Expression);
+	FEmitShaderExpression* EmitPreshaderOrConstant(FEmitScope& Scope, const FRequestedType& RequestedType, const Shader::FType& ResultType, FExpression* Expression);
 	FEmitShaderExpression* EmitConstantZero(FEmitScope& Scope, const Shader::FType& Type);
 	FEmitShaderExpression* EmitCast(FEmitScope& Scope, FEmitShaderExpression* ShaderValue, const Shader::FType& DestType, EEmitCastFlags Flags = EEmitCastFlags::None);
 
@@ -502,6 +502,7 @@ public:
 	const Shader::FStructTypeRegistry* TypeRegistry = nullptr;
 	FTargetParameters TargetParameters;
 	EShaderFrequency ShaderFrequency = SF_Pixel;
+	bool bMarkLiveValues = false;
 	bool bUseAnalyticDerivatives = false;
 
 	TArray<FEmitShaderNode*> EmitNodes;
diff --git a/Engine/Source/Runtime/Engine/Public/MaterialCachedData.h b/Engine/Source/Runtime/Engine/Public/MaterialCachedData.h
--- a/Engine/Source/Runtime/Engine/Public/MaterialCachedData.h
+++ b/Engine/Source/Runtime/Engine/Public/MaterialCachedData.h
@@ -236,12 +236,12 @@ struct FMaterialCachedExpressionData
 	void UpdateForExpressions(const FMaterialCachedExpressionContext& Context, const TArray<TObjectPtr<UMaterialExpression>>& Expressions, EMaterialParameterAssociation Association, int32 ParameterIndex);
 	void UpdateForFunction(const FMaterialCachedExpressionContext& Context, UMaterialFunctionInterface* Function, EMaterialParameterAssociation Association, int32 ParameterIndex);
 	void UpdateForLayerFunctions(const FMaterialCachedExpressionContext& Context, const FMaterialLayersFunctions& LayerFunctions);
-	void UpdateForCachedHLSLTree(const FMaterialCachedHLSLTree& CachedTree, const FStaticParameterSet* StaticParameters);
+	ENGINE_API void UpdateForCachedHLSLTree(const FMaterialCachedHLSLTree& CachedTree, const FStaticParameterSet* StaticParameters);
 #endif // WITH_EDITOR
 
 	void Reset();
 
-	void AddReferencedObjects(FReferenceCollector& Collector);
+	ENGINE_API void AddReferencedObjects(FReferenceCollector& Collector);
 
 	/** Returns an array of the guids of functions used, with the call hierarchy flattened. */
 	void AppendReferencedFunctionIdsTo(TArray<FGuid>& OutIds) const;
