Add forward declare, fix non-unity build

diff --git a/Engine/Source/Runtime/Engine/Public/MaterialCachedHLSLTree.h b/Engine/Source/Runtime/Engine/Public/MaterialCachedHLSLTree.h
--- a/Engine/Source/Runtime/Engine/Public/MaterialCachedHLSLTree.h
+++ b/Engine/Source/Runtime/Engine/Public/MaterialCachedHLSLTree.h
@@ -10,6 +10,7 @@
 #include "Misc/MemStack.h"
 
 class UMaterial;
+class UMaterialExpression;
 class UMaterialExpressionCustomOutput;
 struct FMaterialLayersFunctions;
 
