Missing include, fix non-unity compile

diff --git a/Engine/Source/Runtime/Engine/Private/Materials/MaterialHLSLTree.cpp b/Engine/Source/Runtime/Engine/Private/Materials/MaterialHLSLTree.cpp
--- a/Engine/Source/Runtime/Engine/Private/Materials/MaterialHLSLTree.cpp
+++ b/Engine/Source/Runtime/Engine/Private/Materials/MaterialHLSLTree.cpp
@@ -8,6 +8,7 @@
 #include "MaterialShared.h"
 #include "MaterialCachedData.h"
 #include "MaterialSceneTextureId.h"
+#include "Materials/MaterialFunctionInterface.h"
 #include "Engine/BlendableInterface.h" // BL_AfterTonemapping
 #include "VT/RuntimeVirtualTexture.h"
 #include "Engine/Texture2D.h"
