const-correctness changes in new HLSL translator:
FExpression* is mostly always 'const', since expressions should not change after creation
UMaterialExpression::GenerateHLSLStatements/Expression are const-methods

diff --git a/Engine/Source/Runtime/Engine/Classes/Materials/MaterialExpression.h b/Engine/Source/Runtime/Engine/Classes/Materials/MaterialExpression.h
--- a/Engine/Source/Runtime/Engine/Classes/Materials/MaterialExpression.h
+++ b/Engine/Source/Runtime/Engine/Classes/Materials/MaterialExpression.h
@@ -303,8 +303,8 @@ class ENGINE_API UMaterialExpression : public UObject
 	 * For example, a for-loop expression might generate a statement for the execution input, but also generate an expression to access the loop index
 	 * These methods replace the Compile() method; once we switch over to the new system, Compile() will be removed
 	 */
-	virtual bool GenerateHLSLStatements(FMaterialHLSLGenerator& Generator, UE::HLSLTree::FScope& Scope);
-	virtual bool GenerateHLSLExpression(FMaterialHLSLGenerator& Generator, UE::HLSLTree::FScope& Scope, int32 OutputIndex, UE::HLSLTree::FExpression* &OutExpression);
+	virtual bool GenerateHLSLStatements(FMaterialHLSLGenerator& Generator, UE::HLSLTree::FScope& Scope) const;
+	virtual bool GenerateHLSLExpression(FMaterialHLSLGenerator& Generator, UE::HLSLTree::FScope& Scope, int32 OutputIndex, UE::HLSLTree::FExpression const*& OutExpression) const;
 
 #endif // WITH_EDITOR
 
diff --git a/Engine/Source/Runtime/Engine/Classes/Materials/MaterialExpressionAbs.h b/Engine/Source/Runtime/Engine/Classes/Materials/MaterialExpressionAbs.h
--- a/Engine/Source/Runtime/Engine/Classes/Materials/MaterialExpressionAbs.h
+++ b/Engine/Source/Runtime/Engine/Classes/Materials/MaterialExpressionAbs.h
@@ -27,7 +27,7 @@ class UMaterialExpressionAbs : public UMaterialExpression
 #if WITH_EDITOR
 	virtual int32 Compile(class FMaterialCompiler* Compiler, int32 OutputIndex) override;
 	virtual void GetCaption(TArray<FString>& OutCaptions) const override;
-	virtual bool GenerateHLSLExpression(FMaterialHLSLGenerator& Generator, UE::HLSLTree::FScope& Scope, int32 OutputIndex, UE::HLSLTree::FExpression*& OutExpression) override;
+	virtual bool GenerateHLSLExpression(FMaterialHLSLGenerator& Generator, UE::HLSLTree::FScope& Scope, int32 OutputIndex, UE::HLSLTree::FExpression const*& OutExpression) const override;
 #endif
 	//~ End UMaterialExpression Interface
 
diff --git a/Engine/Source/Runtime/Engine/Classes/Materials/MaterialExpressionActorPositionWS.h b/Engine/Source/Runtime/Engine/Classes/Materials/MaterialExpressionActorPositionWS.h
--- a/Engine/Source/Runtime/Engine/Classes/Materials/MaterialExpressionActorPositionWS.h
+++ b/Engine/Source/Runtime/Engine/Classes/Materials/MaterialExpressionActorPositionWS.h
@@ -18,7 +18,7 @@ class UMaterialExpressionActorPositionWS : public UMaterialExpression
 #if WITH_EDITOR
 	virtual int32 Compile(class FMaterialCompiler* Compiler, int32 OutputIndex) override;
 	virtual void GetCaption(TArray<FString>& OutCaptions) const override;
-	virtual bool GenerateHLSLExpression(FMaterialHLSLGenerator& Generator, UE::HLSLTree::FScope& Scope, int32 OutputIndex, UE::HLSLTree::FExpression*& OutExpression) override;
+	virtual bool GenerateHLSLExpression(FMaterialHLSLGenerator& Generator, UE::HLSLTree::FScope& Scope, int32 OutputIndex, UE::HLSLTree::FExpression const*& OutExpression) const override;
 #endif
 	//~ End UMaterialExpression Interface
 };
diff --git a/Engine/Source/Runtime/Engine/Classes/Materials/MaterialExpressionExecBegin.h b/Engine/Source/Runtime/Engine/Classes/Materials/MaterialExpressionExecBegin.h
--- a/Engine/Source/Runtime/Engine/Classes/Materials/MaterialExpressionExecBegin.h
+++ b/Engine/Source/Runtime/Engine/Classes/Materials/MaterialExpressionExecBegin.h
@@ -23,7 +23,7 @@ class UMaterialExpressionExecBegin : public UMaterialExpression
 	virtual void GetCaption(TArray<FString>& OutCaptions) const override;
 	virtual bool CanUserDeleteExpression() const override { return false; }
 
-	virtual bool GenerateHLSLStatements(FMaterialHLSLGenerator& Generator, UE::HLSLTree::FScope& Scope) override;
+	virtual bool GenerateHLSLStatements(FMaterialHLSLGenerator& Generator, UE::HLSLTree::FScope& Scope) const override;
 #endif
 	//~ End UMaterialExpression Interface
 };
diff --git a/Engine/Source/Runtime/Engine/Classes/Materials/MaterialExpressionExecEnd.h b/Engine/Source/Runtime/Engine/Classes/Materials/MaterialExpressionExecEnd.h
--- a/Engine/Source/Runtime/Engine/Classes/Materials/MaterialExpressionExecEnd.h
+++ b/Engine/Source/Runtime/Engine/Classes/Materials/MaterialExpressionExecEnd.h
@@ -18,7 +18,7 @@ class UMaterialExpressionExecEnd : public UMaterialExpression
 	virtual bool CanUserDeleteExpression() const override { return false; }
 	virtual bool HasExecInput() override { return true; }
 
-	virtual bool GenerateHLSLStatements(FMaterialHLSLGenerator& Generator, UE::HLSLTree::FScope& Scope) override;
+	virtual bool GenerateHLSLStatements(FMaterialHLSLGenerator& Generator, UE::HLSLTree::FScope& Scope) const override;
 #endif
 	//~ End UMaterialExpression Interface
 };
diff --git a/Engine/Source/Runtime/Engine/Classes/Materials/MaterialExpressionForLoop.h b/Engine/Source/Runtime/Engine/Classes/Materials/MaterialExpressionForLoop.h
--- a/Engine/Source/Runtime/Engine/Classes/Materials/MaterialExpressionForLoop.h
+++ b/Engine/Source/Runtime/Engine/Classes/Materials/MaterialExpressionForLoop.h
@@ -37,8 +37,8 @@ class UMaterialExpressionForLoop : public UMaterialExpression
 	virtual uint32 GetOutputType(int32 InputIndex) override;
 	virtual bool HasExecInput() override { return true; }
 
-	virtual bool GenerateHLSLExpression(FMaterialHLSLGenerator& Generator, UE::HLSLTree::FScope& Scope, int32 OutputIndex, UE::HLSLTree::FExpression*& OutExpression) override;
-	virtual bool GenerateHLSLStatements(FMaterialHLSLGenerator& Generator, UE::HLSLTree::FScope& Scope) override;
+	virtual bool GenerateHLSLExpression(FMaterialHLSLGenerator& Generator, UE::HLSLTree::FScope& Scope, int32 OutputIndex, UE::HLSLTree::FExpression const*& OutExpression) const override;
+	virtual bool GenerateHLSLStatements(FMaterialHLSLGenerator& Generator, UE::HLSLTree::FScope& Scope) const override;
 #endif
 	//~ End UMaterialExpression Interface
 };
diff --git a/Engine/Source/Runtime/Engine/Classes/Materials/MaterialExpressionFunctionOutput.h b/Engine/Source/Runtime/Engine/Classes/Materials/MaterialExpressionFunctionOutput.h
--- a/Engine/Source/Runtime/Engine/Classes/Materials/MaterialExpressionFunctionOutput.h
+++ b/Engine/Source/Runtime/Engine/Classes/Materials/MaterialExpressionFunctionOutput.h
@@ -66,7 +66,7 @@ class UMaterialExpressionFunctionOutput : public UMaterialExpression
 	virtual bool IsResultStrataMaterial(int32 OutputIndex) override;
 	virtual void GatherStrataMaterialInfo(FStrataMaterialInfo& StrataMaterialInfo, int32 OutputIndex) override;
 	virtual FStrataOperator* StrataGenerateMaterialTopologyTree(class FMaterialCompiler* Compiler, class UMaterialExpression* Parent, int32 OutputIndex) override;
-	bool GenerateHLSLExpression(FMaterialHLSLGenerator& Generator, UE::HLSLTree::FScope& Scope, int32 OutputIndex, UE::HLSLTree::FExpression*& OutExpression);
+	bool GenerateHLSLExpression(FMaterialHLSLGenerator& Generator, UE::HLSLTree::FScope& Scope, int32 OutputIndex, UE::HLSLTree::FExpression const*& OutExpression) const;
 #endif // WITH_EDITOR
 	//~ End UMaterialExpression Interface
 
diff --git a/Engine/Source/Runtime/Engine/Classes/Materials/MaterialExpressionTextureSample.h b/Engine/Source/Runtime/Engine/Classes/Materials/MaterialExpressionTextureSample.h
--- a/Engine/Source/Runtime/Engine/Classes/Materials/MaterialExpressionTextureSample.h
+++ b/Engine/Source/Runtime/Engine/Classes/Materials/MaterialExpressionTextureSample.h
@@ -64,7 +64,7 @@ protected:
 	uint8 bShowTextureInputPin : 1;
 
 #if WITH_EDITOR
-	bool GenerateHLSLExpressionBase(FMaterialHLSLGenerator& Generator, UE::HLSLTree::FScope& Scope, UE::HLSLTree::FExpression* TextureExpression, UE::HLSLTree::FExpression*& OutExpression);
+	bool GenerateHLSLExpressionBase(FMaterialHLSLGenerator& Generator, UE::HLSLTree::FScope& Scope, const UE::HLSLTree::FExpression* TextureExpression, UE::HLSLTree::FExpression const*& OutExpression) const;
 #endif
 
 public:
@@ -99,7 +99,7 @@ public:
 	virtual bool MatchesSearchQuery( const TCHAR* SearchQuery ) override;
 	virtual uint32 GetInputType(int32 InputIndex) override;
 	virtual bool CanIgnoreOutputIndex() { return true; }
-	virtual bool GenerateHLSLExpression(FMaterialHLSLGenerator& Generator, UE::HLSLTree::FScope& Scope, int32 OutputIndex, UE::HLSLTree::FExpression*& OutExpression) override;
+	virtual bool GenerateHLSLExpression(FMaterialHLSLGenerator& Generator, UE::HLSLTree::FScope& Scope, int32 OutputIndex, UE::HLSLTree::FExpression const*& OutExpression) const override;
 #endif // WITH_EDITOR
 	//~ End UMaterialExpression Interface
 
diff --git a/Engine/Source/Runtime/Engine/Classes/Materials/MaterialExpressionWorldPosition.h b/Engine/Source/Runtime/Engine/Classes/Materials/MaterialExpressionWorldPosition.h
--- a/Engine/Source/Runtime/Engine/Classes/Materials/MaterialExpressionWorldPosition.h
+++ b/Engine/Source/Runtime/Engine/Classes/Materials/MaterialExpressionWorldPosition.h
@@ -37,7 +37,7 @@ class UMaterialExpressionWorldPosition : public UMaterialExpression
 	
 	//~ Begin UMaterialExpression Interface
 #if WITH_EDITOR
-	bool GenerateHLSLExpression(FMaterialHLSLGenerator& Generator, UE::HLSLTree::FScope& Scope, int32 OutputIndex, UE::HLSLTree::FExpression*& OutExpression);
+	bool GenerateHLSLExpression(FMaterialHLSLGenerator& Generator, UE::HLSLTree::FScope& Scope, int32 OutputIndex, UE::HLSLTree::FExpression const*& OutExpression) const;
 	virtual int32 Compile(class FMaterialCompiler* Compiler, int32 OutputIndex) override;
 	virtual void GetCaption(TArray<FString>& OutCaptions) const override;
 	virtual FText GetKeywords() const override {return FText::FromString(TEXT("position"));}
diff --git a/Engine/Source/Runtime/Engine/Private/HLSLTree/HLSLTree.cpp b/Engine/Source/Runtime/Engine/Private/HLSLTree/HLSLTree.cpp
--- a/Engine/Source/Runtime/Engine/Private/HLSLTree/HLSLTree.cpp
+++ b/Engine/Source/Runtime/Engine/Private/HLSLTree/HLSLTree.cpp
@@ -48,7 +48,7 @@ public:
 	FExpressionLocalPHI(const FExpressionLocalPHI* Source, ELocalPHIChainType Type, const FRequestedType& RequestedType = FRequestedType());
 
 	virtual void ComputeAnalyticDerivatives(FTree& Tree, FExpressionDerivatives& OutResult) const override;
-	virtual FExpression* ComputePreviousFrame(FTree& Tree, const FRequestedType& RequestedType) const override;
+	virtual const FExpression* ComputePreviousFrame(FTree& Tree, const FRequestedType& RequestedType) const override;
 	virtual bool PrepareValue(FEmitContext& Context, FEmitScope& Scope, const FRequestedType& RequestedType, FPrepareValueResult& OutResult) const override;
 	virtual void EmitValueShader(FEmitContext& Context, FEmitScope& Scope, const FRequestedType& RequestedType, FEmitValueShaderResult& OutResult) const override;
 	virtual void EmitValuePreshader(FEmitContext& Context, FEmitScope& Scope, const FRequestedType& RequestedType, FEmitValuePreshaderResult& OutResult) const override;
@@ -56,7 +56,7 @@ public:
 	TArray<FLocalPHIChainEntry, TInlineAllocator<8>> Chain;
 	FName LocalName;
 	FScope* Scopes[MaxNumPreviousScopes];
-	FExpression* Values[MaxNumPreviousScopes];
+	const FExpression* Values[MaxNumPreviousScopes];
 	int32 NumValues = 0;
 };
 
@@ -195,7 +195,7 @@ void FExpressionLocalPHI::ComputeAnalyticDerivatives(FTree& Tree, FExpressionDer
 	OutResult.ExpressionDdy = Tree.NewExpression<FExpressionLocalPHI>(this, ELocalPHIChainType::Ddy);
 }
 
-FExpression* FExpressionLocalPHI::ComputePreviousFrame(FTree& Tree, const FRequestedType& RequestedType) const
+const FExpression* FExpressionLocalPHI::ComputePreviousFrame(FTree& Tree, const FRequestedType& RequestedType) const
 {
 	return Tree.NewExpression<FExpressionLocalPHI>(this, ELocalPHIChainType::PreviousFrame, RequestedType);
 }
@@ -206,7 +206,7 @@ bool FExpressionLocalPHI::PrepareValue(FEmitContext& Context, FEmitScope& Scope,
 
 	FEmitScope* EmitValueScopes[MaxNumPreviousScopes];
 	FScope* LiveScopes[MaxNumPreviousScopes];
-	FExpression* LiveValues[MaxNumPreviousScopes];
+	const FExpression* LiveValues[MaxNumPreviousScopes];
 	int32 NumLiveScopes = 0;
 
 	for (int32 i = 0; i < NumValues; ++i)
@@ -319,7 +319,7 @@ struct FLocalPHILiveScopes
 {
 	FEmitScope* EmitDeclarationScope = nullptr;
 	FEmitScope* EmitValueScopes[MaxNumPreviousScopes];
-	FExpression* LiveValues[MaxNumPreviousScopes];
+	const FExpression* LiveValues[MaxNumPreviousScopes];
 	int32 NumScopes = 0;
 	bool bCanForwardValue = true;
 };
@@ -1069,7 +1069,7 @@ void FExpression::ComputeAnalyticDerivatives(FTree& Tree, FExpressionDerivatives
 	// nop
 }
 
-FExpression* FExpression::ComputePreviousFrame(FTree& Tree, const FRequestedType& RequestedType) const
+const FExpression* FExpression::ComputePreviousFrame(FTree& Tree, const FRequestedType& RequestedType) const
 {
 	return nullptr;
 }
@@ -1084,7 +1084,7 @@ void FExpression::EmitValuePreshader(FEmitContext& Context, FEmitScope& Scope, c
 	check(false);
 }
 
-FEmitShaderExpression* FExpression::GetValueShader(FEmitContext& Context, FEmitScope& Scope, const FRequestedType& RequestedType, const FPreparedType& PreparedType, const Shader::FType& ResultType)
+FEmitShaderExpression* FExpression::GetValueShader(FEmitContext& Context, FEmitScope& Scope, const FRequestedType& RequestedType, const FPreparedType& PreparedType, const Shader::FType& ResultType) const
 {
 	FOwnerScope OwnerScope(*Context.Errors, GetOwner());
 
@@ -1110,35 +1110,35 @@ FEmitShaderExpression* FExpression::GetValueShader(FEmitContext& Context, FEmitS
 	return Context.EmitCast(Scope, Value, ResultType);
 }
 
-FEmitShaderExpression* FExpression::GetValueShader(FEmitContext& Context, FEmitScope& Scope, const FRequestedType& RequestedType, const Shader::FType& ResultType)
+FEmitShaderExpression* FExpression::GetValueShader(FEmitContext& Context, FEmitScope& Scope, const FRequestedType& RequestedType, const Shader::FType& ResultType) const
 {
 	const FPreparedType& PreparedType = Context.GetPreparedType(this);
 	return GetValueShader(Context, Scope, RequestedType, PreparedType, ResultType);
 }
 
-FEmitShaderExpression* FExpression::GetValueShader(FEmitContext& Context, FEmitScope& Scope, const FRequestedType& RequestedType)
+FEmitShaderExpression* FExpression::GetValueShader(FEmitContext& Context, FEmitScope& Scope, const FRequestedType& RequestedType) const
 {
 	const FPreparedType& PreparedType = Context.GetPreparedType(this);
 	return GetValueShader(Context, Scope, RequestedType, PreparedType, PreparedType.GetType());
 }
 
-FEmitShaderExpression* FExpression::GetValueShader(FEmitContext& Context, FEmitScope& Scope, const Shader::FType& ResultType)
+FEmitShaderExpression* FExpression::GetValueShader(FEmitContext& Context, FEmitScope& Scope, const Shader::FType& ResultType) const
 {
 	return GetValueShader(Context, Scope, ResultType, ResultType);
 }
 
-FEmitShaderExpression* FExpression::GetValueShader(FEmitContext& Context, FEmitScope& Scope, Shader::EValueType ResultType)
+FEmitShaderExpression* FExpression::GetValueShader(FEmitContext& Context, FEmitScope& Scope, Shader::EValueType ResultType) const
 {
 	return GetValueShader(Context, Scope, ResultType, ResultType);
 }
 
-FEmitShaderExpression* FExpression::GetValueShader(FEmitContext& Context, FEmitScope& Scope)
+FEmitShaderExpression* FExpression::GetValueShader(FEmitContext& Context, FEmitScope& Scope) const
 {
 	const FPreparedType& PreparedType = Context.GetPreparedType(this);
 	return GetValueShader(Context, Scope, PreparedType.GetRequestedType(), PreparedType, PreparedType.GetType());
 }
 
-Shader::FType FExpression::GetValuePreshader(FEmitContext& Context, FEmitScope& Scope, const FRequestedType& RequestedType, const FPreparedType& PreparedType, const Shader::FType& ResultType, Shader::FPreshaderData& OutPreshader)
+Shader::FType FExpression::GetValuePreshader(FEmitContext& Context, FEmitScope& Scope, const FRequestedType& RequestedType, const FPreparedType& PreparedType, const Shader::FType& ResultType, Shader::FPreshaderData& OutPreshader) const
 {
 	FOwnerScope OwnerScope(*Context.Errors, GetOwner());
 	
@@ -1170,19 +1170,19 @@ Shader::FType FExpression::GetValuePreshader(FEmitContext& Context, FEmitScope&
 	return Result.Type;
 }
 
-Shader::FType FExpression::GetValuePreshader(FEmitContext& Context, FEmitScope& Scope, const FRequestedType& RequestedType, const Shader::FType& ResultType, Shader::FPreshaderData& OutPreshader)
+Shader::FType FExpression::GetValuePreshader(FEmitContext& Context, FEmitScope& Scope, const FRequestedType& RequestedType, const Shader::FType& ResultType, Shader::FPreshaderData& OutPreshader) const
 {
 	const FPreparedType& PreparedType = Context.GetPreparedType(this);
 	return GetValuePreshader(Context, Scope, RequestedType, PreparedType, ResultType, OutPreshader);
 }
 
-Shader::FType FExpression::GetValuePreshader(FEmitContext& Context, FEmitScope& Scope, const FRequestedType& RequestedType, Shader::FPreshaderData& OutPreshader)
+Shader::FType FExpression::GetValuePreshader(FEmitContext& Context, FEmitScope& Scope, const FRequestedType& RequestedType, Shader::FPreshaderData& OutPreshader) const
 {
 	const FPreparedType& PreparedType = Context.GetPreparedType(this);
 	return GetValuePreshader(Context, Scope, RequestedType, PreparedType, PreparedType.GetType(), OutPreshader);
 }
 
-Shader::FValue FExpression::GetValueConstant(FEmitContext& Context, FEmitScope& Scope, const FRequestedType& RequestedType, const FPreparedType& PreparedType, const Shader::FType& ResultType)
+Shader::FValue FExpression::GetValueConstant(FEmitContext& Context, FEmitScope& Scope, const FRequestedType& RequestedType, const FPreparedType& PreparedType, const Shader::FType& ResultType) const
 {
 	FOwnerScope OwnerScope(*Context.Errors, GetOwner());
 
@@ -1211,29 +1211,29 @@ Shader::FValue FExpression::GetValueConstant(FEmitContext& Context, FEmitScope&
 	}
 }
 
-Shader::FValue FExpression::GetValueConstant(FEmitContext& Context, FEmitScope& Scope, const FRequestedType& RequestedType, const FPreparedType& PreparedType)
+Shader::FValue FExpression::GetValueConstant(FEmitContext& Context, FEmitScope& Scope, const FRequestedType& RequestedType, const FPreparedType& PreparedType) const
 {
 	return GetValueConstant(Context, Scope, RequestedType, PreparedType, PreparedType.GetType());
 }
 
-Shader::FValue FExpression::GetValueConstant(FEmitContext& Context, FEmitScope& Scope, const FRequestedType& RequestedType, const Shader::FType& ResultType)
+Shader::FValue FExpression::GetValueConstant(FEmitContext& Context, FEmitScope& Scope, const FRequestedType& RequestedType, const Shader::FType& ResultType) const
 {
 	const FPreparedType PreparedType = Context.GetPreparedType(this);
 	return GetValueConstant(Context, Scope, RequestedType, PreparedType, ResultType);
 }
 
-Shader::FValue FExpression::GetValueConstant(FEmitContext& Context, FEmitScope& Scope, const FRequestedType& RequestedType)
+Shader::FValue FExpression::GetValueConstant(FEmitContext& Context, FEmitScope& Scope, const FRequestedType& RequestedType) const
 {
 	const FPreparedType PreparedType = Context.GetPreparedType(this);
 	return GetValueConstant(Context, Scope, RequestedType, PreparedType, PreparedType.GetType());
 }
 
-Shader::FValue FExpression::GetValueConstant(FEmitContext& Context, FEmitScope& Scope, const FPreparedType& PreparedType, const Shader::FType& ResultType)
+Shader::FValue FExpression::GetValueConstant(FEmitContext& Context, FEmitScope& Scope, const FPreparedType& PreparedType, const Shader::FType& ResultType) const
 {
 	return GetValueConstant(Context, Scope, FRequestedType(ResultType), PreparedType, ResultType);
 }
 
-Shader::FValue FExpression::GetValueConstant(FEmitContext& Context, FEmitScope& Scope, const FPreparedType& PreparedType, Shader::EValueType ResultType)
+Shader::FValue FExpression::GetValueConstant(FEmitContext& Context, FEmitScope& Scope, const FPreparedType& PreparedType, Shader::EValueType ResultType) const
 {
 	return GetValueConstant(Context, Scope, FRequestedType(ResultType), PreparedType, ResultType);
 }
@@ -1306,7 +1306,7 @@ bool FTree::Finalize()
 		FExpressionLocalPHI* Expression = PHIExpressions.Pop(false);
 		for (int32 i = 0; i < Expression->NumValues; ++i)
 		{
-			FExpression* LocalValue = AcquireLocal(*Expression->Scopes[i], Expression->LocalName);
+			const FExpression* LocalValue = AcquireLocal(*Expression->Scopes[i], Expression->LocalName);
 			if (!LocalValue)
 			{
 				//Errorf(TEXT("Local %s is not assigned on all control paths"), *Expression->LocalName.ToString());
@@ -1374,9 +1374,9 @@ void FTree::RegisterNode(FNode* Node)
 	Nodes = Node;
 }
 
-FExpression* FTree::FindExpression(FXxHash64 Hash) const
+const FExpression* FTree::FindExpression(FXxHash64 Hash) const
 {
-	FExpression* const* FoundExpression = ExpressionMap.Find(Hash);
+	FExpression const* const* FoundExpression = ExpressionMap.Find(Hash);
 	if (FoundExpression)
 	{
 		return *FoundExpression;
@@ -1403,14 +1403,14 @@ void FTree::RegisterStatement(FScope& Scope, FStatement* Statement)
 	Scope.ContainedStatement = Statement;
 }
 
-void FTree::AssignLocal(FScope& Scope, const FName& LocalName, FExpression* Value)
+void FTree::AssignLocal(FScope& Scope, const FName& LocalName, const FExpression* Value)
 {
 	Scope.LocalMap.Add(LocalName, Value);
 }
 
-FExpression* FTree::AcquireLocal(FScope& Scope, const FName& LocalName)
+const FExpression* FTree::AcquireLocal(FScope& Scope, const FName& LocalName)
 {
-	FExpression** FoundExpression = Scope.LocalMap.Find(LocalName);
+	FExpression const* const* FoundExpression = Scope.LocalMap.Find(LocalName);
 	if (FoundExpression)
 	{
 		return *FoundExpression;
@@ -1419,7 +1419,7 @@ FExpression* FTree::AcquireLocal(FScope& Scope, const FName& LocalName)
 	const TArrayView<FScope*> PreviousScopes = Scope.GetPreviousScopes();
 	if (PreviousScopes.Num() > 1)
 	{
-		FExpression* Expression = NewExpression<FExpressionLocalPHI>(LocalName, PreviousScopes);
+		const FExpression* Expression = NewExpression<FExpressionLocalPHI>(LocalName, PreviousScopes);
 		Scope.LocalMap.Add(LocalName, Expression);
 		return Expression;
 	}
@@ -1432,7 +1432,7 @@ FExpression* FTree::AcquireLocal(FScope& Scope, const FName& LocalName)
 	return nullptr;
 }
 
-FExpression* FTree::NewFunctionCall(FScope& Scope, FFunction* Function, int32 OutputIndex)
+const FExpression* FTree::NewFunctionCall(FScope& Scope, FFunction* Function, int32 OutputIndex)
 {
 	FScope* CalledScope = &Scope;
 	if (Function->CalledScope)
@@ -1444,7 +1444,7 @@ FExpression* FTree::NewFunctionCall(FScope& Scope, FFunction* Function, int32 Ou
 	return NewExpression<FExpressionFunctionCall>(Function, OutputIndex);
 }
 
-FExpressionDerivatives FTree::GetAnalyticDerivatives(FExpression* InExpression)
+FExpressionDerivatives FTree::GetAnalyticDerivatives(const FExpression* InExpression)
 {
 	FExpressionDerivatives Derivatives;
 	if (InExpression)
@@ -1455,13 +1455,13 @@ FExpressionDerivatives FTree::GetAnalyticDerivatives(FExpression* InExpression)
 	return Derivatives;
 }
 
-FExpression* FTree::GetPreviousFrame(FExpression* InExpression, const FRequestedType& RequestedType)
+const FExpression* FTree::GetPreviousFrame(const FExpression* InExpression, const FRequestedType& RequestedType)
 {
-	FExpression* Result = InExpression;
+	const FExpression* Result = InExpression;
 	if (Result && !RequestedType.IsVoid())
 	{
 		FOwnerScope OwnerScope(*this, InExpression->GetOwner()); // Associate any newly created nodes with the same owner as the input expression
-		FExpression* PrevFrameExpression = InExpression->ComputePreviousFrame(*this, RequestedType);
+		const FExpression* PrevFrameExpression = InExpression->ComputePreviousFrame(*this, RequestedType);
 		if (PrevFrameExpression)
 		{
 			Result = PrevFrameExpression;
diff --git a/Engine/Source/Runtime/Engine/Private/HLSLTree/HLSLTreeCommon.cpp b/Engine/Source/Runtime/Engine/Private/HLSLTree/HLSLTreeCommon.cpp
--- a/Engine/Source/Runtime/Engine/Private/HLSLTree/HLSLTreeCommon.cpp
+++ b/Engine/Source/Runtime/Engine/Private/HLSLTree/HLSLTreeCommon.cpp
@@ -94,7 +94,7 @@ void FExpressionForward::ComputeAnalyticDerivatives(FTree& Tree, FExpressionDeri
 	return Expression->ComputeAnalyticDerivatives(Tree, OutResult);
 }
 
-FExpression* FExpressionForward::ComputePreviousFrame(FTree& Tree, const FRequestedType& RequestedType) const
+const FExpression* FExpressionForward::ComputePreviousFrame(FTree& Tree, const FRequestedType& RequestedType) const
 {
 	return Expression->ComputePreviousFrame(Tree, RequestedType);
 }
@@ -114,7 +114,7 @@ void FExpressionForward::EmitValuePreshader(FEmitContext& Context, FEmitScope& S
 	return Expression->EmitValuePreshader(Context, Scope, RequestedType, OutResult);
 }
 
-FExpression* FTree::NewConstant(const Shader::FValue& Value)
+const FExpression* FTree::NewConstant(const Shader::FValue& Value)
 {
 	return NewExpression<FExpressionConstant>(Value);
 }
@@ -321,7 +321,7 @@ void FExpressionGetStructField::ComputeAnalyticDerivatives(FTree& Tree, FExpress
 	}
 }
 
-FExpression* FExpressionGetStructField::ComputePreviousFrame(FTree& Tree, const FRequestedType& RequestedType) const
+const FExpression* FExpressionGetStructField::ComputePreviousFrame(FTree& Tree, const FRequestedType& RequestedType) const
 {
 	FRequestedType RequestedStructType;
 	RequestedStructType.SetField(Field, RequestedType);
@@ -381,14 +381,14 @@ void FExpressionSetStructField::ComputeAnalyticDerivatives(FTree& Tree, FExpress
 	}
 }
 
-FExpression* FExpressionSetStructField::ComputePreviousFrame(FTree& Tree, const FRequestedType& RequestedType) const
+const FExpression* FExpressionSetStructField::ComputePreviousFrame(FTree& Tree, const FRequestedType& RequestedType) const
 {
 	FRequestedType RequestedStructType(RequestedType);
 	RequestedStructType.ClearFieldRequested(Field);
-	FExpression* PrevStructExpression = Tree.GetPreviousFrame(StructExpression, RequestedStructType);
+	const FExpression* PrevStructExpression = Tree.GetPreviousFrame(StructExpression, RequestedStructType);
 
 	const FRequestedType RequestedFieldType = RequestedType.GetField(Field);
-	FExpression* PrevFieldExpression = Tree.GetPreviousFrame(FieldExpression, RequestedFieldType);
+	const FExpression* PrevFieldExpression = Tree.GetPreviousFrame(FieldExpression, RequestedFieldType);
 
 	return Tree.NewExpression<FExpressionSetStructField>(StructType, Field, PrevStructExpression, PrevFieldExpression);
 }
@@ -478,7 +478,7 @@ void FExpressionSelect::ComputeAnalyticDerivatives(FTree& Tree, FExpressionDeriv
 	OutResult.ExpressionDdy = Tree.NewExpression<FExpressionSelect>(ConditionExpression, TrueDerivatives.ExpressionDdy, FalseDerivatives.ExpressionDdy);
 }
 
-FExpression* FExpressionSelect::ComputePreviousFrame(FTree& Tree, const FRequestedType& RequestedType) const
+const FExpression* FExpressionSelect::ComputePreviousFrame(FTree& Tree, const FRequestedType& RequestedType) const
 {
 	return Tree.NewExpression<FExpressionSelect>(
 		Tree.GetPreviousFrame(ConditionExpression, ERequestedType::Scalar),
@@ -520,7 +520,7 @@ void FExpressionSelect::EmitValueShader(FEmitContext& Context, FEmitScope& Scope
 	if (IsConstantEvaluation(ConditionEvaluation))
 	{
 		const bool bCondition = ConditionExpression->GetValueConstant(Context, Scope, ConditionType, Shader::EValueType::Bool1).AsBoolScalar();
-		FExpression* InputExpression = bCondition ? TrueExpression : FalseExpression;
+		const FExpression* InputExpression = bCondition ? TrueExpression : FalseExpression;
 		OutResult.Code = InputExpression->GetValueShader(Context, Scope, RequestedType);
 	}
 	else
@@ -549,7 +549,7 @@ void FExpressionDerivative::ComputeAnalyticDerivatives(FTree& Tree, FExpressionD
 	// TODO
 }
 
-FExpression* FExpressionDerivative::ComputePreviousFrame(FTree& Tree, const FRequestedType& RequestedType) const
+const FExpression* FExpressionDerivative::ComputePreviousFrame(FTree& Tree, const FRequestedType& RequestedType) const
 {
 	return Tree.NewExpression<FExpressionDerivative>(Coord, Tree.GetPreviousFrame(Input, RequestedType));
 }
@@ -604,7 +604,7 @@ void FExpressionSwizzle::ComputeAnalyticDerivatives(FTree& Tree, FExpressionDeri
 	}
 }
 
-FExpression* FExpressionSwizzle::ComputePreviousFrame(FTree& Tree, const FRequestedType& RequestedType) const
+const FExpression* FExpressionSwizzle::ComputePreviousFrame(FTree& Tree, const FRequestedType& RequestedType) const
 {
 	const FRequestedType RequestedInputType = Parameters.GetRequestedInputType(RequestedType);
 	return Tree.NewExpression<FExpressionSwizzle>(Parameters, Tree.GetPreviousFrame(Input, RequestedInputType));
@@ -745,7 +745,7 @@ void FExpressionAppend::ComputeAnalyticDerivatives(FTree& Tree, FExpressionDeriv
 	}
 }
 
-FExpression* FExpressionAppend::ComputePreviousFrame(FTree& Tree, const FRequestedType& RequestedType) const
+const FExpression* FExpressionAppend::ComputePreviousFrame(FTree& Tree, const FRequestedType& RequestedType) const
 {
 	// TODO - requested type?
 	return Tree.NewExpression<FExpressionAppend>(Tree.GetPreviousFrame(Lhs, RequestedType), Tree.GetPreviousFrame(Rhs, RequestedType));
@@ -875,8 +875,8 @@ void FExpressionAppend::EmitValuePreshader(FEmitContext& Context, FEmitScope& Sc
 
 void FExpressionSwitchBase::ComputeAnalyticDerivatives(FTree& Tree, FExpressionDerivatives& OutResult) const
 {
-	FExpression* InputDdx[MaxInputs];
-	FExpression* InputDdy[MaxInputs];
+	const FExpression* InputDdx[MaxInputs];
+	const FExpression* InputDdy[MaxInputs];
 	for (int32 Index = 0; Index < NumInputs; ++Index)
 	{
 		const FExpressionDerivatives InputDerivative = Tree.GetAnalyticDerivatives(Input[Index]);
@@ -887,9 +887,9 @@ void FExpressionSwitchBase::ComputeAnalyticDerivatives(FTree& Tree, FExpressionD
 	OutResult.ExpressionDdy = NewSwitch(Tree, MakeArrayView(InputDdy, NumInputs));
 }
 
-FExpression* FExpressionSwitchBase::ComputePreviousFrame(FTree& Tree, const FRequestedType& RequestedType) const
+const FExpression* FExpressionSwitchBase::ComputePreviousFrame(FTree& Tree, const FRequestedType& RequestedType) const
 {
-	FExpression* InputPrevFrame[MaxInputs];
+	const FExpression* InputPrevFrame[MaxInputs];
 	for (int32 Index = 0; Index < NumInputs; ++Index)
 	{
 		InputPrevFrame[Index] = Tree.GetPreviousFrame(Input[Index], RequestedType);
diff --git a/Engine/Source/Runtime/Engine/Private/HLSLTree/HLSLTreeEmit.cpp b/Engine/Source/Runtime/Engine/Private/HLSLTree/HLSLTreeEmit.cpp
--- a/Engine/Source/Runtime/Engine/Private/HLSLTree/HLSLTreeEmit.cpp
+++ b/Engine/Source/Runtime/Engine/Private/HLSLTree/HLSLTreeEmit.cpp
@@ -372,7 +372,7 @@ EExpressionEvaluation FEmitContext::GetEvaluation(const FExpression* Expression,
 	return GetPreparedType(Expression).GetEvaluation(Scope, RequestedType);
 }
 
-FPreparedType FEmitContext::PrepareExpression(FExpression* InExpression, FEmitScope& Scope, const FRequestedType& RequestedType)
+FPreparedType FEmitContext::PrepareExpression(const FExpression* InExpression, FEmitScope& Scope, const FRequestedType& RequestedType)
 {
 	if (!InExpression || RequestedType.IsVoid())
 	{
@@ -857,7 +857,7 @@ void WriteMaterialUniformAccess(Shader::EValueComponentType ComponentType, uint3
 }
 } // namespace Private
 
-FEmitShaderExpression* FEmitContext::EmitPreshaderOrConstant(FEmitScope& Scope, const FRequestedType& RequestedType, const Shader::FType& ResultType, FExpression* Expression)
+FEmitShaderExpression* FEmitContext::EmitPreshaderOrConstant(FEmitScope& Scope, const FRequestedType& RequestedType, const Shader::FType& ResultType, const FExpression* Expression)
 {
 	Shader::FPreshaderData LocalPreshader;
 	Shader::FType Type = ResultType;
diff --git a/Engine/Source/Runtime/Engine/Private/HLSLTree/HLSLTreeOperation.cpp b/Engine/Source/Runtime/Engine/Private/HLSLTree/HLSLTreeOperation.cpp
--- a/Engine/Source/Runtime/Engine/Private/HLSLTree/HLSLTreeOperation.cpp
+++ b/Engine/Source/Runtime/Engine/Private/HLSLTree/HLSLTreeOperation.cpp
@@ -11,15 +11,15 @@ namespace UE::HLSLTree
 class FExpressionOperation : public FExpression
 {
 public:
-	FExpressionOperation(EOperation InOp, TConstArrayView<FExpression*> InInputs);
+	FExpressionOperation(EOperation InOp, TConstArrayView<const FExpression*> InInputs);
 
 	static constexpr int8 MaxInputs = 2;
 
 	EOperation Op;
-	FExpression* Inputs[MaxInputs];
+	const FExpression* Inputs[MaxInputs];
 
 	virtual void ComputeAnalyticDerivatives(FTree& Tree, FExpressionDerivatives& OutResult) const override;
-	virtual FExpression* ComputePreviousFrame(FTree& Tree, const FRequestedType& RequestedType) const override;
+	virtual const FExpression* ComputePreviousFrame(FTree& Tree, const FRequestedType& RequestedType) const override;
 	virtual bool PrepareValue(FEmitContext& Context, FEmitScope& Scope, const FRequestedType& RequestedType, FPrepareValueResult& OutResult) const override;
 	virtual void EmitValueShader(FEmitContext& Context, FEmitScope& Scope, const FRequestedType& RequestedType, FEmitValueShaderResult& OutResult) const override;
 	virtual void EmitValuePreshader(FEmitContext& Context, FEmitScope& Scope, const FRequestedType& RequestedType, FEmitValuePreshaderResult& OutResult) const override;
@@ -91,37 +91,37 @@ FOperationDescription GetOperationDescription(EOperation Op)
 	}
 }
 
-FExpression* FTree::NewUnaryOp(EOperation Op, FExpression* Input)
+const FExpression* FTree::NewUnaryOp(EOperation Op, const FExpression* Input)
 {
-	FExpression* Inputs[1] = { Input };
+	const FExpression* Inputs[1] = { Input };
 	return NewExpression<FExpressionOperation>(Op, Inputs);
 }
 
-FExpression* FTree::NewBinaryOp(EOperation Op, FExpression* Lhs, FExpression* Rhs)
+const FExpression* FTree::NewBinaryOp(EOperation Op, const FExpression* Lhs, const FExpression* Rhs)
 {
-	FExpression* Inputs[2] = { Lhs, Rhs };
+	const FExpression* Inputs[2] = { Lhs, Rhs };
 	return NewExpression<FExpressionOperation>(Op, Inputs);
 }
 
-FExpression* FTree::NewLog(FExpression* Input)
+const FExpression* FTree::NewLog(const FExpression* Input)
 {
 	const float Log2ToLog = 1.0f / FMath::Log2(EULERS_NUMBER);
 	return NewMul(NewLog2(Input), NewConstant(Log2ToLog));
 }
 
-FExpression* FTree::NewCross(FExpression* Lhs, FExpression* Rhs)
+const FExpression* FTree::NewCross(const FExpression* Lhs, const FExpression* Rhs)
 {
 	//c_P[0] = v_A[1] * v_B[2] - v_A[2] * v_B[1];
 	//c_P[1] = -(v_A[0] * v_B[2] - v_A[2] * v_B[0]);
 	//c_P[2] = v_A[0] * v_B[1] - v_A[1] * v_B[0];
-	FExpression* Lhs0 = NewExpression<FExpressionSwizzle>(FSwizzleParameters(1, 0, 0), Lhs);
-	FExpression* Lhs1 = NewExpression<FExpressionSwizzle>(FSwizzleParameters(2, 2, 1), Lhs);
-	FExpression* Rhs0 = NewExpression<FExpressionSwizzle>(FSwizzleParameters(2, 2, 1), Rhs);
-	FExpression* Rhs1 = NewExpression<FExpressionSwizzle>(FSwizzleParameters(1, 0, 0), Rhs);
+	const FExpression* Lhs0 = NewExpression<FExpressionSwizzle>(FSwizzleParameters(1, 0, 0), Lhs);
+	const FExpression* Lhs1 = NewExpression<FExpressionSwizzle>(FSwizzleParameters(2, 2, 1), Lhs);
+	const FExpression* Rhs0 = NewExpression<FExpressionSwizzle>(FSwizzleParameters(2, 2, 1), Rhs);
+	const FExpression* Rhs1 = NewExpression<FExpressionSwizzle>(FSwizzleParameters(1, 0, 0), Rhs);
 	return NewSub(NewMul(NewMul(Lhs0, Rhs0), NewConstant(FVector3f(1.0f, -1.0f, 1.0f))), NewMul(Lhs1, Rhs1));
 }
 
-FExpressionOperation::FExpressionOperation(EOperation InOp, TConstArrayView<FExpression*> InInputs) : Op(InOp)
+FExpressionOperation::FExpressionOperation(EOperation InOp, TConstArrayView<const FExpression*> InInputs) : Op(InOp)
 {
 	const FOperationDescription OpDesc = GetOperationDescription(InOp);
 	check(OpDesc.NumInputs == InInputs.Num());
@@ -382,22 +382,22 @@ void FExpressionOperation::ComputeAnalyticDerivatives(FTree& Tree, FExpressionDe
 		break;
 	case EOperation::Rcp:
 	{
-		FExpression* Result = Tree.NewRcp(Inputs[0]);
-		FExpression* dFdA = Tree.NewNeg(Tree.NewMul(Result, Result));
+		const FExpression* Result = Tree.NewRcp(Inputs[0]);
+		const FExpression* dFdA = Tree.NewNeg(Tree.NewMul(Result, Result));
 		OutResult.ExpressionDdx = Tree.NewMul(dFdA, InputDerivatives[0].ExpressionDdx);
 		OutResult.ExpressionDdy = Tree.NewMul(dFdA, InputDerivatives[0].ExpressionDdy);
 		break;
 	}
 	case EOperation::Sqrt:
 	{
-		FExpression* dFdA = Tree.NewMul(Tree.NewRsqrt(Tree.NewMax(Inputs[0], Tree.NewConstant(0.00001f))), Tree.NewConstant(0.5f));
+		const FExpression* dFdA = Tree.NewMul(Tree.NewRsqrt(Tree.NewMax(Inputs[0], Tree.NewConstant(0.00001f))), Tree.NewConstant(0.5f));
 		OutResult.ExpressionDdx = Tree.NewMul(dFdA, InputDerivatives[0].ExpressionDdx);
 		OutResult.ExpressionDdy = Tree.NewMul(dFdA, InputDerivatives[0].ExpressionDdy);
 		break;
 	}
 	case EOperation::Rsqrt:
 	{
-		FExpression* dFdA = Tree.NewMul(Tree.NewMul(Tree.NewRsqrt(Inputs[0]), Tree.NewRcp(Inputs[0])), Tree.NewConstant(-0.5f));
+		const FExpression* dFdA = Tree.NewMul(Tree.NewMul(Tree.NewRsqrt(Inputs[0]), Tree.NewRcp(Inputs[0])), Tree.NewConstant(-0.5f));
 		OutResult.ExpressionDdx = Tree.NewMul(dFdA, InputDerivatives[0].ExpressionDdx);
 		OutResult.ExpressionDdy = Tree.NewMul(dFdA, InputDerivatives[0].ExpressionDdy);
 		break;
@@ -411,21 +411,21 @@ void FExpressionOperation::ComputeAnalyticDerivatives(FTree& Tree, FExpressionDe
 		break;
 	case EOperation::Sin:
 	{
-		FExpression* dFdA = Tree.NewCos(Inputs[0]);
+		const FExpression* dFdA = Tree.NewCos(Inputs[0]);
 		OutResult.ExpressionDdx = Tree.NewMul(dFdA, InputDerivatives[0].ExpressionDdx);
 		OutResult.ExpressionDdy = Tree.NewMul(dFdA, InputDerivatives[0].ExpressionDdy);
 		break;
 	}
 	case EOperation::Cos:
 	{
-		FExpression* dFdA = Tree.NewNeg(Tree.NewSin(Inputs[0]));
+		const FExpression* dFdA = Tree.NewNeg(Tree.NewSin(Inputs[0]));
 		OutResult.ExpressionDdx = Tree.NewMul(dFdA, InputDerivatives[0].ExpressionDdx);
 		OutResult.ExpressionDdy = Tree.NewMul(dFdA, InputDerivatives[0].ExpressionDdy);
 		break;
 	}
 	case EOperation::Tan:
 	{
-		FExpression* dFdA = Tree.NewRcp(Tree.NewPow2(Tree.NewCos(Inputs[0])));
+		const FExpression* dFdA = Tree.NewRcp(Tree.NewPow2(Tree.NewCos(Inputs[0])));
 		OutResult.ExpressionDdx = Tree.NewMul(dFdA, InputDerivatives[0].ExpressionDdx);
 		OutResult.ExpressionDdy = Tree.NewMul(dFdA, InputDerivatives[0].ExpressionDdy);
 		break;
@@ -433,7 +433,7 @@ void FExpressionOperation::ComputeAnalyticDerivatives(FTree& Tree, FExpressionDe
 	case EOperation::Asin:
 	case EOperation::AsinFast:
 	{
-		FExpression* dFdA = Tree.NewRsqrt(Tree.NewMax(Tree.NewSub(Tree.NewConstant(1.0f), Tree.NewPow2(Inputs[0])), Tree.NewConstant(0.00001f)));
+		const FExpression* dFdA = Tree.NewRsqrt(Tree.NewMax(Tree.NewSub(Tree.NewConstant(1.0f), Tree.NewPow2(Inputs[0])), Tree.NewConstant(0.00001f)));
 		OutResult.ExpressionDdx = Tree.NewMul(dFdA, InputDerivatives[0].ExpressionDdx);
 		OutResult.ExpressionDdy = Tree.NewMul(dFdA, InputDerivatives[0].ExpressionDdy);
 		break;
@@ -441,7 +441,7 @@ void FExpressionOperation::ComputeAnalyticDerivatives(FTree& Tree, FExpressionDe
 	case EOperation::Acos:
 	case EOperation::AcosFast:
 	{
-		FExpression* dFdA = Tree.NewNeg(Tree.NewRsqrt(Tree.NewMax(Tree.NewSub(Tree.NewConstant(1.0f), Tree.NewPow2(Inputs[0])), Tree.NewConstant(0.00001f))));
+		const FExpression* dFdA = Tree.NewNeg(Tree.NewRsqrt(Tree.NewMax(Tree.NewSub(Tree.NewConstant(1.0f), Tree.NewPow2(Inputs[0])), Tree.NewConstant(0.00001f))));
 		OutResult.ExpressionDdx = Tree.NewMul(dFdA, InputDerivatives[0].ExpressionDdx);
 		OutResult.ExpressionDdy = Tree.NewMul(dFdA, InputDerivatives[0].ExpressionDdy);
 		break;
@@ -449,7 +449,7 @@ void FExpressionOperation::ComputeAnalyticDerivatives(FTree& Tree, FExpressionDe
 	case EOperation::Atan:
 	case EOperation::AtanFast:
 	{
-		FExpression* dFdA = Tree.NewRcp(Tree.NewAdd(Tree.NewPow2(Inputs[0]), Tree.NewConstant(1.0f)));
+		const FExpression* dFdA = Tree.NewRcp(Tree.NewAdd(Tree.NewPow2(Inputs[0]), Tree.NewConstant(1.0f)));
 		OutResult.ExpressionDdx = Tree.NewMul(dFdA, InputDerivatives[0].ExpressionDdx);
 		OutResult.ExpressionDdy = Tree.NewMul(dFdA, InputDerivatives[0].ExpressionDdy);
 		break;
@@ -457,9 +457,9 @@ void FExpressionOperation::ComputeAnalyticDerivatives(FTree& Tree, FExpressionDe
 	case EOperation::Atan2:
 	case EOperation::Atan2Fast:
 	{
-		FExpression* Denom = Tree.NewRcp(Tree.NewAdd(Tree.NewPow2(Inputs[0]), Tree.NewPow2(Inputs[1])));
-		FExpression* dFdA = Tree.NewMul(Inputs[1], Denom);
-		FExpression* dFdB = Tree.NewMul(Tree.NewNeg(Inputs[0]), Denom);
+		const FExpression* Denom = Tree.NewRcp(Tree.NewAdd(Tree.NewPow2(Inputs[0]), Tree.NewPow2(Inputs[1])));
+		const FExpression* dFdA = Tree.NewMul(Inputs[1], Denom);
+		const FExpression* dFdB = Tree.NewMul(Tree.NewNeg(Inputs[0]), Denom);
 		OutResult.ExpressionDdx = Tree.NewAdd(Tree.NewMul(dFdA, InputDerivatives[0].ExpressionDdx), Tree.NewMul(dFdB, InputDerivatives[1].ExpressionDdx));
 		OutResult.ExpressionDdy = Tree.NewAdd(Tree.NewMul(dFdA, InputDerivatives[0].ExpressionDdy), Tree.NewMul(dFdB, InputDerivatives[1].ExpressionDdy));
 		break;
@@ -487,9 +487,9 @@ void FExpressionOperation::ComputeAnalyticDerivatives(FTree& Tree, FExpressionDe
 		break;
 	case EOperation::Div:
 	{
-		FExpression* Denom = Tree.NewRcp(Tree.NewMul(Inputs[1], Inputs[1]));
-		FExpression* dFdA = Tree.NewMul(Inputs[1], Denom);
-		FExpression* dFdB = Tree.NewNeg(Tree.NewMul(Inputs[0], Denom));
+		const FExpression* Denom = Tree.NewRcp(Tree.NewMul(Inputs[1], Inputs[1]));
+		const FExpression* dFdA = Tree.NewMul(Inputs[1], Denom);
+		const FExpression* dFdB = Tree.NewNeg(Tree.NewMul(Inputs[0], Denom));
 		OutResult.ExpressionDdx = Tree.NewAdd(Tree.NewMul(dFdA, InputDerivatives[0].ExpressionDdx), Tree.NewMul(dFdB, InputDerivatives[1].ExpressionDdx));
 		OutResult.ExpressionDdy = Tree.NewAdd(Tree.NewMul(dFdA, InputDerivatives[0].ExpressionDdy), Tree.NewMul(dFdB, InputDerivatives[1].ExpressionDdy));
 		break;
@@ -501,14 +501,14 @@ void FExpressionOperation::ComputeAnalyticDerivatives(FTree& Tree, FExpressionDe
 		break;
 	case EOperation::Min:
 	{
-		FExpression* Cond = Tree.NewLess(Inputs[0], Inputs[1]);
+		const FExpression* Cond = Tree.NewLess(Inputs[0], Inputs[1]);
 		OutResult.ExpressionDdx = Tree.NewExpression<FExpressionSelect>(Cond, InputDerivatives[0].ExpressionDdx, InputDerivatives[1].ExpressionDdx);
 		OutResult.ExpressionDdy = Tree.NewExpression<FExpressionSelect>(Cond, InputDerivatives[0].ExpressionDdy, InputDerivatives[1].ExpressionDdy);
 		break;
 	}
 	case EOperation::Max:
 	{
-		FExpression* Cond = Tree.NewGreater(Inputs[0], Inputs[1]);
+		const FExpression* Cond = Tree.NewGreater(Inputs[0], Inputs[1]);
 		OutResult.ExpressionDdx = Tree.NewExpression<FExpressionSelect>(Cond, InputDerivatives[0].ExpressionDdx, InputDerivatives[1].ExpressionDdx);
 		OutResult.ExpressionDdy = Tree.NewExpression<FExpressionSelect>(Cond, InputDerivatives[0].ExpressionDdy, InputDerivatives[1].ExpressionDdy);
 		break;
@@ -526,11 +526,11 @@ void FExpressionOperation::ComputeAnalyticDerivatives(FTree& Tree, FExpressionDe
 	}
 }
 
-FExpression* FExpressionOperation::ComputePreviousFrame(FTree& Tree, const FRequestedType& RequestedType) const
+const FExpression* FExpressionOperation::ComputePreviousFrame(FTree& Tree, const FRequestedType& RequestedType) const
 {
 	const Private::FOperationRequestedTypes RequestedTypes = Private::GetOperationRequestedTypes(Op, RequestedType);
 	const FOperationDescription OpDesc = GetOperationDescription(Op);
-	FExpression* PrevFrameInputs[MaxInputs];
+	const FExpression* PrevFrameInputs[MaxInputs];
 	for (int32 Index = 0; Index < OpDesc.NumInputs; ++Index)
 	{
 		PrevFrameInputs[Index] = Tree.GetPreviousFrame(Inputs[Index], RequestedTypes.InputType[Index]);
