const-correctness changes in new HLSL translator:

diff --git a/Engine/Source/Runtime/Engine/Private/Materials/MaterialExpressionHLSL.cpp b/Engine/Source/Runtime/Engine/Private/Materials/MaterialExpressionHLSL.cpp
--- a/Engine/Source/Runtime/Engine/Private/Materials/MaterialExpressionHLSL.cpp
+++ b/Engine/Source/Runtime/Engine/Private/Materials/MaterialExpressionHLSL.cpp
@@ -147,62 +147,62 @@
 #include "Misc/MemStackUtility.h"
 #include "MaterialHLSLGenerator.h"
 
-bool UMaterialExpression::GenerateHLSLExpression(FMaterialHLSLGenerator& Generator, UE::HLSLTree::FScope& Scope, int32 OutputIndex, UE::HLSLTree::FExpression*& OutExpression)
+bool UMaterialExpression::GenerateHLSLExpression(FMaterialHLSLGenerator& Generator, UE::HLSLTree::FScope& Scope, int32 OutputIndex, UE::HLSLTree::FExpression const*& OutExpression) const
 {
 	return Generator.Error(TEXT("Node does not support expressions"));
 }
 
-bool UMaterialExpression::GenerateHLSLStatements(FMaterialHLSLGenerator& Generator, UE::HLSLTree::FScope& Scope)
+bool UMaterialExpression::GenerateHLSLStatements(FMaterialHLSLGenerator& Generator, UE::HLSLTree::FScope& Scope) const
 {
 	return Generator.Error(TEXT("Node does not support statements"));
 }
 
-bool UMaterialExpressionReroute::GenerateHLSLExpression(FMaterialHLSLGenerator& Generator, UE::HLSLTree::FScope& Scope, int32 OutputIndex, UE::HLSLTree::FExpression*& OutExpression)
+bool UMaterialExpressionReroute::GenerateHLSLExpression(FMaterialHLSLGenerator& Generator, UE::HLSLTree::FScope& Scope, int32 OutputIndex, UE::HLSLTree::FExpression const*& OutExpression) const
 {
 	OutExpression = Input.TryAcquireHLSLExpression(Generator, Scope);
 	return OutExpression != nullptr;
 }
 
-bool UMaterialExpressionGenericConstant::GenerateHLSLExpression(FMaterialHLSLGenerator& Generator, UE::HLSLTree::FScope& Scope, int32 OutputIndex, UE::HLSLTree::FExpression*& OutExpression)
+bool UMaterialExpressionGenericConstant::GenerateHLSLExpression(FMaterialHLSLGenerator& Generator, UE::HLSLTree::FScope& Scope, int32 OutputIndex, UE::HLSLTree::FExpression const*& OutExpression) const
 {
 	OutExpression = Generator.NewConstant(GetConstantValue());
 	return true;
 }
 
-bool UMaterialExpressionConstant::GenerateHLSLExpression(FMaterialHLSLGenerator& Generator, UE::HLSLTree::FScope& Scope, int32 OutputIndex, UE::HLSLTree::FExpression*& OutExpression)
+bool UMaterialExpressionConstant::GenerateHLSLExpression(FMaterialHLSLGenerator& Generator, UE::HLSLTree::FScope& Scope, int32 OutputIndex, UE::HLSLTree::FExpression const*& OutExpression) const
 {
 	OutExpression = Generator.NewConstant(R);
 	return true;
 }
 
-bool UMaterialExpressionConstant2Vector::GenerateHLSLExpression(FMaterialHLSLGenerator& Generator, UE::HLSLTree::FScope& Scope, int32 OutputIndex, UE::HLSLTree::FExpression*& OutExpression)
+bool UMaterialExpressionConstant2Vector::GenerateHLSLExpression(FMaterialHLSLGenerator& Generator, UE::HLSLTree::FScope& Scope, int32 OutputIndex, UE::HLSLTree::FExpression const*& OutExpression) const
 {
 	OutExpression = Generator.NewConstant(UE::Shader::FValue(R, G));
 	return true;
 }
 
-bool UMaterialExpressionConstant3Vector::GenerateHLSLExpression(FMaterialHLSLGenerator& Generator, UE::HLSLTree::FScope& Scope, int32 OutputIndex, UE::HLSLTree::FExpression*& OutExpression)
+bool UMaterialExpressionConstant3Vector::GenerateHLSLExpression(FMaterialHLSLGenerator& Generator, UE::HLSLTree::FScope& Scope, int32 OutputIndex, UE::HLSLTree::FExpression const*& OutExpression) const
 {
 	OutExpression = Generator.NewConstant(UE::Shader::FValue(Constant.R, Constant.G, Constant.B));
 	return true;
 }
 
-bool UMaterialExpressionConstant4Vector::GenerateHLSLExpression(FMaterialHLSLGenerator& Generator, UE::HLSLTree::FScope& Scope, int32 OutputIndex, UE::HLSLTree::FExpression*& OutExpression)
+bool UMaterialExpressionConstant4Vector::GenerateHLSLExpression(FMaterialHLSLGenerator& Generator, UE::HLSLTree::FScope& Scope, int32 OutputIndex, UE::HLSLTree::FExpression const*& OutExpression) const
 {
 	OutExpression = Generator.NewConstant(UE::Shader::FValue(Constant.R, Constant.G, Constant.B, Constant.A));
 	return true;
 }
 
-bool UMaterialExpressionStaticBool::GenerateHLSLExpression(FMaterialHLSLGenerator& Generator, UE::HLSLTree::FScope& Scope, int32 OutputIndex, UE::HLSLTree::FExpression*& OutExpression)
+bool UMaterialExpressionStaticBool::GenerateHLSLExpression(FMaterialHLSLGenerator& Generator, UE::HLSLTree::FScope& Scope, int32 OutputIndex, UE::HLSLTree::FExpression const*& OutExpression) const
 {
 	OutExpression = Generator.NewConstant((bool)Value);
 	return true;
 }
 
-bool UMaterialExpressionConstantBiasScale::GenerateHLSLExpression(FMaterialHLSLGenerator& Generator, UE::HLSLTree::FScope& Scope, int32 OutputIndex, UE::HLSLTree::FExpression*& OutExpression)
+bool UMaterialExpressionConstantBiasScale::GenerateHLSLExpression(FMaterialHLSLGenerator& Generator, UE::HLSLTree::FScope& Scope, int32 OutputIndex, UE::HLSLTree::FExpression const*& OutExpression) const
 {
 	using namespace UE::HLSLTree;
-	FExpression* ExpressionInput = Input.AcquireHLSLExpression(Generator, Scope);
+	const FExpression* ExpressionInput = Input.AcquireHLSLExpression(Generator, Scope);
 	if (!ExpressionInput)
 	{
 		return false;
@@ -211,24 +211,24 @@ bool UMaterialExpressionConstantBiasScale::GenerateHLSLExpression(FMaterialHLSLG
 	return true;
 }
 
-bool UMaterialExpressionShadingModel::GenerateHLSLExpression(FMaterialHLSLGenerator& Generator, UE::HLSLTree::FScope& Scope, int32 OutputIndex, UE::HLSLTree::FExpression*& OutExpression)
+bool UMaterialExpressionShadingModel::GenerateHLSLExpression(FMaterialHLSLGenerator& Generator, UE::HLSLTree::FScope& Scope, int32 OutputIndex, UE::HLSLTree::FExpression const*& OutExpression) const
 {
 	using namespace UE::HLSLTree::Material;
 	OutExpression = Generator.GetTree().NewExpression<FExpressionShadingModel>(ShadingModel);
 	return true;
 }
 
-bool UMaterialExpressionStaticSwitch::GenerateHLSLExpression(FMaterialHLSLGenerator& Generator, UE::HLSLTree::FScope& Scope, int32 OutputIndex, UE::HLSLTree::FExpression*& OutExpression)
+bool UMaterialExpressionStaticSwitch::GenerateHLSLExpression(FMaterialHLSLGenerator& Generator, UE::HLSLTree::FScope& Scope, int32 OutputIndex, UE::HLSLTree::FExpression const*& OutExpression) const
 {
-	UE::HLSLTree::FExpression* ConditionExpression = Value.AcquireHLSLExpressionOrConstant(Generator, Scope, (bool)DefaultValue);
-	UE::HLSLTree::FExpression* TrueExpression = A.TryAcquireHLSLExpression(Generator, Scope);
-	UE::HLSLTree::FExpression* FalseExpression = B.TryAcquireHLSLExpression(Generator, Scope);
+	const UE::HLSLTree::FExpression* ConditionExpression = Value.AcquireHLSLExpressionOrConstant(Generator, Scope, (bool)DefaultValue);
+	const UE::HLSLTree::FExpression* TrueExpression = A.TryAcquireHLSLExpression(Generator, Scope);
+	const UE::HLSLTree::FExpression* FalseExpression = B.TryAcquireHLSLExpression(Generator, Scope);
 
 	OutExpression = Generator.GetTree().NewExpression<UE::HLSLTree::FExpressionSelect>(ConditionExpression, TrueExpression, FalseExpression);
 	return true;
 }
 
-bool UMaterialExpressionFeatureLevelSwitch::GenerateHLSLExpression(FMaterialHLSLGenerator& Generator, UE::HLSLTree::FScope& Scope, int32 OutputIndex, UE::HLSLTree::FExpression*& OutExpression)
+bool UMaterialExpressionFeatureLevelSwitch::GenerateHLSLExpression(FMaterialHLSLGenerator& Generator, UE::HLSLTree::FScope& Scope, int32 OutputIndex, UE::HLSLTree::FExpression const*& OutExpression) const
 {
 	using namespace UE::HLSLTree;
 	if (!Default.GetTracedInput().Expression)
@@ -236,12 +236,12 @@ bool UMaterialExpressionFeatureLevelSwitch::GenerateHLSLExpression(FMaterialHLSL
 		return Generator.Error(TEXT("Missing default input"));
 	}
 
-	FExpression* ExpressionDefault = nullptr;
-	FExpression* ExpressionInputs[ERHIFeatureLevel::Num] = { nullptr };
+	const FExpression* ExpressionDefault = nullptr;
+	const FExpression* ExpressionInputs[ERHIFeatureLevel::Num] = { nullptr };
 	for (int32 Index = 0; Index < ERHIFeatureLevel::Num; ++Index)
 	{
-		FExpression* Expression = nullptr;
-		FExpressionInput& FeatureInput = Inputs[Index];
+		const FExpression* Expression = nullptr;
+		const FExpressionInput& FeatureInput = Inputs[Index];
 		if (FeatureInput.GetTracedInput().Expression)
 		{
 			Expression = FeatureInput.AcquireHLSLExpression(Generator, Scope);
@@ -261,7 +261,7 @@ bool UMaterialExpressionFeatureLevelSwitch::GenerateHLSLExpression(FMaterialHLSL
 	return true;
 }
 
-bool UMaterialExpressionShadingPathSwitch::GenerateHLSLExpression(FMaterialHLSLGenerator& Generator, UE::HLSLTree::FScope& Scope, int32 OutputIndex, UE::HLSLTree::FExpression*& OutExpression)
+bool UMaterialExpressionShadingPathSwitch::GenerateHLSLExpression(FMaterialHLSLGenerator& Generator, UE::HLSLTree::FScope& Scope, int32 OutputIndex, UE::HLSLTree::FExpression const*& OutExpression) const
 {
 	using namespace UE::HLSLTree;
 	if (!Default.GetTracedInput().Expression)
@@ -269,12 +269,12 @@ bool UMaterialExpressionShadingPathSwitch::GenerateHLSLExpression(FMaterialHLSLG
 		return Generator.Error(TEXT("Missing default input"));
 	}
 
-	FExpression* ExpressionDefault = nullptr;
-	FExpression* ExpressionInputs[ERHIShadingPath::Num] = { nullptr };
+	const FExpression* ExpressionDefault = nullptr;
+	const FExpression* ExpressionInputs[ERHIShadingPath::Num] = { nullptr };
 	for (int32 Index = 0; Index < ERHIShadingPath::Num; ++Index)
 	{
-		FExpression* Expression = nullptr;
-		FExpressionInput& FeatureInput = Inputs[Index];
+		const FExpression* Expression = nullptr;
+		const FExpressionInput& FeatureInput = Inputs[Index];
 		if (FeatureInput.GetTracedInput().Expression)
 		{
 			Expression = FeatureInput.AcquireHLSLExpression(Generator, Scope);
@@ -294,7 +294,7 @@ bool UMaterialExpressionShadingPathSwitch::GenerateHLSLExpression(FMaterialHLSLG
 	return true;
 }
 
-bool UMaterialExpressionGetLocal::GenerateHLSLExpression(FMaterialHLSLGenerator& Generator, UE::HLSLTree::FScope& Scope, int32 OutputIndex, UE::HLSLTree::FExpression*& OutExpression)
+bool UMaterialExpressionGetLocal::GenerateHLSLExpression(FMaterialHLSLGenerator& Generator, UE::HLSLTree::FScope& Scope, int32 OutputIndex, UE::HLSLTree::FExpression const*& OutExpression) const
 {
 	OutExpression = Generator.GetTree().AcquireLocal(Scope, LocalName);
 	if (!OutExpression)
@@ -304,7 +304,7 @@ bool UMaterialExpressionGetLocal::GenerateHLSLExpression(FMaterialHLSLGenerator&
 	return true;
 }
 
-bool UMaterialExpressionParameter::GenerateHLSLExpression(FMaterialHLSLGenerator& Generator, UE::HLSLTree::FScope& Scope, int32 OutputIndex, UE::HLSLTree::FExpression*& OutExpression)
+bool UMaterialExpressionParameter::GenerateHLSLExpression(FMaterialHLSLGenerator& Generator, UE::HLSLTree::FScope& Scope, int32 OutputIndex, UE::HLSLTree::FExpression const*& OutExpression) const
 {
 	FMaterialParameterMetadata ParameterMeta;
 	if (!GetParameterValue(ParameterMeta))
@@ -316,7 +316,7 @@ bool UMaterialExpressionParameter::GenerateHLSLExpression(FMaterialHLSLGenerator
 	return true;
 }
 
-bool UMaterialExpressionCurveAtlasRowParameter::GenerateHLSLExpression(FMaterialHLSLGenerator& Generator, UE::HLSLTree::FScope& Scope, int32 OutputIndex, UE::HLSLTree::FExpression*& OutExpression)
+bool UMaterialExpressionCurveAtlasRowParameter::GenerateHLSLExpression(FMaterialHLSLGenerator& Generator, UE::HLSLTree::FScope& Scope, int32 OutputIndex, UE::HLSLTree::FExpression const*& OutExpression) const
 {
 	using namespace UE::HLSLTree;
 	using namespace UE::Shader;
@@ -334,9 +334,9 @@ bool UMaterialExpressionCurveAtlasRowParameter::GenerateHLSLExpression(FMaterial
 	{
 		return Generator.Error(TEXT("Failed to get parameter value"));
 	}
-	FExpression* ExpressionSlot = Generator.GenerateMaterialParameter(ParameterName, ParameterMeta, ParameterMeta.Value.AsShaderValue());
+	const FExpression* ExpressionSlot = Generator.GenerateMaterialParameter(ParameterName, ParameterMeta, ParameterMeta.Value.AsShaderValue());
 
-	FExpression* ExpressionTexture = nullptr;
+	const FExpression* ExpressionTexture = nullptr;
 	{
 		const FMaterialParameterMetadata TextureParameterMeta(Atlas.Get());
 		const FTextureValue* TextureValue = Generator.AcquireTextureValue(FTextureValue(Atlas.Get(), SAMPLERTYPE_LinearColor));
@@ -344,17 +344,17 @@ bool UMaterialExpressionCurveAtlasRowParameter::GenerateHLSLExpression(FMaterial
 	}
 
 	FTree& Tree = Generator.GetTree();
-	FExpression* ExpressionTextureSize = Tree.NewExpression<Material::FExpressionTextureSize>(FMaterialParameterInfo(), Atlas.Get());
-	FExpression* ExpressionTextureHeight = Tree.NewExpression<FExpressionSwizzle>(MakeSwizzleMask(false, true, false, false), ExpressionTextureSize);
-	FExpression* ExpressionCoordU = InputTime.AcquireHLSLExpressionOrConstant(Generator, Scope, 0.0f);
-	FExpression* ExpressionCoordV = Tree.NewMul(Tree.NewAdd(ExpressionSlot, Tree.NewConstant(0.5f)), Tree.NewRcp(ExpressionTextureHeight));
-	FExpression* ExpressionUV = Tree.NewExpression<FExpressionAppend>(ExpressionCoordU, ExpressionCoordV);
+	const FExpression* ExpressionTextureSize = Tree.NewExpression<Material::FExpressionTextureSize>(FMaterialParameterInfo(), Atlas.Get());
+	const FExpression* ExpressionTextureHeight = Tree.NewExpression<FExpressionSwizzle>(MakeSwizzleMask(false, true, false, false), ExpressionTextureSize);
+	const FExpression* ExpressionCoordU = InputTime.AcquireHLSLExpressionOrConstant(Generator, Scope, 0.0f);
+	const FExpression* ExpressionCoordV = Tree.NewMul(Tree.NewAdd(ExpressionSlot, Tree.NewConstant(0.5f)), Tree.NewRcp(ExpressionTextureHeight));
+	const FExpression* ExpressionUV = Tree.NewExpression<FExpressionAppend>(ExpressionCoordU, ExpressionCoordV);
 
 	OutExpression = Generator.GetTree().NewExpression<FExpressionTextureSample>(ExpressionTexture, ExpressionUV, nullptr, FExpressionDerivatives(), SSM_Clamp_WorldGroupSettings, TMVM_None);
 	return true;
 }
 
-bool UMaterialExpressionStaticSwitchParameter::GenerateHLSLExpression(FMaterialHLSLGenerator& Generator, UE::HLSLTree::FScope& Scope, int32 OutputIndex, UE::HLSLTree::FExpression*& OutExpression)
+bool UMaterialExpressionStaticSwitchParameter::GenerateHLSLExpression(FMaterialHLSLGenerator& Generator, UE::HLSLTree::FScope& Scope, int32 OutputIndex, UE::HLSLTree::FExpression const*& OutExpression) const
 {
 	using namespace UE::HLSLTree;
 
@@ -364,26 +364,26 @@ bool UMaterialExpressionStaticSwitchParameter::GenerateHLSLExpression(FMaterialH
 		return Generator.Error(TEXT("Failed to get parameter value"));
 	}
 
-	FExpression* ExpressionA = A.AcquireHLSLExpression(Generator, Scope);
-	FExpression* ExpressionB = B.AcquireHLSLExpression(Generator, Scope);
+	const FExpression* ExpressionA = A.AcquireHLSLExpression(Generator, Scope);
+	const FExpression* ExpressionB = B.AcquireHLSLExpression(Generator, Scope);
 	if (!ExpressionA || !ExpressionB)
 	{
 		return false;
 	}
 
-	FExpression* ExpressionSwitch = Generator.GenerateMaterialParameter(ParameterName, ParameterMeta, ParameterMeta.Value.AsShaderValue());
+	const FExpression* ExpressionSwitch = Generator.GenerateMaterialParameter(ParameterName, ParameterMeta, ParameterMeta.Value.AsShaderValue());
 	OutExpression = Generator.GetTree().NewExpression <FExpressionSelect>(ExpressionSwitch, ExpressionA, ExpressionB);
 	return true;
 }
 
-bool UMaterialExpressionPixelDepth::GenerateHLSLExpression(FMaterialHLSLGenerator& Generator, UE::HLSLTree::FScope& Scope, int32 OutputIndex, UE::HLSLTree::FExpression*& OutExpression)
+bool UMaterialExpressionPixelDepth::GenerateHLSLExpression(FMaterialHLSLGenerator& Generator, UE::HLSLTree::FScope& Scope, int32 OutputIndex, UE::HLSLTree::FExpression const*& OutExpression) const
 {
 	using namespace UE::HLSLTree::Material;
 	OutExpression = Generator.GetTree().NewExpression<FExpressionExternalInput>(EExternalInput::PixelDepth);
 	return true;
 }
 
-bool UMaterialExpressionWorldPosition::GenerateHLSLExpression(FMaterialHLSLGenerator& Generator, UE::HLSLTree::FScope& Scope, int32 OutputIndex, UE::HLSLTree::FExpression*& OutExpression)
+bool UMaterialExpressionWorldPosition::GenerateHLSLExpression(FMaterialHLSLGenerator& Generator, UE::HLSLTree::FScope& Scope, int32 OutputIndex, UE::HLSLTree::FExpression const*& OutExpression) const
 {
 	using namespace UE::HLSLTree::Material;
 
@@ -401,21 +401,21 @@ bool UMaterialExpressionWorldPosition::GenerateHLSLExpression(FMaterialHLSLGener
 	return true;
 }
 
-bool UMaterialExpressionCameraPositionWS::GenerateHLSLExpression(FMaterialHLSLGenerator& Generator, UE::HLSLTree::FScope& Scope, int32 OutputIndex, UE::HLSLTree::FExpression*& OutExpression)
+bool UMaterialExpressionCameraPositionWS::GenerateHLSLExpression(FMaterialHLSLGenerator& Generator, UE::HLSLTree::FScope& Scope, int32 OutputIndex, UE::HLSLTree::FExpression const*& OutExpression) const
 {
 	using namespace UE::HLSLTree::Material;
 	OutExpression = Generator.GetTree().NewExpression<FExpressionExternalInput>(EExternalInput::CameraWorldPosition);
 	return true;
 }
 
-bool UMaterialExpressionCameraVectorWS::GenerateHLSLExpression(FMaterialHLSLGenerator& Generator, UE::HLSLTree::FScope& Scope, int32 OutputIndex, UE::HLSLTree::FExpression*& OutExpression)
+bool UMaterialExpressionCameraVectorWS::GenerateHLSLExpression(FMaterialHLSLGenerator& Generator, UE::HLSLTree::FScope& Scope, int32 OutputIndex, UE::HLSLTree::FExpression const*& OutExpression) const
 {
 	using namespace UE::HLSLTree::Material;
 	OutExpression = Generator.GetTree().NewExpression<FExpressionExternalInput>(EExternalInput::CameraVector);
 	return true;
 }
 
-bool UMaterialExpressionViewProperty::GenerateHLSLExpression(FMaterialHLSLGenerator& Generator, UE::HLSLTree::FScope& Scope, int32 OutputIndex, UE::HLSLTree::FExpression*& OutExpression)
+bool UMaterialExpressionViewProperty::GenerateHLSLExpression(FMaterialHLSLGenerator& Generator, UE::HLSLTree::FScope& Scope, int32 OutputIndex, UE::HLSLTree::FExpression const*& OutExpression) const
 {
 	using namespace UE::HLSLTree::Material;
 
@@ -460,7 +460,7 @@ bool UMaterialExpressionViewProperty::GenerateHLSLExpression(FMaterialHLSLGenera
 	return true;
 }
 
-bool UMaterialExpressionTime::GenerateHLSLExpression(FMaterialHLSLGenerator& Generator, UE::HLSLTree::FScope& Scope, int32 OutputIndex, UE::HLSLTree::FExpression*& OutExpression)
+bool UMaterialExpressionTime::GenerateHLSLExpression(FMaterialHLSLGenerator& Generator, UE::HLSLTree::FScope& Scope, int32 OutputIndex, UE::HLSLTree::FExpression const*& OutExpression) const
 {
 	using namespace UE::HLSLTree::Material;
 
@@ -479,14 +479,14 @@ bool UMaterialExpressionTime::GenerateHLSLExpression(FMaterialHLSLGenerator& Gen
 	return true;
 }
 
-bool UMaterialExpressionDeltaTime::GenerateHLSLExpression(FMaterialHLSLGenerator& Generator, UE::HLSLTree::FScope& Scope, int32 OutputIndex, UE::HLSLTree::FExpression*& OutExpression)
+bool UMaterialExpressionDeltaTime::GenerateHLSLExpression(FMaterialHLSLGenerator& Generator, UE::HLSLTree::FScope& Scope, int32 OutputIndex, UE::HLSLTree::FExpression const*& OutExpression) const
 {
 	using namespace UE::HLSLTree::Material;
 	OutExpression = Generator.GetTree().NewExpression<FExpressionExternalInput>(EExternalInput::DeltaTime);
 	return true;
 }
 
-bool UMaterialExpressionScreenPosition::GenerateHLSLExpression(FMaterialHLSLGenerator& Generator, UE::HLSLTree::FScope& Scope, int32 OutputIndex, UE::HLSLTree::FExpression*& OutExpression)
+bool UMaterialExpressionScreenPosition::GenerateHLSLExpression(FMaterialHLSLGenerator& Generator, UE::HLSLTree::FScope& Scope, int32 OutputIndex, UE::HLSLTree::FExpression const*& OutExpression) const
 {
 	using namespace UE::HLSLTree::Material;
 	const EExternalInput InputType = (OutputIndex == 1) ? EExternalInput::PixelPosition : EExternalInput::ViewportUV;
@@ -494,7 +494,7 @@ bool UMaterialExpressionScreenPosition::GenerateHLSLExpression(FMaterialHLSLGene
 	return true;
 }
 
-bool UMaterialExpressionSceneTexelSize::GenerateHLSLExpression(FMaterialHLSLGenerator& Generator, UE::HLSLTree::FScope& Scope, int32 OutputIndex, UE::HLSLTree::FExpression*& OutExpression)
+bool UMaterialExpressionSceneTexelSize::GenerateHLSLExpression(FMaterialHLSLGenerator& Generator, UE::HLSLTree::FScope& Scope, int32 OutputIndex, UE::HLSLTree::FExpression const*& OutExpression) const
 {
 	using namespace UE::HLSLTree::Material;
 
@@ -504,21 +504,21 @@ bool UMaterialExpressionSceneTexelSize::GenerateHLSLExpression(FMaterialHLSLGene
 	return true;
 }
 
-bool UMaterialExpressionViewSize::GenerateHLSLExpression(FMaterialHLSLGenerator& Generator, UE::HLSLTree::FScope& Scope, int32 OutputIndex, UE::HLSLTree::FExpression*& OutExpression)
+bool UMaterialExpressionViewSize::GenerateHLSLExpression(FMaterialHLSLGenerator& Generator, UE::HLSLTree::FScope& Scope, int32 OutputIndex, UE::HLSLTree::FExpression const*& OutExpression) const
 {
 	using namespace UE::HLSLTree::Material;
 	OutExpression = Generator.GetTree().NewExpression<FExpressionExternalInput>(EExternalInput::ViewSize);
 	return true;
 }
 
-bool UMaterialExpressionTwoSidedSign::GenerateHLSLExpression(FMaterialHLSLGenerator& Generator, UE::HLSLTree::FScope& Scope, int32 OutputIndex, UE::HLSLTree::FExpression*& OutExpression)
+bool UMaterialExpressionTwoSidedSign::GenerateHLSLExpression(FMaterialHLSLGenerator& Generator, UE::HLSLTree::FScope& Scope, int32 OutputIndex, UE::HLSLTree::FExpression const*& OutExpression) const
 {
 	using namespace UE::HLSLTree::Material;
 	OutExpression = Generator.GetTree().NewExpression<FExpressionExternalInput>(EExternalInput::TwoSidedSign);
 	return true;
 }
 
-bool UMaterialExpressionReflectionVectorWS::GenerateHLSLExpression(FMaterialHLSLGenerator& Generator, UE::HLSLTree::FScope& Scope, int32 OutputIndex, UE::HLSLTree::FExpression*& OutExpression)
+bool UMaterialExpressionReflectionVectorWS::GenerateHLSLExpression(FMaterialHLSLGenerator& Generator, UE::HLSLTree::FScope& Scope, int32 OutputIndex, UE::HLSLTree::FExpression const*& OutExpression) const
 {
 	using namespace UE::HLSLTree::Material;
 	check(!CustomWorldNormal.GetTracedInput().Expression); // TODO
@@ -527,28 +527,28 @@ bool UMaterialExpressionReflectionVectorWS::GenerateHLSLExpression(FMaterialHLSL
 	return true;
 }
 
-bool UMaterialExpressionActorPositionWS::GenerateHLSLExpression(FMaterialHLSLGenerator& Generator, UE::HLSLTree::FScope& Scope, int32 OutputIndex, UE::HLSLTree::FExpression*& OutExpression)
+bool UMaterialExpressionActorPositionWS::GenerateHLSLExpression(FMaterialHLSLGenerator& Generator, UE::HLSLTree::FScope& Scope, int32 OutputIndex, UE::HLSLTree::FExpression const*& OutExpression) const
 {
 	using namespace UE::HLSLTree::Material;
 	OutExpression = Generator.GetTree().NewExpression<FExpressionExternalInput>(EExternalInput::ActorWorldPosition);
 	return true;
 }
 
-bool UMaterialExpressionPreSkinnedPosition::GenerateHLSLExpression(FMaterialHLSLGenerator& Generator, UE::HLSLTree::FScope& Scope, int32 OutputIndex, UE::HLSLTree::FExpression*& OutExpression)
+bool UMaterialExpressionPreSkinnedPosition::GenerateHLSLExpression(FMaterialHLSLGenerator& Generator, UE::HLSLTree::FScope& Scope, int32 OutputIndex, UE::HLSLTree::FExpression const*& OutExpression) const
 {
 	using namespace UE::HLSLTree;
 	OutExpression = Generator.NewExternalInput(Material::EExternalInput::PreSkinnedPosition);
 	return true;
 }
 
-bool UMaterialExpressionPreSkinnedNormal::GenerateHLSLExpression(FMaterialHLSLGenerator& Generator, UE::HLSLTree::FScope& Scope, int32 OutputIndex, UE::HLSLTree::FExpression*& OutExpression)
+bool UMaterialExpressionPreSkinnedNormal::GenerateHLSLExpression(FMaterialHLSLGenerator& Generator, UE::HLSLTree::FScope& Scope, int32 OutputIndex, UE::HLSLTree::FExpression const*& OutExpression) const
 {
 	using namespace UE::HLSLTree;
 	OutExpression = Generator.NewExternalInput(Material::EExternalInput::PreSkinnedNormal);
 	return true;
 }
 
-bool UMaterialExpressionPreSkinnedLocalBounds::GenerateHLSLExpression(FMaterialHLSLGenerator& Generator, UE::HLSLTree::FScope& Scope, int32 OutputIndex, UE::HLSLTree::FExpression*& OutExpression)
+bool UMaterialExpressionPreSkinnedLocalBounds::GenerateHLSLExpression(FMaterialHLSLGenerator& Generator, UE::HLSLTree::FScope& Scope, int32 OutputIndex, UE::HLSLTree::FExpression const*& OutExpression) const
 {
 	using namespace UE::HLSLTree;
 	FTree& Tree = Generator.GetTree();
@@ -571,35 +571,35 @@ bool UMaterialExpressionPreSkinnedLocalBounds::GenerateHLSLExpression(FMaterialH
 	}
 }
 
-bool UMaterialExpressionPixelNormalWS::GenerateHLSLExpression(FMaterialHLSLGenerator& Generator, UE::HLSLTree::FScope& Scope, int32 OutputIndex, UE::HLSLTree::FExpression*& OutExpression)
+bool UMaterialExpressionPixelNormalWS::GenerateHLSLExpression(FMaterialHLSLGenerator& Generator, UE::HLSLTree::FScope& Scope, int32 OutputIndex, UE::HLSLTree::FExpression const*& OutExpression) const
 {
 	using namespace UE::HLSLTree::Material;
 	OutExpression = Generator.GetTree().NewExpression<FExpressionExternalInput>(EExternalInput::WorldNormal);
 	return true;
 }
 
-bool UMaterialExpressionVertexColor::GenerateHLSLExpression(FMaterialHLSLGenerator& Generator, UE::HLSLTree::FScope& Scope, int32 OutputIndex, UE::HLSLTree::FExpression*& OutExpression)
+bool UMaterialExpressionVertexColor::GenerateHLSLExpression(FMaterialHLSLGenerator& Generator, UE::HLSLTree::FScope& Scope, int32 OutputIndex, UE::HLSLTree::FExpression const*& OutExpression) const
 {
 	using namespace UE::HLSLTree::Material;
 	OutExpression = Generator.GetTree().NewExpression<FExpressionExternalInput>(EExternalInput::VertexColor);
 	return true;
 }
 
-bool UMaterialExpressionVertexNormalWS::GenerateHLSLExpression(FMaterialHLSLGenerator& Generator, UE::HLSLTree::FScope& Scope, int32 OutputIndex, UE::HLSLTree::FExpression*& OutExpression)
+bool UMaterialExpressionVertexNormalWS::GenerateHLSLExpression(FMaterialHLSLGenerator& Generator, UE::HLSLTree::FScope& Scope, int32 OutputIndex, UE::HLSLTree::FExpression const*& OutExpression) const
 {
 	using namespace UE::HLSLTree::Material;
 	OutExpression = Generator.GetTree().NewExpression<FExpressionExternalInput>(EExternalInput::WorldVertexNormal);
 	return true;
 }
 
-bool UMaterialExpressionVertexTangentWS::GenerateHLSLExpression(FMaterialHLSLGenerator& Generator, UE::HLSLTree::FScope& Scope, int32 OutputIndex, UE::HLSLTree::FExpression*& OutExpression)
+bool UMaterialExpressionVertexTangentWS::GenerateHLSLExpression(FMaterialHLSLGenerator& Generator, UE::HLSLTree::FScope& Scope, int32 OutputIndex, UE::HLSLTree::FExpression const*& OutExpression) const
 {
 	using namespace UE::HLSLTree::Material;
 	OutExpression = Generator.GetTree().NewExpression<FExpressionExternalInput>(EExternalInput::WorldVertexTangent);
 	return true;
 }
 
-bool UMaterialExpressionTextureCoordinate::GenerateHLSLExpression(FMaterialHLSLGenerator& Generator, UE::HLSLTree::FScope& Scope, int32 OutputIndex, UE::HLSLTree::FExpression*& OutExpression)
+bool UMaterialExpressionTextureCoordinate::GenerateHLSLExpression(FMaterialHLSLGenerator& Generator, UE::HLSLTree::FScope& Scope, int32 OutputIndex, UE::HLSLTree::FExpression const*& OutExpression) const
 {
 	OutExpression = Generator.NewTexCoord(CoordinateIndex);
 
@@ -622,46 +622,46 @@ bool UMaterialExpressionTextureCoordinate::GenerateHLSLExpression(FMaterialHLSLG
 	return true;
 }
 
-bool UMaterialExpressionLightmapUVs::GenerateHLSLExpression(FMaterialHLSLGenerator& Generator, UE::HLSLTree::FScope& Scope, int32 OutputIndex, UE::HLSLTree::FExpression*& OutExpression)
+bool UMaterialExpressionLightmapUVs::GenerateHLSLExpression(FMaterialHLSLGenerator& Generator, UE::HLSLTree::FScope& Scope, int32 OutputIndex, UE::HLSLTree::FExpression const*& OutExpression) const
 {
 	using namespace UE::HLSLTree::Material;
 	OutExpression = Generator.GetTree().NewExpression<FExpressionExternalInput>(EExternalInput::LightmapTexCoord);
 	return true;
 }
 
-bool UMaterialExpressionEyeAdaptation::GenerateHLSLExpression(FMaterialHLSLGenerator& Generator, UE::HLSLTree::FScope& Scope, int32 OutputIndex, UE::HLSLTree::FExpression*& OutExpression)
+bool UMaterialExpressionEyeAdaptation::GenerateHLSLExpression(FMaterialHLSLGenerator& Generator, UE::HLSLTree::FScope& Scope, int32 OutputIndex, UE::HLSLTree::FExpression const*& OutExpression) const
 {
 	using namespace UE::HLSLTree;
 	OutExpression = Generator.GetTree().NewExpression<Material::FExpressionExternalInput>(Material::EExternalInput::EyeAdaptation);
 	return true;
 }
 
-bool UMaterialExpressionEyeAdaptationInverse::GenerateHLSLExpression(FMaterialHLSLGenerator& Generator, UE::HLSLTree::FScope& Scope, int32 OutputIndex, UE::HLSLTree::FExpression*& OutExpression)
+bool UMaterialExpressionEyeAdaptationInverse::GenerateHLSLExpression(FMaterialHLSLGenerator& Generator, UE::HLSLTree::FScope& Scope, int32 OutputIndex, UE::HLSLTree::FExpression const*& OutExpression) const
 {
 	using namespace UE::HLSLTree;
-	FExpression* ExpressionLightValue = LightValueInput.AcquireHLSLExpressionOrConstant(Generator, Scope, FVector3f(1.0f, 1.0f, 1.0f));
-	FExpression* ExpressionAlpha = AlphaInput.AcquireHLSLExpressionOrConstant(Generator, Scope, 1.0f);
+	const FExpression* ExpressionLightValue = LightValueInput.AcquireHLSLExpressionOrConstant(Generator, Scope, FVector3f(1.0f, 1.0f, 1.0f));
+	const FExpression* ExpressionAlpha = AlphaInput.AcquireHLSLExpressionOrConstant(Generator, Scope, 1.0f);
 	if (!ExpressionLightValue || !ExpressionAlpha)
 	{
 		return false;
 	}
 
 	FTree& Tree = Generator.GetTree();
-	FExpression* ExpressionAdaptation = Tree.NewExpression<Material::FExpressionExternalInput>(Material::EExternalInput::EyeAdaptation);
-	FExpression* LerpLogScale = Tree.NewMul(Tree.NewNeg(ExpressionAlpha), Tree.NewLog2(ExpressionAdaptation));
-	FExpression* Scale = Tree.NewExp2(LerpLogScale);
+	const FExpression* ExpressionAdaptation = Tree.NewExpression<Material::FExpressionExternalInput>(Material::EExternalInput::EyeAdaptation);
+	const FExpression* LerpLogScale = Tree.NewMul(Tree.NewNeg(ExpressionAlpha), Tree.NewLog2(ExpressionAdaptation));
+	const FExpression* Scale = Tree.NewExp2(LerpLogScale);
 	OutExpression = Tree.NewMul(ExpressionLightValue, Scale);
 	return true;
 }
 
-bool UMaterialExpressionParticleColor::GenerateHLSLExpression(FMaterialHLSLGenerator& Generator, UE::HLSLTree::FScope& Scope, int32 OutputIndex, UE::HLSLTree::FExpression*& OutExpression)
+bool UMaterialExpressionParticleColor::GenerateHLSLExpression(FMaterialHLSLGenerator& Generator, UE::HLSLTree::FScope& Scope, int32 OutputIndex, UE::HLSLTree::FExpression const*& OutExpression) const
 {
 	using namespace UE::HLSLTree::Material;
 	OutExpression = Generator.GetTree().NewExpression<FExpressionExternalInput>(EExternalInput::ParticleColor);
 	return true;
 }
 
-bool UMaterialExpressionParticlePositionWS::GenerateHLSLExpression(FMaterialHLSLGenerator& Generator, UE::HLSLTree::FScope& Scope, int32 OutputIndex, UE::HLSLTree::FExpression*& OutExpression)
+bool UMaterialExpressionParticlePositionWS::GenerateHLSLExpression(FMaterialHLSLGenerator& Generator, UE::HLSLTree::FScope& Scope, int32 OutputIndex, UE::HLSLTree::FExpression const*& OutExpression) const
 {
 	using namespace UE::HLSLTree::Material;
 	OutExpression = Generator.GetTree().NewExpression<FExpressionExternalInput>(EExternalInput::ParticleTranslatedWorldPosition);
@@ -669,7 +669,7 @@ bool UMaterialExpressionParticlePositionWS::GenerateHLSLExpression(FMaterialHLSL
 	return true;
 }
 
-bool UMaterialExpressionParticleRadius::GenerateHLSLExpression(FMaterialHLSLGenerator& Generator, UE::HLSLTree::FScope& Scope, int32 OutputIndex, UE::HLSLTree::FExpression*& OutExpression)
+bool UMaterialExpressionParticleRadius::GenerateHLSLExpression(FMaterialHLSLGenerator& Generator, UE::HLSLTree::FScope& Scope, int32 OutputIndex, UE::HLSLTree::FExpression const*& OutExpression) const
 {
 	using namespace UE::HLSLTree::Material;
 	OutExpression = Generator.GetTree().NewExpression<FExpressionExternalInput>(EExternalInput::ParticleRadius);
@@ -677,7 +677,7 @@ bool UMaterialExpressionParticleRadius::GenerateHLSLExpression(FMaterialHLSLGene
 	return true;
 }
 
-bool UMaterialExpressionTextureObject::GenerateHLSLExpression(FMaterialHLSLGenerator& Generator, UE::HLSLTree::FScope& Scope, int32 OutputIndex, UE::HLSLTree::FExpression*& OutExpression)
+bool UMaterialExpressionTextureObject::GenerateHLSLExpression(FMaterialHLSLGenerator& Generator, UE::HLSLTree::FScope& Scope, int32 OutputIndex, UE::HLSLTree::FExpression const*& OutExpression) const
 {
 	using namespace UE::HLSLTree;
 	using namespace UE::Shader;
@@ -688,7 +688,7 @@ bool UMaterialExpressionTextureObject::GenerateHLSLExpression(FMaterialHLSLGener
 	return true;
 }
 
-bool UMaterialExpressionTextureObjectParameter::GenerateHLSLExpression(FMaterialHLSLGenerator& Generator, UE::HLSLTree::FScope& Scope, int32 OutputIndex, UE::HLSLTree::FExpression*& OutExpression)
+bool UMaterialExpressionTextureObjectParameter::GenerateHLSLExpression(FMaterialHLSLGenerator& Generator, UE::HLSLTree::FScope& Scope, int32 OutputIndex, UE::HLSLTree::FExpression const*& OutExpression) const
 {
 	using namespace UE::HLSLTree;
 	using namespace UE::Shader;
@@ -704,7 +704,7 @@ bool UMaterialExpressionTextureObjectParameter::GenerateHLSLExpression(FMaterial
 	return true;
 }
 
-bool UMaterialExpressionTextureSample::GenerateHLSLExpressionBase(FMaterialHLSLGenerator& Generator, UE::HLSLTree::FScope& Scope, UE::HLSLTree::FExpression* TextureExpression, UE::HLSLTree::FExpression*& OutExpression)
+bool UMaterialExpressionTextureSample::GenerateHLSLExpressionBase(FMaterialHLSLGenerator& Generator, UE::HLSLTree::FScope& Scope, const UE::HLSLTree::FExpression* TextureExpression, UE::HLSLTree::FExpression const*& OutExpression) const
 {
 	using namespace UE::HLSLTree;
 	if (!TextureExpression)
@@ -712,8 +712,8 @@ bool UMaterialExpressionTextureSample::GenerateHLSLExpressionBase(FMaterialHLSLG
 		return Generator.Error(TEXT("Missing input texture"));
 	}
 
-	FExpression* TexCoordExpression = Coordinates.GetTracedInput().Expression ? Coordinates.TryAcquireHLSLExpression(Generator, Scope) : Generator.NewTexCoord(ConstCoordinate);
-	FExpression* MipLevelExpression = nullptr;
+	const FExpression* TexCoordExpression = Coordinates.GetTracedInput().Expression ? Coordinates.TryAcquireHLSLExpression(Generator, Scope) : Generator.NewTexCoord(ConstCoordinate);
+	const FExpression* MipLevelExpression = nullptr;
 	FExpressionDerivatives TexCoordDerivatives;
 	switch (MipValueMode)
 	{
@@ -723,7 +723,7 @@ bool UMaterialExpressionTextureSample::GenerateHLSLExpressionBase(FMaterialHLSLG
 	case TMVM_MipBias:
 	{
 		MipLevelExpression = MipValue.AcquireHLSLExpressionOrConstant(Generator, Scope, ConstMipValue);
-		FExpression* DerivativeScale = Generator.GetTree().NewExp2(MipLevelExpression);
+		const FExpression* DerivativeScale = Generator.GetTree().NewExp2(MipLevelExpression);
 		TexCoordDerivatives = Generator.GetTree().GetAnalyticDerivatives(TexCoordExpression);
 		TexCoordDerivatives.ExpressionDdx = Generator.GetTree().NewMul(TexCoordDerivatives.ExpressionDdx, DerivativeScale);
 		TexCoordDerivatives.ExpressionDdy = Generator.GetTree().NewMul(TexCoordDerivatives.ExpressionDdy, DerivativeScale);
@@ -745,11 +745,11 @@ bool UMaterialExpressionTextureSample::GenerateHLSLExpressionBase(FMaterialHLSLG
 	return true;
 }
 
-bool UMaterialExpressionTextureSample::GenerateHLSLExpression(FMaterialHLSLGenerator& Generator, UE::HLSLTree::FScope& Scope, int32 OutputIndex, UE::HLSLTree::FExpression*& OutExpression)
+bool UMaterialExpressionTextureSample::GenerateHLSLExpression(FMaterialHLSLGenerator& Generator, UE::HLSLTree::FScope& Scope, int32 OutputIndex, UE::HLSLTree::FExpression const*& OutExpression) const
 {
 	using namespace UE::HLSLTree;
 	using namespace UE::Shader;
-	FExpression* TextureExpression = nullptr;
+	const FExpression* TextureExpression = nullptr;
 	if (TextureObject.GetTracedInput().Expression)
 	{
 		TextureExpression = TextureObject.AcquireHLSLExpression(Generator, Scope);
@@ -764,11 +764,11 @@ bool UMaterialExpressionTextureSample::GenerateHLSLExpression(FMaterialHLSLGener
 	return GenerateHLSLExpressionBase(Generator, Scope, TextureExpression, OutExpression);
 }
 
-bool UMaterialExpressionTextureSampleParameter::GenerateHLSLExpression(FMaterialHLSLGenerator& Generator, UE::HLSLTree::FScope& Scope, int32 OutputIndex, UE::HLSLTree::FExpression*& OutExpression)
+bool UMaterialExpressionTextureSampleParameter::GenerateHLSLExpression(FMaterialHLSLGenerator& Generator, UE::HLSLTree::FScope& Scope, int32 OutputIndex, UE::HLSLTree::FExpression const*& OutExpression) const
 {
 	using namespace UE::HLSLTree;
 	using namespace UE::Shader;
-	FExpression* TextureExpression = nullptr;
+	const FExpression* TextureExpression = nullptr;
 	if (ParameterName.IsNone() && TextureObject.GetTracedInput().Expression)
 	{
 		TextureExpression = TextureObject.AcquireHLSLExpression(Generator, Scope);
@@ -788,7 +788,7 @@ bool UMaterialExpressionTextureSampleParameter::GenerateHLSLExpression(FMaterial
 	return GenerateHLSLExpressionBase(Generator, Scope, TextureExpression, OutExpression);
 }
 
-bool UMaterialExpressionFontSample::GenerateHLSLExpression(FMaterialHLSLGenerator& Generator, UE::HLSLTree::FScope& Scope, int32 OutputIndex, UE::HLSLTree::FExpression*& OutExpression)
+bool UMaterialExpressionFontSample::GenerateHLSLExpression(FMaterialHLSLGenerator& Generator, UE::HLSLTree::FScope& Scope, int32 OutputIndex, UE::HLSLTree::FExpression const*& OutExpression) const
 {
 	using namespace UE::HLSLTree;
 	using namespace UE::Shader;
@@ -838,20 +838,20 @@ bool UMaterialExpressionFontSample::GenerateHLSLExpression(FMaterialHLSLGenerato
 		ParameterMeta.Value = Texture;
 	}
 
-	FExpression* TextureExpression = Generator.GenerateMaterialParameter(GetParameterName(), ParameterMeta, TextureValue);
-	FExpression* TexCoordExpression = Generator.NewTexCoord(0);
+	const FExpression* TextureExpression = Generator.GenerateMaterialParameter(GetParameterName(), ParameterMeta, TextureValue);
+	const FExpression* TexCoordExpression = Generator.NewTexCoord(0);
 	const FExpressionDerivatives TexCoordDerivatives = Generator.GetTree().GetAnalyticDerivatives(TextureExpression);
 
 	OutExpression = Generator.GetTree().NewExpression<FExpressionTextureSample>(TextureExpression, TexCoordExpression, nullptr, TexCoordDerivatives, SSM_FromTextureAsset, TMVM_None);
 	return true;
 }
 
-bool UMaterialExpressionSceneTexture::GenerateHLSLExpression(FMaterialHLSLGenerator& Generator, UE::HLSLTree::FScope& Scope, int32 OutputIndex, UE::HLSLTree::FExpression*& OutExpression)
+bool UMaterialExpressionSceneTexture::GenerateHLSLExpression(FMaterialHLSLGenerator& Generator, UE::HLSLTree::FScope& Scope, int32 OutputIndex, UE::HLSLTree::FExpression const*& OutExpression) const
 {
 	using namespace UE::HLSLTree;
 	if (OutputIndex == 0)
 	{
-		FExpression* ExpressionTexCoord = nullptr;
+		const FExpression* ExpressionTexCoord = nullptr;
 		if (Coordinates.GetTracedInput().Expression)
 		{
 			ExpressionTexCoord = Coordinates.AcquireHLSLExpression(Generator, Scope);
@@ -873,16 +873,16 @@ bool UMaterialExpressionSceneTexture::GenerateHLSLExpression(FMaterialHLSLGenera
 	return Generator.Error(TEXT("Invalid input parameter"));
 }
 
-bool UMaterialExpressionNoise::GenerateHLSLExpression(FMaterialHLSLGenerator& Generator, UE::HLSLTree::FScope& Scope, int32 OutputIndex, UE::HLSLTree::FExpression*& OutExpression)
+bool UMaterialExpressionNoise::GenerateHLSLExpression(FMaterialHLSLGenerator& Generator, UE::HLSLTree::FScope& Scope, int32 OutputIndex, UE::HLSLTree::FExpression const*& OutExpression) const
 {
 	using namespace UE::HLSLTree;
 
-	FExpression* ExpressionPosition = Position.TryAcquireHLSLExpression(Generator, Scope);
+	const FExpression* ExpressionPosition = Position.TryAcquireHLSLExpression(Generator, Scope);
 	if (!ExpressionPosition)
 	{
 		ExpressionPosition = Generator.GetTree().NewExpression<Material::FExpressionExternalInput>(Material::EExternalInput::WorldPosition);
 	}
-	FExpression* ExpressionFilterWidth = FilterWidth.AcquireHLSLExpressionOrConstant(Generator, Scope, 0.0f);
+	const FExpression* ExpressionFilterWidth = FilterWidth.AcquireHLSLExpressionOrConstant(Generator, Scope, 0.0f);
 
 	Material::FNoiseParameters NoiseParameters;
 	NoiseParameters.Quality = Quality;
@@ -901,15 +901,15 @@ bool UMaterialExpressionNoise::GenerateHLSLExpression(FMaterialHLSLGenerator& Ge
 	return true;
 }
 
-bool UMaterialExpressionSceneDepth::GenerateHLSLExpression(FMaterialHLSLGenerator& Generator, UE::HLSLTree::FScope& Scope, int32 OutputIndex, UE::HLSLTree::FExpression*& OutExpression)
+bool UMaterialExpressionSceneDepth::GenerateHLSLExpression(FMaterialHLSLGenerator& Generator, UE::HLSLTree::FScope& Scope, int32 OutputIndex, UE::HLSLTree::FExpression const*& OutExpression) const
 {
 	OutExpression = Generator.NewConstant(0.0f); // TODO
 	return true;
 }
 
-bool UMaterialExpressionOneMinus::GenerateHLSLExpression(FMaterialHLSLGenerator& Generator, UE::HLSLTree::FScope& Scope, int32 OutputIndex, UE::HLSLTree::FExpression*& OutExpression)
+bool UMaterialExpressionOneMinus::GenerateHLSLExpression(FMaterialHLSLGenerator& Generator, UE::HLSLTree::FScope& Scope, int32 OutputIndex, UE::HLSLTree::FExpression const*& OutExpression) const
 {
-	UE::HLSLTree::FExpression* InputExpression = Input.AcquireHLSLExpression(Generator, Scope);
+	const UE::HLSLTree::FExpression* InputExpression = Input.AcquireHLSLExpression(Generator, Scope);
 	if (!InputExpression)
 	{
 		return false;
@@ -918,10 +918,10 @@ bool UMaterialExpressionOneMinus::GenerateHLSLExpression(FMaterialHLSLGenerator&
 	return true;
 }
 
-bool UMaterialExpressionAbs::GenerateHLSLExpression(FMaterialHLSLGenerator& Generator, UE::HLSLTree::FScope& Scope, int32 OutputIndex, UE::HLSLTree::FExpression*& OutExpression)
+bool UMaterialExpressionAbs::GenerateHLSLExpression(FMaterialHLSLGenerator& Generator, UE::HLSLTree::FScope& Scope, int32 OutputIndex, UE::HLSLTree::FExpression const*& OutExpression) const
 {
 	using namespace UE::HLSLTree;
-	FExpression* InputExpression = Input.AcquireHLSLExpression(Generator, Scope);
+	const FExpression* InputExpression = Input.AcquireHLSLExpression(Generator, Scope);
 	if (!InputExpression)
 	{
 		return false;
@@ -930,10 +930,10 @@ bool UMaterialExpressionAbs::GenerateHLSLExpression(FMaterialHLSLGenerator& Gene
 	return true;
 }
 
-bool UMaterialExpressionSquareRoot::GenerateHLSLExpression(FMaterialHLSLGenerator& Generator, UE::HLSLTree::FScope& Scope, int32 OutputIndex, UE::HLSLTree::FExpression*& OutExpression)
+bool UMaterialExpressionSquareRoot::GenerateHLSLExpression(FMaterialHLSLGenerator& Generator, UE::HLSLTree::FScope& Scope, int32 OutputIndex, UE::HLSLTree::FExpression const*& OutExpression) const
 {
 	using namespace UE::HLSLTree;
-	FExpression* InputExpression = Input.AcquireHLSLExpression(Generator, Scope);
+	const FExpression* InputExpression = Input.AcquireHLSLExpression(Generator, Scope);
 	if (!InputExpression)
 	{
 		return false;
@@ -942,10 +942,10 @@ bool UMaterialExpressionSquareRoot::GenerateHLSLExpression(FMaterialHLSLGenerato
 	return true;
 }
 
-bool UMaterialExpressionLogarithm2::GenerateHLSLExpression(FMaterialHLSLGenerator& Generator, UE::HLSLTree::FScope& Scope, int32 OutputIndex, UE::HLSLTree::FExpression*& OutExpression)
+bool UMaterialExpressionLogarithm2::GenerateHLSLExpression(FMaterialHLSLGenerator& Generator, UE::HLSLTree::FScope& Scope, int32 OutputIndex, UE::HLSLTree::FExpression const*& OutExpression) const
 {
 	using namespace UE::HLSLTree;
-	FExpression* InputExpression = X.AcquireHLSLExpression(Generator, Scope);
+	const FExpression* InputExpression = X.AcquireHLSLExpression(Generator, Scope);
 	if (!InputExpression)
 	{
 		return false;
@@ -954,10 +954,10 @@ bool UMaterialExpressionLogarithm2::GenerateHLSLExpression(FMaterialHLSLGenerato
 	return true;
 }
 
-bool UMaterialExpressionLogarithm10::GenerateHLSLExpression(FMaterialHLSLGenerator& Generator, UE::HLSLTree::FScope& Scope, int32 OutputIndex, UE::HLSLTree::FExpression*& OutExpression)
+bool UMaterialExpressionLogarithm10::GenerateHLSLExpression(FMaterialHLSLGenerator& Generator, UE::HLSLTree::FScope& Scope, int32 OutputIndex, UE::HLSLTree::FExpression const*& OutExpression) const
 {
 	using namespace UE::HLSLTree;
-	FExpression* InputExpression = X.AcquireHLSLExpression(Generator, Scope);
+	const FExpression* InputExpression = X.AcquireHLSLExpression(Generator, Scope);
 	if (!InputExpression)
 	{
 		return false;
@@ -967,10 +967,10 @@ bool UMaterialExpressionLogarithm10::GenerateHLSLExpression(FMaterialHLSLGenerat
 	return true;
 }
 
-bool UMaterialExpressionFrac::GenerateHLSLExpression(FMaterialHLSLGenerator& Generator, UE::HLSLTree::FScope& Scope, int32 OutputIndex, UE::HLSLTree::FExpression*& OutExpression)
+bool UMaterialExpressionFrac::GenerateHLSLExpression(FMaterialHLSLGenerator& Generator, UE::HLSLTree::FScope& Scope, int32 OutputIndex, UE::HLSLTree::FExpression const*& OutExpression) const
 {
 	using namespace UE::HLSLTree;
-	FExpression* InputExpression = Input.AcquireHLSLExpression(Generator, Scope);
+	const FExpression* InputExpression = Input.AcquireHLSLExpression(Generator, Scope);
 	if (!InputExpression)
 	{
 		return false;
@@ -979,10 +979,10 @@ bool UMaterialExpressionFrac::GenerateHLSLExpression(FMaterialHLSLGenerator& Gen
 	return true;
 }
 
-bool UMaterialExpressionFloor::GenerateHLSLExpression(FMaterialHLSLGenerator& Generator, UE::HLSLTree::FScope& Scope, int32 OutputIndex, UE::HLSLTree::FExpression*& OutExpression)
+bool UMaterialExpressionFloor::GenerateHLSLExpression(FMaterialHLSLGenerator& Generator, UE::HLSLTree::FScope& Scope, int32 OutputIndex, UE::HLSLTree::FExpression const*& OutExpression) const
 {
 	using namespace UE::HLSLTree;
-	FExpression* InputExpression = Input.AcquireHLSLExpression(Generator, Scope);
+	const FExpression* InputExpression = Input.AcquireHLSLExpression(Generator, Scope);
 	if (!InputExpression)
 	{
 		return false;
@@ -991,10 +991,10 @@ bool UMaterialExpressionFloor::GenerateHLSLExpression(FMaterialHLSLGenerator& Ge
 	return true;
 }
 
-bool UMaterialExpressionCeil::GenerateHLSLExpression(FMaterialHLSLGenerator& Generator, UE::HLSLTree::FScope& Scope, int32 OutputIndex, UE::HLSLTree::FExpression*& OutExpression)
+bool UMaterialExpressionCeil::GenerateHLSLExpression(FMaterialHLSLGenerator& Generator, UE::HLSLTree::FScope& Scope, int32 OutputIndex, UE::HLSLTree::FExpression const*& OutExpression) const
 {
 	using namespace UE::HLSLTree;
-	FExpression* InputExpression = Input.AcquireHLSLExpression(Generator, Scope);
+	const FExpression* InputExpression = Input.AcquireHLSLExpression(Generator, Scope);
 	if (!InputExpression)
 	{
 		return false;
@@ -1003,10 +1003,10 @@ bool UMaterialExpressionCeil::GenerateHLSLExpression(FMaterialHLSLGenerator& Gen
 	return true;
 }
 
-bool UMaterialExpressionRound::GenerateHLSLExpression(FMaterialHLSLGenerator& Generator, UE::HLSLTree::FScope& Scope, int32 OutputIndex, UE::HLSLTree::FExpression*& OutExpression)
+bool UMaterialExpressionRound::GenerateHLSLExpression(FMaterialHLSLGenerator& Generator, UE::HLSLTree::FScope& Scope, int32 OutputIndex, UE::HLSLTree::FExpression const*& OutExpression) const
 {
 	using namespace UE::HLSLTree;
-	FExpression* InputExpression = Input.AcquireHLSLExpression(Generator, Scope);
+	const FExpression* InputExpression = Input.AcquireHLSLExpression(Generator, Scope);
 	if (!InputExpression)
 	{
 		return false;
@@ -1015,10 +1015,10 @@ bool UMaterialExpressionRound::GenerateHLSLExpression(FMaterialHLSLGenerator& Ge
 	return true;
 }
 
-bool UMaterialExpressionTruncate::GenerateHLSLExpression(FMaterialHLSLGenerator& Generator, UE::HLSLTree::FScope& Scope, int32 OutputIndex, UE::HLSLTree::FExpression*& OutExpression)
+bool UMaterialExpressionTruncate::GenerateHLSLExpression(FMaterialHLSLGenerator& Generator, UE::HLSLTree::FScope& Scope, int32 OutputIndex, UE::HLSLTree::FExpression const*& OutExpression) const
 {
 	using namespace UE::HLSLTree;
-	FExpression* InputExpression = Input.AcquireHLSLExpression(Generator, Scope);
+	const FExpression* InputExpression = Input.AcquireHLSLExpression(Generator, Scope);
 	if (!InputExpression)
 	{
 		return false;
@@ -1027,10 +1027,10 @@ bool UMaterialExpressionTruncate::GenerateHLSLExpression(FMaterialHLSLGenerator&
 	return true;
 }
 
-bool UMaterialExpressionSaturate::GenerateHLSLExpression(FMaterialHLSLGenerator& Generator, UE::HLSLTree::FScope& Scope, int32 OutputIndex, UE::HLSLTree::FExpression*& OutExpression)
+bool UMaterialExpressionSaturate::GenerateHLSLExpression(FMaterialHLSLGenerator& Generator, UE::HLSLTree::FScope& Scope, int32 OutputIndex, UE::HLSLTree::FExpression const*& OutExpression) const
 {
 	using namespace UE::HLSLTree;
-	FExpression* InputExpression = Input.AcquireHLSLExpression(Generator, Scope);
+	const FExpression* InputExpression = Input.AcquireHLSLExpression(Generator, Scope);
 	if (!InputExpression)
 	{
 		return false;
@@ -1039,10 +1039,10 @@ bool UMaterialExpressionSaturate::GenerateHLSLExpression(FMaterialHLSLGenerator&
 	return true;
 }
 
-bool UMaterialExpressionSign::GenerateHLSLExpression(FMaterialHLSLGenerator& Generator, UE::HLSLTree::FScope& Scope, int32 OutputIndex, UE::HLSLTree::FExpression*& OutExpression)
+bool UMaterialExpressionSign::GenerateHLSLExpression(FMaterialHLSLGenerator& Generator, UE::HLSLTree::FScope& Scope, int32 OutputIndex, UE::HLSLTree::FExpression const*& OutExpression) const
 {
 	using namespace UE::HLSLTree;
-	FExpression* InputExpression = Input.AcquireHLSLExpression(Generator, Scope);
+	const FExpression* InputExpression = Input.AcquireHLSLExpression(Generator, Scope);
 	if (!InputExpression)
 	{
 		return false;
@@ -1053,9 +1053,9 @@ bool UMaterialExpressionSign::GenerateHLSLExpression(FMaterialHLSLGenerator& Gen
 
 static bool GenerateHLSLExpressionTrig(FMaterialHLSLGenerator& Generator,
 	UE::HLSLTree::EOperation Op,
-	UE::HLSLTree::FExpression* InputExpression,
+	const UE::HLSLTree::FExpression* InputExpression,
 	float Period,
-	UE::HLSLTree::FExpression*& OutExpression)
+	UE::HLSLTree::FExpression const*& OutExpression)
 {
 	if (!InputExpression)
 	{
@@ -1069,41 +1069,41 @@ static bool GenerateHLSLExpressionTrig(FMaterialHLSLGenerator& Generator,
 	return true;
 }
 
-bool UMaterialExpressionSine::GenerateHLSLExpression(FMaterialHLSLGenerator& Generator, UE::HLSLTree::FScope& Scope, int32 OutputIndex, UE::HLSLTree::FExpression*& OutExpression)
+bool UMaterialExpressionSine::GenerateHLSLExpression(FMaterialHLSLGenerator& Generator, UE::HLSLTree::FScope& Scope, int32 OutputIndex, UE::HLSLTree::FExpression const*& OutExpression) const
 {
 	return GenerateHLSLExpressionTrig(Generator, UE::HLSLTree::EOperation::Sin, Input.AcquireHLSLExpression(Generator, Scope), Period, OutExpression);
 }
 
-bool UMaterialExpressionCosine::GenerateHLSLExpression(FMaterialHLSLGenerator& Generator, UE::HLSLTree::FScope& Scope, int32 OutputIndex, UE::HLSLTree::FExpression*& OutExpression)
+bool UMaterialExpressionCosine::GenerateHLSLExpression(FMaterialHLSLGenerator& Generator, UE::HLSLTree::FScope& Scope, int32 OutputIndex, UE::HLSLTree::FExpression const*& OutExpression) const
 {
 	return GenerateHLSLExpressionTrig(Generator, UE::HLSLTree::EOperation::Cos, Input.AcquireHLSLExpression(Generator, Scope), Period, OutExpression);
 }
 
-bool UMaterialExpressionTangent::GenerateHLSLExpression(FMaterialHLSLGenerator& Generator, UE::HLSLTree::FScope& Scope, int32 OutputIndex, UE::HLSLTree::FExpression*& OutExpression)
+bool UMaterialExpressionTangent::GenerateHLSLExpression(FMaterialHLSLGenerator& Generator, UE::HLSLTree::FScope& Scope, int32 OutputIndex, UE::HLSLTree::FExpression const*& OutExpression) const
 {
 	return GenerateHLSLExpressionTrig(Generator, UE::HLSLTree::EOperation::Tan, Input.AcquireHLSLExpression(Generator, Scope), Period, OutExpression);
 }
 
-bool UMaterialExpressionArcsine::GenerateHLSLExpression(FMaterialHLSLGenerator& Generator, UE::HLSLTree::FScope& Scope, int32 OutputIndex, UE::HLSLTree::FExpression*& OutExpression)
+bool UMaterialExpressionArcsine::GenerateHLSLExpression(FMaterialHLSLGenerator& Generator, UE::HLSLTree::FScope& Scope, int32 OutputIndex, UE::HLSLTree::FExpression const*& OutExpression) const
 {
 	return GenerateHLSLExpressionTrig(Generator, UE::HLSLTree::EOperation::Asin, Input.AcquireHLSLExpression(Generator, Scope), 0.0f, OutExpression);
 }
 
-bool UMaterialExpressionArccosine::GenerateHLSLExpression(FMaterialHLSLGenerator& Generator, UE::HLSLTree::FScope& Scope, int32 OutputIndex, UE::HLSLTree::FExpression*& OutExpression)
+bool UMaterialExpressionArccosine::GenerateHLSLExpression(FMaterialHLSLGenerator& Generator, UE::HLSLTree::FScope& Scope, int32 OutputIndex, UE::HLSLTree::FExpression const*& OutExpression) const
 {
 	return GenerateHLSLExpressionTrig(Generator, UE::HLSLTree::EOperation::Acos, Input.AcquireHLSLExpression(Generator, Scope), 0.0f, OutExpression);
 }
 
-bool UMaterialExpressionArctangent::GenerateHLSLExpression(FMaterialHLSLGenerator& Generator, UE::HLSLTree::FScope& Scope, int32 OutputIndex, UE::HLSLTree::FExpression*& OutExpression)
+bool UMaterialExpressionArctangent::GenerateHLSLExpression(FMaterialHLSLGenerator& Generator, UE::HLSLTree::FScope& Scope, int32 OutputIndex, UE::HLSLTree::FExpression const*& OutExpression) const
 {
 	return GenerateHLSLExpressionTrig(Generator, UE::HLSLTree::EOperation::Atan, Input.AcquireHLSLExpression(Generator, Scope), 0.0f, OutExpression);
 }
 
-bool UMaterialExpressionArctangent2::GenerateHLSLExpression(FMaterialHLSLGenerator& Generator, UE::HLSLTree::FScope& Scope, int32 OutputIndex, UE::HLSLTree::FExpression*& OutExpression)
+bool UMaterialExpressionArctangent2::GenerateHLSLExpression(FMaterialHLSLGenerator& Generator, UE::HLSLTree::FScope& Scope, int32 OutputIndex, UE::HLSLTree::FExpression const*& OutExpression) const
 {
 	using namespace UE::HLSLTree;
-	FExpression* Lhs = Y.AcquireHLSLExpression(Generator, Scope);
-	FExpression* Rhs = X.AcquireHLSLExpression(Generator, Scope);
+	const FExpression* Lhs = Y.AcquireHLSLExpression(Generator, Scope);
+	const FExpression* Rhs = X.AcquireHLSLExpression(Generator, Scope);
 	if (!Lhs || !Rhs)
 	{
 		return false;
@@ -1112,26 +1112,26 @@ bool UMaterialExpressionArctangent2::GenerateHLSLExpression(FMaterialHLSLGenerat
 	return true;
 }
 
-bool UMaterialExpressionArcsineFast::GenerateHLSLExpression(FMaterialHLSLGenerator& Generator, UE::HLSLTree::FScope& Scope, int32 OutputIndex, UE::HLSLTree::FExpression*& OutExpression)
+bool UMaterialExpressionArcsineFast::GenerateHLSLExpression(FMaterialHLSLGenerator& Generator, UE::HLSLTree::FScope& Scope, int32 OutputIndex, UE::HLSLTree::FExpression const*& OutExpression) const
 {
 	return GenerateHLSLExpressionTrig(Generator, UE::HLSLTree::EOperation::AsinFast, Input.AcquireHLSLExpression(Generator, Scope), 0.0f, OutExpression);
 }
 
-bool UMaterialExpressionArccosineFast::GenerateHLSLExpression(FMaterialHLSLGenerator& Generator, UE::HLSLTree::FScope& Scope, int32 OutputIndex, UE::HLSLTree::FExpression*& OutExpression)
+bool UMaterialExpressionArccosineFast::GenerateHLSLExpression(FMaterialHLSLGenerator& Generator, UE::HLSLTree::FScope& Scope, int32 OutputIndex, UE::HLSLTree::FExpression const*& OutExpression) const
 {
 	return GenerateHLSLExpressionTrig(Generator, UE::HLSLTree::EOperation::AcosFast, Input.AcquireHLSLExpression(Generator, Scope), 0.0f, OutExpression);
 }
 
-bool UMaterialExpressionArctangentFast::GenerateHLSLExpression(FMaterialHLSLGenerator& Generator, UE::HLSLTree::FScope& Scope, int32 OutputIndex, UE::HLSLTree::FExpression*& OutExpression)
+bool UMaterialExpressionArctangentFast::GenerateHLSLExpression(FMaterialHLSLGenerator& Generator, UE::HLSLTree::FScope& Scope, int32 OutputIndex, UE::HLSLTree::FExpression const*& OutExpression) const
 {
 	return GenerateHLSLExpressionTrig(Generator, UE::HLSLTree::EOperation::AtanFast, Input.AcquireHLSLExpression(Generator, Scope), 0.0f, OutExpression);
 }
 
-bool UMaterialExpressionArctangent2Fast::GenerateHLSLExpression(FMaterialHLSLGenerator& Generator, UE::HLSLTree::FScope& Scope, int32 OutputIndex, UE::HLSLTree::FExpression*& OutExpression)
+bool UMaterialExpressionArctangent2Fast::GenerateHLSLExpression(FMaterialHLSLGenerator& Generator, UE::HLSLTree::FScope& Scope, int32 OutputIndex, UE::HLSLTree::FExpression const*& OutExpression) const
 {
 	using namespace UE::HLSLTree;
-	FExpression* Lhs = Y.AcquireHLSLExpression(Generator, Scope);
-	FExpression* Rhs = X.AcquireHLSLExpression(Generator, Scope);
+	const FExpression* Lhs = Y.AcquireHLSLExpression(Generator, Scope);
+	const FExpression* Rhs = X.AcquireHLSLExpression(Generator, Scope);
 	if (!Lhs || !Rhs)
 	{
 		return false;
@@ -1140,10 +1140,10 @@ bool UMaterialExpressionArctangent2Fast::GenerateHLSLExpression(FMaterialHLSLGen
 	return true;
 }
 
-bool UMaterialExpressionDDX::GenerateHLSLExpression(FMaterialHLSLGenerator& Generator, UE::HLSLTree::FScope& Scope, int32 OutputIndex, UE::HLSLTree::FExpression*& OutExpression)
+bool UMaterialExpressionDDX::GenerateHLSLExpression(FMaterialHLSLGenerator& Generator, UE::HLSLTree::FScope& Scope, int32 OutputIndex, UE::HLSLTree::FExpression const*& OutExpression) const
 {
 	using namespace UE::HLSLTree;
-	FExpression* ExpressionInput = Value.AcquireHLSLExpression(Generator, Scope);
+	const FExpression* ExpressionInput = Value.AcquireHLSLExpression(Generator, Scope);
 	if (!ExpressionInput)
 	{
 		return false;
@@ -1152,10 +1152,10 @@ bool UMaterialExpressionDDX::GenerateHLSLExpression(FMaterialHLSLGenerator& Gene
 	return true;
 }
 
-bool UMaterialExpressionDDY::GenerateHLSLExpression(FMaterialHLSLGenerator& Generator, UE::HLSLTree::FScope& Scope, int32 OutputIndex, UE::HLSLTree::FExpression*& OutExpression)
+bool UMaterialExpressionDDY::GenerateHLSLExpression(FMaterialHLSLGenerator& Generator, UE::HLSLTree::FScope& Scope, int32 OutputIndex, UE::HLSLTree::FExpression const*& OutExpression) const
 {
 	using namespace UE::HLSLTree;
-	FExpression* ExpressionInput = Value.AcquireHLSLExpression(Generator, Scope);
+	const FExpression* ExpressionInput = Value.AcquireHLSLExpression(Generator, Scope);
 	if (!ExpressionInput)
 	{
 		return false;
@@ -1164,10 +1164,10 @@ bool UMaterialExpressionDDY::GenerateHLSLExpression(FMaterialHLSLGenerator& Gene
 	return true;
 }
 
-bool UMaterialExpressionBinaryOp::GenerateHLSLExpression(FMaterialHLSLGenerator& Generator, UE::HLSLTree::FScope& Scope, int32 OutputIndex, UE::HLSLTree::FExpression*& OutExpression)
+bool UMaterialExpressionBinaryOp::GenerateHLSLExpression(FMaterialHLSLGenerator& Generator, UE::HLSLTree::FScope& Scope, int32 OutputIndex, UE::HLSLTree::FExpression const*& OutExpression) const
 {
-	UE::HLSLTree::FExpression* Lhs = A.AcquireHLSLExpressionOrConstant(Generator, Scope, ConstA);
-	UE::HLSLTree::FExpression* Rhs = B.AcquireHLSLExpressionOrConstant(Generator, Scope, ConstB);
+	const UE::HLSLTree::FExpression* Lhs = A.AcquireHLSLExpressionOrConstant(Generator, Scope, ConstA);
+	const UE::HLSLTree::FExpression* Rhs = B.AcquireHLSLExpressionOrConstant(Generator, Scope, ConstB);
 	if (!Lhs || !Rhs)
 	{
 		return false;
@@ -1177,10 +1177,10 @@ bool UMaterialExpressionBinaryOp::GenerateHLSLExpression(FMaterialHLSLGenerator&
 	return true;
 }
 
-bool UMaterialExpressionAdd::GenerateHLSLExpression(FMaterialHLSLGenerator& Generator, UE::HLSLTree::FScope& Scope, int32 OutputIndex, UE::HLSLTree::FExpression*& OutExpression)
+bool UMaterialExpressionAdd::GenerateHLSLExpression(FMaterialHLSLGenerator& Generator, UE::HLSLTree::FScope& Scope, int32 OutputIndex, UE::HLSLTree::FExpression const*& OutExpression) const
 {
-	UE::HLSLTree::FExpression* Lhs = A.AcquireHLSLExpressionOrConstant(Generator, Scope, ConstA);
-	UE::HLSLTree::FExpression* Rhs = B.AcquireHLSLExpressionOrConstant(Generator, Scope, ConstB);
+	const UE::HLSLTree::FExpression* Lhs = A.AcquireHLSLExpressionOrConstant(Generator, Scope, ConstA);
+	const UE::HLSLTree::FExpression* Rhs = B.AcquireHLSLExpressionOrConstant(Generator, Scope, ConstB);
 	if (!Lhs || !Rhs)
 	{
 		return false;
@@ -1189,10 +1189,10 @@ bool UMaterialExpressionAdd::GenerateHLSLExpression(FMaterialHLSLGenerator& Gene
 	return true;
 }
 
-bool UMaterialExpressionSubtract::GenerateHLSLExpression(FMaterialHLSLGenerator& Generator, UE::HLSLTree::FScope& Scope, int32 OutputIndex, UE::HLSLTree::FExpression*& OutExpression)
+bool UMaterialExpressionSubtract::GenerateHLSLExpression(FMaterialHLSLGenerator& Generator, UE::HLSLTree::FScope& Scope, int32 OutputIndex, UE::HLSLTree::FExpression const*& OutExpression) const
 {
-	UE::HLSLTree::FExpression* Lhs = A.AcquireHLSLExpressionOrConstant(Generator, Scope, ConstA);
-	UE::HLSLTree::FExpression* Rhs = B.AcquireHLSLExpressionOrConstant(Generator, Scope, ConstB);
+	const UE::HLSLTree::FExpression* Lhs = A.AcquireHLSLExpressionOrConstant(Generator, Scope, ConstA);
+	const UE::HLSLTree::FExpression* Rhs = B.AcquireHLSLExpressionOrConstant(Generator, Scope, ConstB);
 	if (!Lhs || !Rhs)
 	{
 		return false;
@@ -1201,10 +1201,10 @@ bool UMaterialExpressionSubtract::GenerateHLSLExpression(FMaterialHLSLGenerator&
 	return true;
 }
 
-bool UMaterialExpressionMultiply::GenerateHLSLExpression(FMaterialHLSLGenerator& Generator, UE::HLSLTree::FScope& Scope, int32 OutputIndex, UE::HLSLTree::FExpression*& OutExpression)
+bool UMaterialExpressionMultiply::GenerateHLSLExpression(FMaterialHLSLGenerator& Generator, UE::HLSLTree::FScope& Scope, int32 OutputIndex, UE::HLSLTree::FExpression const*& OutExpression) const
 {
-	UE::HLSLTree::FExpression* Lhs = A.AcquireHLSLExpressionOrConstant(Generator, Scope, ConstA);
-	UE::HLSLTree::FExpression* Rhs = B.AcquireHLSLExpressionOrConstant(Generator, Scope, ConstB);
+	const UE::HLSLTree::FExpression* Lhs = A.AcquireHLSLExpressionOrConstant(Generator, Scope, ConstA);
+	const UE::HLSLTree::FExpression* Rhs = B.AcquireHLSLExpressionOrConstant(Generator, Scope, ConstB);
 	if (!Lhs || !Rhs)
 	{
 		return false;
@@ -1213,10 +1213,10 @@ bool UMaterialExpressionMultiply::GenerateHLSLExpression(FMaterialHLSLGenerator&
 	return true;
 }
 
-bool UMaterialExpressionDivide::GenerateHLSLExpression(FMaterialHLSLGenerator& Generator, UE::HLSLTree::FScope& Scope, int32 OutputIndex, UE::HLSLTree::FExpression*& OutExpression)
+bool UMaterialExpressionDivide::GenerateHLSLExpression(FMaterialHLSLGenerator& Generator, UE::HLSLTree::FScope& Scope, int32 OutputIndex, UE::HLSLTree::FExpression const*& OutExpression) const
 {
-	UE::HLSLTree::FExpression* Lhs = A.AcquireHLSLExpressionOrConstant(Generator, Scope, ConstA);
-	UE::HLSLTree::FExpression* Rhs = B.AcquireHLSLExpressionOrConstant(Generator, Scope, ConstB);
+	const UE::HLSLTree::FExpression* Lhs = A.AcquireHLSLExpressionOrConstant(Generator, Scope, ConstA);
+	const UE::HLSLTree::FExpression* Rhs = B.AcquireHLSLExpressionOrConstant(Generator, Scope, ConstB);
 	if (!Lhs || !Rhs)
 	{
 		return false;
@@ -1225,10 +1225,10 @@ bool UMaterialExpressionDivide::GenerateHLSLExpression(FMaterialHLSLGenerator& G
 	return true;
 }
 
-bool UMaterialExpressionFmod::GenerateHLSLExpression(FMaterialHLSLGenerator& Generator, UE::HLSLTree::FScope& Scope, int32 OutputIndex, UE::HLSLTree::FExpression*& OutExpression)
+bool UMaterialExpressionFmod::GenerateHLSLExpression(FMaterialHLSLGenerator& Generator, UE::HLSLTree::FScope& Scope, int32 OutputIndex, UE::HLSLTree::FExpression const*& OutExpression) const
 {
-	UE::HLSLTree::FExpression* Lhs = A.AcquireHLSLExpression(Generator, Scope);
-	UE::HLSLTree::FExpression* Rhs = B.AcquireHLSLExpression(Generator, Scope);
+	const UE::HLSLTree::FExpression* Lhs = A.AcquireHLSLExpression(Generator, Scope);
+	const UE::HLSLTree::FExpression* Rhs = B.AcquireHLSLExpression(Generator, Scope);
 	if (!Lhs || !Rhs)
 	{
 		return false;
@@ -1237,11 +1237,11 @@ bool UMaterialExpressionFmod::GenerateHLSLExpression(FMaterialHLSLGenerator& Gen
 	return true;
 }
 
-bool UMaterialExpressionPower::GenerateHLSLExpression(FMaterialHLSLGenerator& Generator, UE::HLSLTree::FScope& Scope, int32 OutputIndex, UE::HLSLTree::FExpression*& OutExpression)
+bool UMaterialExpressionPower::GenerateHLSLExpression(FMaterialHLSLGenerator& Generator, UE::HLSLTree::FScope& Scope, int32 OutputIndex, UE::HLSLTree::FExpression const*& OutExpression) const
 {
 	using namespace UE::HLSLTree;
-	FExpression* Lhs = Base.AcquireHLSLExpression(Generator, Scope);
-	FExpression* Rhs = Exponent.AcquireHLSLExpressionOrConstant(Generator, Scope, ConstExponent);
+	const FExpression* Lhs = Base.AcquireHLSLExpression(Generator, Scope);
+	const FExpression* Rhs = Exponent.AcquireHLSLExpressionOrConstant(Generator, Scope, ConstExponent);
 	if (!Lhs || !Rhs)
 	{
 		return false;
@@ -1250,10 +1250,10 @@ bool UMaterialExpressionPower::GenerateHLSLExpression(FMaterialHLSLGenerator& Ge
 	return true;
 }
 
-bool UMaterialExpressionDotProduct::GenerateHLSLExpression(FMaterialHLSLGenerator& Generator, UE::HLSLTree::FScope& Scope, int32 OutputIndex, UE::HLSLTree::FExpression*& OutExpression)
+bool UMaterialExpressionDotProduct::GenerateHLSLExpression(FMaterialHLSLGenerator& Generator, UE::HLSLTree::FScope& Scope, int32 OutputIndex, UE::HLSLTree::FExpression const*& OutExpression) const
 {
-	UE::HLSLTree::FExpression* Lhs = A.AcquireHLSLExpression(Generator, Scope);
-	UE::HLSLTree::FExpression* Rhs = B.AcquireHLSLExpression(Generator, Scope);
+	const UE::HLSLTree::FExpression* Lhs = A.AcquireHLSLExpression(Generator, Scope);
+	const UE::HLSLTree::FExpression* Rhs = B.AcquireHLSLExpression(Generator, Scope);
 	if (!Lhs || !Rhs)
 	{
 		return false;
@@ -1262,10 +1262,10 @@ bool UMaterialExpressionDotProduct::GenerateHLSLExpression(FMaterialHLSLGenerato
 	return true;
 }
 
-bool UMaterialExpressionCrossProduct::GenerateHLSLExpression(FMaterialHLSLGenerator& Generator, UE::HLSLTree::FScope& Scope, int32 OutputIndex, UE::HLSLTree::FExpression*& OutExpression)
+bool UMaterialExpressionCrossProduct::GenerateHLSLExpression(FMaterialHLSLGenerator& Generator, UE::HLSLTree::FScope& Scope, int32 OutputIndex, UE::HLSLTree::FExpression const*& OutExpression) const
 {
-	UE::HLSLTree::FExpression* Lhs = A.AcquireHLSLExpression(Generator, Scope);
-	UE::HLSLTree::FExpression* Rhs = B.AcquireHLSLExpression(Generator, Scope);
+	const UE::HLSLTree::FExpression* Lhs = A.AcquireHLSLExpression(Generator, Scope);
+	const UE::HLSLTree::FExpression* Rhs = B.AcquireHLSLExpression(Generator, Scope);
 	if (!Lhs || !Rhs)
 	{
 		return false;
@@ -1274,10 +1274,10 @@ bool UMaterialExpressionCrossProduct::GenerateHLSLExpression(FMaterialHLSLGenera
 	return true;
 }
 
-bool UMaterialExpressionMin::GenerateHLSLExpression(FMaterialHLSLGenerator& Generator, UE::HLSLTree::FScope& Scope, int32 OutputIndex, UE::HLSLTree::FExpression*& OutExpression)
+bool UMaterialExpressionMin::GenerateHLSLExpression(FMaterialHLSLGenerator& Generator, UE::HLSLTree::FScope& Scope, int32 OutputIndex, UE::HLSLTree::FExpression const*& OutExpression) const
 {
-	UE::HLSLTree::FExpression* Lhs = A.AcquireHLSLExpressionOrConstant(Generator, Scope, ConstA);
-	UE::HLSLTree::FExpression* Rhs = B.AcquireHLSLExpressionOrConstant(Generator, Scope, ConstB);
+	const UE::HLSLTree::FExpression* Lhs = A.AcquireHLSLExpressionOrConstant(Generator, Scope, ConstA);
+	const UE::HLSLTree::FExpression* Rhs = B.AcquireHLSLExpressionOrConstant(Generator, Scope, ConstB);
 	if (!Lhs || !Rhs)
 	{
 		return false;
@@ -1286,10 +1286,10 @@ bool UMaterialExpressionMin::GenerateHLSLExpression(FMaterialHLSLGenerator& Gene
 	return true;
 }
 
-bool UMaterialExpressionMax::GenerateHLSLExpression(FMaterialHLSLGenerator& Generator, UE::HLSLTree::FScope& Scope, int32 OutputIndex, UE::HLSLTree::FExpression*& OutExpression)
+bool UMaterialExpressionMax::GenerateHLSLExpression(FMaterialHLSLGenerator& Generator, UE::HLSLTree::FScope& Scope, int32 OutputIndex, UE::HLSLTree::FExpression const*& OutExpression) const
 {
-	UE::HLSLTree::FExpression* Lhs = A.AcquireHLSLExpressionOrConstant(Generator, Scope, ConstA);
-	UE::HLSLTree::FExpression* Rhs = B.AcquireHLSLExpressionOrConstant(Generator, Scope, ConstB);
+	const UE::HLSLTree::FExpression* Lhs = A.AcquireHLSLExpressionOrConstant(Generator, Scope, ConstA);
+	const UE::HLSLTree::FExpression* Rhs = B.AcquireHLSLExpressionOrConstant(Generator, Scope, ConstB);
 	if (!Lhs || !Rhs)
 	{
 		return false;
@@ -1298,11 +1298,11 @@ bool UMaterialExpressionMax::GenerateHLSLExpression(FMaterialHLSLGenerator& Gene
 	return true;
 }
 
-bool UMaterialExpressionClamp::GenerateHLSLExpression(FMaterialHLSLGenerator& Generator, UE::HLSLTree::FScope& Scope, int32 OutputIndex, UE::HLSLTree::FExpression*& OutExpression)
+bool UMaterialExpressionClamp::GenerateHLSLExpression(FMaterialHLSLGenerator& Generator, UE::HLSLTree::FScope& Scope, int32 OutputIndex, UE::HLSLTree::FExpression const*& OutExpression) const
 {
-	UE::HLSLTree::FExpression* ExpressionMin = Min.AcquireHLSLExpressionOrConstant(Generator, Scope, MinDefault);
-	UE::HLSLTree::FExpression* ExpressionMax = Max.AcquireHLSLExpressionOrConstant(Generator, Scope, MaxDefault);
-	UE::HLSLTree::FExpression* ExpressionInput = Input.AcquireHLSLExpression(Generator, Scope);
+	const UE::HLSLTree::FExpression* ExpressionMin = Min.AcquireHLSLExpressionOrConstant(Generator, Scope, MinDefault);
+	const UE::HLSLTree::FExpression* ExpressionMax = Max.AcquireHLSLExpressionOrConstant(Generator, Scope, MaxDefault);
+	const UE::HLSLTree::FExpression* ExpressionInput = Input.AcquireHLSLExpression(Generator, Scope);
 	if (!ExpressionMin || !ExpressionMax || !ExpressionInput)
 	{
 		return false;
@@ -1321,11 +1321,11 @@ bool UMaterialExpressionClamp::GenerateHLSLExpression(FMaterialHLSLGenerator& Ge
 	return true;
 }
 
-bool UMaterialExpressionLinearInterpolate::GenerateHLSLExpression(FMaterialHLSLGenerator& Generator, UE::HLSLTree::FScope& Scope, int32 OutputIndex, UE::HLSLTree::FExpression*& OutExpression)
+bool UMaterialExpressionLinearInterpolate::GenerateHLSLExpression(FMaterialHLSLGenerator& Generator, UE::HLSLTree::FScope& Scope, int32 OutputIndex, UE::HLSLTree::FExpression const*& OutExpression) const
 {
-	UE::HLSLTree::FExpression* ExpressionA = A.AcquireHLSLExpressionOrConstant(Generator, Scope, ConstA);
-	UE::HLSLTree::FExpression* ExpressionB = B.AcquireHLSLExpressionOrConstant(Generator, Scope, ConstB);
-	UE::HLSLTree::FExpression* ExpressionAlpha = Alpha.AcquireHLSLExpressionOrConstant(Generator, Scope, ConstAlpha);
+	const UE::HLSLTree::FExpression* ExpressionA = A.AcquireHLSLExpressionOrConstant(Generator, Scope, ConstA);
+	const UE::HLSLTree::FExpression* ExpressionB = B.AcquireHLSLExpressionOrConstant(Generator, Scope, ConstB);
+	const UE::HLSLTree::FExpression* ExpressionAlpha = Alpha.AcquireHLSLExpressionOrConstant(Generator, Scope, ConstAlpha);
 	if (!ExpressionA || !ExpressionB || !ExpressionAlpha)
 	{
 		return false;
@@ -1334,10 +1334,10 @@ bool UMaterialExpressionLinearInterpolate::GenerateHLSLExpression(FMaterialHLSLG
 	return true;
 }
 
-bool UMaterialExpressionDistance::GenerateHLSLExpression(FMaterialHLSLGenerator& Generator, UE::HLSLTree::FScope& Scope, int32 OutputIndex, UE::HLSLTree::FExpression*& OutExpression)
+bool UMaterialExpressionDistance::GenerateHLSLExpression(FMaterialHLSLGenerator& Generator, UE::HLSLTree::FScope& Scope, int32 OutputIndex, UE::HLSLTree::FExpression const*& OutExpression) const
 {
-	UE::HLSLTree::FExpression* ExpressionA = A.AcquireHLSLExpression(Generator, Scope);
-	UE::HLSLTree::FExpression* ExpressionB = B.AcquireHLSLExpression(Generator, Scope);
+	const UE::HLSLTree::FExpression* ExpressionA = A.AcquireHLSLExpression(Generator, Scope);
+	const UE::HLSLTree::FExpression* ExpressionB = B.AcquireHLSLExpression(Generator, Scope);
 	if (!ExpressionA || !ExpressionB)
 	{
 		return false;
@@ -1347,9 +1347,9 @@ bool UMaterialExpressionDistance::GenerateHLSLExpression(FMaterialHLSLGenerator&
 	return true;
 }
 
-bool UMaterialExpressionNormalize::GenerateHLSLExpression(FMaterialHLSLGenerator& Generator, UE::HLSLTree::FScope& Scope, int32 OutputIndex, UE::HLSLTree::FExpression*& OutExpression)
+bool UMaterialExpressionNormalize::GenerateHLSLExpression(FMaterialHLSLGenerator& Generator, UE::HLSLTree::FScope& Scope, int32 OutputIndex, UE::HLSLTree::FExpression const*& OutExpression) const
 {
-	UE::HLSLTree::FExpression* ExpressionInput = VectorInput.AcquireHLSLExpression(Generator, Scope);
+	const UE::HLSLTree::FExpression* ExpressionInput = VectorInput.AcquireHLSLExpression(Generator, Scope);
 	if (!ExpressionInput)
 	{
 		return false;
@@ -1359,10 +1359,10 @@ bool UMaterialExpressionNormalize::GenerateHLSLExpression(FMaterialHLSLGenerator
 	return true;
 }
 
-bool UMaterialExpressionAppendVector::GenerateHLSLExpression(FMaterialHLSLGenerator& Generator, UE::HLSLTree::FScope& Scope, int32 OutputIndex, UE::HLSLTree::FExpression*& OutExpression)
+bool UMaterialExpressionAppendVector::GenerateHLSLExpression(FMaterialHLSLGenerator& Generator, UE::HLSLTree::FScope& Scope, int32 OutputIndex, UE::HLSLTree::FExpression const*& OutExpression) const
 {
-	UE::HLSLTree::FExpression* Lhs = A.AcquireHLSLExpression(Generator, Scope);
-	UE::HLSLTree::FExpression* Rhs = B.AcquireHLSLExpression(Generator, Scope);
+	const UE::HLSLTree::FExpression* Lhs = A.AcquireHLSLExpression(Generator, Scope);
+	const UE::HLSLTree::FExpression* Rhs = B.AcquireHLSLExpression(Generator, Scope);
 	if (!Lhs || !Rhs)
 	{
 		return false;
@@ -1371,9 +1371,9 @@ bool UMaterialExpressionAppendVector::GenerateHLSLExpression(FMaterialHLSLGenera
 	return true;
 }
 
-bool UMaterialExpressionComponentMask::GenerateHLSLExpression(FMaterialHLSLGenerator& Generator, UE::HLSLTree::FScope& Scope, int32 OutputIndex, UE::HLSLTree::FExpression*& OutExpression)
+bool UMaterialExpressionComponentMask::GenerateHLSLExpression(FMaterialHLSLGenerator& Generator, UE::HLSLTree::FScope& Scope, int32 OutputIndex, UE::HLSLTree::FExpression const*& OutExpression) const
 {
-	UE::HLSLTree::FExpression* InputExpression = Input.AcquireHLSLExpression(Generator, Scope);
+	const UE::HLSLTree::FExpression* InputExpression = Input.AcquireHLSLExpression(Generator, Scope);
 	if (!InputExpression)
 	{
 		return false;
@@ -1384,9 +1384,9 @@ bool UMaterialExpressionComponentMask::GenerateHLSLExpression(FMaterialHLSLGener
 }
 
 
-bool UMaterialExpressionGetMaterialAttributes::GenerateHLSLExpression(FMaterialHLSLGenerator& Generator, UE::HLSLTree::FScope& Scope, int32 OutputIndex, UE::HLSLTree::FExpression*& OutExpression)
+bool UMaterialExpressionGetMaterialAttributes::GenerateHLSLExpression(FMaterialHLSLGenerator& Generator, UE::HLSLTree::FScope& Scope, int32 OutputIndex, UE::HLSLTree::FExpression const*& OutExpression) const
 {
-	UE::HLSLTree::FExpression* AttributesExpression = MaterialAttributes.AcquireHLSLExpression(Generator, Scope);
+	const UE::HLSLTree::FExpression* AttributesExpression = MaterialAttributes.AcquireHLSLExpression(Generator, Scope);
 	if (!AttributesExpression)
 	{
 		return false;
@@ -1411,9 +1411,9 @@ bool UMaterialExpressionGetMaterialAttributes::GenerateHLSLExpression(FMaterialH
 
 }
 
-bool UMaterialExpressionSetMaterialAttributes::GenerateHLSLExpression(FMaterialHLSLGenerator& Generator, UE::HLSLTree::FScope& Scope, int32 OutputIndex, UE::HLSLTree::FExpression*& OutExpression)
+bool UMaterialExpressionSetMaterialAttributes::GenerateHLSLExpression(FMaterialHLSLGenerator& Generator, UE::HLSLTree::FScope& Scope, int32 OutputIndex, UE::HLSLTree::FExpression const*& OutExpression) const
 {
-	UE::HLSLTree::FExpression* AttributesExpression = Inputs[0].AcquireHLSLExpressionOrConstant(Generator, Scope, Generator.GetMaterialAttributesDefaultValue());
+	const UE::HLSLTree::FExpression* AttributesExpression = Inputs[0].AcquireHLSLExpressionOrConstant(Generator, Scope, Generator.GetMaterialAttributesDefaultValue());
 	
 	for (int32 PinIndex = 0; PinIndex < AttributeSetTypes.Num(); ++PinIndex)
 	{
@@ -1426,7 +1426,7 @@ bool UMaterialExpressionSetMaterialAttributes::GenerateHLSLExpression(FMaterialH
 			const FString& AttributeName = FMaterialAttributeDefinitionMap::GetAttributeName(AttributeID);
 			//if (AttributeFrequency == Compiler->GetCurrentShaderFrequency())
 			{
-				UE::HLSLTree::FExpression* ValueExpression = AttributeInput.TryAcquireHLSLExpression(Generator, Scope);
+				const UE::HLSLTree::FExpression* ValueExpression = AttributeInput.TryAcquireHLSLExpression(Generator, Scope);
 				if (ValueExpression)
 				{
 					const UE::Shader::FStructField* AttributeField = Generator.GetMaterialAttributesType()->FindFieldByName(*AttributeName);
@@ -1440,16 +1440,16 @@ bool UMaterialExpressionSetMaterialAttributes::GenerateHLSLExpression(FMaterialH
 	return true;
 }
 
-static UE::HLSLTree::FExpression* SetAttribute(FMaterialHLSLGenerator& Generator,
+static const UE::HLSLTree::FExpression* SetAttribute(FMaterialHLSLGenerator& Generator,
 	UE::HLSLTree::FScope& Scope,
 	EMaterialProperty Property,
 	const FExpressionInput& Input,
-	UE::HLSLTree::FExpression* AttributesExpression)
+	const UE::HLSLTree::FExpression* AttributesExpression)
 {
 	using namespace UE::HLSLTree;
 	using namespace UE::Shader;
 
-	FExpression* InputExpression = Input.TryAcquireHLSLExpression(Generator, Scope);
+	const FExpression* InputExpression = Input.TryAcquireHLSLExpression(Generator, Scope);
 	if (InputExpression)
 	{
 		const FString& AttributeName = FMaterialAttributeDefinitionMap::GetAttributeName(Property);
@@ -1459,10 +1459,10 @@ static UE::HLSLTree::FExpression* SetAttribute(FMaterialHLSLGenerator& Generator
 	return AttributesExpression;
 }
 
-bool UMaterialExpressionMakeMaterialAttributes::GenerateHLSLExpression(FMaterialHLSLGenerator& Generator, UE::HLSLTree::FScope& Scope, int32 OutputIndex, UE::HLSLTree::FExpression*& OutExpression)
+bool UMaterialExpressionMakeMaterialAttributes::GenerateHLSLExpression(FMaterialHLSLGenerator& Generator, UE::HLSLTree::FScope& Scope, int32 OutputIndex, UE::HLSLTree::FExpression const*& OutExpression) const
 {
 	using namespace UE::HLSLTree;
-	FExpression* AttributesExpression = Generator.NewConstant(Generator.GetMaterialAttributesDefaultValue());
+	const FExpression* AttributesExpression = Generator.NewConstant(Generator.GetMaterialAttributesDefaultValue());
 	AttributesExpression = SetAttribute(Generator, Scope, MP_BaseColor, BaseColor, AttributesExpression);
 	AttributesExpression = SetAttribute(Generator, Scope, MP_Metallic, Metallic, AttributesExpression);
 
@@ -1493,7 +1493,7 @@ bool UMaterialExpressionMakeMaterialAttributes::GenerateHLSLExpression(FMaterial
 	return true;
 }
 
-bool UMaterialExpressionBreakMaterialAttributes::GenerateHLSLExpression(FMaterialHLSLGenerator& Generator, UE::HLSLTree::FScope& Scope, int32 OutputIndex, UE::HLSLTree::FExpression*& OutExpression)
+bool UMaterialExpressionBreakMaterialAttributes::GenerateHLSLExpression(FMaterialHLSLGenerator& Generator, UE::HLSLTree::FScope& Scope, int32 OutputIndex, UE::HLSLTree::FExpression const*& OutExpression) const
 {
 	using namespace UE::HLSLTree;
 	using namespace UE::Shader;
@@ -1508,7 +1508,7 @@ bool UMaterialExpressionBreakMaterialAttributes::GenerateHLSLExpression(FMateria
 		return Generator.Error(TEXT("Invalid output"));
 	}
 
-	FExpression* AttributesExpression = MaterialAttributes.AcquireHLSLExpression(Generator, Scope);
+	const FExpression* AttributesExpression = MaterialAttributes.AcquireHLSLExpression(Generator, Scope);
 	if (!AttributesExpression)
 	{
 		return false;
@@ -1520,29 +1520,29 @@ bool UMaterialExpressionBreakMaterialAttributes::GenerateHLSLExpression(FMateria
 	return true;
 }
 
-bool UMaterialExpressionFunctionOutput::GenerateHLSLExpression(FMaterialHLSLGenerator& Generator, UE::HLSLTree::FScope& Scope, int32 OutputIndex, UE::HLSLTree::FExpression*& OutExpression)
+bool UMaterialExpressionFunctionOutput::GenerateHLSLExpression(FMaterialHLSLGenerator& Generator, UE::HLSLTree::FScope& Scope, int32 OutputIndex, UE::HLSLTree::FExpression const*& OutExpression) const
 {
 	// This should only be called when editing/previewing the function directly
 	OutExpression = A.AcquireHLSLExpression(Generator, Scope);
 	return true;
 }
 
-bool UMaterialExpressionFunctionInput::GenerateHLSLExpression(FMaterialHLSLGenerator& Generator, UE::HLSLTree::FScope& Scope, int32 OutputIndex, UE::HLSLTree::FExpression*& OutExpression)
+bool UMaterialExpressionFunctionInput::GenerateHLSLExpression(FMaterialHLSLGenerator& Generator, UE::HLSLTree::FScope& Scope, int32 OutputIndex, UE::HLSLTree::FExpression const*& OutExpression) const
 {
 	OutExpression = Generator.AcquireFunctionInputExpression(Scope, this);
 	return true;
 }
 
-bool UMaterialExpressionMaterialFunctionCall::GenerateHLSLExpression(FMaterialHLSLGenerator& Generator, UE::HLSLTree::FScope& Scope, int32 OutputIndex, UE::HLSLTree::FExpression*& OutExpression)
+bool UMaterialExpressionMaterialFunctionCall::GenerateHLSLExpression(FMaterialHLSLGenerator& Generator, UE::HLSLTree::FScope& Scope, int32 OutputIndex, UE::HLSLTree::FExpression const*& OutExpression) const
 {
 	using namespace UE::HLSLTree;
-	TArray<FExpression*, TInlineAllocator<16>> ConnectedInputs;
+	TArray<const FExpression*, TInlineAllocator<16>> ConnectedInputs;
 	ConnectedInputs.Reserve(FunctionInputs.Num());
 	for (int32 InputIndex = 0; InputIndex < FunctionInputs.Num(); ++InputIndex)
 	{
 		// ConnectedInputs are the inputs from the UMaterialFunctionCall object
 		// We want to connect the UMaterialExpressionFunctionInput from the UMaterialFunction to whatever UMaterialExpression is passed to the UMaterialFunctionCall
-		FExpression* ConnectedInput = FunctionInputs[InputIndex].Input.TryAcquireHLSLExpression(Generator, Scope);
+		const FExpression* ConnectedInput = FunctionInputs[InputIndex].Input.TryAcquireHLSLExpression(Generator, Scope);
 		ConnectedInputs.Add(ConnectedInput);
 	}
 
@@ -1550,7 +1550,7 @@ bool UMaterialExpressionMaterialFunctionCall::GenerateHLSLExpression(FMaterialHL
 	return true;
 }
 
-bool UMaterialExpressionMaterialAttributeLayers::GenerateHLSLExpression(FMaterialHLSLGenerator& Generator, UE::HLSLTree::FScope& Scope, int32 OutputIndex, UE::HLSLTree::FExpression*& OutExpression)
+bool UMaterialExpressionMaterialAttributeLayers::GenerateHLSLExpression(FMaterialHLSLGenerator& Generator, UE::HLSLTree::FScope& Scope, int32 OutputIndex, UE::HLSLTree::FExpression const*& OutExpression) const
 {
 	using namespace UE::HLSLTree;
 
@@ -1564,15 +1564,15 @@ bool UMaterialExpressionMaterialAttributeLayers::GenerateHLSLExpression(FMateria
 	TArray<FFunctionExpressionInput> FunctionInputs;
 	TArray<FFunctionExpressionOutput> FunctionOutputs;
 
-	FExpression* ExpressionLayerInput = Input.AcquireHLSLExpressionOrConstant(Generator, Scope, Generator.GetMaterialAttributesDefaultValue());
-	TArray<FExpression*, TInlineAllocator<1>> LayerInputExpressions;
+	const FExpression* ExpressionLayerInput = Input.AcquireHLSLExpressionOrConstant(Generator, Scope, Generator.GetMaterialAttributesDefaultValue());
+	TArray<const FExpression*, TInlineAllocator<1>> LayerInputExpressions;
 
-	TArray<FExpression*, TInlineAllocator<16>> LayerExpressions;
+	TArray<const FExpression*, TInlineAllocator<16>> LayerExpressions;
 	LayerExpressions.Reserve(MaterialLayers.Layers.Num());
 	for (int32 LayerIndex = 0; LayerIndex < MaterialLayers.Layers.Num(); ++LayerIndex)
 	{
 		UMaterialFunctionInterface* LayerFunction = MaterialLayers.Layers[LayerIndex];
-		FExpression* LayerExpression = nullptr;
+		const FExpression* LayerExpression = nullptr;
 		if (LayerFunction && MaterialLayers.LayerStates[LayerIndex])
 		{
 			const EMaterialFunctionUsage Usage = LayerFunction->GetMaterialFunctionUsage();
@@ -1601,13 +1601,13 @@ bool UMaterialExpressionMaterialAttributeLayers::GenerateHLSLExpression(FMateria
 		LayerExpressions.Add(LayerExpression);
 	}
 
-	FExpression* BottomLayerExpression = LayerExpressions[0];
+	const FExpression* BottomLayerExpression = LayerExpressions[0];
 	if (!BottomLayerExpression)
 	{
 		return Generator.Error(TEXT("No layers"));
 	}
 
-	TArray<FExpression*, TInlineAllocator<2>> BlendInputExpressions;
+	TArray<const FExpression*, TInlineAllocator<2>> BlendInputExpressions;
 	for (int32 BlendIndex = 0; BlendIndex < MaterialLayers.Blends.Num(); ++BlendIndex)
 	{
 		const int32 LayerIndex = BlendIndex + 1;
@@ -1618,7 +1618,7 @@ bool UMaterialExpressionMaterialAttributeLayers::GenerateHLSLExpression(FMateria
 
 		if (MaterialLayers.Layers[LayerIndex] && MaterialLayers.LayerStates[LayerIndex])
 		{
-			FExpression* LayerExpression = LayerExpressions[LayerIndex];
+			const FExpression* LayerExpression = LayerExpressions[LayerIndex];
 			if (!LayerExpression)
 			{
 				return Generator.Errorf(TEXT("Missing layer %d"), LayerIndex);
@@ -1659,14 +1659,14 @@ bool UMaterialExpressionMaterialAttributeLayers::GenerateHLSLExpression(FMateria
 	return true;
 }
 
-bool UMaterialExpressionBlendMaterialAttributes::GenerateHLSLExpression(FMaterialHLSLGenerator& Generator, UE::HLSLTree::FScope& Scope, int32 OutputIndex, UE::HLSLTree::FExpression*& OutExpression)
+bool UMaterialExpressionBlendMaterialAttributes::GenerateHLSLExpression(FMaterialHLSLGenerator& Generator, UE::HLSLTree::FScope& Scope, int32 OutputIndex, UE::HLSLTree::FExpression const*& OutExpression) const
 {
 	using namespace UE::HLSLTree;
 	using namespace UE::Shader;
 
-	FExpression* ExpressionA = A.AcquireHLSLExpressionOrConstant(Generator, Scope, Generator.GetMaterialAttributesDefaultValue());
-	FExpression* ExpressionB = B.AcquireHLSLExpressionOrConstant(Generator, Scope, Generator.GetMaterialAttributesDefaultValue());
-	FExpression* ExpressionAlpha = Alpha.AcquireHLSLExpression(Generator, Scope);
+	const FExpression* ExpressionA = A.AcquireHLSLExpressionOrConstant(Generator, Scope, Generator.GetMaterialAttributesDefaultValue());
+	const FExpression* ExpressionB = B.AcquireHLSLExpressionOrConstant(Generator, Scope, Generator.GetMaterialAttributesDefaultValue());
+	const FExpression* ExpressionAlpha = Alpha.AcquireHLSLExpression(Generator, Scope);
 	if (!ExpressionA || !ExpressionB || !ExpressionAlpha)
 	{
 		return false;
@@ -1674,7 +1674,7 @@ bool UMaterialExpressionBlendMaterialAttributes::GenerateHLSLExpression(FMateria
 
 	const UE::Shader::FStructType* MaterialAttributesType = Generator.GetMaterialAttributesType();
 
-	FExpression* ExpressionResult = Generator.GetTree().NewConstant(Generator.GetMaterialAttributesDefaultValue());
+	const FExpression* ExpressionResult = Generator.GetTree().NewConstant(Generator.GetMaterialAttributesDefaultValue());
 	const TArray<FGuid>& OrderedVisibleAttributes = FMaterialAttributeDefinitionMap::GetOrderedVisibleAttributeList();
 	for (const FGuid& AttributeID : OrderedVisibleAttributes)
 	{
@@ -1693,7 +1693,7 @@ bool UMaterialExpressionBlendMaterialAttributes::GenerateHLSLExpression(FMateria
 		default: AttributeBlend = PixelAttributeBlendType; break;
 		}
 
-		FExpression* ExpressionFieldResult = nullptr;
+		const FExpression* ExpressionFieldResult = nullptr;
 		if (AttributeBlend == EMaterialAttributeBlend::UseA)
 		{
 			ExpressionFieldResult = Generator.GetTree().NewExpression<FExpressionGetStructField>(MaterialAttributesType, Field, ExpressionA);
@@ -1704,8 +1704,8 @@ bool UMaterialExpressionBlendMaterialAttributes::GenerateHLSLExpression(FMateria
 		}
 		else
 		{
-			FExpression* ExpressionFieldA = Generator.GetTree().NewExpression<FExpressionGetStructField>(MaterialAttributesType, Field, ExpressionA);
-			FExpression* ExpressionFieldB = Generator.GetTree().NewExpression<FExpressionGetStructField>(MaterialAttributesType, Field, ExpressionB);
+			const FExpression* ExpressionFieldA = Generator.GetTree().NewExpression<FExpressionGetStructField>(MaterialAttributesType, Field, ExpressionA);
+			const FExpression* ExpressionFieldB = Generator.GetTree().NewExpression<FExpressionGetStructField>(MaterialAttributesType, Field, ExpressionB);
 			ExpressionFieldResult = Generator.GetTree().NewLerp(ExpressionFieldA, ExpressionFieldB, ExpressionAlpha);
 		}
 		ExpressionResult = Generator.GetTree().NewExpression<FExpressionSetStructField>(MaterialAttributesType, Field, ExpressionResult, ExpressionFieldResult);
@@ -1715,10 +1715,10 @@ bool UMaterialExpressionBlendMaterialAttributes::GenerateHLSLExpression(FMateria
 	return true;
 }
 
-static UE::HLSLTree::FExpression* TransformBase(UE::HLSLTree::FTree& Tree,
+static const UE::HLSLTree::FExpression* TransformBase(UE::HLSLTree::FTree& Tree,
 	EMaterialCommonBasis SourceCoordBasis,
 	EMaterialCommonBasis DestCoordBasis,
-	UE::HLSLTree::FExpression* Input,
+	const UE::HLSLTree::FExpression* Input,
 	bool bWComponent)
 {
 	using namespace UE::HLSLTree;
@@ -1736,7 +1736,7 @@ static UE::HLSLTree::FExpression* TransformBase(UE::HLSLTree::FTree& Tree,
 		return Input;
 	}
 
-	FExpression* Result = nullptr;
+	const FExpression* Result = nullptr;
 	EMaterialCommonBasis IntermediaryBasis = MCB_World;
 	const EOperation Op = bWComponent ? EOperation::VecMulMatrix4 : EOperation::VecMulMatrix3;
 	switch (SourceCoordBasis)
@@ -1901,14 +1901,14 @@ static UE::HLSLTree::FExpression* TransformBase(UE::HLSLTree::FTree& Tree,
 		check(IntermediaryBasis != DestCoordBasis);
 
 		// use intermediary basis
-		FExpression* IntermediaryExpression = TransformBase(Tree, SourceCoordBasis, IntermediaryBasis, Input, bWComponent);
+		const FExpression* IntermediaryExpression = TransformBase(Tree, SourceCoordBasis, IntermediaryBasis, Input, bWComponent);
 		return TransformBase(Tree, IntermediaryBasis, DestCoordBasis, IntermediaryExpression, bWComponent);
 	}
 
 	return Result;
 }
 
-bool UMaterialExpressionTransform::GenerateHLSLExpression(FMaterialHLSLGenerator& Generator, UE::HLSLTree::FScope& Scope, int32 OutputIndex, UE::HLSLTree::FExpression*& OutExpression)
+bool UMaterialExpressionTransform::GenerateHLSLExpression(FMaterialHLSLGenerator& Generator, UE::HLSLTree::FScope& Scope, int32 OutputIndex, UE::HLSLTree::FExpression const*& OutExpression) const
 {
 	using namespace UE::HLSLTree;
 
@@ -1922,7 +1922,7 @@ bool UMaterialExpressionTransform::GenerateHLSLExpression(FMaterialHLSLGenerator
 		MCB_Instance,					// TRANSFORM_Instance
 	};
 
-	FExpression* ExpressionInput = Input.AcquireHLSLExpression(Generator, Scope);
+	const FExpression* ExpressionInput = Input.AcquireHLSLExpression(Generator, Scope);
 	if (!ExpressionInput)
 	{
 		return false;
@@ -1933,7 +1933,7 @@ bool UMaterialExpressionTransform::GenerateHLSLExpression(FMaterialHLSLGenerator
 }
 
 
-bool UMaterialExpressionTransformPosition::GenerateHLSLExpression(FMaterialHLSLGenerator& Generator, UE::HLSLTree::FScope& Scope, int32 OutputIndex, UE::HLSLTree::FExpression*& OutExpression)
+bool UMaterialExpressionTransformPosition::GenerateHLSLExpression(FMaterialHLSLGenerator& Generator, UE::HLSLTree::FScope& Scope, int32 OutputIndex, UE::HLSLTree::FExpression const*& OutExpression) const
 {
 	using namespace UE::HLSLTree;
 
@@ -1947,7 +1947,7 @@ bool UMaterialExpressionTransformPosition::GenerateHLSLExpression(FMaterialHLSLG
 		MCB_Instance,					// TRANSFORMPOSSOURCE_Instance
 	};
 
-	FExpression* ExpressionInput = Input.AcquireHLSLExpression(Generator, Scope);
+	const FExpression* ExpressionInput = Input.AcquireHLSLExpression(Generator, Scope);
 	if (!ExpressionInput)
 	{
 		return false;
@@ -1957,80 +1957,80 @@ bool UMaterialExpressionTransformPosition::GenerateHLSLExpression(FMaterialHLSLG
 	return true;
 }
 
-bool UMaterialExpressionIf::GenerateHLSLExpression(FMaterialHLSLGenerator& Generator, UE::HLSLTree::FScope& Scope, int32 OutputIndex, UE::HLSLTree::FExpression*& OutExpression)
+bool UMaterialExpressionIf::GenerateHLSLExpression(FMaterialHLSLGenerator& Generator, UE::HLSLTree::FScope& Scope, int32 OutputIndex, UE::HLSLTree::FExpression const*& OutExpression) const
 {
 	using namespace UE::HLSLTree;
 	using namespace UE::Shader;
 
-	FExpression* ExpressionA = A.AcquireHLSLExpression(Generator, Scope);
-	FExpression* ExpressionB = A.AcquireHLSLExpressionOrConstant(Generator, Scope, ConstB);
-	FExpression* ExpressionAGreaterThanB = AGreaterThanB.AcquireHLSLExpression(Generator, Scope);
-	FExpression* ExpressionALessThanB = ALessThanB.AcquireHLSLExpression(Generator, Scope);
+	const FExpression* ExpressionA = A.AcquireHLSLExpression(Generator, Scope);
+	const FExpression* ExpressionB = A.AcquireHLSLExpressionOrConstant(Generator, Scope, ConstB);
+	const FExpression* ExpressionAGreaterThanB = AGreaterThanB.AcquireHLSLExpression(Generator, Scope);
+	const FExpression* ExpressionALessThanB = ALessThanB.AcquireHLSLExpression(Generator, Scope);
 	if (!ExpressionA || !ExpressionB || !ExpressionAGreaterThanB || !ExpressionALessThanB)
 	{
 		return false;
 	}
 
 	UE::HLSLTree::FTree& Tree = Generator.GetTree();
-	FExpression* ExpressionCondAGreaterEqualB = Tree.NewGreaterEqual(ExpressionA, ExpressionB);
+	const FExpression* ExpressionCondAGreaterEqualB = Tree.NewGreaterEqual(ExpressionA, ExpressionB);
 	OutExpression = Tree.NewExpression<FExpressionSelect>(ExpressionCondAGreaterEqualB, ExpressionAGreaterThanB, ExpressionALessThanB);
 
-	FExpression* ExpressionAEqualsB = AEqualsB.TryAcquireHLSLExpression(Generator, Scope);
+	const FExpression* ExpressionAEqualsB = AEqualsB.TryAcquireHLSLExpression(Generator, Scope);
 	if (ExpressionAEqualsB)
 	{
-		FExpression* ExpressionThreshold = Generator.NewConstant(EqualsThreshold);
-		FExpression* ExpressionCondANotEqualsB = Tree.NewGreater(Tree.NewAbs(Tree.NewSub(ExpressionA, ExpressionB)), ExpressionThreshold);
+		const FExpression* ExpressionThreshold = Generator.NewConstant(EqualsThreshold);
+		const FExpression* ExpressionCondANotEqualsB = Tree.NewGreater(Tree.NewAbs(Tree.NewSub(ExpressionA, ExpressionB)), ExpressionThreshold);
 		OutExpression = Tree.NewExpression<FExpressionSelect>(ExpressionCondANotEqualsB, OutExpression, ExpressionAEqualsB);
 	}
 
 	return true;
 }
 
-bool UMaterialExpressionFresnel::GenerateHLSLExpression(FMaterialHLSLGenerator& Generator, UE::HLSLTree::FScope& Scope, int32 OutputIndex, UE::HLSLTree::FExpression*& OutExpression)
+bool UMaterialExpressionFresnel::GenerateHLSLExpression(FMaterialHLSLGenerator& Generator, UE::HLSLTree::FScope& Scope, int32 OutputIndex, UE::HLSLTree::FExpression const*& OutExpression) const
 {
 	using namespace UE::HLSLTree;
 
 	// pow(1 - max(0,Normal dot Camera),Exponent) * (1 - BaseReflectFraction) + BaseReflectFraction
 
-	FExpression* ExpressionExponent = ExponentIn.AcquireHLSLExpressionOrConstant(Generator, Scope, Exponent);
-	FExpression* ExpressionReflectFraction = BaseReflectFractionIn.AcquireHLSLExpressionOrConstant(Generator, Scope, BaseReflectFraction);
+	const FExpression* ExpressionExponent = ExponentIn.AcquireHLSLExpressionOrConstant(Generator, Scope, Exponent);
+	const FExpression* ExpressionReflectFraction = BaseReflectFractionIn.AcquireHLSLExpressionOrConstant(Generator, Scope, BaseReflectFraction);
 	if (!ExpressionExponent || !ExpressionReflectFraction)
 	{
 		return false;
 	}
 
-	FExpression* ExpressionNormal = Normal.TryAcquireHLSLExpression(Generator, Scope);
+	const FExpression* ExpressionNormal = Normal.TryAcquireHLSLExpression(Generator, Scope);
 	if (!ExpressionNormal)
 	{
 		ExpressionNormal = Generator.GetTree().NewExpression<Material::FExpressionExternalInput>(Material::EExternalInput::WorldNormal);
 	}
 
-	FExpression* ExpressionNdotV = Generator.GetTree().NewDot(ExpressionNormal, Generator.GetTree().NewExpression<Material::FExpressionExternalInput>(Material::EExternalInput::CameraVector));
-	FExpression* ExpressionMax = Generator.GetTree().NewMax(Generator.NewConstant(0.0f), ExpressionNdotV);
-	FExpression* ExpressionMinus = Generator.GetTree().NewSub(Generator.NewConstant(1.0f), ExpressionMax);
+	const FExpression* ExpressionNdotV = Generator.GetTree().NewDot(ExpressionNormal, Generator.GetTree().NewExpression<Material::FExpressionExternalInput>(Material::EExternalInput::CameraVector));
+	const FExpression* ExpressionMax = Generator.GetTree().NewMax(Generator.NewConstant(0.0f), ExpressionNdotV);
+	const FExpression* ExpressionMinus = Generator.GetTree().NewSub(Generator.NewConstant(1.0f), ExpressionMax);
 
 	// Compiler->Power got changed to call PositiveClampedPow instead of ClampedPow
 	// Manually implement ClampedPow to maintain backwards compatibility in the case where the input normal is not normalized (length > 1)
-	FExpression* ExpressionAbs = Generator.GetTree().NewMax(Generator.GetTree().NewAbs(ExpressionMinus), Generator.NewConstant(KINDA_SMALL_NUMBER));
+	const FExpression* ExpressionAbs = Generator.GetTree().NewMax(Generator.GetTree().NewAbs(ExpressionMinus), Generator.NewConstant(KINDA_SMALL_NUMBER));
 
-	FExpression* ExpressionPow = Generator.GetTree().NewBinaryOp(EOperation::PowPositiveClamped, ExpressionAbs, ExpressionExponent);
+	const FExpression* ExpressionPow = Generator.GetTree().NewBinaryOp(EOperation::PowPositiveClamped, ExpressionAbs, ExpressionExponent);
 
 	OutExpression = Generator.GetTree().NewMul(ExpressionPow, Generator.GetTree().NewSub(Generator.NewConstant(1.0f), ExpressionReflectFraction));
 	OutExpression = Generator.GetTree().NewAdd(OutExpression, ExpressionReflectFraction);
 	return true;
 }
 
-bool UMaterialExpressionDesaturation::GenerateHLSLExpression(FMaterialHLSLGenerator& Generator, UE::HLSLTree::FScope& Scope, int32 OutputIndex, UE::HLSLTree::FExpression*& OutExpression)
+bool UMaterialExpressionDesaturation::GenerateHLSLExpression(FMaterialHLSLGenerator& Generator, UE::HLSLTree::FScope& Scope, int32 OutputIndex, UE::HLSLTree::FExpression const*& OutExpression) const
 {
 	using namespace UE::HLSLTree;
-	FExpression* ExpressionInput = Input.AcquireHLSLExpression(Generator, Scope);
+	const FExpression* ExpressionInput = Input.AcquireHLSLExpression(Generator, Scope);
 	if (!ExpressionInput)
 	{
 		return false;
 	}
 
-	FExpression* ExpressionGrey = Generator.GetTree().NewDot(ExpressionInput, Generator.NewConstant(FVector3f(LuminanceFactors.R, LuminanceFactors.G, LuminanceFactors.B)));
-	FExpression* ExpressionFraction = Fraction.TryAcquireHLSLExpression(Generator, Scope);
+	const FExpression* ExpressionGrey = Generator.GetTree().NewDot(ExpressionInput, Generator.NewConstant(FVector3f(LuminanceFactors.R, LuminanceFactors.G, LuminanceFactors.B)));
+	const FExpression* ExpressionFraction = Fraction.TryAcquireHLSLExpression(Generator, Scope);
 	if (ExpressionFraction)
 	{
 		OutExpression = Generator.GetTree().NewLerp(ExpressionInput, ExpressionGrey, ExpressionFraction);
@@ -2042,38 +2042,38 @@ bool UMaterialExpressionDesaturation::GenerateHLSLExpression(FMaterialHLSLGenera
 	return true;
 }
 
-bool UMaterialExpressionSphereMask::GenerateHLSLExpression(FMaterialHLSLGenerator& Generator, UE::HLSLTree::FScope& Scope, int32 OutputIndex, UE::HLSLTree::FExpression*& OutExpression)
+bool UMaterialExpressionSphereMask::GenerateHLSLExpression(FMaterialHLSLGenerator& Generator, UE::HLSLTree::FScope& Scope, int32 OutputIndex, UE::HLSLTree::FExpression const*& OutExpression) const
 {
 	using namespace UE::HLSLTree;
 
-	FExpression* ExpressionA = A.AcquireHLSLExpression(Generator, Scope);
-	FExpression* ExpressionB = B.AcquireHLSLExpression(Generator, Scope);
+	const FExpression* ExpressionA = A.AcquireHLSLExpression(Generator, Scope);
+	const FExpression* ExpressionB = B.AcquireHLSLExpression(Generator, Scope);
 	if (!ExpressionA || !ExpressionB)
 	{
 		return false;
 	}
 
-	FExpression* ExpressionDistance = Generator.GetTree().NewLength(Generator.GetTree().NewSub(ExpressionA, ExpressionB));
-	FExpression* ExpressionRadius = Radius.AcquireHLSLExpressionOrConstant(Generator, Scope, AttenuationRadius);
-	FExpression* ExpressionRcpRadius = Generator.GetTree().NewRcp(Generator.GetTree().NewMax(ExpressionRadius, Generator.NewConstant(0.00001f)));
-	FExpression* ExpressionNormalizedDistance = Generator.GetTree().NewMul(ExpressionDistance, ExpressionRcpRadius);
-	FExpression* ExpressionHardness = Hardness.AcquireHLSLExpressionOrConstant(Generator, Scope, HardnessPercent * 0.01f);
-	FExpression* ExpressionSoftness = Generator.GetTree().NewSub(Generator.NewConstant(1.0f), ExpressionHardness);
-	FExpression* ExpressionRcpSoftness = Generator.GetTree().NewRcp(Generator.GetTree().NewMax(ExpressionSoftness, Generator.NewConstant(0.00001f)));
-	FExpression* ExpressionOneMinusDistance = Generator.GetTree().NewSub(Generator.NewConstant(1.0f), ExpressionNormalizedDistance);
+	const FExpression* ExpressionDistance = Generator.GetTree().NewLength(Generator.GetTree().NewSub(ExpressionA, ExpressionB));
+	const FExpression* ExpressionRadius = Radius.AcquireHLSLExpressionOrConstant(Generator, Scope, AttenuationRadius);
+	const FExpression* ExpressionRcpRadius = Generator.GetTree().NewRcp(Generator.GetTree().NewMax(ExpressionRadius, Generator.NewConstant(0.00001f)));
+	const FExpression* ExpressionNormalizedDistance = Generator.GetTree().NewMul(ExpressionDistance, ExpressionRcpRadius);
+	const FExpression* ExpressionHardness = Hardness.AcquireHLSLExpressionOrConstant(Generator, Scope, HardnessPercent * 0.01f);
+	const FExpression* ExpressionSoftness = Generator.GetTree().NewSub(Generator.NewConstant(1.0f), ExpressionHardness);
+	const FExpression* ExpressionRcpSoftness = Generator.GetTree().NewRcp(Generator.GetTree().NewMax(ExpressionSoftness, Generator.NewConstant(0.00001f)));
+	const FExpression* ExpressionOneMinusDistance = Generator.GetTree().NewSub(Generator.NewConstant(1.0f), ExpressionNormalizedDistance);
 	
 	OutExpression = Generator.GetTree().NewSaturate(Generator.GetTree().NewMul(ExpressionOneMinusDistance, ExpressionRcpSoftness));
 	return true;
 }
 
-bool UMaterialExpressionRotateAboutAxis::GenerateHLSLExpression(FMaterialHLSLGenerator& Generator, UE::HLSLTree::FScope& Scope, int32 OutputIndex, UE::HLSLTree::FExpression*& OutExpression)
+bool UMaterialExpressionRotateAboutAxis::GenerateHLSLExpression(FMaterialHLSLGenerator& Generator, UE::HLSLTree::FScope& Scope, int32 OutputIndex, UE::HLSLTree::FExpression const*& OutExpression) const
 {
 	using namespace UE::HLSLTree;
 
-	FExpression* ExpressionAxis = NormalizedRotationAxis.AcquireHLSLExpression(Generator, Scope);
-	FExpression* ExpressionRotationAngle = RotationAngle.AcquireHLSLExpression(Generator, Scope);
-	FExpression* ExpressionPivotPoint = PivotPoint.AcquireHLSLExpression(Generator, Scope);
-	FExpression* ExpressionPosition = Position.AcquireHLSLExpression(Generator, Scope);
+	const FExpression* ExpressionAxis = NormalizedRotationAxis.AcquireHLSLExpression(Generator, Scope);
+	const FExpression* ExpressionRotationAngle = RotationAngle.AcquireHLSLExpression(Generator, Scope);
+	const FExpression* ExpressionPivotPoint = PivotPoint.AcquireHLSLExpression(Generator, Scope);
+	const FExpression* ExpressionPosition = Position.AcquireHLSLExpression(Generator, Scope);
 	if (!ExpressionAxis || !ExpressionRotationAngle || !ExpressionPivotPoint || !ExpressionPosition)
 	{
 		return false;
@@ -2082,26 +2082,26 @@ bool UMaterialExpressionRotateAboutAxis::GenerateHLSLExpression(FMaterialHLSLGen
 	FTree& Tree = Generator.GetTree();
 	ExpressionRotationAngle = Tree.NewMul(ExpressionRotationAngle, Generator.NewConstant(2.0f * (float)PI / Period));
 
-	FExpression* ClosestPointOnAxis = Tree.NewAdd(ExpressionPivotPoint, Tree.NewMul(ExpressionAxis, Tree.NewDot(ExpressionAxis, Tree.NewSub(ExpressionPosition, ExpressionPivotPoint))));
-	FExpression* UAxis = Tree.NewSub(ExpressionPosition, ClosestPointOnAxis);
-	FExpression* VAxis = Tree.NewCross(ExpressionAxis, UAxis);
-	FExpression* CosAngle = Tree.NewCos(ExpressionRotationAngle);
-	FExpression* SinAngle = Tree.NewSin(ExpressionRotationAngle);
-	FExpression* R = Tree.NewAdd(Tree.NewMul(UAxis, CosAngle), Tree.NewMul(VAxis, SinAngle));
-	FExpression* RotatedPosition = Tree.NewAdd(ClosestPointOnAxis, R);
+	const FExpression* ClosestPointOnAxis = Tree.NewAdd(ExpressionPivotPoint, Tree.NewMul(ExpressionAxis, Tree.NewDot(ExpressionAxis, Tree.NewSub(ExpressionPosition, ExpressionPivotPoint))));
+	const FExpression* UAxis = Tree.NewSub(ExpressionPosition, ClosestPointOnAxis);
+	const FExpression* VAxis = Tree.NewCross(ExpressionAxis, UAxis);
+	const FExpression* CosAngle = Tree.NewCos(ExpressionRotationAngle);
+	const FExpression* SinAngle = Tree.NewSin(ExpressionRotationAngle);
+	const FExpression* R = Tree.NewAdd(Tree.NewMul(UAxis, CosAngle), Tree.NewMul(VAxis, SinAngle));
+	const FExpression* RotatedPosition = Tree.NewAdd(ClosestPointOnAxis, R);
 
 	// TODO - see MaterialTemplate.ush LWC version of RotateAboutAxis to optimize
 	OutExpression = Tree.NewSub(RotatedPosition, ExpressionPosition);
 	return true;
 }
 
-bool UMaterialExpressionBumpOffset::GenerateHLSLExpression(FMaterialHLSLGenerator& Generator, UE::HLSLTree::FScope& Scope, int32 OutputIndex, UE::HLSLTree::FExpression*& OutExpression)
+bool UMaterialExpressionBumpOffset::GenerateHLSLExpression(FMaterialHLSLGenerator& Generator, UE::HLSLTree::FScope& Scope, int32 OutputIndex, UE::HLSLTree::FExpression const*& OutExpression) const
 {
 	using namespace UE::HLSLTree;
 
-	FExpression* ExpressionHeightRatio = HeightRatioInput.AcquireHLSLExpressionOrConstant(Generator, Scope, HeightRatio);
-	FExpression* ExpressionHeight = Height.AcquireHLSLExpression(Generator, Scope);
-	FExpression* ExpressionCoordinate = Coordinate.AcquireHLSLExpression(Generator, Scope);
+	const FExpression* ExpressionHeightRatio = HeightRatioInput.AcquireHLSLExpressionOrConstant(Generator, Scope, HeightRatio);
+	const FExpression* ExpressionHeight = Height.AcquireHLSLExpression(Generator, Scope);
+	const FExpression* ExpressionCoordinate = Coordinate.AcquireHLSLExpression(Generator, Scope);
 	if (!ExpressionCoordinate)
 	{
 		ExpressionCoordinate = Generator.NewTexCoord(ConstCoordinate);
@@ -2112,30 +2112,30 @@ bool UMaterialExpressionBumpOffset::GenerateHLSLExpression(FMaterialHLSLGenerato
 	}
 
 	FTree& Tree = Generator.GetTree();
-	FExpression* ExpressionCamera = TransformBase(Tree, MCB_World, MCB_Tangent, Generator.NewExternalInput(Material::EExternalInput::CameraVector), false);
+	const FExpression* ExpressionCamera = TransformBase(Tree, MCB_World, MCB_Tangent, Generator.NewExternalInput(Material::EExternalInput::CameraVector), false);
 	ExpressionCamera = Generator.NewSwizzle(MakeSwizzleMask(true, true, false, false), ExpressionCamera);
 
-	FExpression* ExpressionHeightOffset = Tree.NewAdd(Tree.NewMul(ExpressionHeightRatio, ExpressionHeight), Tree.NewMul(ExpressionHeightRatio, Tree.NewConstant(-ReferencePlane)));
+	const FExpression* ExpressionHeightOffset = Tree.NewAdd(Tree.NewMul(ExpressionHeightRatio, ExpressionHeight), Tree.NewMul(ExpressionHeightRatio, Tree.NewConstant(-ReferencePlane)));
 	OutExpression = Tree.NewAdd(Tree.NewMul(ExpressionHeightOffset, ExpressionCamera), ExpressionCoordinate);
 	return true;
 }
 
-bool UMaterialExpressionPanner::GenerateHLSLExpression(FMaterialHLSLGenerator& Generator, UE::HLSLTree::FScope& Scope, int32 OutputIndex, UE::HLSLTree::FExpression*& OutExpression)
+bool UMaterialExpressionPanner::GenerateHLSLExpression(FMaterialHLSLGenerator& Generator, UE::HLSLTree::FScope& Scope, int32 OutputIndex, UE::HLSLTree::FExpression const*& OutExpression) const
 {
 	using namespace UE::HLSLTree;
 
-	FExpression* ExpressionTime = Time.TryAcquireHLSLExpression(Generator, Scope);
+	const FExpression* ExpressionTime = Time.TryAcquireHLSLExpression(Generator, Scope);
 	if (!ExpressionTime)
 	{
 		ExpressionTime = Generator.GetTree().NewExpression<Material::FExpressionExternalInput>(Material::EExternalInput::GameTime);
 	}
-	FExpression* ExpressionSpeed = Speed.AcquireHLSLExpressionOrConstant(Generator, Scope, FVector2f(SpeedX, SpeedY));
-	FExpression* ExpressionOffset = Generator.GetTree().NewMul(ExpressionSpeed, ExpressionTime);
+	const FExpression* ExpressionSpeed = Speed.AcquireHLSLExpressionOrConstant(Generator, Scope, FVector2f(SpeedX, SpeedY));
+	const FExpression* ExpressionOffset = Generator.GetTree().NewMul(ExpressionSpeed, ExpressionTime);
 	if (bFractionalPart)
 	{
 		ExpressionOffset = Generator.GetTree().NewFrac(ExpressionOffset);
 	}
-	FExpression* ExpressionTexCoord = Coordinate.TryAcquireHLSLExpression(Generator, Scope);
+	const FExpression* ExpressionTexCoord = Coordinate.TryAcquireHLSLExpression(Generator, Scope);
 	if (!ExpressionTexCoord)
 	{
 		ExpressionTexCoord = Generator.NewTexCoord(ConstCoordinate);
@@ -2145,33 +2145,33 @@ bool UMaterialExpressionPanner::GenerateHLSLExpression(FMaterialHLSLGenerator& G
 	return true;
 }
 
-bool UMaterialExpressionRotator::GenerateHLSLExpression(FMaterialHLSLGenerator& Generator, UE::HLSLTree::FScope& Scope, int32 OutputIndex, UE::HLSLTree::FExpression*& OutExpression)
+bool UMaterialExpressionRotator::GenerateHLSLExpression(FMaterialHLSLGenerator& Generator, UE::HLSLTree::FScope& Scope, int32 OutputIndex, UE::HLSLTree::FExpression const*& OutExpression) const
 {
 	using namespace UE::HLSLTree;
 
-	FExpression* ExpressionTime = Time.TryAcquireHLSLExpression(Generator, Scope);
+	const FExpression* ExpressionTime = Time.TryAcquireHLSLExpression(Generator, Scope);
 	if (!ExpressionTime)
 	{
 		ExpressionTime = Generator.GetTree().NewExpression<Material::FExpressionExternalInput>(Material::EExternalInput::GameTime);
 	}
 	ExpressionTime = Generator.GetTree().NewMul(ExpressionTime, Generator.NewConstant(Speed));
 
-	FExpression* ExpressionCos = Generator.GetTree().NewUnaryOp(EOperation::Cos, ExpressionTime);
-	FExpression* ExpressionSin = Generator.GetTree().NewUnaryOp(EOperation::Sin, ExpressionTime);
-	FExpression* ExpressionRowX = Generator.GetTree().NewExpression<FExpressionAppend>(ExpressionCos, Generator.GetTree().NewMul(ExpressionSin, Generator.NewConstant(-1.0f)));
-	FExpression* ExpressionRowY = Generator.GetTree().NewExpression<FExpressionAppend>(ExpressionSin, ExpressionCos);
-	FExpression* ExpressionOrigin = Generator.GetTree().NewConstant(FVector2f(CenterX, CenterY));
+	const FExpression* ExpressionCos = Generator.GetTree().NewUnaryOp(EOperation::Cos, ExpressionTime);
+	const FExpression* ExpressionSin = Generator.GetTree().NewUnaryOp(EOperation::Sin, ExpressionTime);
+	const FExpression* ExpressionRowX = Generator.GetTree().NewExpression<FExpressionAppend>(ExpressionCos, Generator.GetTree().NewMul(ExpressionSin, Generator.NewConstant(-1.0f)));
+	const FExpression* ExpressionRowY = Generator.GetTree().NewExpression<FExpressionAppend>(ExpressionSin, ExpressionCos);
+	const FExpression* ExpressionOrigin = Generator.GetTree().NewConstant(FVector2f(CenterX, CenterY));
 
-	FExpression* ExpressionCoord = Coordinate.TryAcquireHLSLExpression(Generator, Scope);
+	const FExpression* ExpressionCoord = Coordinate.TryAcquireHLSLExpression(Generator, Scope);
 	if (!ExpressionCoord)
 	{
 		ExpressionCoord = Generator.NewTexCoord(ConstCoordinate);
 	}
-	FExpression* ExpressionCoordXY = Generator.GetTree().NewSub(Generator.NewSwizzle(FSwizzleParameters(0, 1), ExpressionCoord), ExpressionOrigin);
+	const FExpression* ExpressionCoordXY = Generator.GetTree().NewSub(Generator.NewSwizzle(FSwizzleParameters(0, 1), ExpressionCoord), ExpressionOrigin);
 
-	FExpression* ExpressionResultX = Generator.GetTree().NewDot(ExpressionRowX, ExpressionCoordXY);
-	FExpression* ExpressionResultY = Generator.GetTree().NewDot(ExpressionRowY, ExpressionCoordXY);
-	FExpression* ExpressionResult = Generator.GetTree().NewAdd(Generator.GetTree().NewExpression<FExpressionAppend>(ExpressionResultX, ExpressionResultY), ExpressionOrigin);
+	const FExpression* ExpressionResultX = Generator.GetTree().NewDot(ExpressionRowX, ExpressionCoordXY);
+	const FExpression* ExpressionResultY = Generator.GetTree().NewDot(ExpressionRowY, ExpressionCoordXY);
+	const FExpression* ExpressionResult = Generator.GetTree().NewAdd(Generator.GetTree().NewExpression<FExpressionAppend>(ExpressionResultX, ExpressionResultY), ExpressionOrigin);
 
 	// If given coordinate was a float3, we append the Z coordinate onto the result
 	OutExpression = Generator.GetTree().NewExpression<FExpressionAppend>(ExpressionResult, Generator.NewSwizzle(FSwizzleParameters(2), ExpressionCoord));
@@ -2192,7 +2192,7 @@ static UE::Shader::FType GetCustomOutputType(const FMaterialHLSLGenerator& Gener
 	}
 }
 
-bool UMaterialExpressionVolumetricAdvancedMaterialInput::GenerateHLSLExpression(FMaterialHLSLGenerator& Generator, UE::HLSLTree::FScope& Scope, int32 OutputIndex, UE::HLSLTree::FExpression*& OutExpression)
+bool UMaterialExpressionVolumetricAdvancedMaterialInput::GenerateHLSLExpression(FMaterialHLSLGenerator& Generator, UE::HLSLTree::FScope& Scope, int32 OutputIndex, UE::HLSLTree::FExpression const*& OutExpression) const
 {
 	if (OutputIndex != 0)
 	{
@@ -2202,7 +2202,7 @@ bool UMaterialExpressionVolumetricAdvancedMaterialInput::GenerateHLSLExpression(
 	return true;
 }
 
-bool UMaterialExpressionVolumetricAdvancedMaterialOutput::GenerateHLSLExpression(FMaterialHLSLGenerator& Generator, UE::HLSLTree::FScope& Scope, int32 OutputIndex, UE::HLSLTree::FExpression*& OutExpression)
+bool UMaterialExpressionVolumetricAdvancedMaterialOutput::GenerateHLSLExpression(FMaterialHLSLGenerator& Generator, UE::HLSLTree::FScope& Scope, int32 OutputIndex, UE::HLSLTree::FExpression const*& OutExpression) const
 {
 	switch (OutputIndex)
 	{
@@ -2225,7 +2225,7 @@ UE::Shader::EValueType UMaterialExpressionVolumetricAdvancedMaterialOutput::GetC
 	else return UE::Shader::EValueType::Void;
 }
 
-bool UMaterialExpressionHairAttributes::GenerateHLSLExpression(FMaterialHLSLGenerator& Generator, UE::HLSLTree::FScope& Scope, int32 OutputIndex, UE::HLSLTree::FExpression*& OutExpression)
+bool UMaterialExpressionHairAttributes::GenerateHLSLExpression(FMaterialHLSLGenerator& Generator, UE::HLSLTree::FScope& Scope, int32 OutputIndex, UE::HLSLTree::FExpression const*& OutExpression) const
 {
 	using namespace UE::HLSLTree;
 	using namespace UE::Shader;
@@ -2255,7 +2255,7 @@ bool UMaterialExpressionHairAttributes::GenerateHLSLExpression(FMaterialHLSLGene
 	return OutExpression != nullptr;
 }
 
-bool UMaterialExpressionCloudSampleAttribute::GenerateHLSLExpression(FMaterialHLSLGenerator& Generator, UE::HLSLTree::FScope& Scope, int32 OutputIndex, UE::HLSLTree::FExpression*& OutExpression)
+bool UMaterialExpressionCloudSampleAttribute::GenerateHLSLExpression(FMaterialHLSLGenerator& Generator, UE::HLSLTree::FScope& Scope, int32 OutputIndex, UE::HLSLTree::FExpression const*& OutExpression) const
 {
 	using namespace UE::HLSLTree;
 	using namespace UE::Shader;
@@ -2271,7 +2271,7 @@ bool UMaterialExpressionCloudSampleAttribute::GenerateHLSLExpression(FMaterialHL
 	return OutExpression != nullptr;
 }
 
-bool UMaterialExpressionPerInstanceFadeAmount::GenerateHLSLExpression(FMaterialHLSLGenerator& Generator, UE::HLSLTree::FScope& Scope, int32 OutputIndex, UE::HLSLTree::FExpression*& OutExpression)
+bool UMaterialExpressionPerInstanceFadeAmount::GenerateHLSLExpression(FMaterialHLSLGenerator& Generator, UE::HLSLTree::FScope& Scope, int32 OutputIndex, UE::HLSLTree::FExpression const*& OutExpression) const
 {
 	using namespace UE::HLSLTree;
 	using namespace UE::Shader;
@@ -2279,7 +2279,7 @@ bool UMaterialExpressionPerInstanceFadeAmount::GenerateHLSLExpression(FMaterialH
 	return true;
 }
 
-bool UMaterialExpressionObjectBounds::GenerateHLSLExpression(FMaterialHLSLGenerator& Generator, UE::HLSLTree::FScope& Scope, int32 OutputIndex, UE::HLSLTree::FExpression*& OutExpression)
+bool UMaterialExpressionObjectBounds::GenerateHLSLExpression(FMaterialHLSLGenerator& Generator, UE::HLSLTree::FScope& Scope, int32 OutputIndex, UE::HLSLTree::FExpression const*& OutExpression) const
 {
 	using namespace UE::HLSLTree;
 	using namespace UE::Shader;
@@ -2288,7 +2288,7 @@ bool UMaterialExpressionObjectBounds::GenerateHLSLExpression(FMaterialHLSLGenera
 	return true;
 }
 
-bool UMaterialExpressionObjectOrientation::GenerateHLSLExpression(FMaterialHLSLGenerator& Generator, UE::HLSLTree::FScope& Scope, int32 OutputIndex, UE::HLSLTree::FExpression*& OutExpression)
+bool UMaterialExpressionObjectOrientation::GenerateHLSLExpression(FMaterialHLSLGenerator& Generator, UE::HLSLTree::FScope& Scope, int32 OutputIndex, UE::HLSLTree::FExpression const*& OutExpression) const
 {
 	using namespace UE::HLSLTree;
 	using namespace UE::Shader;
@@ -2296,7 +2296,7 @@ bool UMaterialExpressionObjectOrientation::GenerateHLSLExpression(FMaterialHLSLG
 	return true;
 }
 
-bool UMaterialExpressionObjectPositionWS::GenerateHLSLExpression(FMaterialHLSLGenerator& Generator, UE::HLSLTree::FScope& Scope, int32 OutputIndex, UE::HLSLTree::FExpression*& OutExpression)
+bool UMaterialExpressionObjectPositionWS::GenerateHLSLExpression(FMaterialHLSLGenerator& Generator, UE::HLSLTree::FScope& Scope, int32 OutputIndex, UE::HLSLTree::FExpression const*& OutExpression) const
 {
 	using namespace UE::HLSLTree;
 	using namespace UE::Shader;
@@ -2304,7 +2304,7 @@ bool UMaterialExpressionObjectPositionWS::GenerateHLSLExpression(FMaterialHLSLGe
 	return true;
 }
 
-bool UMaterialExpressionObjectRadius::GenerateHLSLExpression(FMaterialHLSLGenerator& Generator, UE::HLSLTree::FScope& Scope, int32 OutputIndex, UE::HLSLTree::FExpression*& OutExpression)
+bool UMaterialExpressionObjectRadius::GenerateHLSLExpression(FMaterialHLSLGenerator& Generator, UE::HLSLTree::FScope& Scope, int32 OutputIndex, UE::HLSLTree::FExpression const*& OutExpression) const
 {
 	using namespace UE::HLSLTree;
 	using namespace UE::Shader;
@@ -2312,7 +2312,7 @@ bool UMaterialExpressionObjectRadius::GenerateHLSLExpression(FMaterialHLSLGenera
 	return true;
 }
 
-bool UMaterialExpressionCustom::GenerateHLSLExpression(FMaterialHLSLGenerator& Generator, UE::HLSLTree::FScope& Scope, int32 OutputIndex, UE::HLSLTree::FExpression*& OutExpression)
+bool UMaterialExpressionCustom::GenerateHLSLExpression(FMaterialHLSLGenerator& Generator, UE::HLSLTree::FScope& Scope, int32 OutputIndex, UE::HLSLTree::FExpression const*& OutExpression) const
 {
 	using namespace UE::HLSLTree;
 	using namespace UE::Shader;
@@ -2331,7 +2331,7 @@ bool UMaterialExpressionCustom::GenerateHLSLExpression(FMaterialHLSLGenerator& G
 		const FCustomInput& Input = Inputs[Index];
 		if (!Input.InputName.IsNone())
 		{
-			FExpression* Expression = Input.Input.AcquireHLSLExpression(Generator, Scope);
+			const FExpression* Expression = Input.Input.AcquireHLSLExpression(Generator, Scope);
 			if (!Expression)
 			{
 				return false;
@@ -2391,7 +2391,7 @@ bool UMaterialExpressionCustom::GenerateHLSLExpression(FMaterialHLSLGenerator& G
 		FunctionCode = UE::MemStack::AllocateStringView(Allocator, FormattedCode.ToView());
 	}
 
-	FExpression* ExpressionCustom = Generator.GetTree().NewExpression<FExpressionCustomHLSL>(
+	const FExpression* ExpressionCustom = Generator.GetTree().NewExpression<FExpressionCustomHLSL>(
 		UE::MemStack::AllocateStringView(Allocator, DeclarationCode.ToView()),
 		FunctionCode,
 		LocalInputs,
@@ -2401,27 +2401,27 @@ bool UMaterialExpressionCustom::GenerateHLSLExpression(FMaterialHLSLGenerator& G
 	return true;
 }
 
-bool UMaterialExpressionClearCoatNormalCustomOutput::GenerateHLSLExpression(FMaterialHLSLGenerator& Generator, UE::HLSLTree::FScope& Scope, int32 OutputIndex, UE::HLSLTree::FExpression*& OutExpression)
+bool UMaterialExpressionClearCoatNormalCustomOutput::GenerateHLSLExpression(FMaterialHLSLGenerator& Generator, UE::HLSLTree::FScope& Scope, int32 OutputIndex, UE::HLSLTree::FExpression const*& OutExpression) const
 {
 	OutExpression = Input.AcquireHLSLExpression(Generator, Scope);
 	return OutExpression != nullptr;
 }
 
-bool UMaterialExpressionExecBegin::GenerateHLSLStatements(FMaterialHLSLGenerator& Generator, UE::HLSLTree::FScope& Scope)
+bool UMaterialExpressionExecBegin::GenerateHLSLStatements(FMaterialHLSLGenerator& Generator, UE::HLSLTree::FScope& Scope) const
 {
 	Exec.GenerateHLSLStatements(Generator, Scope);
 	return true;
 }
 
-bool UMaterialExpressionExecEnd::GenerateHLSLStatements(FMaterialHLSLGenerator& Generator, UE::HLSLTree::FScope& Scope)
+bool UMaterialExpressionExecEnd::GenerateHLSLStatements(FMaterialHLSLGenerator& Generator, UE::HLSLTree::FScope& Scope) const
 {
 	Generator.GenerateResult(Scope);
 	return true;
 }
 
-bool UMaterialExpressionSetLocal::GenerateHLSLStatements(FMaterialHLSLGenerator& Generator, UE::HLSLTree::FScope& Scope)
+bool UMaterialExpressionSetLocal::GenerateHLSLStatements(FMaterialHLSLGenerator& Generator, UE::HLSLTree::FScope& Scope) const
 {
-	UE::HLSLTree::FExpression* ValueExpression = Value.AcquireHLSLExpression(Generator, Scope);
+	const UE::HLSLTree::FExpression* ValueExpression = Value.AcquireHLSLExpression(Generator, Scope);
 	if (!ValueExpression)
 	{
 		return false;
@@ -2432,9 +2432,9 @@ bool UMaterialExpressionSetLocal::GenerateHLSLStatements(FMaterialHLSLGenerator&
 	return true;
 }
 
-bool UMaterialExpressionIfThenElse::GenerateHLSLStatements(FMaterialHLSLGenerator& Generator, UE::HLSLTree::FScope& Scope)
+bool UMaterialExpressionIfThenElse::GenerateHLSLStatements(FMaterialHLSLGenerator& Generator, UE::HLSLTree::FScope& Scope) const
 {
-	UE::HLSLTree::FExpression* ConditionExpression = Condition.AcquireHLSLExpression(Generator, Scope);
+	const UE::HLSLTree::FExpression* ConditionExpression = Condition.AcquireHLSLExpression(Generator, Scope);
 	if (!ConditionExpression)
 	{
 		return false;
@@ -2449,7 +2449,7 @@ bool UMaterialExpressionIfThenElse::GenerateHLSLStatements(FMaterialHLSLGenerato
 	return true;
 }
 
-bool UMaterialExpressionWhileLoop::GenerateHLSLStatements(FMaterialHLSLGenerator& Generator, UE::HLSLTree::FScope& Scope)
+bool UMaterialExpressionWhileLoop::GenerateHLSLStatements(FMaterialHLSLGenerator& Generator, UE::HLSLTree::FScope& Scope) const
 {
 	using namespace UE::HLSLTree;
 	if (!Condition.IsConnected())
@@ -2494,7 +2494,7 @@ struct FExpressionDataForLoop
 };
 DECLARE_MATERIAL_HLSLGENERATOR_DATA(FExpressionDataForLoop);
 
-bool UMaterialExpressionForLoop::GenerateHLSLExpression(FMaterialHLSLGenerator& Generator, UE::HLSLTree::FScope& Scope, int32 OutputIndex, UE::HLSLTree::FExpression*& OutExpression)
+bool UMaterialExpressionForLoop::GenerateHLSLExpression(FMaterialHLSLGenerator& Generator, UE::HLSLTree::FScope& Scope, int32 OutputIndex, UE::HLSLTree::FExpression const*& OutExpression) const
 {
 	FExpressionDataForLoop* ExpressionData = Generator.FindExpressionData<FExpressionDataForLoop>(this);
 	if (!ExpressionData || !Scope.HasParentScope(*ExpressionData->LoopScope))
@@ -2506,7 +2506,7 @@ bool UMaterialExpressionForLoop::GenerateHLSLExpression(FMaterialHLSLGenerator&
 	return true;
 }
 
-bool UMaterialExpressionForLoop::GenerateHLSLStatements(FMaterialHLSLGenerator& Generator, UE::HLSLTree::FScope& Scope)
+bool UMaterialExpressionForLoop::GenerateHLSLStatements(FMaterialHLSLGenerator& Generator, UE::HLSLTree::FScope& Scope) const
 {
 	using namespace UE::HLSLTree;
 	if (!LoopBody.GetExpression())
@@ -2514,13 +2514,13 @@ bool UMaterialExpressionForLoop::GenerateHLSLStatements(FMaterialHLSLGenerator&
 		return Generator.Error(TEXT("Missing LoopBody connection"));
 	}
 
-	FExpression* StartExpression = StartIndex.AcquireHLSLExpression(Generator, Scope);
+	const FExpression* StartExpression = StartIndex.AcquireHLSLExpression(Generator, Scope);
 	if (!StartExpression)
 	{
 		return false;
 	}
 
-	FExpression* EndExpression = EndIndex.AcquireHLSLExpression(Generator, Scope);
+	const FExpression* EndExpression = EndIndex.AcquireHLSLExpression(Generator, Scope);
 	if (!EndExpression)
 	{
 		return false;
@@ -2530,7 +2530,7 @@ bool UMaterialExpressionForLoop::GenerateHLSLStatements(FMaterialHLSLGenerator&
 	FExpressionDataForLoop* ExpressionData = Generator.NewExpressionData<FExpressionDataForLoop>(this);
 	ExpressionData->LocalName = *FString::Printf(TEXT("ForLoopControl%d"), GlobalData->NumLoops++);
 
-	FExpression* StepExpression = IndexStep.AcquireHLSLExpressionOrConstant(Generator, Scope, int32(1));
+	const FExpression* StepExpression = IndexStep.AcquireHLSLExpressionOrConstant(Generator, Scope, int32(1));
 
 	Generator.GetTree().AssignLocal(Scope, ExpressionData->LocalName, StartExpression);
 
@@ -2547,12 +2547,12 @@ bool UMaterialExpressionForLoop::GenerateHLSLStatements(FMaterialHLSLGenerator&
 
 	LoopStatement->BreakStatement = Generator.GetTree().NewStatement<FStatementBreak>(*IfStatement->ElseScope);
 
-	FExpression* LocalExpression = Generator.GetTree().AcquireLocal(*LoopStatement->LoopScope, ExpressionData->LocalName);
+	const FExpression* LocalExpression = Generator.GetTree().AcquireLocal(*LoopStatement->LoopScope, ExpressionData->LocalName);
 
 	IfStatement->ConditionExpression = Generator.GetTree().NewLess(LocalExpression, EndExpression);
 	LoopBody.GenerateHLSLStatements(Generator, *IfStatement->ThenScope);
 
-	FExpression* NewLocalExpression = Generator.GetTree().NewAdd(Generator.GetTree().AcquireLocal(*IfStatement->ThenScope, ExpressionData->LocalName), StepExpression);
+	const FExpression* NewLocalExpression = Generator.GetTree().NewAdd(Generator.GetTree().AcquireLocal(*IfStatement->ThenScope, ExpressionData->LocalName), StepExpression);
 	Generator.GetTree().AssignLocal(*IfStatement->ThenScope, ExpressionData->LocalName, NewLocalExpression);
 
 	Completed.GenerateHLSLStatements(Generator, *LoopStatement->NextScope);
