const-correctness changes in new HLSL translator:

diff --git a/Engine/Source/Runtime/Engine/Private/Materials/MaterialHLSLGenerator.cpp b/Engine/Source/Runtime/Engine/Private/Materials/MaterialHLSLGenerator.cpp
--- a/Engine/Source/Runtime/Engine/Private/Materials/MaterialHLSLGenerator.cpp
+++ b/Engine/Source/Runtime/Engine/Private/Materials/MaterialHLSLGenerator.cpp
@@ -110,14 +110,14 @@ bool FMaterialHLSLGenerator::Generate()
 	return GetTree().Finalize();
 }
 
-static UE::HLSLTree::FExpression* CompileMaterialInput(FMaterialHLSLGenerator& Generator,
+static const UE::HLSLTree::FExpression* CompileMaterialInput(FMaterialHLSLGenerator& Generator,
 	UE::HLSLTree::FScope& Scope,
 	EMaterialProperty InputProperty,
 	UMaterial* Material)
 {
 	using namespace UE::HLSLTree;
 
-	FExpression* Expression = nullptr;
+	const FExpression* Expression = nullptr;
 	if (Material->IsPropertyActive(InputProperty))
 	{
 		FMaterialInputDescription InputDescription;
@@ -171,7 +171,7 @@ bool FMaterialHLSLGenerator::GenerateResult(UE::HLSLTree::FScope& Scope)
 		check(!CachedTree.ResultStatement);
 		check(!CachedTree.ResultExpression);
 
-		FExpression* AttributesExpression = nullptr;
+		const FExpression* AttributesExpression = nullptr;
 		if (TargetMaterial)
 		{
 			const FStructField* PrevWPOField = CachedTree.GetMaterialAttributesType()->FindFieldByName(TEXT("PrevWorldPositionOffset"));
@@ -191,9 +191,9 @@ bool FMaterialHLSLGenerator::GenerateResult(UE::HLSLTree::FScope& Scope)
 						FRequestedType PrevRequestedType;
 						PrevRequestedType.SetFieldRequested(WPOField);
 
-						FExpression* PrevAttributesExpression = GetTree().GetPreviousFrame(AttributesExpression, PrevRequestedType);
+						const FExpression* PrevAttributesExpression = GetTree().GetPreviousFrame(AttributesExpression, PrevRequestedType);
 						ensure(PrevAttributesExpression);
-						FExpression* PrevWPOExpression = GetTree().NewExpression<FExpressionGetStructField>(CachedTree.GetMaterialAttributesType(), WPOField, PrevAttributesExpression);
+						const FExpression* PrevWPOExpression = GetTree().NewExpression<FExpressionGetStructField>(CachedTree.GetMaterialAttributesType(), WPOField, PrevAttributesExpression);
 						AttributesExpression = GetTree().NewExpression<FExpressionSetStructField>(CachedTree.GetMaterialAttributesType(), PrevWPOField, AttributesExpression, PrevWPOExpression);
 					}
 				}
@@ -209,7 +209,7 @@ bool FMaterialHLSLGenerator::GenerateResult(UE::HLSLTree::FScope& Scope)
 					const UE::Shader::FStructField* AttributeField = CachedTree.GetMaterialAttributesType()->FindFieldByName(*FMaterialAttributeDefinitionMap::GetAttributeName(Property));
 					if (AttributeField)
 					{
-						FExpression* InputExpression = CompileMaterialInput(*this, Scope, Property, TargetMaterial);
+						const FExpression* InputExpression = CompileMaterialInput(*this, Scope, Property, TargetMaterial);
 						if (InputExpression)
 						{
 							AttributesExpression = GetTree().NewExpression<FExpressionSetStructField>(
@@ -219,7 +219,7 @@ bool FMaterialHLSLGenerator::GenerateResult(UE::HLSLTree::FScope& Scope)
 								InputExpression);
 							if (Property == MP_WorldPositionOffset)
 							{
-								FExpression* PrevWPOExpression = GetTree().GetPreviousFrame(InputExpression, ERequestedType::Vector3);
+								const FExpression* PrevWPOExpression = GetTree().GetPreviousFrame(InputExpression, ERequestedType::Vector3);
 								ensure(PrevWPOExpression);
 								AttributesExpression = GetTree().NewExpression<FExpressionSetStructField>(CachedTree.GetMaterialAttributesType(), PrevWPOField, AttributesExpression, PrevWPOExpression);
 							}
@@ -239,7 +239,7 @@ bool FMaterialHLSLGenerator::GenerateResult(UE::HLSLTree::FScope& Scope)
 					const FStructField* CustomOutputField = CachedTree.GetMaterialAttributesType()->FindFieldByName(FieldName.ToString());
 					check(CustomOutputField);
 
-					FExpression* CustomOutputExpression = AcquireExpression(Scope, CustomOutput, OutputIndex);
+					const FExpression* CustomOutputExpression = AcquireExpression(Scope, CustomOutput, OutputIndex);
 					AttributesExpression = GetTree().NewExpression<FExpressionSetStructField>(CachedTree.GetMaterialAttributesType(), CustomOutputField, AttributesExpression, CustomOutputExpression);
 				}
 			}
@@ -258,7 +258,7 @@ bool FMaterialHLSLGenerator::GenerateResult(UE::HLSLTree::FScope& Scope)
 			const FStructField* EmissiveColorField = CachedTree.GetMaterialAttributesType()->FindFieldByName(*EmissiveColorName);
 
 			// Get back into gamma corrected space, as DrawTile does not do this adjustment.
-			FExpression* ExpressionEmissive = GetTree().NewPowClamped(PreviewExpressionResult, NewConstant(1.f / 2.2f));
+			const FExpression* ExpressionEmissive = GetTree().NewPowClamped(PreviewExpressionResult, NewConstant(1.f / 2.2f));
 
 			AttributesExpression = GetTree().NewExpression<FExpressionConstant>(CachedTree.GetMaterialAttributesDefaultValue());
 			AttributesExpression = GetTree().NewExpression<FExpressionSetStructField>(CachedTree.GetMaterialAttributesType(), EmissiveColorField, AttributesExpression, ExpressionEmissive);
@@ -303,24 +303,24 @@ UE::HLSLTree::FScope* FMaterialHLSLGenerator::NewJoinedScope(UE::HLSLTree::FScop
 	return NewScope;
 }
 
-UE::HLSLTree::FExpression* FMaterialHLSLGenerator::NewConstant(const UE::Shader::FValue& Value)
+const UE::HLSLTree::FExpression* FMaterialHLSLGenerator::NewConstant(const UE::Shader::FValue& Value)
 {
 	return GetTree().NewConstant(Value);
 }
 
-UE::HLSLTree::FExpression* FMaterialHLSLGenerator::NewTexCoord(int32 Index)
+const UE::HLSLTree::FExpression* FMaterialHLSLGenerator::NewTexCoord(int32 Index)
 {
 	using namespace UE::HLSLTree;
 	return NewExternalInput(Material::MakeInputTexCoord(Index));
 }
 
-UE::HLSLTree::FExpression* FMaterialHLSLGenerator::NewExternalInput(UE::HLSLTree::Material::EExternalInput Input)
+const UE::HLSLTree::FExpression* FMaterialHLSLGenerator::NewExternalInput(UE::HLSLTree::Material::EExternalInput Input)
 {
 	using namespace UE::HLSLTree;
 	return GetTree().NewExpression<Material::FExpressionExternalInput>(Input);
 }
 
-UE::HLSLTree::FExpression* FMaterialHLSLGenerator::NewSwizzle(const UE::HLSLTree::FSwizzleParameters& Params, UE::HLSLTree::FExpression* Input)
+const UE::HLSLTree::FExpression* FMaterialHLSLGenerator::NewSwizzle(const UE::HLSLTree::FSwizzleParameters& Params, const UE::HLSLTree::FExpression* Input)
 {
 	return GetTree().NewExpression<UE::HLSLTree::FExpressionSwizzle>(Params, Input);
 }
@@ -366,12 +366,12 @@ bool FMaterialHLSLGenerator::InternalError(FStringView ErrorMessage)
 	return false;
 }
 
-UE::HLSLTree::FExpression* FMaterialHLSLGenerator::AcquireExpression(UE::HLSLTree::FScope& Scope, UMaterialExpression* MaterialExpression, int32 OutputIndex)
+const UE::HLSLTree::FExpression* FMaterialHLSLGenerator::AcquireExpression(UE::HLSLTree::FScope& Scope, UMaterialExpression* MaterialExpression, int32 OutputIndex)
 {
 	using namespace UE::HLSLTree;
 	FOwnerScope TreeOwnerScope(GetTree(), MaterialExpression);
 
-	FExpression* Expression = nullptr;
+	const FExpression* Expression = nullptr;
 	if (MaterialExpression->GenerateHLSLExpression(*this, Scope, OutputIndex, Expression))
 	{
 		if (MaterialExpression == PreviewExpression &&
@@ -388,11 +388,11 @@ UE::HLSLTree::FExpression* FMaterialHLSLGenerator::AcquireExpression(UE::HLSLTre
 	return GetTree().NewExpression<UE::HLSLTree::FExpressionError>(ErrorMessage);
 }
 
-UE::HLSLTree::FExpression* FMaterialHLSLGenerator::AcquireFunctionInputExpression(UE::HLSLTree::FScope& Scope, const UMaterialExpressionFunctionInput* MaterialExpression)
+const UE::HLSLTree::FExpression* FMaterialHLSLGenerator::AcquireFunctionInputExpression(UE::HLSLTree::FScope& Scope, const UMaterialExpressionFunctionInput* MaterialExpression)
 {
 	using namespace UE::HLSLTree;
 	const FFunctionCallEntry* FunctionEntry = FunctionCallStack.Last();
-	FExpression* InputExpression = nullptr;
+	const FExpression* InputExpression = nullptr;
 	if (FunctionEntry->MaterialFunction)
 	{
 		bool bFoundInput = false;
@@ -492,7 +492,7 @@ bool FMaterialHLSLGenerator::GenerateStatements(UE::HLSLTree::FScope& Scope, UMa
 	return bResult;
 }
 
-UE::HLSLTree::FExpression* FMaterialHLSLGenerator::GenerateMaterialParameter(FName InParameterName, const FMaterialParameterMetadata& InParameterMeta, const UE::Shader::FValue& InDefaultValue)
+const UE::HLSLTree::FExpression* FMaterialHLSLGenerator::GenerateMaterialParameter(FName InParameterName, const FMaterialParameterMetadata& InParameterMeta, const UE::Shader::FValue& InDefaultValue)
 {
 	using namespace UE::Shader;
 
@@ -523,11 +523,11 @@ UE::HLSLTree::FExpression* FMaterialHLSLGenerator::GenerateMaterialParameter(FNa
 	return GetTree().NewExpression<UE::HLSLTree::Material::FExpressionParameter>(GetParameterInfo(InParameterName), ParameterMeta, DefaultValue);
 }
 
-UE::HLSLTree::FExpression* FMaterialHLSLGenerator::GenerateFunctionCall(UE::HLSLTree::FScope& Scope,
+const UE::HLSLTree::FExpression* FMaterialHLSLGenerator::GenerateFunctionCall(UE::HLSLTree::FScope& Scope,
 	UMaterialFunctionInterface* MaterialFunction,
 	EMaterialParameterAssociation InParameterAssociation,
 	int32 InParameterIndex,
-	TArrayView<UE::HLSLTree::FExpression*> ConnectedInputs,
+	TArrayView<const UE::HLSLTree::FExpression*> ConnectedInputs,
 	int32 OutputIndex)
 {
 	using namespace UE::HLSLTree;
@@ -580,7 +580,7 @@ UE::HLSLTree::FExpression* FMaterialHLSLGenerator::GenerateFunctionCall(UE::HLSL
 
 			// ConnectedInputs are the inputs from the UMaterialFunctionCall object
 			// We want to connect the UMaterialExpressionFunctionInput from the UMaterialFunction to whatever UMaterialExpression is passed to the UMaterialFunctionCall
-			FExpression* ConnectedInput = ConnectedInputs[InputIndex];
+			const FExpression* ConnectedInput = ConnectedInputs[InputIndex];
 
 			LocalFunctionInputs.Add(FunctionInput.ExpressionInput);
 			Hasher.Update((uint8*)&ConnectedInput, sizeof(ConnectedInput));
@@ -620,7 +620,7 @@ UE::HLSLTree::FExpression* FMaterialHLSLGenerator::GenerateFunctionCall(UE::HLSL
 		}
 	}
 
-	FExpression* Result = nullptr;
+	const FExpression* Result = nullptr;
 	FunctionCallStack.Add(FunctionCall);
 	if (bInlineFunction)
 	{
diff --git a/Engine/Source/Runtime/Engine/Private/Materials/MaterialHLSLTree.cpp b/Engine/Source/Runtime/Engine/Private/Materials/MaterialHLSLTree.cpp
--- a/Engine/Source/Runtime/Engine/Private/Materials/MaterialHLSLTree.cpp
+++ b/Engine/Source/Runtime/Engine/Private/Materials/MaterialHLSLTree.cpp
@@ -166,8 +166,8 @@ void FExpressionExternalInput::ComputeAnalyticDerivatives(FTree& Tree, FExpressi
 		{
 			// Ddx = float2(RcpViewSize.x, 0.0f)
 			// Ddy = float2(0.0f, RcpViewSize.y)
-			FExpression* RcpViewSize = Tree.NewExpression<FExpressionExternalInput>(EExternalInput::RcpViewSize);
-			FExpression* Constant0 = Tree.NewConstant(0.0f);
+			const FExpression* RcpViewSize = Tree.NewExpression<FExpressionExternalInput>(EExternalInput::RcpViewSize);
+			const FExpression* Constant0 = Tree.NewConstant(0.0f);
 			OutResult.ExpressionDdx = Tree.NewExpression<FExpressionAppend>(Tree.NewExpression<FExpressionSwizzle>(MakeSwizzleMask(true, false, false, false), RcpViewSize), Constant0);
 			OutResult.ExpressionDdy = Tree.NewExpression<FExpressionAppend>(Constant0, Tree.NewExpression<FExpressionSwizzle>(MakeSwizzleMask(false, true, false, false), RcpViewSize));
 			break;
@@ -182,7 +182,7 @@ void FExpressionExternalInput::ComputeAnalyticDerivatives(FTree& Tree, FExpressi
 	}
 }
 
-FExpression* FExpressionExternalInput::ComputePreviousFrame(FTree& Tree, const FRequestedType& RequestedType) const
+const FExpression* FExpressionExternalInput::ComputePreviousFrame(FTree& Tree, const FRequestedType& RequestedType) const
 {
 	if (InputType == EExternalInput::ActorWorldPosition)
 	{
diff --git a/Engine/Source/Runtime/Engine/Private/Materials/MaterialHLSLTree.h b/Engine/Source/Runtime/Engine/Private/Materials/MaterialHLSLTree.h
--- a/Engine/Source/Runtime/Engine/Private/Materials/MaterialHLSLTree.h
+++ b/Engine/Source/Runtime/Engine/Private/Materials/MaterialHLSLTree.h
@@ -186,7 +186,7 @@ public:
 	EExternalInput InputType;
 
 	virtual void ComputeAnalyticDerivatives(FTree& Tree, FExpressionDerivatives& OutResult) const override;
-	virtual FExpression* ComputePreviousFrame(FTree& Tree, const FRequestedType& RequestedType) const override;
+	virtual const FExpression* ComputePreviousFrame(FTree& Tree, const FRequestedType& RequestedType) const override;
 	virtual bool PrepareValue(FEmitContext& Context, FEmitScope& Scope, const FRequestedType& RequestedType, FPrepareValueResult& OutResult) const override;
 	virtual void EmitValueShader(FEmitContext& Context, FEmitScope& Scope, const FRequestedType& RequestedType, FEmitValueShaderResult& OutResult) const override;
 };
@@ -268,7 +268,7 @@ public:
 class FExpressionFunctionCall : public FExpressionForward
 {
 public:
-	FExpressionFunctionCall(FExpression* InExpression, UMaterialFunctionInterface* InMaterialFunction)
+	FExpressionFunctionCall(const FExpression* InExpression, UMaterialFunctionInterface* InMaterialFunction)
 		: FExpressionForward(InExpression)
 		, MaterialFunction(InMaterialFunction)
 	{}
@@ -281,7 +281,7 @@ public:
 class FExpressionMaterialLayers : public FExpressionForward
 {
 public:
-	FExpressionMaterialLayers(FExpression* InExpression, const FMaterialLayersFunctions* InMaterialLayers)
+	FExpressionMaterialLayers(const FExpression* InExpression, const FMaterialLayersFunctions* InMaterialLayers)
 		: FExpressionForward(InExpression)
 		, MaterialLayers(InMaterialLayers)
 	{}
@@ -294,13 +294,13 @@ public:
 class FExpressionSceneTexture : public FExpression
 {
 public:
-	FExpressionSceneTexture(FExpression* InTexCoordExpression, uint32 InSceneTextureId, bool bInFiltered)
+	FExpressionSceneTexture(const FExpression* InTexCoordExpression, uint32 InSceneTextureId, bool bInFiltered)
 		: TexCoordExpression(InTexCoordExpression)
 		, SceneTextureId(InSceneTextureId)
 		, bFiltered(bInFiltered)
 	{}
 
-	FExpression* TexCoordExpression;
+	const FExpression* TexCoordExpression;
 	uint32 SceneTextureId;
 	bool bFiltered;
 
@@ -327,14 +327,14 @@ struct FNoiseParameters
 class FExpressionNoise : public FExpression
 {
 public:
-	FExpressionNoise(const FNoiseParameters& InParams, FExpression* InPositionExpression, FExpression* InFilterWidthExpression)
+	FExpressionNoise(const FNoiseParameters& InParams, const FExpression* InPositionExpression, const FExpression* InFilterWidthExpression)
 		: PositionExpression(InPositionExpression)
 		, FilterWidthExpression(InFilterWidthExpression)
 		, Parameters(InParams)
 	{}
 
-	FExpression* PositionExpression;
-	FExpression* FilterWidthExpression;
+	const FExpression* PositionExpression;
+	const FExpression* FilterWidthExpression;
 	FNoiseParameters Parameters;
 
 	virtual bool PrepareValue(FEmitContext& Context, FEmitScope& Scope, const FRequestedType& RequestedType, FPrepareValueResult& OutResult) const override;
diff --git a/Engine/Source/Runtime/Engine/Private/Materials/MaterialShared.cpp b/Engine/Source/Runtime/Engine/Private/Materials/MaterialShared.cpp
--- a/Engine/Source/Runtime/Engine/Private/Materials/MaterialShared.cpp
+++ b/Engine/Source/Runtime/Engine/Private/Materials/MaterialShared.cpp
@@ -278,9 +278,9 @@ int32 FExpressionInput::Compile(class FMaterialCompiler* Compiler)
 		return INDEX_NONE;
 }
 
-UE::HLSLTree::FExpression* FExpressionInput::TryAcquireHLSLExpression(FMaterialHLSLGenerator& Generator, UE::HLSLTree::FScope& Scope) const
+const UE::HLSLTree::FExpression* FExpressionInput::TryAcquireHLSLExpression(FMaterialHLSLGenerator& Generator, UE::HLSLTree::FScope& Scope) const
 {
-	UE::HLSLTree::FExpression* Result = nullptr;
+	const UE::HLSLTree::FExpression* Result = nullptr;
 	if (Expression)
 	{
 		Expression->ValidateState();
@@ -295,7 +295,7 @@ UE::HLSLTree::FExpression* FExpressionInput::TryAcquireHLSLExpression(FMaterialH
 	return Result;
 }
 
-UE::HLSLTree::FExpression* FExpressionInput::AcquireHLSLExpression(FMaterialHLSLGenerator& Generator, UE::HLSLTree::FScope& Scope) const
+const UE::HLSLTree::FExpression* FExpressionInput::AcquireHLSLExpression(FMaterialHLSLGenerator& Generator, UE::HLSLTree::FScope& Scope) const
 {
 	const FExpressionInput TracedInput = GetTracedInput();
 	if (!TracedInput.Expression)
@@ -306,7 +306,7 @@ UE::HLSLTree::FExpression* FExpressionInput::AcquireHLSLExpression(FMaterialHLSL
 	return TryAcquireHLSLExpression(Generator, Scope);
 }
 
-UE::HLSLTree::FExpression* FExpressionInput::AcquireHLSLExpressionOrConstant(FMaterialHLSLGenerator& Generator, UE::HLSLTree::FScope& Scope, const UE::Shader::FValue& ConstantValue) const
+const UE::HLSLTree::FExpression* FExpressionInput::AcquireHLSLExpressionOrConstant(FMaterialHLSLGenerator& Generator, UE::HLSLTree::FScope& Scope, const UE::Shader::FValue& ConstantValue) const
 {
 	const FExpressionInput TracedInput = GetTracedInput();
 	if (!TracedInput.Expression)
diff --git a/Engine/Source/Runtime/Engine/Public/HLSLTree/HLSLTree.h b/Engine/Source/Runtime/Engine/Public/HLSLTree/HLSLTree.h
--- a/Engine/Source/Runtime/Engine/Public/HLSLTree/HLSLTree.h
+++ b/Engine/Source/Runtime/Engine/Public/HLSLTree/HLSLTree.h
@@ -134,10 +134,10 @@ private:
 struct FEmitPreshaderScope
 {
 	FEmitPreshaderScope() = default;
-	FEmitPreshaderScope(FEmitScope* InScope, FExpression* InValue) : Scope(InScope), Value(InValue) {}
+	FEmitPreshaderScope(FEmitScope* InScope, const FExpression* InValue) : Scope(InScope), Value(InValue) {}
 
 	FEmitScope* Scope = nullptr;
-	FExpression* Value = nullptr;
+	const FExpression* Value = nullptr;
 };
 
 /**
@@ -364,10 +364,10 @@ enum class EDerivativeCoordinate : uint8
 
 struct FExpressionDerivatives
 {
-	FExpression* ExpressionDdx = nullptr;
-	FExpression* ExpressionDdy = nullptr;
+	const FExpression* ExpressionDdx = nullptr;
+	const FExpression* ExpressionDdy = nullptr;
 
-	FExpression* Get(EDerivativeCoordinate Coord) const { return (Coord == EDerivativeCoordinate::Ddx) ? ExpressionDdx : ExpressionDdy; }
+	const FExpression* Get(EDerivativeCoordinate Coord) const { return (Coord == EDerivativeCoordinate::Ddx) ? ExpressionDdx : ExpressionDdy; }
 
 	bool IsValid() const { return (bool)ExpressionDdx && (bool)ExpressionDdy; }
 };
@@ -381,28 +381,35 @@ struct FExpressionDerivatives
 class FExpression : public FNode
 {
 public:
-	FEmitShaderExpression* GetValueShader(FEmitContext& Context, FEmitScope& Scope, const FRequestedType& RequestedType, const FPreparedType& PreparedType, const Shader::FType& ResultType);
-	Shader::FType GetValuePreshader(FEmitContext& Context, FEmitScope& Scope, const FRequestedType& RequestedType, const FPreparedType& PreparedType, const Shader::FType& ResultType, Shader::FPreshaderData& OutPreshader);
-	Shader::FValue GetValueConstant(FEmitContext& Context, FEmitScope& Scope, const FRequestedType& RequestedType, const FPreparedType& PreparedType, const Shader::FType& ResultType);
-
-	FEmitShaderExpression* GetValueShader(FEmitContext& Context, FEmitScope& Scope, const FRequestedType& RequestedType, const Shader::FType& ResultType);
-	FEmitShaderExpression* GetValueShader(FEmitContext& Context, FEmitScope& Scope, const FRequestedType& RequestedType);
-	FEmitShaderExpression* GetValueShader(FEmitContext& Context, FEmitScope& Scope, const Shader::FType& ResultType);
-	FEmitShaderExpression* GetValueShader(FEmitContext& Context, FEmitScope& Scope, Shader::EValueType ResultType);
-	FEmitShaderExpression* GetValueShader(FEmitContext& Context, FEmitScope& Scope);
-
-	Shader::FType GetValuePreshader(FEmitContext& Context, FEmitScope& Scope, const FRequestedType& RequestedType, const Shader::FType& ResultType, Shader::FPreshaderData& OutPreshader);
-	Shader::FType GetValuePreshader(FEmitContext& Context, FEmitScope& Scope, const FRequestedType& RequestedType, Shader::FPreshaderData& OutPreshader);
-
-	Shader::FValue GetValueConstant(FEmitContext& Context, FEmitScope& Scope, const FRequestedType& RequestedType, const FPreparedType& PreparedType);
-	Shader::FValue GetValueConstant(FEmitContext& Context, FEmitScope& Scope, const FRequestedType& RequestedType, const Shader::FType& ResultType);
-	Shader::FValue GetValueConstant(FEmitContext& Context, FEmitScope& Scope, const FRequestedType& RequestedType);
-	Shader::FValue GetValueConstant(FEmitContext& Context, FEmitScope& Scope, const FPreparedType& PreparedType, const Shader::FType& ResultType);
-	Shader::FValue GetValueConstant(FEmitContext& Context, FEmitScope& Scope, const FPreparedType& PreparedType, Shader::EValueType ResultType);
+	/**
+	 * Get Shader/Preshader/Constant value for the expression
+	 * @param RequestedType the componets of the result that are needed
+	 * @param PreparedType the prepared type of this expression, should be return of Context.PrepareExpression/GetPreparedType
+	 * @param ResultType the requested output type, value will be cast to this type if needed
+	 * Various overloads below will derive the missing values from those that are provided, and/or object PreparedType from the given Context directly
+	 */
+	FEmitShaderExpression* GetValueShader(FEmitContext& Context, FEmitScope& Scope, const FRequestedType& RequestedType, const FPreparedType& PreparedType, const Shader::FType& ResultType) const;
+	Shader::FType GetValuePreshader(FEmitContext& Context, FEmitScope& Scope, const FRequestedType& RequestedType, const FPreparedType& PreparedType, const Shader::FType& ResultType, Shader::FPreshaderData& OutPreshader) const;
+	Shader::FValue GetValueConstant(FEmitContext& Context, FEmitScope& Scope, const FRequestedType& RequestedType, const FPreparedType& PreparedType, const Shader::FType& ResultType) const;
+
+	FEmitShaderExpression* GetValueShader(FEmitContext& Context, FEmitScope& Scope, const FRequestedType& RequestedType, const Shader::FType& ResultType) const;
+	FEmitShaderExpression* GetValueShader(FEmitContext& Context, FEmitScope& Scope, const FRequestedType& RequestedType) const;
+	FEmitShaderExpression* GetValueShader(FEmitContext& Context, FEmitScope& Scope, const Shader::FType& ResultType) const;
+	FEmitShaderExpression* GetValueShader(FEmitContext& Context, FEmitScope& Scope, Shader::EValueType ResultType) const;
+	FEmitShaderExpression* GetValueShader(FEmitContext& Context, FEmitScope& Scope) const;
+
+	Shader::FType GetValuePreshader(FEmitContext& Context, FEmitScope& Scope, const FRequestedType& RequestedType, const Shader::FType& ResultType, Shader::FPreshaderData& OutPreshader) const;
+	Shader::FType GetValuePreshader(FEmitContext& Context, FEmitScope& Scope, const FRequestedType& RequestedType, Shader::FPreshaderData& OutPreshader) const;
+
+	Shader::FValue GetValueConstant(FEmitContext& Context, FEmitScope& Scope, const FRequestedType& RequestedType, const FPreparedType& PreparedType) const;
+	Shader::FValue GetValueConstant(FEmitContext& Context, FEmitScope& Scope, const FRequestedType& RequestedType, const Shader::FType& ResultType) const;
+	Shader::FValue GetValueConstant(FEmitContext& Context, FEmitScope& Scope, const FRequestedType& RequestedType) const;
+	Shader::FValue GetValueConstant(FEmitContext& Context, FEmitScope& Scope, const FPreparedType& PreparedType, const Shader::FType& ResultType) const;
+	Shader::FValue GetValueConstant(FEmitContext& Context, FEmitScope& Scope, const FPreparedType& PreparedType, Shader::EValueType ResultType) const;
 
 protected:
 	virtual void ComputeAnalyticDerivatives(FTree& Tree, FExpressionDerivatives& OutResult) const;
-	virtual FExpression* ComputePreviousFrame(FTree& Tree, const FRequestedType& RequestedType) const;
+	virtual const FExpression* ComputePreviousFrame(FTree& Tree, const FRequestedType& RequestedType) const;
 	virtual bool PrepareValue(FEmitContext& Context, FEmitScope& Scope, const FRequestedType& RequestedType, FPrepareValueResult& OutResult) const = 0;
 	virtual void EmitValueShader(FEmitContext& Context, FEmitScope& Scope, const FRequestedType& RequestedType, FEmitValueShaderResult& OutResult) const;
 	virtual void EmitValuePreshader(FEmitContext& Context, FEmitScope& Scope, const FRequestedType& RequestedType, FEmitValuePreshaderResult& OutResult) const;
@@ -419,7 +426,7 @@ public:
 
 	FScope* RootScope = nullptr;
 	FScope* CalledScope = nullptr;
-	TArray<FExpression*, TInlineAllocator<8>> OutputExpressions;
+	TArray<const FExpression*, TInlineAllocator<8>> OutputExpressions;
 };
 
 /**
@@ -452,7 +459,7 @@ private:
 	FStatement* OwnerStatement = nullptr;
 	FStatement* ContainedStatement = nullptr;
 	FScope* PreviousScope[MaxNumPreviousScopes];
-	TMap<FName, FExpression*> LocalMap;
+	TMap<FName, const FExpression*> LocalMap;
 	int32 NumPreviousScopes = 0;
 	int32 NestedLevel = 0;
 };
@@ -475,13 +482,13 @@ public:
 	FScope& GetRootScope() const { return *RootScope; }
 
 	template<typename T, typename... ArgTypes>
-	inline FExpression* NewExpression(ArgTypes&&... Args)
+	inline const FExpression* NewExpression(ArgTypes&&... Args)
 	{
 		FHasher Hasher;
 		AppendHash(Hasher, GetGeneratedTypeName<T>());
 		AppendHashes(Hasher, Forward<ArgTypes>(Args)...);
 		const FXxHash64 Hash = Hasher.Finalize();
-		FExpression* Expression = FindExpression(Hash);
+		const FExpression* Expression = FindExpression(Hash);
 		if (!Expression)
 		{
 			T* TypedExpression = NewNode<T>(Forward<ArgTypes>(Args)...);
@@ -499,56 +506,56 @@ public:
 		return Statement;
 	}
 
-	void AssignLocal(FScope& Scope, const FName& LocalName, FExpression* Value);
-	FExpression* AcquireLocal(FScope& Scope, const FName& LocalName);
+	void AssignLocal(FScope& Scope, const FName& LocalName, const FExpression* Value);
+	const FExpression* AcquireLocal(FScope& Scope, const FName& LocalName);
 
-	FExpression* NewFunctionCall(FScope& Scope, FFunction* Function, int32 OutputIndex);
+	const FExpression* NewFunctionCall(FScope& Scope, FFunction* Function, int32 OutputIndex);
 
-	FExpressionDerivatives GetAnalyticDerivatives(FExpression* InExpression);
-	FExpression* GetPreviousFrame(FExpression* InExpression, const FRequestedType& RequestedType);
+	FExpressionDerivatives GetAnalyticDerivatives(const FExpression* InExpression);
+	const FExpression* GetPreviousFrame(const FExpression* InExpression, const FRequestedType& RequestedType);
 
 	FScope* NewScope(FScope& Scope);
 	FScope* NewOwnedScope(FStatement& Owner);
 	FFunction* NewFunction();
 
 	/** Shortcuts to create various common expression types */
-	FExpression* NewConstant(const Shader::FValue& Value);
-	FExpression* NewUnaryOp(EOperation Op, FExpression* Input);
-	FExpression* NewBinaryOp(EOperation Op, FExpression* Lhs, FExpression* Rhs);
-
-	FExpression* NewAbs(FExpression* Input) { return NewUnaryOp(EOperation::Abs, Input); }
-	FExpression* NewNeg(FExpression* Input) { return NewUnaryOp(EOperation::Neg, Input); }
-	FExpression* NewSaturate(FExpression* Input) { return NewUnaryOp(EOperation::Saturate, Input); }
-	FExpression* NewSum(FExpression* Input) { return NewUnaryOp(EOperation::Sum, Input); }
-	FExpression* NewRcp(FExpression* Input) { return NewUnaryOp(EOperation::Rcp, Input); }
-	FExpression* NewSqrt(FExpression* Input) { return NewUnaryOp(EOperation::Sqrt, Input); }
-	FExpression* NewRsqrt(FExpression* Input) { return NewUnaryOp(EOperation::Rsqrt, Input); }
-	FExpression* NewLog2(FExpression* Input) { return NewUnaryOp(EOperation::Log2, Input); }
-	FExpression* NewExp2(FExpression* Input) { return NewUnaryOp(EOperation::Exp2, Input); }
-	FExpression* NewFrac(FExpression* Input) { return NewUnaryOp(EOperation::Frac, Input); }
-	FExpression* NewLength(FExpression* Input) { return NewUnaryOp(EOperation::Length, Input); }
-	FExpression* NewNormalize(FExpression* Input) { return NewUnaryOp(EOperation::Normalize, Input); }
-	FExpression* NewSin(FExpression* Input) { return NewUnaryOp(EOperation::Sin, Input); }
-	FExpression* NewCos(FExpression* Input) { return NewUnaryOp(EOperation::Cos, Input); }
-
-	FExpression* NewAdd(FExpression* Lhs, FExpression* Rhs) { return NewBinaryOp(EOperation::Add, Lhs, Rhs); }
-	FExpression* NewSub(FExpression* Lhs, FExpression* Rhs) { return NewBinaryOp(EOperation::Sub, Lhs, Rhs); }
-	FExpression* NewMul(FExpression* Lhs, FExpression* Rhs) { return NewBinaryOp(EOperation::Mul, Lhs, Rhs); }
-	FExpression* NewDiv(FExpression* Lhs, FExpression* Rhs) { return NewBinaryOp(EOperation::Div, Lhs, Rhs); }
-	FExpression* NewFmod(FExpression* Lhs, FExpression* Rhs) { return NewBinaryOp(EOperation::Fmod, Lhs, Rhs); }
-	FExpression* NewPowClamped(FExpression* Lhs, FExpression* Rhs) { return NewBinaryOp(EOperation::PowPositiveClamped, Lhs, Rhs); }
-	FExpression* NewMin(FExpression* Lhs, FExpression* Rhs) { return NewBinaryOp(EOperation::Min, Lhs, Rhs); }
-	FExpression* NewMax(FExpression* Lhs, FExpression* Rhs) { return NewBinaryOp(EOperation::Max, Lhs, Rhs); }
-	FExpression* NewLess(FExpression* Lhs, FExpression* Rhs) { return NewBinaryOp(EOperation::Less, Lhs, Rhs); }
-	FExpression* NewGreater(FExpression* Lhs, FExpression* Rhs) { return NewBinaryOp(EOperation::Greater, Lhs, Rhs); }
-	FExpression* NewLessEqual(FExpression* Lhs, FExpression* Rhs) { return NewBinaryOp(EOperation::LessEqual, Lhs, Rhs); }
-	FExpression* NewGreaterEqual(FExpression* Lhs, FExpression* Rhs) { return NewBinaryOp(EOperation::GreaterEqual, Lhs, Rhs); }
-
-	FExpression* NewLog(FExpression* Input);
-	FExpression* NewPow2(FExpression* Input) { return NewMul(Input, Input); }
-	FExpression* NewCross(FExpression* Lhs, FExpression* Rhs);
-	FExpression* NewDot(FExpression* Lhs, FExpression* Rhs) { return NewSum(NewMul(Lhs, Rhs)); }
-	FExpression* NewLerp(FExpression* A, FExpression* B, FExpression* T) { return NewAdd(A, NewMul(NewSub(B, A), T)); }
+	const FExpression* NewConstant(const Shader::FValue& Value);
+	const FExpression* NewUnaryOp(EOperation Op, const FExpression* Input);
+	const FExpression* NewBinaryOp(EOperation Op, const FExpression* Lhs, const FExpression* Rhs);
+
+	const FExpression* NewAbs(const FExpression* Input) { return NewUnaryOp(EOperation::Abs, Input); }
+	const FExpression* NewNeg(const FExpression* Input) { return NewUnaryOp(EOperation::Neg, Input); }
+	const FExpression* NewSaturate(const FExpression* Input) { return NewUnaryOp(EOperation::Saturate, Input); }
+	const FExpression* NewSum(const FExpression* Input) { return NewUnaryOp(EOperation::Sum, Input); }
+	const FExpression* NewRcp(const FExpression* Input) { return NewUnaryOp(EOperation::Rcp, Input); }
+	const FExpression* NewSqrt(const FExpression* Input) { return NewUnaryOp(EOperation::Sqrt, Input); }
+	const FExpression* NewRsqrt(const FExpression* Input) { return NewUnaryOp(EOperation::Rsqrt, Input); }
+	const FExpression* NewLog2(const FExpression* Input) { return NewUnaryOp(EOperation::Log2, Input); }
+	const FExpression* NewExp2(const FExpression* Input) { return NewUnaryOp(EOperation::Exp2, Input); }
+	const FExpression* NewFrac(const FExpression* Input) { return NewUnaryOp(EOperation::Frac, Input); }
+	const FExpression* NewLength(const FExpression* Input) { return NewUnaryOp(EOperation::Length, Input); }
+	const FExpression* NewNormalize(const FExpression* Input) { return NewUnaryOp(EOperation::Normalize, Input); }
+	const FExpression* NewSin(const FExpression* Input) { return NewUnaryOp(EOperation::Sin, Input); }
+	const FExpression* NewCos(const FExpression* Input) { return NewUnaryOp(EOperation::Cos, Input); }
+
+	const FExpression* NewAdd(const FExpression* Lhs, const FExpression* Rhs) { return NewBinaryOp(EOperation::Add, Lhs, Rhs); }
+	const FExpression* NewSub(const FExpression* Lhs, const FExpression* Rhs) { return NewBinaryOp(EOperation::Sub, Lhs, Rhs); }
+	const FExpression* NewMul(const FExpression* Lhs, const FExpression* Rhs) { return NewBinaryOp(EOperation::Mul, Lhs, Rhs); }
+	const FExpression* NewDiv(const FExpression* Lhs, const FExpression* Rhs) { return NewBinaryOp(EOperation::Div, Lhs, Rhs); }
+	const FExpression* NewFmod(const FExpression* Lhs, const FExpression* Rhs) { return NewBinaryOp(EOperation::Fmod, Lhs, Rhs); }
+	const FExpression* NewPowClamped(const FExpression* Lhs, const FExpression* Rhs) { return NewBinaryOp(EOperation::PowPositiveClamped, Lhs, Rhs); }
+	const FExpression* NewMin(const FExpression* Lhs, const FExpression* Rhs) { return NewBinaryOp(EOperation::Min, Lhs, Rhs); }
+	const FExpression* NewMax(const FExpression* Lhs, const FExpression* Rhs) { return NewBinaryOp(EOperation::Max, Lhs, Rhs); }
+	const FExpression* NewLess(const FExpression* Lhs, const FExpression* Rhs) { return NewBinaryOp(EOperation::Less, Lhs, Rhs); }
+	const FExpression* NewGreater(const FExpression* Lhs, const FExpression* Rhs) { return NewBinaryOp(EOperation::Greater, Lhs, Rhs); }
+	const FExpression* NewLessEqual(const FExpression* Lhs, const FExpression* Rhs) { return NewBinaryOp(EOperation::LessEqual, Lhs, Rhs); }
+	const FExpression* NewGreaterEqual(const FExpression* Lhs, const FExpression* Rhs) { return NewBinaryOp(EOperation::GreaterEqual, Lhs, Rhs); }
+
+	const FExpression* NewLog(const FExpression* Input);
+	const FExpression* NewPow2(const FExpression* Input) { return NewMul(Input, Input); }
+	const FExpression* NewCross(const FExpression* Lhs, const FExpression* Rhs);
+	const FExpression* NewDot(const FExpression* Lhs, const FExpression* Rhs) { return NewSum(NewMul(Lhs, Rhs)); }
+	const FExpression* NewLerp(const FExpression* A, const FExpression* B, const FExpression* T) { return NewAdd(A, NewMul(NewSub(B, A), T)); }
 
 private:
 	template<typename T, typename... ArgTypes>
@@ -563,12 +570,12 @@ private:
 	void RegisterExpression(FExpression* Expression, FXxHash64 Hash);
 	void RegisterExpression(FExpressionLocalPHI* Expression, FXxHash64 Hash);
 	void RegisterStatement(FScope& Scope, FStatement* Statement);
-	FExpression* FindExpression(FXxHash64 Hash) const;
+	const FExpression* FindExpression(FXxHash64 Hash) const;
 
 	FMemStackBase* Allocator = nullptr;
 	FNode* Nodes = nullptr;
 	FScope* RootScope = nullptr;
-	TMap<FXxHash64, FExpression*> ExpressionMap;
+	TMap<FXxHash64, const FExpression*> ExpressionMap;
 	TArray<FExpressionLocalPHI*> PHIExpressions;
 
 	friend class FExpressionLocalPHI;
diff --git a/Engine/Source/Runtime/Engine/Public/HLSLTree/HLSLTreeCommon.h b/Engine/Source/Runtime/Engine/Public/HLSLTree/HLSLTreeCommon.h
--- a/Engine/Source/Runtime/Engine/Public/HLSLTree/HLSLTreeCommon.h
+++ b/Engine/Source/Runtime/Engine/Public/HLSLTree/HLSLTreeCommon.h
@@ -28,12 +28,12 @@ public:
 class FExpressionForward : public FExpression
 {
 public:
-	explicit FExpressionForward(FExpression* InExpression) : Expression(InExpression) {}
+	explicit FExpressionForward(const FExpression* InExpression) : Expression(InExpression) {}
 
-	FExpression* Expression;
+	const FExpression* Expression;
 
 	virtual void ComputeAnalyticDerivatives(FTree& Tree, FExpressionDerivatives& OutResult) const override;
-	virtual FExpression* ComputePreviousFrame(FTree& Tree, const FRequestedType& RequestedType) const override;
+	virtual const FExpression* ComputePreviousFrame(FTree& Tree, const FRequestedType& RequestedType) const override;
 	virtual bool PrepareValue(FEmitContext& Context, FEmitScope& Scope, const FRequestedType& RequestedType, FPrepareValueResult& OutResult) const override;
 	virtual void EmitValueShader(FEmitContext& Context, FEmitScope& Scope, const FRequestedType& RequestedType, FEmitValueShaderResult& OutResult) const override;
 	virtual void EmitValuePreshader(FEmitContext& Context, FEmitScope& Scope, const FRequestedType& RequestedType, FEmitValuePreshaderResult& OutResult) const override;
@@ -56,9 +56,9 @@ public:
 class FExpressionTextureSample : public FExpression
 {
 public:
-	FExpressionTextureSample(FExpression* InTextureExpression,
-		FExpression* InTexCoordExpression,
-		FExpression* InMipValueExpression,
+	FExpressionTextureSample(const FExpression* InTextureExpression,
+		const FExpression* InTexCoordExpression,
+		const FExpression* InMipValueExpression,
 		const FExpressionDerivatives& InTexCoordDerivatives,
 		ESamplerSourceMode InSamplerSource,
 		ETextureMipValueMode InMipValueMode)
@@ -70,9 +70,9 @@ public:
 		, MipValueMode(InMipValueMode)
 	{}
 
-	FExpression* TextureExpression;
-	FExpression* TexCoordExpression;
-	FExpression* MipValueExpression;
+	const FExpression* TextureExpression;
+	const FExpression* TexCoordExpression;
+	const FExpression* MipValueExpression;
 	FExpressionDerivatives TexCoordDerivatives;
 	ESamplerSourceMode SamplerSource;
 	ETextureMipValueMode MipValueMode;
@@ -84,7 +84,7 @@ public:
 class FExpressionGetStructField : public FExpression
 {
 public:
-	FExpressionGetStructField(const Shader::FStructType* InStructType, const Shader::FStructField* InField, FExpression* InStructExpression)
+	FExpressionGetStructField(const Shader::FStructType* InStructType, const Shader::FStructField* InField, const FExpression* InStructExpression)
 		: StructType(InStructType)
 		, Field(InField)
 		, StructExpression(InStructExpression)
@@ -93,10 +93,10 @@ public:
 
 	const Shader::FStructType* StructType;
 	const Shader::FStructField* Field;
-	FExpression* StructExpression;
+	const FExpression* StructExpression;
 
 	virtual void ComputeAnalyticDerivatives(FTree& Tree, FExpressionDerivatives& OutResult) const override;
-	virtual FExpression* ComputePreviousFrame(FTree& Tree, const FRequestedType& RequestedType) const override;
+	virtual const FExpression* ComputePreviousFrame(FTree& Tree, const FRequestedType& RequestedType) const override;
 	virtual bool PrepareValue(FEmitContext& Context, FEmitScope& Scope, const FRequestedType& RequestedType, FPrepareValueResult& OutResult) const override;
 	virtual void EmitValueShader(FEmitContext& Context, FEmitScope& Scope, const FRequestedType& RequestedType, FEmitValueShaderResult& OutResult) const override;
 	virtual void EmitValuePreshader(FEmitContext& Context, FEmitScope& Scope, const FRequestedType& RequestedType, FEmitValuePreshaderResult& OutResult) const override;
@@ -105,7 +105,7 @@ public:
 class FExpressionSetStructField : public FExpression
 {
 public:
-	FExpressionSetStructField(const Shader::FStructType* InStructType, const Shader::FStructField* InField, FExpression* InStructExpression, FExpression* InFieldExpression)
+	FExpressionSetStructField(const Shader::FStructType* InStructType, const Shader::FStructField* InField, const FExpression* InStructExpression, const FExpression* InFieldExpression)
 		: StructType(InStructType)
 		, Field(InField)
 		, StructExpression(InStructExpression)
@@ -117,11 +117,11 @@ public:
 
 	const Shader::FStructType* StructType;
 	const Shader::FStructField* Field;
-	FExpression* StructExpression;
-	FExpression* FieldExpression;
+	const FExpression* StructExpression;
+	const FExpression* FieldExpression;
 
 	virtual void ComputeAnalyticDerivatives(FTree& Tree, FExpressionDerivatives& OutResult) const override;
-	virtual FExpression* ComputePreviousFrame(FTree& Tree, const FRequestedType& RequestedType) const override;
+	virtual const FExpression* ComputePreviousFrame(FTree& Tree, const FRequestedType& RequestedType) const override;
 	virtual bool PrepareValue(FEmitContext& Context, FEmitScope& Scope, const FRequestedType& RequestedType, FPrepareValueResult& OutResult) const override;
 	virtual void EmitValueShader(FEmitContext& Context, FEmitScope& Scope, const FRequestedType& RequestedType, FEmitValueShaderResult& OutResult) const override;
 	virtual void EmitValuePreshader(FEmitContext& Context, FEmitScope& Scope, const FRequestedType& RequestedType, FEmitValuePreshaderResult& OutResult) const override;
@@ -130,18 +130,18 @@ public:
 class FExpressionSelect : public FExpression
 {
 public:
-	FExpressionSelect(FExpression* InCondition, FExpression* InTrue, FExpression* InFalse)
+	FExpressionSelect(const FExpression* InCondition, const FExpression* InTrue, const FExpression* InFalse)
 		: ConditionExpression(InCondition)
 		, TrueExpression(InTrue)
 		, FalseExpression(InFalse)
 	{}
 
-	FExpression* ConditionExpression;
-	FExpression* TrueExpression;
-	FExpression* FalseExpression;
+	const FExpression* ConditionExpression;
+	const FExpression* TrueExpression;
+	const FExpression* FalseExpression;
 
 	virtual void ComputeAnalyticDerivatives(FTree& Tree, FExpressionDerivatives& OutResult) const override;
-	virtual FExpression* ComputePreviousFrame(FTree& Tree, const FRequestedType& RequestedType) const override;
+	virtual const FExpression* ComputePreviousFrame(FTree& Tree, const FRequestedType& RequestedType) const override;
 	virtual bool PrepareValue(FEmitContext& Context, FEmitScope& Scope, const FRequestedType& RequestedType, FPrepareValueResult& OutResult) const override;
 	virtual void EmitValueShader(FEmitContext& Context, FEmitScope& Scope, const FRequestedType& RequestedType, FEmitValueShaderResult& OutResult) const override;
 	virtual void EmitValuePreshader(FEmitContext& Context, FEmitScope& Scope, const FRequestedType& RequestedType, FEmitValuePreshaderResult& OutResult) const override;
@@ -150,13 +150,13 @@ public:
 class FExpressionDerivative : public FExpression
 {
 public:
-	FExpressionDerivative(EDerivativeCoordinate InCoord, FExpression* InInput) : Input(InInput), Coord(InCoord) {}
+	FExpressionDerivative(EDerivativeCoordinate InCoord, const FExpression* InInput) : Input(InInput), Coord(InCoord) {}
 
-	FExpression* Input;
+	const FExpression* Input;
 	EDerivativeCoordinate Coord;
 
 	virtual void ComputeAnalyticDerivatives(FTree& Tree, FExpressionDerivatives& OutResult) const override;
-	virtual FExpression* ComputePreviousFrame(FTree& Tree, const FRequestedType& RequestedType) const override;
+	virtual const FExpression* ComputePreviousFrame(FTree& Tree, const FRequestedType& RequestedType) const override;
 	virtual bool PrepareValue(FEmitContext& Context, FEmitScope& Scope, const FRequestedType& RequestedType, FPrepareValueResult& OutResult) const override;
 	virtual void EmitValueShader(FEmitContext& Context, FEmitScope& Scope, const FRequestedType& RequestedType, FEmitValueShaderResult& OutResult) const override;
 	virtual void EmitValuePreshader(FEmitContext& Context, FEmitScope& Scope, const FRequestedType& RequestedType, FEmitValuePreshaderResult& OutResult) const override;
@@ -178,16 +178,16 @@ FSwizzleParameters MakeSwizzleMask(bool bInR, bool bInG, bool bInB, bool bInA);
 class FExpressionSwizzle : public FExpression
 {
 public:
-	FExpressionSwizzle(const FSwizzleParameters& InParams, FExpression* InInput)
+	FExpressionSwizzle(const FSwizzleParameters& InParams, const FExpression* InInput)
 		: Parameters(InParams)
 		, Input(InInput)
 	{}
 
 	FSwizzleParameters Parameters;
-	FExpression* Input;
+	const FExpression* Input;
 
 	virtual void ComputeAnalyticDerivatives(FTree& Tree, FExpressionDerivatives& OutResult) const override;
-	virtual FExpression* ComputePreviousFrame(FTree& Tree, const FRequestedType& RequestedType) const override;
+	virtual const FExpression* ComputePreviousFrame(FTree& Tree, const FRequestedType& RequestedType) const override;
 	virtual bool PrepareValue(FEmitContext& Context, FEmitScope& Scope, const FRequestedType& RequestedType, FPrepareValueResult& OutResult) const override;
 	virtual void EmitValueShader(FEmitContext& Context, FEmitScope& Scope, const FRequestedType& RequestedType, FEmitValueShaderResult& OutResult) const override;
 	virtual void EmitValuePreshader(FEmitContext& Context, FEmitScope& Scope, const FRequestedType& RequestedType, FEmitValuePreshaderResult& OutResult) const override;
@@ -196,16 +196,16 @@ public:
 class FExpressionAppend : public FExpression
 {
 public:
-	FExpressionAppend(FExpression* InLhs, FExpression* InRhs)
+	FExpressionAppend(const FExpression* InLhs, const FExpression* InRhs)
 		: Lhs(InLhs)
 		, Rhs(InRhs)
 	{}
 
-	FExpression* Lhs;
-	FExpression* Rhs;
+	const FExpression* Lhs;
+	const FExpression* Rhs;
 
 	virtual void ComputeAnalyticDerivatives(FTree& Tree, FExpressionDerivatives& OutResult) const override;
-	virtual FExpression* ComputePreviousFrame(FTree& Tree, const FRequestedType& RequestedType) const override;
+	virtual const FExpression* ComputePreviousFrame(FTree& Tree, const FRequestedType& RequestedType) const override;
 	virtual bool PrepareValue(FEmitContext& Context, FEmitScope& Scope, const FRequestedType& RequestedType, FPrepareValueResult& OutResult) const override;
 	virtual void EmitValueShader(FEmitContext& Context, FEmitScope& Scope, const FRequestedType& RequestedType, FEmitValueShaderResult& OutResult) const override;
 	virtual void EmitValuePreshader(FEmitContext& Context, FEmitScope& Scope, const FRequestedType& RequestedType, FEmitValuePreshaderResult& OutResult) const override;
@@ -216,7 +216,7 @@ class FExpressionSwitchBase : public FExpression
 public:
 	static constexpr int8 MaxInputs = 8;
 
-	FExpressionSwitchBase(TConstArrayView<FExpression*> InInputs) : NumInputs(InInputs.Num())
+	FExpressionSwitchBase(TConstArrayView<const FExpression*> InInputs) : NumInputs(InInputs.Num())
 	{
 		check(InInputs.Num() <= MaxInputs);
 		for (int32 i = 0; i < InInputs.Num(); ++i)
@@ -225,14 +225,14 @@ public:
 		}
 	}
 
-	FExpression* Input[MaxInputs] = { nullptr };
+	const FExpression* Input[MaxInputs] = { nullptr };
 	int8 NumInputs = 0;
 
-	virtual FExpression* NewSwitch(FTree& Tree, TConstArrayView<FExpression*> InInputs) const = 0;
+	virtual const FExpression* NewSwitch(FTree& Tree, TConstArrayView<const FExpression*> InInputs) const = 0;
 	virtual bool IsInputActive(const FEmitContext& Context, int32 Index) const = 0;
 
 	virtual void ComputeAnalyticDerivatives(FTree& Tree, FExpressionDerivatives& OutResult) const override;
-	virtual FExpression* ComputePreviousFrame(FTree& Tree, const FRequestedType& RequestedType) const override;
+	virtual const FExpression* ComputePreviousFrame(FTree& Tree, const FRequestedType& RequestedType) const override;
 	virtual bool PrepareValue(FEmitContext& Context, FEmitScope& Scope, const FRequestedType& RequestedType, FPrepareValueResult& OutResult) const override;
 	virtual void EmitValueShader(FEmitContext& Context, FEmitScope& Scope, const FRequestedType& RequestedType, FEmitValueShaderResult& OutResult) const override;
 	virtual void EmitValuePreshader(FEmitContext& Context, FEmitScope& Scope, const FRequestedType& RequestedType, FEmitValuePreshaderResult& OutResult) const override;
@@ -242,12 +242,12 @@ class FExpressionFeatureLevelSwitch : public FExpressionSwitchBase
 {
 public:
 	static_assert(MaxInputs >= (int32)ERHIFeatureLevel::Num, "FExpressionSwitchBase is too small for FExpressionFeatureLevelSwitch");
-	FExpressionFeatureLevelSwitch(TConstArrayView<FExpression*> InInputs) : FExpressionSwitchBase(InInputs)
+	FExpressionFeatureLevelSwitch(TConstArrayView<const FExpression*> InInputs) : FExpressionSwitchBase(InInputs)
 	{
 		check(InInputs.Num() == (int32)ERHIFeatureLevel::Num);
 	}
 
-	virtual FExpression* NewSwitch(FTree& Tree, TConstArrayView<FExpression*> InInputs) const override { return Tree.NewExpression<FExpressionFeatureLevelSwitch>(InInputs); }
+	virtual const FExpression* NewSwitch(FTree& Tree, TConstArrayView<const FExpression*> InInputs) const override { return Tree.NewExpression<FExpressionFeatureLevelSwitch>(InInputs); }
 	virtual bool IsInputActive(const FEmitContext& Context, int32 Index) const override;
 };
 
@@ -255,12 +255,12 @@ class FExpressionShadingPathSwitch : public FExpressionSwitchBase
 {
 public:
 	static_assert(MaxInputs >= (int32)ERHIShadingPath::Num, "FExpressionSwitchBase is too small for FExpressionShadingPathSwitch");
-	FExpressionShadingPathSwitch(TConstArrayView<FExpression*> InInputs) : FExpressionSwitchBase(InInputs)
+	FExpressionShadingPathSwitch(TConstArrayView<const FExpression*> InInputs) : FExpressionSwitchBase(InInputs)
 	{
 		check(InInputs.Num() == (int32)ERHIShadingPath::Num);
 	}
 
-	virtual FExpression* NewSwitch(FTree& Tree, TConstArrayView<FExpression*> InInputs) const override { return Tree.NewExpression<FExpressionShadingPathSwitch>(InInputs); }
+	virtual const FExpression* NewSwitch(FTree& Tree, TConstArrayView<const FExpression*> InInputs) const override { return Tree.NewExpression<FExpressionShadingPathSwitch>(InInputs); }
 	virtual bool IsInputActive(const FEmitContext& Context, int32 Index) const override;
 };
 
@@ -299,7 +299,7 @@ public:
 class FStatementReturn : public FStatement
 {
 public:
-	FExpression* Expression;
+	const FExpression* Expression;
 
 	virtual bool Prepare(FEmitContext& Context, FEmitScope& Scope) const override;
 	virtual void EmitShader(FEmitContext& Context, FEmitScope& Scope) const override;
@@ -316,7 +316,7 @@ public:
 class FStatementIf : public FStatement
 {
 public:
-	FExpression* ConditionExpression;
+	const FExpression* ConditionExpression;
 	FScope* ThenScope;
 	FScope* ElseScope;
 	FScope* NextScope;
diff --git a/Engine/Source/Runtime/Engine/Public/HLSLTree/HLSLTreeEmit.h b/Engine/Source/Runtime/Engine/Public/HLSLTree/HLSLTreeEmit.h
--- a/Engine/Source/Runtime/Engine/Public/HLSLTree/HLSLTreeEmit.h
+++ b/Engine/Source/Runtime/Engine/Public/HLSLTree/HLSLTreeEmit.h
@@ -313,7 +313,7 @@ public:
 	Shader::FType GetType(const FExpression* Expression) const;
 	EExpressionEvaluation GetEvaluation(const FExpression* Expression, const FEmitScope& Scope, const FRequestedType& RequestedType) const;
 
-	FPreparedType PrepareExpression(FExpression* InExpression, FEmitScope& Scope, const FRequestedType& RequestedType);
+	FPreparedType PrepareExpression(const FExpression* InExpression, FEmitScope& Scope, const FRequestedType& RequestedType);
 
 	FEmitScope* InternalPrepareScope(FScope* Scope, FScope* ParentScope);
 	FEmitScope* PrepareScope(FScope* Scope);
@@ -474,7 +474,7 @@ public:
 		return nullptr;
 	}
 
-	FEmitShaderExpression* EmitPreshaderOrConstant(FEmitScope& Scope, const FRequestedType& RequestedType, const Shader::FType& ResultType, FExpression* Expression);
+	FEmitShaderExpression* EmitPreshaderOrConstant(FEmitScope& Scope, const FRequestedType& RequestedType, const Shader::FType& ResultType, const FExpression* Expression);
 	FEmitShaderExpression* EmitConstantZero(FEmitScope& Scope, const Shader::FType& Type);
 	FEmitShaderExpression* EmitCast(FEmitScope& Scope, FEmitShaderExpression* ShaderValue, const Shader::FType& DestType, EEmitCastFlags Flags = EEmitCastFlags::None);
 
diff --git a/Engine/Source/Runtime/Engine/Public/HLSLTree/HLSLTreeTypes.h b/Engine/Source/Runtime/Engine/Public/HLSLTree/HLSLTreeTypes.h
--- a/Engine/Source/Runtime/Engine/Public/HLSLTree/HLSLTreeTypes.h
+++ b/Engine/Source/Runtime/Engine/Public/HLSLTree/HLSLTreeTypes.h
@@ -140,10 +140,10 @@ FOperationDescription GetOperationDescription(EOperation Op);
 struct FCustomHLSLInput
 {
 	FCustomHLSLInput() = default;
-	FCustomHLSLInput(FStringView InName, FExpression* InExpression) : Name(InName), Expression(InExpression) {}
+	FCustomHLSLInput(FStringView InName, const FExpression* InExpression) : Name(InName), Expression(InExpression) {}
 
 	FStringView Name;
-	FExpression* Expression = nullptr;
+	const FExpression* Expression = nullptr;
 };
 
 } // namespace UE::HLSLTree
diff --git a/Engine/Source/Runtime/Engine/Public/MaterialCachedHLSLTree.h b/Engine/Source/Runtime/Engine/Public/MaterialCachedHLSLTree.h
--- a/Engine/Source/Runtime/Engine/Public/MaterialCachedHLSLTree.h
+++ b/Engine/Source/Runtime/Engine/Public/MaterialCachedHLSLTree.h
@@ -32,7 +32,7 @@ public:
 	const UE::Shader::FStructTypeRegistry& GetTypeRegistry() const { return TypeRegistry; }
 	UE::HLSLTree::FTree& GetTree() const { return *HLSLTree; }
 
-	UE::HLSLTree::FExpression* GetResultExpression() const { return ResultExpression; }
+	const UE::HLSLTree::FExpression* GetResultExpression() const { return ResultExpression; }
 	UE::HLSLTree::FStatement* GetResultStatement() const { return ResultStatement; }
 
 	const UE::Shader::FStructType* GetMaterialAttributesType() const { return MaterialAttributesType; }
@@ -50,7 +50,7 @@ private:
 	FMemStackBase Allocator;
 	UE::Shader::FStructTypeRegistry TypeRegistry;
 	UE::HLSLTree::FTree* HLSLTree = nullptr;
-	UE::HLSLTree::FExpression* ResultExpression = nullptr;
+	const UE::HLSLTree::FExpression* ResultExpression = nullptr;
 	UE::HLSLTree::FStatement* ResultStatement = nullptr;
 
 	TArray<UMaterialExpressionCustomOutput*> MaterialCustomOutputs;
diff --git a/Engine/Source/Runtime/Engine/Public/MaterialExpressionIO.h b/Engine/Source/Runtime/Engine/Public/MaterialExpressionIO.h
--- a/Engine/Source/Runtime/Engine/Public/MaterialExpressionIO.h
+++ b/Engine/Source/Runtime/Engine/Public/MaterialExpressionIO.h
@@ -77,9 +77,9 @@ struct FExpressionInput
 
 #if WITH_EDITOR
 	ENGINE_API int32 Compile(class FMaterialCompiler* Compiler);
-	ENGINE_API UE::HLSLTree::FExpression* TryAcquireHLSLExpression(FMaterialHLSLGenerator& Generator, UE::HLSLTree::FScope& Scope) const;
-	ENGINE_API UE::HLSLTree::FExpression* AcquireHLSLExpression(FMaterialHLSLGenerator& Generator, UE::HLSLTree::FScope& Scope) const;
-	ENGINE_API UE::HLSLTree::FExpression* AcquireHLSLExpressionOrConstant(FMaterialHLSLGenerator& Generator, UE::HLSLTree::FScope& Scope, const UE::Shader::FValue& ConstantValue) const;
+	ENGINE_API const UE::HLSLTree::FExpression* TryAcquireHLSLExpression(FMaterialHLSLGenerator& Generator, UE::HLSLTree::FScope& Scope) const;
+	ENGINE_API const UE::HLSLTree::FExpression* AcquireHLSLExpression(FMaterialHLSLGenerator& Generator, UE::HLSLTree::FScope& Scope) const;
+	ENGINE_API const UE::HLSLTree::FExpression* AcquireHLSLExpressionOrConstant(FMaterialHLSLGenerator& Generator, UE::HLSLTree::FScope& Scope, const UE::Shader::FValue& ConstantValue) const;
 #endif // WITH_EDITOR
 
 	/**
diff --git a/Engine/Source/Runtime/Engine/Public/MaterialHLSLGenerator.h b/Engine/Source/Runtime/Engine/Public/MaterialHLSLGenerator.h
--- a/Engine/Source/Runtime/Engine/Public/MaterialHLSLGenerator.h
+++ b/Engine/Source/Runtime/Engine/Public/MaterialHLSLGenerator.h
@@ -97,10 +97,10 @@ public:
 
 	UE::HLSLTree::FScope* NewJoinedScope(UE::HLSLTree::FScope& Scope);
 
-	UE::HLSLTree::FExpression* NewConstant(const UE::Shader::FValue& Value);
-	UE::HLSLTree::FExpression* NewTexCoord(int32 Index);
-	UE::HLSLTree::FExpression* NewExternalInput(UE::HLSLTree::Material::EExternalInput Input);
-	UE::HLSLTree::FExpression* NewSwizzle(const UE::HLSLTree::FSwizzleParameters& Params, UE::HLSLTree::FExpression* Input);
+	const UE::HLSLTree::FExpression* NewConstant(const UE::Shader::FValue& Value);
+	const UE::HLSLTree::FExpression* NewTexCoord(int32 Index);
+	const UE::HLSLTree::FExpression* NewExternalInput(UE::HLSLTree::Material::EExternalInput Input);
+	const UE::HLSLTree::FExpression* NewSwizzle(const UE::HLSLTree::FSwizzleParameters& Params, const UE::HLSLTree::FExpression* Input);
 
 	const UE::Shader::FTextureValue* AcquireTextureValue(const UE::Shader::FTextureValue& Value);
 
@@ -109,18 +109,18 @@ public:
 	 * The node will be created if it doesn't exist. Otherwise, the tree will be updated to ensure the node is visible in the given scope
 	 * Note that a given UMaterialExpression may only support 1 of these node types, attempting to access an invalid node type will generate an error
 	 */
-	UE::HLSLTree::FExpression* AcquireExpression(UE::HLSLTree::FScope& Scope, UMaterialExpression* MaterialExpression, int32 OutputIndex);
-	UE::HLSLTree::FExpression* AcquireFunctionInputExpression(UE::HLSLTree::FScope& Scope, const UMaterialExpressionFunctionInput* MaterialExpression);
+	const UE::HLSLTree::FExpression* AcquireExpression(UE::HLSLTree::FScope& Scope, UMaterialExpression* MaterialExpression, int32 OutputIndex);
+	const UE::HLSLTree::FExpression* AcquireFunctionInputExpression(UE::HLSLTree::FScope& Scope, const UMaterialExpressionFunctionInput* MaterialExpression);
 
 	bool GenerateStatements(UE::HLSLTree::FScope& Scope, UMaterialExpression* MaterialExpression);
 
-	UE::HLSLTree::FExpression* GenerateMaterialParameter(FName InParameterName, const FMaterialParameterMetadata& InParameterMeta, const UE::Shader::FValue& InDefaultValue);
+	const UE::HLSLTree::FExpression* GenerateMaterialParameter(FName InParameterName, const FMaterialParameterMetadata& InParameterMeta, const UE::Shader::FValue& InDefaultValue);
 
-	UE::HLSLTree::FExpression* GenerateFunctionCall(UE::HLSLTree::FScope& Scope,
+	const UE::HLSLTree::FExpression* GenerateFunctionCall(UE::HLSLTree::FScope& Scope,
 		UMaterialFunctionInterface* Function,
 		EMaterialParameterAssociation ParameterAssociation,
 		int32 ParameterIndex,
-		TArrayView<UE::HLSLTree::FExpression*> ConnectedInputs,
+		TArrayView<const UE::HLSLTree::FExpression*> ConnectedInputs,
 		int32 OutputIndex);
 
 	template<typename T, typename... ArgTypes>
@@ -175,7 +175,7 @@ private:
 
 	using FFunctionInputArray = TArray<const UMaterialExpressionFunctionInput*, TInlineAllocator<4>>;
 	using FFunctionOutputArray = TArray<const UMaterialExpressionFunctionOutput*, TInlineAllocator<4>>;
-	using FConnectedInputArray = TArray<UE::HLSLTree::FExpression*, TInlineAllocator<4>>;
+	using FConnectedInputArray = TArray<const UE::HLSLTree::FExpression*, TInlineAllocator<4>>;
 
 	struct FFunctionCallEntry
 	{
@@ -214,7 +214,7 @@ private:
 	TMap<UMaterialFunctionInterface*, UE::HLSLTree::FFunction*> FunctionMap;
 	TMap<UMaterialExpression*, FStatementEntry> StatementMap;
 	TMap<FExpressionDataKey, void*> ExpressionDataMap;
-	UE::HLSLTree::FExpression* PreviewExpressionResult = nullptr;
+	const UE::HLSLTree::FExpression* PreviewExpressionResult = nullptr;
 	bool bGeneratedResult;
 };
 
