Cloth - Unregister the cloth tick function whenever the component tick is disabled.

diff --git a/Engine/Source/Runtime/Engine/Classes/Components/SkeletalMeshComponent.h b/Engine/Source/Runtime/Engine/Classes/Components/SkeletalMeshComponent.h
--- a/Engine/Source/Runtime/Engine/Classes/Components/SkeletalMeshComponent.h
+++ b/Engine/Source/Runtime/Engine/Classes/Components/SkeletalMeshComponent.h
@@ -1645,6 +1645,7 @@ public:
 	virtual void InitializeComponent() override;
 	virtual void TickComponent(float DeltaTime, enum ELevelTick TickType, FActorComponentTickFunction *ThisTickFunction) override;
 	virtual void BeginPlay() override;
+	virtual void SetComponentTickEnabled(bool bEnabled) override;
 
 	//Handle registering our end physics tick function
 	virtual void RegisterEndPhysicsTick(bool bRegister);
diff --git a/Engine/Source/Runtime/Engine/Private/Components/SkeletalMeshComponent.cpp b/Engine/Source/Runtime/Engine/Private/Components/SkeletalMeshComponent.cpp
--- a/Engine/Source/Runtime/Engine/Private/Components/SkeletalMeshComponent.cpp
+++ b/Engine/Source/Runtime/Engine/Private/Components/SkeletalMeshComponent.cpp
@@ -372,6 +372,16 @@ void USkeletalMeshComponent::RegisterComponentTickFunctions(bool bRegister)
 	UpdateClothTickRegisteredState();
 }
 
+void USkeletalMeshComponent::SetComponentTickEnabled(bool bEnabled)
+{
+	Super::SetComponentTickEnabled(bEnabled);
+
+	// Unregister the cloth tick function when the component's tick is being disabled (it will be re-enabled at the next component tick update once it runs again)
+	if (!bEnabled)
+	{
+		RegisterClothTick(false);
+	}
+}
 
 void USkeletalMeshComponent::RegisterEndPhysicsTick(bool bRegister)
 {
