Avoid calling `GetMaterialWithFallback` which was forcing us to compile _all_ shaders for materials even when ODSC was turned on.

- In CitySample we went from compiling 125,000 shaders down to 9,000 shaders.
- In CitySample we went from spending 13 minutes submitting shaders down to ~15 seconds.
- Simply call `GetIncompleteMaterialWithFallback`, which should be safe to do in this case.  We don't actually need shaders from the material and we are simply querying it for blend mode and for certain ray tracing flags.

diff --git a/Engine/Source/Runtime/Engine/Private/RayTracingInstance.cpp b/Engine/Source/Runtime/Engine/Private/RayTracingInstance.cpp
--- a/Engine/Source/Runtime/Engine/Private/RayTracingInstance.cpp
+++ b/Engine/Source/Runtime/Engine/Private/RayTracingInstance.cpp
@@ -40,8 +40,7 @@ FRayTracingMaskAndFlags BuildRayTracingInstanceMaskAndFlags(TArrayView<const FMe
 		// Mesh Batches can "null" when they have zero triangles.  Check the MaterialRenderProxy before accessing.
 		if (MeshBatch.bUseForMaterial && MeshBatch.MaterialRenderProxy)
 		{
-			const FMaterialRenderProxy* FallbackMaterialRenderProxyPtr = nullptr;
-			const FMaterial& Material = MeshBatch.MaterialRenderProxy->GetMaterialWithFallback(FeatureLevel, FallbackMaterialRenderProxyPtr);
+			const FMaterial& Material = MeshBatch.MaterialRenderProxy->GetIncompleteMaterialWithFallback(FeatureLevel);
 			const EBlendMode BlendMode = Material.GetBlendMode();
 			Result.Mask |= ComputeBlendModeMask(BlendMode);
 			bAllSegmentsOpaque &= BlendMode == BLEND_Opaque;
