Add LLM tracking for audio spatialization plugins

diff --git a/Engine/Source/Runtime/Core/Private/HAL/LowLevelMemTracker.cpp b/Engine/Source/Runtime/Core/Private/HAL/LowLevelMemTracker.cpp
--- a/Engine/Source/Runtime/Core/Private/HAL/LowLevelMemTracker.cpp
+++ b/Engine/Source/Runtime/Core/Private/HAL/LowLevelMemTracker.cpp
@@ -168,6 +168,7 @@ DECLARE_LLM_MEMORY_STAT(TEXT("Audio"), STAT_AudioLLM, STATGROUP_LLMFULL);
 DECLARE_LLM_MEMORY_STAT(TEXT("AudioMisc"), STAT_AudioMiscLLM, STATGROUP_LLMFULL);
 DECLARE_LLM_MEMORY_STAT(TEXT("AudioSoundWaves"), STAT_AudioSoundWavesLLM, STATGROUP_LLMFULL);
 DECLARE_LLM_MEMORY_STAT(TEXT("AudioSoundWaveProxies"), STAT_AudioSoundWaveProxiesLLM, STATGROUP_LLMFULL);
+DECLARE_LLM_MEMORY_STAT(TEXT("AudioSpatializationPlugins"), STAT_AudioSpatializationPluginsLLM, STATGROUP_LLMFULL);
 DECLARE_LLM_MEMORY_STAT(TEXT("AudioMixer"), STAT_AudioMixerLLM, STATGROUP_LLMFULL);
 DECLARE_LLM_MEMORY_STAT(TEXT("AudioMixerPlugins"), STAT_AudioMixerPluginsLLM, STATGROUP_LLMFULL);
 DECLARE_LLM_MEMORY_STAT(TEXT("AudioPrecache"), STAT_AudioPrecacheLLM, STATGROUP_LLMFULL);
diff --git a/Engine/Source/Runtime/Core/Public/HAL/LowLevelMemTracker.h b/Engine/Source/Runtime/Core/Public/HAL/LowLevelMemTracker.h
--- a/Engine/Source/Runtime/Core/Public/HAL/LowLevelMemTracker.h
+++ b/Engine/Source/Runtime/Core/Public/HAL/LowLevelMemTracker.h
@@ -153,6 +153,7 @@ enum class ELLMTagSet : uint8
 	macro(AudioMisc,							"AudioMisc",					GET_STATFNAME(STAT_AudioMiscLLM),							GET_STATFNAME(STAT_AudioSummaryLLM),			ELLMTag::Audio)\
 	macro(AudioSoundWaves,						"AudioSoundWaves",				GET_STATFNAME(STAT_AudioSoundWavesLLM),						GET_STATFNAME(STAT_AudioSummaryLLM),			ELLMTag::Audio)\
 	macro(AudioSoundWaveProxies,				"AudioSoundWaveProxies",		GET_STATFNAME(STAT_AudioSoundWaveProxiesLLM),				GET_STATFNAME(STAT_AudioSummaryLLM),			ELLMTag::Audio)\
+	macro(AudioSpatializationPlugins,			"AudioSpatializationPlugins",	GET_STATFNAME(STAT_AudioSpatializationPluginsLLM),			GET_STATFNAME(STAT_AudioSummaryLLM),			ELLMTag::Audio)\
 	macro(AudioMixer,							"AudioMixer",					GET_STATFNAME(STAT_AudioMixerLLM),							GET_STATFNAME(STAT_AudioSummaryLLM),			ELLMTag::Audio)\
 	macro(AudioMixerPlugins,					"AudioMixerPlugins",			GET_STATFNAME(STAT_AudioMixerPluginsLLM),					GET_STATFNAME(STAT_AudioSummaryLLM),			ELLMTag::Audio)\
 	macro(AudioPrecache,						"AudioPrecache",				GET_STATFNAME(STAT_AudioPrecacheLLM),						GET_STATFNAME(STAT_AudioSummaryLLM),			ELLMTag::Audio)\
@@ -249,7 +250,7 @@ enum class ELLMTag : LLM_TAG_TYPE
 
 	//------------------------------
 	// Platform tags
-	PlatformTagStart = 108,
+	PlatformTagStart = 109,
 	PlatformTagEnd = 149,
 
 	//------------------------------
diff --git a/Engine/Source/Runtime/Engine/Private/AudioDevice.cpp b/Engine/Source/Runtime/Engine/Private/AudioDevice.cpp
--- a/Engine/Source/Runtime/Engine/Private/AudioDevice.cpp
+++ b/Engine/Source/Runtime/Engine/Private/AudioDevice.cpp
@@ -783,6 +783,8 @@ void FAudioDevice::Teardown()
 
 	for (TAudioPluginListenerPtr PluginListener : PluginListeners)
 	{
+		LLM_SCOPE(ELLMTag::AudioSpatializationPlugins);
+
 		PluginListener->OnListenerShutdown(this);
 	}
 
@@ -3222,6 +3224,8 @@ void FAudioDevice::SetListener(UWorld* World, const int32 InViewportIndex, const
 
 	if (World)
 	{
+		LLM_SCOPE(ELLMTag::AudioSpatializationPlugins);
+
 		for (TAudioPluginListenerPtr PluginManager : PluginListeners)
 		{
 			PluginManager->OnTick(World, InViewportIndex, ListenerTransformCopy, InDeltaSeconds);
@@ -3233,6 +3237,8 @@ void FAudioDevice::SetListener(UWorld* World, const int32 InViewportIndex, const
 		// Broadcast to a 3rd party plugin listener observer if enabled
 		for (TAudioPluginListenerPtr PluginManager : PluginListeners)
 		{
+			LLM_SCOPE(ELLMTag::AudioSpatializationPlugins);
+
 			PluginManager->OnListenerUpdated(this, InViewportIndex, ListenerTransformCopy, InDeltaSeconds);
 		}
 
@@ -4784,6 +4790,8 @@ void FAudioDevice::InitializePluginListeners(UWorld* World)
 	check(IsInGameThread());
 	check(!bPluginListenersInitialized);
 
+	LLM_SCOPE(ELLMTag::AudioSpatializationPlugins);
+
 	for (TAudioPluginListenerPtr PluginListener : PluginListeners)
 	{
 		PluginListener->OnListenerInitialize(this, World);
@@ -4794,6 +4802,8 @@ void FAudioDevice::NotifyPluginListenersWorldChanged(UWorld* World)
 {
 	check(IsInGameThread());
 
+	LLM_SCOPE(ELLMTag::AudioSpatializationPlugins);
+
 	for (TAudioPluginListenerPtr PluginListener : PluginListeners)
 	{
 		PluginListener->OnWorldChanged(this, World);
