Add support for vertex interpolators to new HLSL translator

diff --git a/Engine/Shaders/Private/GeometryCollectionVertexFactory.ush b/Engine/Shaders/Private/GeometryCollectionVertexFactory.ush
--- a/Engine/Shaders/Private/GeometryCollectionVertexFactory.ush
+++ b/Engine/Shaders/Private/GeometryCollectionVertexFactory.ush
@@ -269,7 +269,7 @@ FMaterialVertexParameters GetMaterialVertexParameters(FVertexFactoryInput Input,
 #endif
 
 #if ENABLE_NEW_HLSL_GENERATOR
-	Result.MaterialAttributes = EvaluateVertexMaterialAttributes(Result);
+	EvaluateVertexMaterialAttributes(Result);
 #endif
 	return Result;
 }
diff --git a/Engine/Shaders/Private/LocalVertexFactory.ush b/Engine/Shaders/Private/LocalVertexFactory.ush
--- a/Engine/Shaders/Private/LocalVertexFactory.ush
+++ b/Engine/Shaders/Private/LocalVertexFactory.ush
@@ -661,7 +661,7 @@ FMaterialVertexParameters GetMaterialVertexParameters(FVertexFactoryInput Input,
 #endif
 
 #if ENABLE_NEW_HLSL_GENERATOR
-	Result.MaterialAttributes = EvaluateVertexMaterialAttributes(Result);
+	EvaluateVertexMaterialAttributes(Result);
 #endif
 	return Result;
 }
diff --git a/Engine/Shaders/Private/MaterialTemplate.ush b/Engine/Shaders/Private/MaterialTemplate.ush
--- a/Engine/Shaders/Private/MaterialTemplate.ush
+++ b/Engine/Shaders/Private/MaterialTemplate.ush
@@ -529,6 +529,11 @@ struct FMaterialVertexParameters
 	#endif
 #endif
 
+#if NUM_CUSTOM_VERTEX_INTERPOLATORS
+	// Used by new HLSL translator
+	float2 CustomInterpolators[NUM_CUSTOM_VERTEX_INTERPOLATORS];
+#endif
+
 	/** Per-particle properties. Only valid for particle vertex factories. */
 	FMaterialParticleParameters Particle;
 
@@ -546,6 +551,22 @@ struct FMaterialVertexParameters
 	uint PrimitiveId;
 };
 
+float MaterialReadInterpolatorComponent(FMaterialPixelParameters Parameters, int InterpolatorIndex)
+{
+#if NUM_TEX_COORD_INTERPOLATORS
+	return Parameters.TexCoords[NUM_MATERIAL_TEXCOORDS + InterpolatorIndex / 2][InterpolatorIndex & 1];
+#else
+	return 0.0f;
+#endif
+}
+
+void MaterialPackInterpolatorComponent(in out FMaterialVertexParameters Parameters, int InterpolatorIndex, float Value)
+{
+#if NUM_CUSTOM_VERTEX_INTERPOLATORS
+	Parameters.CustomInterpolators[InterpolatorIndex / 2][InterpolatorIndex & 1] = Value;
+#endif
+}
+
 FPrimitiveSceneData GetPrimitiveData(FMaterialVertexParameters Parameters)
 {
 	return Parameters.SceneData.Primitive;
diff --git a/Engine/Source/Runtime/Engine/Classes/Materials/MaterialExpression.h b/Engine/Source/Runtime/Engine/Classes/Materials/MaterialExpression.h
--- a/Engine/Source/Runtime/Engine/Classes/Materials/MaterialExpression.h
+++ b/Engine/Source/Runtime/Engine/Classes/Materials/MaterialExpression.h
@@ -306,6 +306,8 @@ class ENGINE_API UMaterialExpression : public UObject
 	virtual bool GenerateHLSLStatements(FMaterialHLSLGenerator& Generator, UE::HLSLTree::FScope& Scope) const;
 	virtual bool GenerateHLSLExpression(FMaterialHLSLGenerator& Generator, UE::HLSLTree::FScope& Scope, int32 OutputIndex, UE::HLSLTree::FExpression const*& OutExpression) const;
 
+	bool IsUsingNewHLSLGenerator() const;
+
 #endif // WITH_EDITOR
 
 	/**
diff --git a/Engine/Source/Runtime/Engine/Classes/Materials/MaterialExpressionVertexInterpolator.h b/Engine/Source/Runtime/Engine/Classes/Materials/MaterialExpressionVertexInterpolator.h
--- a/Engine/Source/Runtime/Engine/Classes/Materials/MaterialExpressionVertexInterpolator.h
+++ b/Engine/Source/Runtime/Engine/Classes/Materials/MaterialExpressionVertexInterpolator.h
@@ -17,11 +17,14 @@ class UMaterialExpressionVertexInterpolator : public UMaterialExpressionCustomOu
 	virtual int32 Compile(class FMaterialCompiler* Compiler, int32 OutputIndex) override;
 	ENGINE_API int32 CompileInput(class FMaterialCompiler* Compiler, int32 AssignedInterpolatorIndex);
 	virtual void GetCaption(TArray<FString>& OutCaptions) const override;
-	virtual uint32 GetInputType(int32 InputIndex) override { return MCT_Float4; }
+	virtual uint32 GetInputType(int32 InputIndex) override;
+	virtual uint32 GetOutputType(int32 OutputIndex) override;
 	virtual bool HasCustomSourceOutput() override { return true; }
 
 	virtual FExpressionInput* GetInput(int32 InputIndex) override;
 	virtual FName GetInputName(int32 InputIndex) const override { return TEXT("VS"); }
+
+	virtual bool GenerateHLSLExpression(FMaterialHLSLGenerator& Generator, UE::HLSLTree::FScope& Scope, int32 OutputIndex, UE::HLSLTree::FExpression const*& OutExpression) const override;
 #endif
 
 	virtual FString GetFunctionName() const override { return TEXT("VertexInterpolator"); }
diff --git a/Engine/Source/Runtime/Engine/Private/HLSLTree/HLSLTreeCommon.cpp b/Engine/Source/Runtime/Engine/Private/HLSLTree/HLSLTreeCommon.cpp
--- a/Engine/Source/Runtime/Engine/Private/HLSLTree/HLSLTreeCommon.cpp
+++ b/Engine/Source/Runtime/Engine/Private/HLSLTree/HLSLTreeCommon.cpp
@@ -1016,16 +1016,6 @@ void FStatementBreak::EmitPreshader(FEmitContext& Context, FEmitScope& Scope, co
 	LoopScope->BreakLabel = OutPreshader.WriteJump(Shader::EPreshaderOpcode::Jump);
 }
 
-bool FStatementReturn::Prepare(FEmitContext& Context, FEmitScope& Scope) const
-{
-	return true;
-}
-
-void FStatementReturn::EmitShader(FEmitContext& Context, FEmitScope& Scope) const
-{
-	Context.EmitStatement(Scope, TEXT("return %;"), Expression->GetValueShader(Context, Scope));
-}
-
 bool FStatementIf::Prepare(FEmitContext& Context, FEmitScope& Scope) const
 {
 	const FPreparedType& ConditionType = Context.PrepareExpression(ConditionExpression, Scope, ERequestedType::Scalar);
diff --git a/Engine/Source/Runtime/Engine/Private/HLSLTree/HLSLTreeEmit.cpp b/Engine/Source/Runtime/Engine/Private/HLSLTree/HLSLTreeEmit.cpp
--- a/Engine/Source/Runtime/Engine/Private/HLSLTree/HLSLTreeEmit.cpp
+++ b/Engine/Source/Runtime/Engine/Private/HLSLTree/HLSLTreeEmit.cpp
@@ -859,6 +859,11 @@ void WriteMaterialUniformAccess(Shader::EValueComponentType ComponentType, uint3
 
 FEmitShaderExpression* FEmitContext::EmitPreshaderOrConstant(FEmitScope& Scope, const FRequestedType& RequestedType, const Shader::FType& ResultType, const FExpression* Expression)
 {
+	if (RequestedType.IsVoid())
+	{
+		return EmitConstantZero(Scope, ResultType);
+	}
+
 	Shader::FPreshaderData LocalPreshader;
 	Shader::FType Type = ResultType;
 	{
diff --git a/Engine/Source/Runtime/Engine/Private/Materials/MaterialCachedData.cpp b/Engine/Source/Runtime/Engine/Private/Materials/MaterialCachedData.cpp
--- a/Engine/Source/Runtime/Engine/Private/Materials/MaterialCachedData.cpp
+++ b/Engine/Source/Runtime/Engine/Private/Materials/MaterialCachedData.cpp
@@ -476,7 +476,7 @@ void PrepareHLSLTree(UE::HLSLTree::FEmitContext& EmitContext,
 	EmitContext.ShaderFrequency = ShaderFrequency;
 	EmitContext.bUseAnalyticDerivatives = true; // We want to consider expressions used for analytic derivatives
 	EmitContext.bMarkLiveValues = false;
-	FEmitScope* EmitResultScope = EmitContext.PrepareScope(&CachedTree.GetResultStatement()->GetParentScope());
+	FEmitScope* EmitResultScope = EmitContext.PrepareScope(CachedTree.GetResultScope());
 
 	FRequestedType RequestedAttributesType;
 	CachedTree.SetRequestedFields(ShaderFrequency, RequestedAttributesType);
diff --git a/Engine/Source/Runtime/Engine/Private/Materials/MaterialCachedHLSLTree.cpp b/Engine/Source/Runtime/Engine/Private/Materials/MaterialCachedHLSLTree.cpp
--- a/Engine/Source/Runtime/Engine/Private/Materials/MaterialCachedHLSLTree.cpp
+++ b/Engine/Source/Runtime/Engine/Private/Materials/MaterialCachedHLSLTree.cpp
@@ -53,7 +53,15 @@ bool FMaterialCachedHLSLTree::GenerateTree(UMaterial* Material, const FMaterialL
 {
 	const EMaterialShadingModel DefaultShadingModel = Material->GetShadingModels().GetFirstShadingModel();
 
-	Material->GetAllCustomOutputExpressions(MaterialCustomOutputs);
+	for (UMaterialExpression* Expression : Material->Expressions)
+	{
+		UMaterialExpressionCustomOutput* CustomOutput = Cast<UMaterialExpressionCustomOutput>(Expression);
+		// We don't want anything with HasCustomSourceOutput() here (VertexInterpolators)
+		if (CustomOutput && !CustomOutput->HasCustomSourceOutput())
+		{
+			MaterialCustomOutputs.Add(CustomOutput);
+		}
+	}
 
 	TArray<UE::Shader::FStructFieldInitializer, TInlineAllocator<MP_MAX + 16>> MaterialAttributeFields;
 
diff --git a/Engine/Source/Runtime/Engine/Private/Materials/MaterialExpressionHLSL.cpp b/Engine/Source/Runtime/Engine/Private/Materials/MaterialExpressionHLSL.cpp
--- a/Engine/Source/Runtime/Engine/Private/Materials/MaterialExpressionHLSL.cpp
+++ b/Engine/Source/Runtime/Engine/Private/Materials/MaterialExpressionHLSL.cpp
@@ -54,6 +54,7 @@
 #include "Materials/MaterialExpressionSceneTexture.h"
 #include "Materials/MaterialExpressionSceneDepth.h"
 #include "Materials/MaterialExpressionNoise.h"
+#include "Materials/MaterialExpressionVertexInterpolator.h"
 #include "Materials/MaterialExpressionFunctionInput.h"
 #include "Materials/MaterialExpressionFunctionOutput.h"
 #include "Materials/MaterialExpressionMaterialFunctionCall.h"
@@ -1520,6 +1521,19 @@ bool UMaterialExpressionBreakMaterialAttributes::GenerateHLSLExpression(FMateria
 	return true;
 }
 
+bool UMaterialExpressionVertexInterpolator::GenerateHLSLExpression(FMaterialHLSLGenerator& Generator, UE::HLSLTree::FScope& Scope, int32 OutputIndex, UE::HLSLTree::FExpression const*& OutExpression) const
+{
+	using namespace UE::HLSLTree;
+	const FExpression* VertexExpression = Input.AcquireHLSLExpression(Generator, Scope);
+	if (!VertexExpression)
+	{
+		return false;
+	}
+
+	OutExpression = Generator.GetTree().NewExpression<Material::FExpressionVertexInterpolator>(VertexExpression);
+	return true;
+}
+
 bool UMaterialExpressionFunctionOutput::GenerateHLSLExpression(FMaterialHLSLGenerator& Generator, UE::HLSLTree::FScope& Scope, int32 OutputIndex, UE::HLSLTree::FExpression const*& OutExpression) const
 {
 	// This should only be called when editing/previewing the function directly
diff --git a/Engine/Source/Runtime/Engine/Private/Materials/MaterialExpressions.cpp b/Engine/Source/Runtime/Engine/Private/Materials/MaterialExpressions.cpp
--- a/Engine/Source/Runtime/Engine/Private/Materials/MaterialExpressions.cpp
+++ b/Engine/Source/Runtime/Engine/Private/Materials/MaterialExpressions.cpp
@@ -1566,6 +1566,20 @@ bool UMaterialExpression::ContainsInputLoopInternal(TArray<FMaterialExpressionKe
 
 	return false;
 }
+
+bool UMaterialExpression::IsUsingNewHLSLGenerator() const
+{
+	if (Material)
+	{
+		return Material->IsUsingNewHLSLGenerator();
+	}
+	if (Function)
+	{
+		return Function->IsUsingNewHLSLGenerator();
+	}
+	return false;
+}
+
 #endif // WITH_EDITOR
 
 UMaterialExpressionTextureBase::UMaterialExpressionTextureBase(const FObjectInitializer& ObjectInitializer)
@@ -18874,6 +18888,18 @@ FExpressionInput* UMaterialExpressionVertexInterpolator::GetInput(int32 InputInd
 {
 	return &Input;
 }
+
+uint32 UMaterialExpressionVertexInterpolator::GetInputType(int32 InputIndex)
+{
+	// New HLSL generator allows struct interpolators
+	return IsUsingNewHLSLGenerator() ? MCT_Unknown : MCT_Float4;
+}
+
+uint32 UMaterialExpressionVertexInterpolator::GetOutputType(int32 OutputIndex)
+{
+	return IsUsingNewHLSLGenerator() ? MCT_Unknown : UMaterialExpression::GetOutputType(OutputIndex);
+}
+
 #endif // WITH_EDITOR
 
 ///////////////////////////////////////////////////////////////////////////////
diff --git a/Engine/Source/Runtime/Engine/Private/Materials/MaterialHLSLEmitter.cpp b/Engine/Source/Runtime/Engine/Private/Materials/MaterialHLSLEmitter.cpp
--- a/Engine/Source/Runtime/Engine/Private/Materials/MaterialHLSLEmitter.cpp
+++ b/Engine/Source/Runtime/Engine/Private/Materials/MaterialHLSLEmitter.cpp
@@ -52,7 +52,8 @@ static FString GenerateMaterialTemplateHLSL(EShaderPlatform ShaderPlatform,
 	const TCHAR* SharedShaderCode,
 	const TCHAR* VertexShaderCode,
 	const TCHAR* PixelShaderCodePhase0[2],
-	const TCHAR* PixelShaderCodePhase1[2])
+	const TCHAR* PixelShaderCodePhase1[2],
+	FMaterialCompilationOutput& OutCompilationOutput)
 {
 	using namespace UE::HLSLTree;
 
@@ -82,23 +83,31 @@ static FString GenerateMaterialTemplateHLSL(EShaderPlatform ShaderPlatform,
 
 	FLazyPrintf LazyPrintf(*MaterialTemplateSource);
 
-	uint32 NumTexCoords = 0u;
-	for (int32 TexCoordIndex = Material::MaxNumTexCoords - 1; TexCoordIndex >= 0; --TexCoordIndex)
+	uint32 NumVertexTexCoords = 0u;
+	uint32 NumPixelTexCoords = 0u;
+	for (int32 TexCoordIndex = 0; TexCoordIndex < Material::MaxNumTexCoords; ++TexCoordIndex)
 	{
-		if (EmitMaterialData.IsExternalInputUsed(SF_Pixel, Material::MakeInputTexCoord(TexCoordIndex)))
+		const Material::EExternalInput TexCoordInput = Material::MakeInputTexCoord(TexCoordIndex);
+		if (EmitMaterialData.IsExternalInputUsed(SF_Vertex, TexCoordInput))
 		{
-			NumTexCoords = TexCoordIndex + 1;
-			break;
+			NumVertexTexCoords = TexCoordIndex + 1;
+		}
+		if (EmitMaterialData.IsExternalInputUsed(SF_Pixel, TexCoordInput))
+		{
+			NumPixelTexCoords = TexCoordIndex + 1;
 		}
 	}
 
-	const uint32 NumUserVertexTexCoords = NumTexCoords;
-	const uint32 NumUserTexCoords = NumTexCoords;
-	const uint32 NumCustomVectors = 0u;
-	const uint32 NumTexCoordVectors = NumTexCoords;
+	NumVertexTexCoords = FMath::Max(NumVertexTexCoords, NumPixelTexCoords);
+
+	const uint32 NumCustomVectors = (EmitMaterialData.NumInterpolatorComponents + 1) / 2;
+	const uint32 NumTexCoordVectors = NumPixelTexCoords + NumCustomVectors;
+
+	OutCompilationOutput.NumUsedUVScalars = NumPixelTexCoords * 2;
+	OutCompilationOutput.NumUsedCustomInterpolatorScalars = EmitMaterialData.NumInterpolatorComponents;
 
-	LazyPrintf.PushParam(*FString::Printf(TEXT("%u"), NumUserVertexTexCoords));
-	LazyPrintf.PushParam(*FString::Printf(TEXT("%u"), NumUserTexCoords));
+	LazyPrintf.PushParam(*FString::Printf(TEXT("%u"), NumVertexTexCoords));
+	LazyPrintf.PushParam(*FString::Printf(TEXT("%u"), NumPixelTexCoords));
 	LazyPrintf.PushParam(*FString::Printf(TEXT("%u"), NumCustomVectors));
 	LazyPrintf.PushParam(*FString::Printf(TEXT("%u"), NumTexCoordVectors));
 
@@ -229,73 +238,25 @@ static FString GenerateMaterialTemplateHLSL(EShaderPlatform ShaderPlatform,
 	LazyPrintf.PushParam(TEXT("return Parameters.MaterialAttributes.ClearCoatRoughness"));
 
 	// Print custom texture coordinate assignments, should be fine with regular derivatives
-	FString CustomUVAssignments;
-
-	int32 LastProperty = -1;
-	for (uint32 CustomUVIndex = 0; CustomUVIndex < NumUserTexCoords; CustomUVIndex++)
 	{
-		//if (bEnableExecutionFlow)
+		FString CustomUVAssignments;
+		for (uint32 CustomUVIndex = 0; CustomUVIndex < NumPixelTexCoords; CustomUVIndex++)
 		{
 			const FString AttributeName = FMaterialAttributeDefinitionMap::GetAttributeName((EMaterialProperty)(MP_CustomizedUVs0 + CustomUVIndex));
 			CustomUVAssignments += FString::Printf(TEXT("\tOutTexCoords[%u] = Parameters.MaterialAttributes.%s;") LINE_TERMINATOR, CustomUVIndex, *AttributeName);
 		}
-		/*else
-		{
-			if (CustomUVIndex == 0)
-			{
-				CustomUVAssignments += DerivativeVariations[BaseDerivativeVariation].TranslatedCodeChunkDefinitions[MP_CustomizedUVs0 + CustomUVIndex];
-			}
-
-			if (DerivativeVariations[BaseDerivativeVariation].TranslatedCodeChunkDefinitions[MP_CustomizedUVs0 + CustomUVIndex].Len() > 0)
-			{
-				if (LastProperty >= 0)
-				{
-					check(DerivativeVariations[BaseDerivativeVariation].TranslatedCodeChunkDefinitions[LastProperty].Len() == DerivativeVariations[BaseDerivativeVariation].TranslatedCodeChunkDefinitions[MP_CustomizedUVs0 + CustomUVIndex].Len());
-				}
-				LastProperty = MP_CustomizedUVs0 + CustomUVIndex;
-			}
-			CustomUVAssignments += FString::Printf(TEXT("\tOutTexCoords[%u] = %s;") LINE_TERMINATOR, CustomUVIndex, *DerivativeVariations[BaseDerivativeVariation].TranslatedCodeChunks[MP_CustomizedUVs0 + CustomUVIndex]);
-		}*/
+		LazyPrintf.PushParam(*CustomUVAssignments);
 	}
 
-	LazyPrintf.PushParam(*CustomUVAssignments);
-
 	// Print custom vertex shader interpolator assignments
-	FString CustomInterpolatorAssignments;
-
-	/*for (UMaterialExpressionVertexInterpolator* Interpolator : CustomVertexInterpolators)
 	{
-		if (Interpolator->InterpolatorOffset != INDEX_NONE)
+		FString CustomInterpolatorAssignments;
+		for (uint32 InterpolatorIndex = 0; InterpolatorIndex < NumCustomVectors; ++InterpolatorIndex)
 		{
-			check(Interpolator->InterpolatorIndex != INDEX_NONE);
-			check(Interpolator->InterpolatedType & MCT_Float);
-
-			const EMaterialValueType Type = Interpolator->InterpolatedType == MCT_Float ? MCT_Float1 : Interpolator->InterpolatedType;
-			const TCHAR* Swizzle[2] = { TEXT("x"), TEXT("y") };
-			const int32 Offset = Interpolator->InterpolatorOffset;
-			const int32 Index = Interpolator->InterpolatorIndex;
-
-			// Note: We reference the UV define directly to avoid having to pre-accumulate UV counts before property translation
-			CustomInterpolatorAssignments += FString::Printf(TEXT("\tOutTexCoords[VERTEX_INTERPOLATOR_%i_TEXCOORDS_X].%s = VertexInterpolator%i(Parameters).x;") LINE_TERMINATOR, Index, Swizzle[Offset % 2], Index);
-
-			if (Type >= MCT_Float2)
-			{
-				CustomInterpolatorAssignments += FString::Printf(TEXT("\tOutTexCoords[VERTEX_INTERPOLATOR_%i_TEXCOORDS_Y].%s = VertexInterpolator%i(Parameters).y;") LINE_TERMINATOR, Index, Swizzle[(Offset + 1) % 2], Index);
-
-				if (Type >= MCT_Float3)
-				{
-					CustomInterpolatorAssignments += FString::Printf(TEXT("\tOutTexCoords[VERTEX_INTERPOLATOR_%i_TEXCOORDS_Z].%s = VertexInterpolator%i(Parameters).z;") LINE_TERMINATOR, Index, Swizzle[(Offset + 2) % 2], Index);
-
-					if (Type == MCT_Float4)
-					{
-						CustomInterpolatorAssignments += FString::Printf(TEXT("\tOutTexCoords[VERTEX_INTERPOLATOR_%i_TEXCOORDS_W].%s = VertexInterpolator%i(Parameters).w;") LINE_TERMINATOR, Index, Swizzle[(Offset + 3) % 2], Index);
-					}
-				}
-			}
+			CustomInterpolatorAssignments += FString::Printf(TEXT("\tOutTexCoords[NUM_MATERIAL_TEXCOORDS + %u] = Parameters.CustomInterpolators[%u];") LINE_TERMINATOR, InterpolatorIndex, InterpolatorIndex);
 		}
-	}*/
-
-	LazyPrintf.PushParam(*CustomInterpolatorAssignments);
+		LazyPrintf.PushParam(*CustomInterpolatorAssignments);
+	}
 
 	//LazyPrintf.PushParam(*MaterialAttributesDefault);
 
@@ -303,25 +264,24 @@ static FString GenerateMaterialTemplateHLSL(EShaderPlatform ShaderPlatform,
 	{
 		FString EvaluateMaterialDeclaration;
 
-		EvaluateMaterialDeclaration += TEXT("FMaterialAttributes EvaluateVertexMaterialAttributesInternal(FMaterialVertexParameters Parameters)" LINE_TERMINATOR);
+		EvaluateMaterialDeclaration += TEXT("void EvaluateVertexMaterialAttributesInternal(in out FMaterialVertexParameters Parameters)" LINE_TERMINATOR);
 		EvaluateMaterialDeclaration += TEXT("{" LINE_TERMINATOR);
 		EvaluateMaterialDeclaration += VertexShaderCode;
 		EvaluateMaterialDeclaration += TEXT("}" LINE_TERMINATOR);
 
-		EvaluateMaterialDeclaration += TEXT("FMaterialAttributes EvaluateVertexMaterialAttributes(FMaterialVertexParameters Parameters)" LINE_TERMINATOR);
+		EvaluateMaterialDeclaration += TEXT("void EvaluateVertexMaterialAttributes(in out FMaterialVertexParameters Parameters)" LINE_TERMINATOR);
 		EvaluateMaterialDeclaration += TEXT("{" LINE_TERMINATOR);
-		EvaluateMaterialDeclaration += TEXT("    FMaterialAttributes Result = EvaluateVertexMaterialAttributesInternal(Parameters);" LINE_TERMINATOR);
-		for (uint32 TexCoordIndex = 0; TexCoordIndex < NumTexCoords; ++TexCoordIndex)
+		EvaluateMaterialDeclaration += TEXT("    EvaluateVertexMaterialAttributesInternal(Parameters);" LINE_TERMINATOR);
+		for (uint32 TexCoordIndex = 0; TexCoordIndex < NumPixelTexCoords; ++TexCoordIndex)
 		{
 			const Material::EExternalInput TexCoordInput = Material::MakeInputTexCoord(TexCoordIndex);
 
 			if (EmitMaterialData.IsExternalInputUsed(SF_Pixel, TexCoordInput) && !EmitMaterialData.IsExternalInputUsed(SF_Vertex, TexCoordInput))
 			{
 				const FString AttributeName = FMaterialAttributeDefinitionMap::GetAttributeName((EMaterialProperty)(MP_CustomizedUVs0 + TexCoordIndex));
-				EvaluateMaterialDeclaration += FString::Printf(TEXT("    Result.%s = Parameters.TexCoords[%d];") LINE_TERMINATOR, *AttributeName, TexCoordIndex);
+				EvaluateMaterialDeclaration += FString::Printf(TEXT("    Parameters.MaterialAttributes.%s = Parameters.TexCoords[%d];") LINE_TERMINATOR, *AttributeName, TexCoordIndex);
 			}
 		}
-		EvaluateMaterialDeclaration += TEXT("   return Result;" LINE_TERMINATOR);
 		EvaluateMaterialDeclaration += TEXT("}" LINE_TERMINATOR);
 
 		FString EvaluateMaterialAttributesPhase0[2];
@@ -330,21 +290,23 @@ static FString GenerateMaterialTemplateHLSL(EShaderPlatform ShaderPlatform,
 		{
 			const TCHAR* DerivativeName = (DerivativeIndex == 0) ? TEXT("HWDerivative") : TEXT("AnalyticDerivative");
 
-			EvaluateMaterialAttributesPhase0[DerivativeIndex] = FString::Printf(TEXT("    FMaterialAttributes MaterialAttributesPhase0 = EvaluatePixelMaterialAttributesPhase0_%s(Parameters);" LINE_TERMINATOR), DerivativeName);
+			EvaluateMaterialAttributesPhase0[DerivativeIndex] =  FString::Printf(TEXT("    FMaterialAttributes MaterialAttributesPhase0;" LINE_TERMINATOR), DerivativeName);
+			EvaluateMaterialAttributesPhase0[DerivativeIndex] += FString::Printf(TEXT("    EvaluatePixelMaterialAttributesPhase0_%s(Parameters, MaterialAttributesPhase0);" LINE_TERMINATOR), DerivativeName);
 
-			EvaluateMaterialDeclaration += FString::Printf(TEXT("FMaterialAttributes EvaluatePixelMaterialAttributesPhase0_%s(in out FMaterialPixelParameters Parameters)" LINE_TERMINATOR), DerivativeName);
+			EvaluateMaterialDeclaration += FString::Printf(TEXT("void EvaluatePixelMaterialAttributesPhase0_%s(in out FMaterialPixelParameters Parameters, out FMaterialAttributes OutResult)" LINE_TERMINATOR), DerivativeName);
 			EvaluateMaterialDeclaration += TEXT("{" LINE_TERMINATOR);
 			EvaluateMaterialDeclaration += PixelShaderCodePhase0[DerivativeIndex];
 			EvaluateMaterialDeclaration += TEXT("}" LINE_TERMINATOR);
 
 			if (PixelShaderCodePhase1[DerivativeIndex])
 			{
-				EvaluateMaterialDeclaration += FString::Printf(TEXT("FMaterialAttributes EvaluatePixelMaterialAttributesPhase1_%s(in out FMaterialPixelParameters Parameters)" LINE_TERMINATOR), DerivativeName);
+				EvaluateMaterialDeclaration += FString::Printf(TEXT("void EvaluatePixelMaterialAttributesPhase1_%s(in out FMaterialPixelParameters Parameters, out FMaterialAttributes OutResult)" LINE_TERMINATOR), DerivativeName);
 				EvaluateMaterialDeclaration += TEXT("{" LINE_TERMINATOR);
 				EvaluateMaterialDeclaration += PixelShaderCodePhase1[DerivativeIndex];
 				EvaluateMaterialDeclaration += TEXT("}" LINE_TERMINATOR);
 
-				EvaluateMaterialAttributesPhase1[DerivativeIndex] = FString::Printf(TEXT("    FMaterialAttributes MaterialAttributesPhase1 = EvaluatePixelMaterialAttributesPhase1_%s(Parameters);" LINE_TERMINATOR), DerivativeName);
+				EvaluateMaterialAttributesPhase1[DerivativeIndex] =  FString::Printf(TEXT("    FMaterialAttributes MaterialAttributesPhase1;" LINE_TERMINATOR), DerivativeName);
+				EvaluateMaterialAttributesPhase1[DerivativeIndex] += FString::Printf(TEXT("    EvaluatePixelMaterialAttributesPhase1_%s(Parameters, MaterialAttributesPhase1);" LINE_TERMINATOR), DerivativeName);
 				EvaluateMaterialAttributesPhase1[DerivativeIndex] += TEXT("    Parameters.MaterialAttributes = MaterialAttributesPhase1;" LINE_TERMINATOR);
 			}
 			else
@@ -841,7 +803,7 @@ bool MaterialEmitHLSL(const FMaterialCompileTargetParameters& InCompilerTarget,
 		EmitContext.ShaderFrequency = SF_Pixel;
 		EmitContext.bUseAnalyticDerivatives = (DerivativeIndex == 1);
 		EmitContext.bMarkLiveValues = false;
-		FEmitScope* EmitResultScope = EmitContext.PrepareScope(&CachedTree->GetResultStatement()->GetParentScope());
+		FEmitScope* EmitResultScope = EmitContext.PrepareScope(CachedTree->GetResultScope());
 
 		// Prepare all fields *except* normal
 		FRequestedType RequestedPixelAttributesType;
@@ -889,16 +851,22 @@ bool MaterialEmitHLSL(const FMaterialCompileTargetParameters& InCompilerTarget,
 			EmitContext.bMarkLiveValues = true;
 			EmitContext.PrepareExpression(CachedTree->GetResultExpression(), *EmitResultScope, RequestedPixelAttributesType);
 
+			FEmitShaderExpression* EmitResultExpression = CachedTree->GetResultExpression()->GetValueShader(EmitContext, *EmitResultScope, RequestedPixelAttributesType, CachedTree->GetMaterialAttributesType());
+			EmitContext.EmitStatement(*EmitResultScope, TEXT("OutResult = %;"), EmitResultExpression);
+
 			PixelCodePhase0[DerivativeIndex] = new(Allocator) FStringBuilderMemstack(Allocator, 128 * 1024);
 			CachedTree->GetTree().EmitShader(EmitContext, *PixelCodePhase0[DerivativeIndex]);
 		}
 		else
 		{
 			// Execute everything *except* normal in phase1
+			FEmitShaderExpression* EmitResultExpression1 = CachedTree->GetResultExpression()->GetValueShader(EmitContext, *EmitResultScope, RequestedPixelAttributesType, CachedTree->GetMaterialAttributesType());
+			EmitContext.EmitStatement(*EmitResultScope, TEXT("OutResult = %;"), EmitResultExpression1);
+
 			PixelCodePhase1[DerivativeIndex] = new(Allocator) FStringBuilderMemstack(Allocator, 128 * 1024);
 			CachedTree->GetTree().EmitShader(EmitContext, *PixelCodePhase1[DerivativeIndex]);
 
-			EmitResultScope = EmitContext.PrepareScope(&CachedTree->GetResultStatement()->GetParentScope());
+			EmitResultScope = EmitContext.PrepareScope(CachedTree->GetResultScope());
 
 			// Prepare code for just the normal
 			FRequestedType RequestedMaterialNormal;
@@ -911,6 +879,9 @@ bool MaterialEmitHLSL(const FMaterialCompileTargetParameters& InCompilerTarget,
 
 			EmitContext.bMarkLiveValues = true;
 			EmitContext.PrepareExpression(CachedTree->GetResultExpression(), *EmitResultScope, RequestedMaterialNormal);
+			
+			FEmitShaderExpression* EmitResultExpression0 = CachedTree->GetResultExpression()->GetValueShader(EmitContext, *EmitResultScope, RequestedMaterialNormal, CachedTree->GetMaterialAttributesType());
+			EmitContext.EmitStatement(*EmitResultScope, TEXT("OutResult = %;"), EmitResultExpression0);
 
 			// Execute the normal in phase0
 			PixelCodePhase0[DerivativeIndex] = new(Allocator) FStringBuilderMemstack(Allocator, 128 * 1024);
@@ -940,8 +911,9 @@ bool MaterialEmitHLSL(const FMaterialCompileTargetParameters& InCompilerTarget,
 		EmitContext.ShaderFrequency = SF_Vertex;
 		EmitContext.bUseAnalyticDerivatives = false;
 		EmitContext.bMarkLiveValues = false;
-		FEmitScope* EmitResultScope = EmitContext.PrepareScope(&CachedTree->GetResultStatement()->GetParentScope());
+		FEmitScope* EmitResultScope = EmitContext.PrepareScope(CachedTree->GetResultScope());
 
+		EmitMaterialData.PrepareInterpolators(EmitContext, *EmitResultScope);
 		const FPreparedType VertexResultType = EmitContext.PrepareExpression(CachedTree->GetResultExpression(), *EmitResultScope, RequestedVertexAttributesType);
 		if (VertexResultType.IsVoid())
 		{
@@ -949,8 +921,13 @@ bool MaterialEmitHLSL(const FMaterialCompileTargetParameters& InCompilerTarget,
 		}
 
 		EmitContext.bMarkLiveValues = true;
+		EmitMaterialData.PrepareInterpolators(EmitContext, *EmitResultScope);
 		EmitContext.PrepareExpression(CachedTree->GetResultExpression(), *EmitResultScope, RequestedVertexAttributesType);
 
+		FEmitShaderExpression* EmitResultExpression = CachedTree->GetResultExpression()->GetValueShader(EmitContext, *EmitResultScope, RequestedVertexAttributesType, CachedTree->GetMaterialAttributesType());
+		EmitContext.EmitStatement(*EmitResultScope, TEXT("Parameters.MaterialAttributes = %;"), EmitResultExpression);
+		EmitMaterialData.EmitInterpolatorStatements(EmitContext, *EmitResultScope);
+
 		bUsesWorldPositionOffset = CachedTree->IsAttributeUsed(EmitContext, *EmitResultScope, VertexResultType, MP_WorldPositionOffset);
 
 		CachedTree->GetTree().EmitShader(EmitContext, VertexCode);
@@ -972,6 +949,7 @@ bool MaterialEmitHLSL(const FMaterialCompileTargetParameters& InCompilerTarget,
 	FStringBuilderMemstack SharedCode(Allocator, 32 * 1024);
 	CachedTree->EmitSharedCode(SharedCode);
 	EmitContext.EmitDeclarationsCode(SharedCode);
+	EmitMaterialData.EmitInterpolatorShader(EmitContext, SharedCode);
 
 	FString MaterialTemplateSource;
 	{
@@ -992,7 +970,8 @@ bool MaterialEmitHLSL(const FMaterialCompileTargetParameters& InCompilerTarget,
 			SharedCode.ToString(),
 			VertexCode.ToString(),
 			InputPixelCodePhase0,
-			InputPixelCodePhase1);
+			InputPixelCodePhase1,
+			OutCompilationOutput);
 	}
 
 	OutMaterialEnvironment = new FSharedShaderCompilerEnvironment();
diff --git a/Engine/Source/Runtime/Engine/Private/Materials/MaterialHLSLGenerator.cpp b/Engine/Source/Runtime/Engine/Private/Materials/MaterialHLSLGenerator.cpp
--- a/Engine/Source/Runtime/Engine/Private/Materials/MaterialHLSLGenerator.cpp
+++ b/Engine/Source/Runtime/Engine/Private/Materials/MaterialHLSLGenerator.cpp
@@ -9,6 +9,7 @@
 #include "Materials/MaterialExpressionFunctionInput.h"
 #include "Materials/MaterialExpressionFunctionOutput.h"
 #include "Materials/MaterialExpressionCustomOutput.h"
+#include "Materials/MaterialExpressionVertexInterpolator.h"
 #include "Materials/MaterialExpressionExecBegin.h"
 #include "Materials/MaterialHLSLTree.h"
 #include "Materials/Material.h"
@@ -87,7 +88,7 @@ bool FMaterialHLSLGenerator::Generate()
 		return Error(TEXT("Missing connection to material output"));
 	}
 
-	if (!CachedTree.GetResultExpression() || !CachedTree.GetResultStatement())
+	if (!CachedTree.GetResultExpression() || !CachedTree.GetResultScope())
 	{
 		return Error(TEXT("Failed to initialize result"));
 	}
@@ -168,7 +169,7 @@ bool FMaterialHLSLGenerator::GenerateResult(UE::HLSLTree::FScope& Scope)
 	}
 	else
 	{
-		check(!CachedTree.ResultStatement);
+		check(!CachedTree.ResultScope);
 		check(!CachedTree.ResultExpression);
 
 		const FExpression* AttributesExpression = nullptr;
@@ -266,10 +267,8 @@ bool FMaterialHLSLGenerator::GenerateResult(UE::HLSLTree::FScope& Scope)
 
 		if (AttributesExpression)
 		{
-			FStatementReturn* ReturnStatement = GetTree().NewStatement<FStatementReturn>(Scope);
-			ReturnStatement->Expression = AttributesExpression;
+			CachedTree.ResultScope = &Scope;
 			CachedTree.ResultExpression = AttributesExpression;
-			CachedTree.ResultStatement = ReturnStatement;
 			bResult = true;
 		}
 
diff --git a/Engine/Source/Runtime/Engine/Private/Materials/MaterialHLSLTree.cpp b/Engine/Source/Runtime/Engine/Private/Materials/MaterialHLSLTree.cpp
--- a/Engine/Source/Runtime/Engine/Private/Materials/MaterialHLSLTree.cpp
+++ b/Engine/Source/Runtime/Engine/Private/Materials/MaterialHLSLTree.cpp
@@ -691,6 +691,225 @@ void FExpressionNoise::EmitValueShader(FEmitContext& Context, FEmitScope& Scope,
 		Parameters.RepeatSize);
 }
 
+void FExpressionVertexInterpolator::ComputeAnalyticDerivatives(FTree& Tree, FExpressionDerivatives& OutResult) const
+{
+	// TODO - we can access interpolator derivatives via TexCoord_DDX/DDY
+}
+
+const FExpression* FExpressionVertexInterpolator::ComputePreviousFrame(FTree& Tree, const FRequestedType& RequestedType) const
+{
+	return Tree.NewExpression<FExpressionVertexInterpolator>(Tree.GetPreviousFrame(VertexExpression, RequestedType));
+}
+
+bool FExpressionVertexInterpolator::PrepareValue(FEmitContext& Context, FEmitScope& Scope, const FRequestedType& RequestedType, FPrepareValueResult& OutResult) const
+{
+	FPreparedType PreparedVertexType;
+	{
+		// Switch to vertex shader while preparing the interpolator expression
+		const EShaderFrequency PrevShaderFrequency = Context.ShaderFrequency;
+		Context.ShaderFrequency = SF_Vertex;
+		PreparedVertexType = Context.PrepareExpression(VertexExpression, Scope, RequestedType);
+		Context.ShaderFrequency = PrevShaderFrequency;
+	}
+
+	if (PreparedVertexType.IsVoid())
+	{
+		return false;
+	}
+
+	// Don't need to allocate an interpolator if we're already in the vertex shader
+	if (Context.ShaderFrequency != SF_Vertex)
+	{
+		// Only allocate an interpolator if we have shader evaluation
+		// Otherwise, just insert the constant/preshader directly into the pixel shader
+		const EExpressionEvaluation Evaluation = PreparedVertexType.GetEvaluation(Scope, RequestedType);
+		if (Evaluation == EExpressionEvaluation::Shader && Context.bMarkLiveValues)
+		{
+			FEmitData& EmitMaterialData = Context.FindData<FEmitData>();
+			EmitMaterialData.AddInterpolator(VertexExpression, RequestedType, PreparedVertexType);
+		}
+	}
+	return OutResult.SetType(Context, RequestedType, PreparedVertexType);
+}
+
+void FExpressionVertexInterpolator::EmitValueShader(FEmitContext& Context, FEmitScope& Scope, const FRequestedType& RequestedType, FEmitValueShaderResult& OutResult) const
+{
+	const FEmitData& EmitMaterialData = Context.FindData<FEmitData>();
+	const int32 InterpolatorIndex = EmitMaterialData.FindInterpolatorIndex(VertexExpression);
+	if (InterpolatorIndex >= 0)
+	{
+		const FVertexInterpolator& Interpolator = EmitMaterialData.VertexInterpolators[InterpolatorIndex];
+
+		FRequestedType RequestedPreshaderType;
+		for (int32 Index = 0; Index < RequestedType.RequestedComponents.Num(); ++Index)
+		{
+			// Requested components that are *not* requested by the interpolator will be either constant or preshader evaluation
+			if (RequestedType.IsComponentRequested(Index) &&
+				!Interpolator.RequestedType.IsComponentRequested(Index))
+			{
+				RequestedPreshaderType.SetComponentRequest(Index);
+			}
+		}
+
+		const Shader::FType LocalType = Interpolator.PreparedType.GetType();
+		FEmitShaderExpression* EmitPreshader = Context.EmitPreshaderOrConstant(Scope, RequestedPreshaderType, LocalType, VertexExpression);
+		OutResult.Code = Context.EmitExpression(Scope, LocalType, TEXT("MaterialVertexInterpolator%(Parameters, %)"), InterpolatorIndex, EmitPreshader);
+	}
+	else
+	{
+		// May not be an interpolator if we're already in the vertex shader
+		check(Context.ShaderFrequency == SF_Vertex);
+		OutResult.Code = VertexExpression->GetValueShader(Context, Scope, RequestedType);
+	}
+}
+
+void FExpressionVertexInterpolator::EmitValuePreshader(FEmitContext& Context, FEmitScope& Scope, const FRequestedType& RequestedType, FEmitValuePreshaderResult& OutResult) const
+{
+	OutResult.Type = VertexExpression->GetValuePreshader(Context, Scope, RequestedType, OutResult.Preshader);
+}
+
+int32 FEmitData::FindInterpolatorIndex(const FExpression* Expression) const
+{
+	for (int32 Index = 0; Index < VertexInterpolators.Num(); ++Index)
+	{
+		if (VertexInterpolators[Index].Expression == Expression)
+		{
+			return Index;
+		}
+	}
+	return INDEX_NONE;
+}
+
+void FEmitData::AddInterpolator(const FExpression* Expression, const FRequestedType& RequestedType, const FPreparedType& PreparedType)
+{
+	int32 Index = FindInterpolatorIndex(Expression);
+	if (Index == INDEX_NONE)
+	{
+		Index = VertexInterpolators.Emplace(Expression);
+	}
+
+	FVertexInterpolator& Interpolator = VertexInterpolators[Index];
+	Interpolator.PreparedType = PreparedType;
+	for (int32 ComponentIndex = 0; ComponentIndex < RequestedType.RequestedComponents.Num(); ++ComponentIndex)
+	{
+		if (RequestedType.IsComponentRequested(Index))
+		{
+			const EExpressionEvaluation ComponentEvaluation = PreparedType.GetComponent(ComponentIndex).Evaluation;
+			if (ComponentEvaluation == EExpressionEvaluation::Shader)
+			{
+				// Only request components that need 'Shader' evaluation
+				Interpolator.RequestedType.SetComponentRequest(ComponentIndex);
+			}
+		}
+	}
+}
+
+void FEmitData::PrepareInterpolators(FEmitContext& Context, FEmitScope& Scope)
+{
+	for (int32 Index = 0; Index < VertexInterpolators.Num(); ++Index)
+	{
+		FVertexInterpolator& Interpolator = VertexInterpolators[Index];
+		Context.PrepareExpression(Interpolator.Expression, Scope, Interpolator.RequestedType);
+	}
+}
+
+void FEmitData::EmitInterpolatorStatements(FEmitContext& Context, FEmitScope& Scope) const
+{
+	for (int32 Index = 0; Index < VertexInterpolators.Num(); ++Index)
+	{
+		const FVertexInterpolator& Interpolator = VertexInterpolators[Index];
+		FEmitShaderExpression* EmitExpression = Interpolator.Expression->GetValueShader(Context, Scope, Interpolator.RequestedType, Interpolator.PreparedType, Interpolator.PreparedType.GetType());
+		Context.EmitStatement(Scope, TEXT("MaterialPackVertexInterpolator%(Parameters, %);"), Index, EmitExpression);
+	}
+}
+
+namespace Private
+{
+void WriteInterpolator(const Shader::FType& Type,
+	const FRequestedType& RequestedType,
+	const TCHAR* FieldName,
+	int32& GlobalComponentIndex,
+	int32& InterpolatorOffset,
+	FStringBuilderBase& ReadCode,
+	FStringBuilderBase& WriteCode)
+{
+	static const TCHAR* ComponentSwizzle[] =
+	{
+		TEXT(".x"),
+		TEXT(".y"),
+		TEXT(".z"),
+		TEXT(".w"),
+	};
+
+	if (Type.IsStruct())
+	{
+		for (const Shader::FStructField& Field : Type.StructType->Fields)
+		{
+			TStringBuilder<128> StructFieldName;
+			StructFieldName.Appendf(TEXT("%s.%s"), FieldName, Field.Name);
+			WriteInterpolator(Field.Type, RequestedType, StructFieldName.ToString(), GlobalComponentIndex, InterpolatorOffset, ReadCode, WriteCode);
+		}
+	}
+	else
+	{
+		const Shader::FValueTypeDescription TypeDesc = Shader::GetValueTypeDescription(Type);
+		for (int32 ComponentIndex = 0; ComponentIndex < TypeDesc.NumComponents; ++ComponentIndex)
+		{
+			const bool bComponentRequested = RequestedType.IsComponentRequested(GlobalComponentIndex++);
+			if (bComponentRequested)
+			{
+				const TCHAR* Swizzle = (TypeDesc.NumComponents > 1) ? ComponentSwizzle[ComponentIndex] : TEXT("");
+				if (TypeDesc.ComponentType == Shader::EValueComponentType::Double)
+				{
+					// Each 'double' component requires 2 slots, 1 for tile and 1 for offset
+					const int32 TileComponentOffset = InterpolatorOffset++;
+					const int32 OffsetComponentOffset = InterpolatorOffset++;
+					ReadCode.Appendf(TEXT("    Value%s.Tile%s = MaterialReadInterpolatorComponent(Parameters, %d);\n"), FieldName, Swizzle, TileComponentOffset);
+					ReadCode.Appendf(TEXT("    Value%s.Offset%s = MaterialReadInterpolatorComponent(Parameters, %d);\n"), FieldName, Swizzle, OffsetComponentOffset);
+					WriteCode.Appendf(TEXT("    MaterialPackInterpolatorComponent(Parameters, %d, Value%s.Tile%s);\n"), TileComponentOffset, FieldName, Swizzle);
+					WriteCode.Appendf(TEXT("    MaterialPackInterpolatorComponent(Parameters, %d, Value%s.Offset%s);\n"), OffsetComponentOffset, FieldName, Swizzle);
+				}
+				else
+				{
+					const int32 ComponentOffset = InterpolatorOffset++;
+					ReadCode.Appendf(TEXT("    Value%s%s = MaterialReadInterpolatorComponent(Parameters, %d);\n"), FieldName, Swizzle, ComponentOffset);
+					WriteCode.Appendf(TEXT("    MaterialPackInterpolatorComponent(Parameters, %d, Value%s%s);\n"), ComponentOffset, FieldName, Swizzle);
+				}
+			}
+		}
+	}
+}
+} // namespace Private
+
+void FEmitData::EmitInterpolatorShader(FEmitContext& Context, FStringBuilderBase& OutCode)
+{
+	int32 InterpolatorOffset = 0;
+	for (int32 Index = 0; Index < VertexInterpolators.Num(); ++Index)
+	{
+		const FVertexInterpolator& Interpolator = VertexInterpolators[Index];
+		const Shader::FType Type = Interpolator.PreparedType.GetType();
+		const TCHAR* TypeName = Type.GetName();
+
+		TStringBuilder<1024> ReadCode;
+		TStringBuilder<1024> WriteCode;
+		ReadCode.Appendf(TEXT("    %s Value = DefaultValue;\n"), TypeName);
+		int32 ComponentIndex = 0;
+		Private::WriteInterpolator(Type, Interpolator.RequestedType, TEXT(""), ComponentIndex, InterpolatorOffset, ReadCode, WriteCode);
+		check(ComponentIndex == Type.GetNumComponents());
+		ReadCode.Append(TEXT("    return Value;\n"));
+
+		OutCode.Appendf(TEXT("%s MaterialVertexInterpolator%d(FMaterialPixelParameters Parameters, %s DefaultValue)\n{\n"), TypeName, Index, TypeName);
+		OutCode.Append(ReadCode.ToView());
+		OutCode.Append(TEXT("}\n"));
+
+		OutCode.Appendf(TEXT("void MaterialPackVertexInterpolator%d(in out FMaterialVertexParameters Parameters, %s Value)\n{\n"), Index, TypeName);
+		OutCode.Append(WriteCode.ToView());
+		OutCode.Append(TEXT("}\n"));
+	}
+
+	NumInterpolatorComponents = InterpolatorOffset;
+}
+
 } // namespace UE::HLSLTree::Material
 
 #endif // WITH_EDITOR
diff --git a/Engine/Source/Runtime/Engine/Private/Materials/MaterialHLSLTree.h b/Engine/Source/Runtime/Engine/Private/Materials/MaterialHLSLTree.h
--- a/Engine/Source/Runtime/Engine/Private/Materials/MaterialHLSLTree.h
+++ b/Engine/Source/Runtime/Engine/Private/Materials/MaterialHLSLTree.h
@@ -341,6 +341,34 @@ public:
 	virtual void EmitValueShader(FEmitContext& Context, FEmitScope& Scope, const FRequestedType& RequestedType, FEmitValueShaderResult& OutResult) const override;
 };
 
+/**
+ * Mechanics of passing values from PS->VS are specific to material shaders, to VertexInterpolator is included in the Material module
+ * Possible that parts of this could be moved to the common HLSLTree module, if interpolators are needed by another system
+ */
+class FExpressionVertexInterpolator : public FExpression
+{
+public:
+	FExpressionVertexInterpolator(const FExpression* InVertexExpression) : VertexExpression(InVertexExpression) {}
+
+	const FExpression* VertexExpression;
+
+	virtual void ComputeAnalyticDerivatives(FTree& Tree, FExpressionDerivatives& OutResult) const override;
+	virtual const FExpression* ComputePreviousFrame(FTree& Tree, const FRequestedType& RequestedType) const override;
+	virtual bool PrepareValue(FEmitContext& Context, FEmitScope& Scope, const FRequestedType& RequestedType, FPrepareValueResult& OutResult) const override;
+	virtual void EmitValueShader(FEmitContext& Context, FEmitScope& Scope, const FRequestedType& RequestedType, FEmitValueShaderResult& OutResult) const override;
+	virtual void EmitValuePreshader(FEmitContext& Context, FEmitScope& Scope, const FRequestedType& RequestedType, FEmitValuePreshaderResult& OutResult) const override;
+};
+
+struct FVertexInterpolator
+{
+	FVertexInterpolator() = default;
+	FVertexInterpolator(const FExpression* InExpression) : Expression(InExpression) {}
+
+	const FExpression* Expression = nullptr;
+	FRequestedType RequestedType;
+	FPreparedType PreparedType;
+};
+
 class FEmitData
 {
 public:
@@ -355,11 +383,20 @@ public:
 	const FStaticParameterSet* StaticParameters = nullptr;
 	FMaterialCachedExpressionData* CachedExpressionData = nullptr;
 	TMap<Shader::FValue, uint32> DefaultUniformValues;
+	TArray<FVertexInterpolator, TInlineAllocator<8>> VertexInterpolators;
 	TBitArray<> ExternalInputMask[SF_NumFrequencies];
 	FMaterialShadingModelField ShadingModelsFromCompilation;
+	int32 NumInterpolatorComponents = 0;
 
 	bool IsExternalInputUsed(EShaderFrequency Frequency, EExternalInput Input) const { return ExternalInputMask[Frequency][(int32)Input]; }
 	bool IsExternalInputUsed(EExternalInput Input) const { return IsExternalInputUsed(SF_Vertex, Input) || IsExternalInputUsed(SF_Pixel, Input); }
+
+	int32 FindInterpolatorIndex(const FExpression* Expression) const;
+	void AddInterpolator(const FExpression* Expression, const FRequestedType& RequestedType, const FPreparedType& PreparedType);
+
+	void PrepareInterpolators(FEmitContext& Context, FEmitScope& Scope);
+	void EmitInterpolatorStatements(FEmitContext& Context, FEmitScope& Scope) const;
+	void EmitInterpolatorShader(FEmitContext& Context, FStringBuilderBase& OutCode);
 };
 
 } // namespace UE::HLSLTree::Material
diff --git a/Engine/Source/Runtime/Engine/Public/HLSLTree/HLSLTree.h b/Engine/Source/Runtime/Engine/Public/HLSLTree/HLSLTree.h
--- a/Engine/Source/Runtime/Engine/Public/HLSLTree/HLSLTree.h
+++ b/Engine/Source/Runtime/Engine/Public/HLSLTree/HLSLTree.h
@@ -295,6 +295,7 @@ public:
 
 	FPreparedComponent GetComponent(int32 Index) const;
 	Shader::FComponentBounds GetComponentBounds(int32 Index) const;
+	inline bool IsComponentRequested(int32 Index) const { return GetComponent(Index).IsRequested(); }
 
 	void SetComponent(int32 Index, const FPreparedComponent& InComponent);
 	void SetComponentBounds(int32 Index, const Shader::FComponentBounds Bounds);
diff --git a/Engine/Source/Runtime/Engine/Public/HLSLTree/HLSLTreeCommon.h b/Engine/Source/Runtime/Engine/Public/HLSLTree/HLSLTreeCommon.h
--- a/Engine/Source/Runtime/Engine/Public/HLSLTree/HLSLTreeCommon.h
+++ b/Engine/Source/Runtime/Engine/Public/HLSLTree/HLSLTreeCommon.h
@@ -296,15 +296,6 @@ public:
 	const Shader::FStructType* OutputStructType = nullptr;
 };
 
-class FStatementReturn : public FStatement
-{
-public:
-	const FExpression* Expression;
-
-	virtual bool Prepare(FEmitContext& Context, FEmitScope& Scope) const override;
-	virtual void EmitShader(FEmitContext& Context, FEmitScope& Scope) const override;
-};
-
 class FStatementBreak : public FStatement
 {
 public:
diff --git a/Engine/Source/Runtime/Engine/Public/MaterialCachedHLSLTree.h b/Engine/Source/Runtime/Engine/Public/MaterialCachedHLSLTree.h
--- a/Engine/Source/Runtime/Engine/Public/MaterialCachedHLSLTree.h
+++ b/Engine/Source/Runtime/Engine/Public/MaterialCachedHLSLTree.h
@@ -33,7 +33,7 @@ public:
 	UE::HLSLTree::FTree& GetTree() const { return *HLSLTree; }
 
 	const UE::HLSLTree::FExpression* GetResultExpression() const { return ResultExpression; }
-	UE::HLSLTree::FStatement* GetResultStatement() const { return ResultStatement; }
+	UE::HLSLTree::FScope* GetResultScope() const { return ResultScope; }
 
 	const UE::Shader::FStructType* GetMaterialAttributesType() const { return MaterialAttributesType; }
 	const UE::Shader::FValue& GetMaterialAttributesDefaultValue() const { return MaterialAttributesDefaultValue; }
@@ -51,7 +51,7 @@ private:
 	UE::Shader::FStructTypeRegistry TypeRegistry;
 	UE::HLSLTree::FTree* HLSLTree = nullptr;
 	const UE::HLSLTree::FExpression* ResultExpression = nullptr;
-	UE::HLSLTree::FStatement* ResultStatement = nullptr;
+	UE::HLSLTree::FScope* ResultScope = nullptr;
 
 	TArray<UMaterialExpressionCustomOutput*> MaterialCustomOutputs;
 	const UE::Shader::FStructType* MaterialAttributesType = nullptr;
diff --git a/Engine/Source/Runtime/Engine/Public/MaterialHLSLGenerator.h b/Engine/Source/Runtime/Engine/Public/MaterialHLSLGenerator.h
--- a/Engine/Source/Runtime/Engine/Public/MaterialHLSLGenerator.h
+++ b/Engine/Source/Runtime/Engine/Public/MaterialHLSLGenerator.h
@@ -26,6 +26,7 @@ class UMaterialExpression;
 class UMaterialExpressionFunctionInput;
 class UMaterialExpressionFunctionOutput;
 class UMaterialExpressionCustomOutput;
+class UMaterialExpressionVertexInterpolator;
 struct FFunctionExpressionInput;
 class ITargetPlatform;
 enum class EMaterialParameterType : uint8;
