Remove "Within=PlayerController" on the Enhanced Player Input. This allows you to spawn a new player input without an owning player controller, such as with the editor input subsystem.

diff --git a/Engine/Plugins/EnhancedInput/Source/EnhancedInput/Private/EnhancedPlayerInput.cpp b/Engine/Plugins/EnhancedInput/Source/EnhancedInput/Private/EnhancedPlayerInput.cpp
--- a/Engine/Plugins/EnhancedInput/Source/EnhancedInput/Private/EnhancedPlayerInput.cpp
+++ b/Engine/Plugins/EnhancedInput/Source/EnhancedInput/Private/EnhancedPlayerInput.cpp
@@ -179,6 +179,22 @@ bool UEnhancedPlayerInput::InputKey(const FInputKeyParams& Params)
 	return bResult;
 }
 
+float UEnhancedPlayerInput::GetEffectiveTimeDilation() const
+{
+	if (const APlayerController* PC = GetOuterAPlayerController())
+	{
+		return PC->GetActorTimeDilation();
+	}
+	else if (const UWorld* World = GetWorld())
+	{
+		if (const AWorldSettings* WorldSettings = World->GetWorldSettings())
+		{
+			return WorldSettings->GetEffectiveTimeDilation();	
+		}
+	}
+	return 1.0f;
+}
+
 void UEnhancedPlayerInput::ProcessInputStack(const TArray<UInputComponent*>& InputComponentStack, const float DeltaTime, const bool bGamePaused)
 {
 	// We need to grab the down states of all keys before calling Super::ProcessInputStack as it will leave bDownPrevious in the same state as bDown (i.e. this frame, not last).
@@ -213,8 +229,7 @@ void UEnhancedPlayerInput::ProcessInputStack(const TArray<UInputComponent*>& Inp
 	RealTimeDeltaSeconds = CurrentTime - LastFrameTime;
 	
 	// Use non-dilated delta time for processing
-	check(GetOuterAPlayerController());
-	const float Dilation = GetOuterAPlayerController()->GetActorTimeDilation();
+	const float Dilation = GetEffectiveTimeDilation();
 	const float NonDilatedDeltaTime = Dilation != 0.0f ? DeltaTime / Dilation : RealTimeDeltaSeconds;
 
 	// Handle input devices, applying modifiers and triggers
diff --git a/Engine/Plugins/EnhancedInput/Source/EnhancedInput/Private/InputModifiers.cpp b/Engine/Plugins/EnhancedInput/Source/EnhancedInput/Private/InputModifiers.cpp
--- a/Engine/Plugins/EnhancedInput/Source/EnhancedInput/Private/InputModifiers.cpp
+++ b/Engine/Plugins/EnhancedInput/Source/EnhancedInput/Private/InputModifiers.cpp
@@ -224,7 +224,10 @@ FInputActionValue UInputModifierFOVScaling::ModifyRaw_Implementation(const UEnha
 {
 	const APlayerController* PC = PlayerInput->GetOuterAPlayerController();
 
-	check(PC);
+	if (!PC)
+	{
+		return CurrentValue;
+	}
 
 	const float FOVAngle = PC->PlayerCameraManager ? PC->PlayerCameraManager->GetFOVAngle() : 1.f;
 	float Scale = FOVScale;
diff --git a/Engine/Plugins/EnhancedInput/Source/EnhancedInput/Private/InputTriggers.cpp b/Engine/Plugins/EnhancedInput/Source/EnhancedInput/Private/InputTriggers.cpp
--- a/Engine/Plugins/EnhancedInput/Source/EnhancedInput/Private/InputTriggers.cpp
+++ b/Engine/Plugins/EnhancedInput/Source/EnhancedInput/Private/InputTriggers.cpp
@@ -72,9 +72,9 @@ ETriggerState UInputTriggerTimedBase::UpdateState_Implementation(const UEnhanced
 
 float UInputTriggerTimedBase::CalculateHeldDuration(const UEnhancedPlayerInput* const PlayerInput, const float DeltaTime) const
 {
-	if(ensureMsgf(PlayerInput && PlayerInput->GetOuterAPlayerController(), TEXT("No Player Input was given to Calculate with! Returning 1.0")))
+	if (ensureMsgf(PlayerInput, TEXT("No Player Input was given to Calculate with! Returning 1.0")))
 	{
-		const float TimeDilation = PlayerInput->GetOuterAPlayerController()->GetActorTimeDilation();
+		const float TimeDilation = PlayerInput->GetEffectiveTimeDilation();
 	
 		// Calculates the new held duration, applying time dilation if desired
 		return HeldDuration + (!bAffectedByTimeDilation ? DeltaTime : DeltaTime * TimeDilation);
diff --git a/Engine/Plugins/EnhancedInput/Source/EnhancedInput/Public/EnhancedPlayerInput.h b/Engine/Plugins/EnhancedInput/Source/EnhancedInput/Public/EnhancedPlayerInput.h
--- a/Engine/Plugins/EnhancedInput/Source/EnhancedInput/Public/EnhancedPlayerInput.h
+++ b/Engine/Plugins/EnhancedInput/Source/EnhancedInput/Public/EnhancedPlayerInput.h
@@ -18,7 +18,7 @@ enum class EKeyEvent : uint8;
 /**
 * UEnhancedPlayerInput : UPlayerInput extensions for enhanced player input system
 */
-UCLASS(Within = PlayerController, config = Input, transient)
+UCLASS(config = Input, transient)
 class ENHANCEDINPUT_API UEnhancedPlayerInput : public UPlayerInput
 {
 	friend class IEnhancedInputSubsystemInterface;
@@ -42,11 +42,14 @@ public:
 
 	virtual bool InputKey(const FInputKeyParams& Params) override;
 	
-protected:
-
 	// Applies modifiers and triggers without affecting keys read by the base input system
 	virtual void ProcessInputStack(const TArray<UInputComponent*>& InputComponentStack, const float DeltaTime, const bool bGamePaused) override;
 
+	/** Returns the Time Dilation value that is currently effecting this input. */
+	float GetEffectiveTimeDilation() const;
+	
+protected:
+	
 	// Causes key to be consumed if it is affecting an action.
 	virtual bool IsKeyHandledByAction(FKey Key) const override;
 
diff --git a/Engine/Source/Runtime/Engine/Classes/GameFramework/PlayerInput.h b/Engine/Source/Runtime/Engine/Classes/GameFramework/PlayerInput.h
--- a/Engine/Source/Runtime/Engine/Classes/GameFramework/PlayerInput.h
+++ b/Engine/Source/Runtime/Engine/Classes/GameFramework/PlayerInput.h
@@ -388,7 +388,7 @@ struct ENGINE_API FInputKeyParams
  *
  * @see https://docs.unrealengine.com/latest/INT/Gameplay/Input/index.html
  */
-UCLASS(Within=PlayerController, config=Input, transient)
+UCLASS(config=Input, transient)
 class ENGINE_API UPlayerInput : public UObject
 {
 	GENERATED_BODY()
@@ -498,6 +498,10 @@ public:
 	/** Clear the current cached key maps and rebuild from the source arrays. */
 	void ForceRebuildingKeyMaps(const bool bRestoreDefaults = false);
 
+	/** Return's this object casted to a player controller. This can be null if there is no player controller. */
+	UFUNCTION(BlueprintCallable, BlueprintPure, Category = "Input")
+	APlayerController* GetOuterAPlayerController() const;
+
 private:
 
 	/** Runtime struct that caches the list of mappings for a given Action Name and the capturing chord if applicable */
diff --git a/Engine/Source/Runtime/Engine/Private/UserInterface/PlayerInput.cpp b/Engine/Source/Runtime/Engine/Private/UserInterface/PlayerInput.cpp
--- a/Engine/Source/Runtime/Engine/Private/UserInterface/PlayerInput.cpp
+++ b/Engine/Source/Runtime/Engine/Private/UserInterface/PlayerInput.cpp
@@ -93,7 +93,7 @@ void UPlayerInput::PostInitProperties()
 void UPlayerInput::FlushPressedKeys()
 {
 	APlayerController* PlayerController = GetOuterAPlayerController();
-	ULocalPlayer* LocalPlayer = Cast<ULocalPlayer>(PlayerController->Player);
+	ULocalPlayer* LocalPlayer = PlayerController ? Cast<ULocalPlayer>(PlayerController->Player) : nullptr;
 	if (LocalPlayer != nullptr)
 	{
 		TArray<FKey> PressedKeys;
@@ -749,6 +749,11 @@ void UPlayerInput::ForceRebuildingKeyMaps(const bool bRestoreDefaults)
 	bKeyMapsBuilt = false;
 }
 
+APlayerController* UPlayerInput::GetOuterAPlayerController() const
+{
+	return Cast<APlayerController>(GetOuter());
+}
+
 void UPlayerInput::ConditionalBuildKeyMappings_Internal() const
 {
 	if (ActionKeyMap.Num() == 0)
@@ -1108,8 +1113,10 @@ void UPlayerInput::ProcessInputStack(const TArray<UInputComponent*>& InputCompon
 	// in this processing batch
 
 	APlayerController* PlayerController = GetOuterAPlayerController();
-
-	PlayerController->PreProcessInput(DeltaTime, bGamePaused);
+	if (PlayerController)
+	{
+		PlayerController->PreProcessInput(DeltaTime, bGamePaused);
+	}
 
 	// copy data from accumulators to the real values
 	for (TMap<FKey,FKeyState>::TIterator It(KeyStateMap); It; ++It)
@@ -1447,8 +1454,11 @@ void UPlayerInput::ProcessInputStack(const TArray<UInputComponent*>& InputCompon
 		}
 	}
 
-	PlayerController->PostProcessInput(DeltaTime, bGamePaused);
-
+	if (PlayerController)
+	{
+		PlayerController->PostProcessInput(DeltaTime, bGamePaused);
+	}
+	
 	FinishProcessingPlayerInput();
 	AxisDelegates.Reset();
 	VectorAxisDelegates.Reset();
@@ -1899,7 +1909,7 @@ float UPlayerInput::MassageAxisInput(FKey Key, float RawValue)
 
 		// Take FOV into account (lower FOV == less sensitivity).
 		APlayerController const* const PlayerController = GetOuterAPlayerController();
-		float const FOVScale = (DefaultInputSettings->bEnableFOVScaling && PlayerController->PlayerCameraManager) ? (DefaultInputSettings->FOVScale*PlayerController->PlayerCameraManager->GetFOVAngle()) : 1.0f;
+		float const FOVScale = (DefaultInputSettings->bEnableFOVScaling && PlayerController && PlayerController->PlayerCameraManager) ? (DefaultInputSettings->FOVScale*PlayerController->PlayerCameraManager->GetFOVAngle()) : 1.0f;
 		NewVal *= FOVScale;
 
 		// debug
@@ -2221,7 +2231,10 @@ void UPlayerInput::SetBind(FName BindName, const FString& Command)
 
 class UWorld* UPlayerInput::GetWorld() const
 {
-	check(GetOuterAPlayerController());
-	UWorld* World = GetOuterAPlayerController()->GetWorld();
-	return World;
+	if (APlayerController* PC = GetOuterAPlayerController())
+	{
+		return PC->GetWorld();	
+	}
+	
+	return Super::GetWorld();
 }
