Added central FNetPing class for handling NetConnection ping tracking/calculation, for various types of ping (RoundTrip/ICMP/UDP), and associated analytics/stats.

diff --git a/Engine/Source/Runtime/Engine/Classes/Engine/NetConnection.h b/Engine/Source/Runtime/Engine/Classes/Engine/NetConnection.h
--- a/Engine/Source/Runtime/Engine/Classes/Engine/NetConnection.h
+++ b/Engine/Source/Runtime/Engine/Classes/Engine/NetConnection.h
@@ -47,6 +47,11 @@ class UActorChannel;
 class UChildConnection;
 struct FEncryptionKeyResponse;
 
+namespace UE::Net
+{
+	class FNetPing;
+}
+
 typedef TMap<TWeakObjectPtr<AActor>, UActorChannel*, FDefaultSetAllocator, TWeakObjectPtrMapKeyFuncs<TWeakObjectPtr<AActor>, UActorChannel*>> FActorChannelMap;
 
 namespace NetConnectionHelper
@@ -1660,6 +1665,9 @@ private:
 	/** Whether or not this NetConnection has already received an NMT_CloseReason message */
 	bool bReceivedCloseReason = false;
 
+	/** Ping collection and calculation */
+	TPimplPtr<UE::Net::FNetPing> NetPing;
+
 
 	int32 GetFreeChannelIndex(const FName& ChName) const;
 
@@ -1680,6 +1688,11 @@ public:
 		return &FaultRecovery;
 	}
 
+	UE::Net::FNetPing* GetNetPing()
+	{
+		return NetPing.Get();
+	}
+
 	/**
 	 * Handles parsing/validation and logging of NMT_CloseReason messages
 	 *
diff --git a/Engine/Source/Runtime/Engine/Engine.Build.cs b/Engine/Source/Runtime/Engine/Engine.Build.cs
--- a/Engine/Source/Runtime/Engine/Engine.Build.cs
+++ b/Engine/Source/Runtime/Engine/Engine.Build.cs
@@ -130,6 +130,7 @@ public class Engine : ModuleRules
 				"IntelISPC",
 				"TraceLog",
 				"ColorManagement",
+				"Icmp",
 				"XmlParser"
 			}
 		);
diff --git a/Engine/Source/Runtime/Engine/Private/Analytics/EngineNetAnalytics.cpp b/Engine/Source/Runtime/Engine/Private/Analytics/EngineNetAnalytics.cpp
--- a/Engine/Source/Runtime/Engine/Private/Analytics/EngineNetAnalytics.cpp
+++ b/Engine/Source/Runtime/Engine/Private/Analytics/EngineNetAnalytics.cpp
@@ -33,7 +33,9 @@ bool FNetConnAnalyticsVars::operator == (const FNetConnAnalyticsVars& A) const
 			/** Close results can't be shared - if either are set, equality comparison fails */
 			!CloseReason.IsValid() && !A.CloseReason.IsValid() &&
 			ClientCloseReasons == A.ClientCloseReasons &&
-			RecoveredFaults.OrderIndependentCompareEqual(A.RecoveredFaults);
+			RecoveredFaults.OrderIndependentCompareEqual(A.RecoveredFaults) &&
+			FailedPingAddressesICMP == A.FailedPingAddressesICMP &&
+			FailedPingAddressesUDP == A.FailedPingAddressesUDP;
 }
 
 void FNetConnAnalyticsVars::CommitAnalytics(FNetConnAnalyticsVars& AggregatedData)
@@ -58,6 +60,9 @@ void FNetConnAnalyticsVars::CommitAnalytics(FNetConnAnalyticsVars& AggregatedDat
 		CurData.CloseReason = MoveTemp(CloseReason);
 		CurData.ClientCloseReasons = MoveTemp(ClientCloseReasons);
 	}
+
+	AggregatedData.FailedPingAddressesICMP.Append(FailedPingAddressesICMP);
+	AggregatedData.FailedPingAddressesUDP.Append(FailedPingAddressesUDP);
 }
 
 
@@ -175,15 +180,66 @@ void FNetConnAnalyticsData::SendAnalytics()
 			}
 		}
 
-		auto CounterSort = [](const FReasonCounter& A, const FReasonCounter& B) -> bool
+		{
+			auto CounterSort = [](const FReasonCounter& A, const FReasonCounter& B) -> bool
+				{
+					return A.Counter < B.Counter;
+				};
+
+			HeadCloseReasons.Sort(CounterSort);
+			FullCloseReasons.Sort(CounterSort);
+			HeadClientCloseReasons.Sort(CounterSort);
+			FullClientCloseReasons.Sort(CounterSort);
+		}
+
+
+		struct FFailedPingCounter
+		{
+			FString PingAddress;
+			int32 Counter = 0;
+
+			bool operator == (const FString& A) const
 			{
-				return A.Counter < B.Counter;
-			};
+				return PingAddress == A;
+			}
+		};
+
+		TArray<FFailedPingCounter> ProcessedFailedPingAddressesICMP;
+		TArray<FFailedPingCounter> ProcessedFailedPingAddressesUDP;
+
+		for (const FString& CurFailedAddress : FailedPingAddressesICMP)
+		{
+			int32 ProcessedIdx = ProcessedFailedPingAddressesICMP.IndexOfByKey(CurFailedAddress);
+
+			if (ProcessedIdx == INDEX_NONE)
+			{
+				ProcessedIdx = ProcessedFailedPingAddressesICMP.Add({CurFailedAddress, 0});
+			}
 
-		HeadCloseReasons.Sort(CounterSort);
-		FullCloseReasons.Sort(CounterSort);
-		HeadClientCloseReasons.Sort(CounterSort);
-		FullClientCloseReasons.Sort(CounterSort);
+			ProcessedFailedPingAddressesICMP[ProcessedIdx].Counter++;
+		}
+
+		for (const FString& CurFailedAddress : FailedPingAddressesUDP)
+		{
+			int32 ProcessedIdx = ProcessedFailedPingAddressesUDP.IndexOfByKey(CurFailedAddress);
+
+			if (ProcessedIdx == INDEX_NONE)
+			{
+				ProcessedIdx = ProcessedFailedPingAddressesUDP.Add({CurFailedAddress, 0});
+			}
+
+			ProcessedFailedPingAddressesUDP[ProcessedIdx].Counter++;
+		}
+
+		{
+			auto CounterSort = [](const FFailedPingCounter& A, const FFailedPingCounter& B) -> bool
+				{
+					return A.Counter < B.Counter;
+				};
+
+			ProcessedFailedPingAddressesICMP.Sort(CounterSort);
+			ProcessedFailedPingAddressesUDP.Sort(CounterSort);
+		}
 
 
 		UE_LOG(LogNet, Log, TEXT("NetConnection Analytics:"));
@@ -195,18 +251,24 @@ void FNetConnAnalyticsData::SendAnalytics()
 		UE_LOG(LogNet, Log, TEXT(" - OutOfOrderPacketsDuplicateCount: %llu"), OutOfOrderPacketsDuplicateCount);
 
 
-		UE_LOG(LogNet, Log, TEXT(" - CloseReasons:"));
-
-		for (const FReasonCounter& CurCounter : HeadCloseReasons)
+		if (HeadCloseReasons.Num() > 0)
 		{
-			UE_LOG(LogNet, Log, TEXT("  - %s: %i"), ToCStr(CurCounter.ReasonStr), CurCounter.Counter);
-		}
+			UE_LOG(LogNet, Log, TEXT(" - CloseReasons:"));
 
-		UE_LOG(LogNet, Log, TEXT(" - FullCloseReasons:"));
+			for (const FReasonCounter& CurCounter : HeadCloseReasons)
+			{
+				UE_LOG(LogNet, Log, TEXT("  - %s: %i"), ToCStr(CurCounter.ReasonStr), CurCounter.Counter);
+			}
+		}
 
-		for (const FReasonCounter& CurCounter : FullCloseReasons)
+		if (FullCloseReasons.Num() > 0)
 		{
-			UE_LOG(LogNet, Log, TEXT("  - %s: %i"), ToCStr(CurCounter.ReasonStr), CurCounter.Counter);
+			UE_LOG(LogNet, Log, TEXT(" - FullCloseReasons:"));
+
+			for (const FReasonCounter& CurCounter : FullCloseReasons)
+			{
+				UE_LOG(LogNet, Log, TEXT("  - %s: %i"), ToCStr(CurCounter.ReasonStr), CurCounter.Counter);
+			}
 		}
 
 
@@ -230,12 +292,34 @@ void FNetConnAnalyticsData::SendAnalytics()
 			}
 		}
 
+		if (RecoveredFaults.Num() > 0)
+		{
+			UE_LOG(LogNet, Log, TEXT(" - RecoveredFaults:"));
 
-		UE_LOG(LogNet, Log, TEXT(" - RecoveredFaults:"));
+			for (TMap<FString, int32>::TConstIterator It(RecoveredFaults); It; ++It)
+			{
+				UE_LOG(LogNet, Log, TEXT("  - %s: %i"), ToCStr(It.Key()), It.Value());
+			}
+		}
 
-		for (TMap<FString, int32>::TConstIterator It(RecoveredFaults); It; ++It)
+		if (ProcessedFailedPingAddressesICMP.Num() > 0)
 		{
-			UE_LOG(LogNet, Log, TEXT("  - %s: %i"), ToCStr(It.Key()), It.Value());
+			UE_LOG(LogNet, Log, TEXT(" - FailedPingAddressesICMP:"));
+
+			for (const FFailedPingCounter& CurCounter : ProcessedFailedPingAddressesICMP)
+			{
+				UE_LOG(LogNet, Log, TEXT("  - %s: %i"), ToCStr(CurCounter.PingAddress), CurCounter.Counter);
+			}
+		}
+
+		if (ProcessedFailedPingAddressesUDP.Num() > 0)
+		{
+			UE_LOG(LogNet, Log, TEXT(" - FailedPingAddressesUDP:"));
+
+			for (const FFailedPingCounter& CurCounter : ProcessedFailedPingAddressesUDP)
+			{
+				UE_LOG(LogNet, Log, TEXT("  - %s: %i"), ToCStr(CurCounter.PingAddress), CurCounter.Counter);
+			}
 		}
 
 
@@ -250,8 +334,11 @@ void FNetConnAnalyticsData::SendAnalytics()
 		static const FString EZAttrib_ClientCloseReasons = TEXT("ClientCloseReasons");
 		static const FString EZAttrib_FullClientCloseReasons = TEXT("FullClientCloseReasons");
 		static const FString EZAttrib_RecoveredFaults = TEXT("RecoveredFaults");
+		static const FString EZAttrib_FailedPingICMP = TEXT("FailedPingICMP");
+		static const FString EZAttrib_FailedPingUDP = TEXT("FailedPingUDP");
 		static const FString EZAttrib_Reason = TEXT("Reason");
 		static const FString EZAttrib_Count = TEXT("Count");
+		static const FString EZAttrib_IP = TEXT("IP");
 
 
 		// Json writer subclass to allow us to avoid using a SharedPtr to write basic Json
@@ -266,109 +353,211 @@ void FNetConnAnalyticsData::SendAnalytics()
 
 
 		FString CloseReasonsJsonStr;
-		FAnalyticsJsonWriter CloseReasonsJsonWriter(&CloseReasonsJsonStr);
-
-		CloseReasonsJsonWriter.WriteArrayStart();
 
-		for (const FReasonCounter& CurCounter : HeadCloseReasons)
+		if (HeadCloseReasons.Num() > 0)
 		{
-			CloseReasonsJsonWriter.WriteObjectStart();
+			FAnalyticsJsonWriter CloseReasonsJsonWriter(&CloseReasonsJsonStr);
 
-			CloseReasonsJsonWriter.WriteValue(EZAttrib_Reason, ToCStr(CurCounter.ReasonStr));
-			CloseReasonsJsonWriter.WriteValue(EZAttrib_Count, CurCounter.Counter);
+			CloseReasonsJsonWriter.WriteArrayStart();
 
-			CloseReasonsJsonWriter.WriteObjectEnd();
-		}
+			for (const FReasonCounter& CurCounter : HeadCloseReasons)
+			{
+				CloseReasonsJsonWriter.WriteObjectStart();
 
-		CloseReasonsJsonWriter.WriteArrayEnd();
-		CloseReasonsJsonWriter.Close();
+				CloseReasonsJsonWriter.WriteValue(EZAttrib_Reason, ToCStr(CurCounter.ReasonStr));
+				CloseReasonsJsonWriter.WriteValue(EZAttrib_Count, CurCounter.Counter);
 
-		FString FullCloseReasonsJsonStr;
-		FAnalyticsJsonWriter FullCloseReasonsJsonWriter(&FullCloseReasonsJsonStr);
+				CloseReasonsJsonWriter.WriteObjectEnd();
+			}
 
-		FullCloseReasonsJsonWriter.WriteArrayStart();
+			CloseReasonsJsonWriter.WriteArrayEnd();
+			CloseReasonsJsonWriter.Close();
+		}
+
+		FString FullCloseReasonsJsonStr;
 
-		for (const FReasonCounter& CurCounter : FullCloseReasons)
+		if (FullCloseReasons.Num() > 0)
 		{
-			FullCloseReasonsJsonWriter.WriteObjectStart();
+			FAnalyticsJsonWriter FullCloseReasonsJsonWriter(&FullCloseReasonsJsonStr);
 
-			FullCloseReasonsJsonWriter.WriteValue(EZAttrib_Reason, ToCStr(CurCounter.ReasonStr));
-			FullCloseReasonsJsonWriter.WriteValue(EZAttrib_Count, CurCounter.Counter);
+			FullCloseReasonsJsonWriter.WriteArrayStart();
 
-			FullCloseReasonsJsonWriter.WriteObjectEnd();
-		}
+			for (const FReasonCounter& CurCounter : FullCloseReasons)
+			{
+				FullCloseReasonsJsonWriter.WriteObjectStart();
 
-		FullCloseReasonsJsonWriter.WriteArrayEnd();
-		FullCloseReasonsJsonWriter.Close();
+				FullCloseReasonsJsonWriter.WriteValue(EZAttrib_Reason, ToCStr(CurCounter.ReasonStr));
+				FullCloseReasonsJsonWriter.WriteValue(EZAttrib_Count, CurCounter.Counter);
 
+				FullCloseReasonsJsonWriter.WriteObjectEnd();
+			}
 
-		FString ClientCloseReasonsJsonStr;
-		FAnalyticsJsonWriter ClientCloseReasonsJsonWriter(&ClientCloseReasonsJsonStr);
+			FullCloseReasonsJsonWriter.WriteArrayEnd();
+			FullCloseReasonsJsonWriter.Close();
+		}
 
-		ClientCloseReasonsJsonWriter.WriteArrayStart();
 
-		for (const FReasonCounter& CurCounter : HeadClientCloseReasons)
+		FString ClientCloseReasonsJsonStr;
+
+		if (HeadClientCloseReasons.Num() > 0)
 		{
-			ClientCloseReasonsJsonWriter.WriteObjectStart();
+			FAnalyticsJsonWriter ClientCloseReasonsJsonWriter(&ClientCloseReasonsJsonStr);
 
-			ClientCloseReasonsJsonWriter.WriteValue(EZAttrib_Reason, ToCStr(CurCounter.ReasonStr));
-			ClientCloseReasonsJsonWriter.WriteValue(EZAttrib_Count, CurCounter.Counter);
+			ClientCloseReasonsJsonWriter.WriteArrayStart();
 
-			ClientCloseReasonsJsonWriter.WriteObjectEnd();
-		}
+			for (const FReasonCounter& CurCounter : HeadClientCloseReasons)
+			{
+				ClientCloseReasonsJsonWriter.WriteObjectStart();
+
+				ClientCloseReasonsJsonWriter.WriteValue(EZAttrib_Reason, ToCStr(CurCounter.ReasonStr));
+				ClientCloseReasonsJsonWriter.WriteValue(EZAttrib_Count, CurCounter.Counter);
 
-		ClientCloseReasonsJsonWriter.WriteArrayEnd();
-		ClientCloseReasonsJsonWriter.Close();
+				ClientCloseReasonsJsonWriter.WriteObjectEnd();
+			}
+
+			ClientCloseReasonsJsonWriter.WriteArrayEnd();
+			ClientCloseReasonsJsonWriter.Close();
+		}
 
 		FString FullClientCloseReasonsJsonStr;
-		FAnalyticsJsonWriter FullClientCloseReasonsJsonWriter(&FullClientCloseReasonsJsonStr);
 
-		FullClientCloseReasonsJsonWriter.WriteArrayStart();
+		if (FullClientCloseReasons.Num() > 0)
+		{
+			FAnalyticsJsonWriter FullClientCloseReasonsJsonWriter(&FullClientCloseReasonsJsonStr);
+
+			FullClientCloseReasonsJsonWriter.WriteArrayStart();
+
+			for (const FReasonCounter& CurCounter : FullClientCloseReasons)
+			{
+				FullClientCloseReasonsJsonWriter.WriteObjectStart();
+
+				FullClientCloseReasonsJsonWriter.WriteValue(EZAttrib_Reason, ToCStr(CurCounter.ReasonStr));
+				FullClientCloseReasonsJsonWriter.WriteValue(EZAttrib_Count, CurCounter.Counter);
 
-		for (const FReasonCounter& CurCounter : FullClientCloseReasons)
+				FullClientCloseReasonsJsonWriter.WriteObjectEnd();
+			}
+
+			FullClientCloseReasonsJsonWriter.WriteArrayEnd();
+			FullClientCloseReasonsJsonWriter.Close();
+		}
+
+
+		FString RecoveredFaultsJsonStr;
+
+		if (RecoveredFaults.Num() > 0)
 		{
-			FullClientCloseReasonsJsonWriter.WriteObjectStart();
+			FAnalyticsJsonWriter RecoveredFaultsJsonWriter(&RecoveredFaultsJsonStr);
 
-			FullClientCloseReasonsJsonWriter.WriteValue(EZAttrib_Reason, ToCStr(CurCounter.ReasonStr));
-			FullClientCloseReasonsJsonWriter.WriteValue(EZAttrib_Count, CurCounter.Counter);
+			RecoveredFaultsJsonWriter.WriteArrayStart();
+
+			for (TMap<FString, int32>::TConstIterator It(RecoveredFaults); It; ++It)
+			{
+				RecoveredFaultsJsonWriter.WriteObjectStart();
 
-			FullClientCloseReasonsJsonWriter.WriteObjectEnd();
+				RecoveredFaultsJsonWriter.WriteValue(EZAttrib_Reason, ToCStr(It.Key()));
+				RecoveredFaultsJsonWriter.WriteValue(EZAttrib_Count, It.Value());
+
+				RecoveredFaultsJsonWriter.WriteObjectEnd();
+			}
+
+			RecoveredFaultsJsonWriter.WriteArrayEnd();
+			RecoveredFaultsJsonWriter.Close();
 		}
 
-		FullClientCloseReasonsJsonWriter.WriteArrayEnd();
-		FullClientCloseReasonsJsonWriter.Close();
 
+		FString FailedPingAddressesICMPJsonStr;
 
-		FString RecoveredFaultsJsonStr;
-		FAnalyticsJsonWriter RecoveredFaultsJsonWriter(&RecoveredFaultsJsonStr);
+		if (ProcessedFailedPingAddressesICMP.Num() > 0)
+		{
+			FAnalyticsJsonWriter FailedPingAddressesICMPJsonWriter(&FailedPingAddressesICMPJsonStr);
+
+			FailedPingAddressesICMPJsonWriter.WriteArrayStart();
+
+			for (const FFailedPingCounter& CurCounter : ProcessedFailedPingAddressesICMP)
+			{
+				FailedPingAddressesICMPJsonWriter.WriteObjectStart();
 
-		RecoveredFaultsJsonWriter.WriteArrayStart();
+				FailedPingAddressesICMPJsonWriter.WriteValue(EZAttrib_IP, ToCStr(CurCounter.PingAddress));
+				FailedPingAddressesICMPJsonWriter.WriteValue(EZAttrib_Count, CurCounter.Counter);
 
-		for (TMap<FString, int32>::TConstIterator It(RecoveredFaults); It; ++It)
+				FailedPingAddressesICMPJsonWriter.WriteObjectEnd();
+			}
+
+			FailedPingAddressesICMPJsonWriter.WriteArrayEnd();
+			FailedPingAddressesICMPJsonWriter.Close();
+		}
+
+		FString FailedPingAddressesUDPJsonStr;
+
+		if (ProcessedFailedPingAddressesUDP.Num() > 0)
 		{
-			RecoveredFaultsJsonWriter.WriteObjectStart();
+			FAnalyticsJsonWriter FailedPingAddressesUDPJsonWriter(&FailedPingAddressesUDPJsonStr);
 
-			RecoveredFaultsJsonWriter.WriteValue(EZAttrib_Reason, ToCStr(It.Key()));
-			RecoveredFaultsJsonWriter.WriteValue(EZAttrib_Count, It.Value());
+			FailedPingAddressesUDPJsonWriter.WriteArrayStart();
 
-			RecoveredFaultsJsonWriter.WriteObjectEnd();
+			for (const FFailedPingCounter& CurCounter : ProcessedFailedPingAddressesUDP)
+			{
+				FailedPingAddressesUDPJsonWriter.WriteObjectStart();
+
+				FailedPingAddressesUDPJsonWriter.WriteValue(EZAttrib_IP, ToCStr(CurCounter.PingAddress));
+				FailedPingAddressesUDPJsonWriter.WriteValue(EZAttrib_Count, CurCounter.Counter);
+
+				FailedPingAddressesUDPJsonWriter.WriteObjectEnd();
+			}
+
+			FailedPingAddressesUDPJsonWriter.WriteArrayEnd();
+			FailedPingAddressesUDPJsonWriter.Close();
 		}
 
-		RecoveredFaultsJsonWriter.WriteArrayEnd();
-		RecoveredFaultsJsonWriter.Close();
 
+		// IMPORTANT: Make sure the number of attributes is accurate when adding new analytics - for efficient allocation
+		const int32 NumAttribs = 12;
+		TArray<FAnalyticsEventAttribute> Attrs;
+
+		Attrs.Reserve(NumAttribs);
 
-		AnalyticsProvider->RecordEvent(EZEventName, MakeAnalyticsEventAttributeArray(
+		AppendAnalyticsEventAttributeArray(Attrs,
 			EZAttrib_OutAckOnlyCount, OutAckOnlyCount,
 			EZAttrib_OutKeepAliveCount, OutKeepAliveCount,
 			EZAttrib_OutOfOrderPacketsLostCount, OutOfOrderPacketsLostCount,
 			EZAttrib_OutOfOrderPacketsRecoveredCount, OutOfOrderPacketsRecoveredCount,
-			EZAttrib_OutOfOrderPacketsDuplicateCount, OutOfOrderPacketsDuplicateCount,
-			EZAttrib_CloseReasons, FJsonFragment(MoveTemp(CloseReasonsJsonStr)),
-			EZAttrib_FullCloseReasons, FJsonFragment(MoveTemp(FullCloseReasonsJsonStr)),
-			EZAttrib_ClientCloseReasons, FJsonFragment(MoveTemp(ClientCloseReasonsJsonStr)),
-			EZAttrib_FullClientCloseReasons, FJsonFragment(MoveTemp(FullClientCloseReasonsJsonStr)),
-			EZAttrib_RecoveredFaults, FJsonFragment(MoveTemp(RecoveredFaultsJsonStr))
-		));
+			EZAttrib_OutOfOrderPacketsDuplicateCount, OutOfOrderPacketsDuplicateCount);
+
+		if (ClientCloseReasonsJsonStr.Len() > 0)
+		{
+			AppendAnalyticsEventAttributeArray(Attrs, EZAttrib_CloseReasons, FJsonFragment(MoveTemp(CloseReasonsJsonStr)));
+		}
+
+		if (FullCloseReasonsJsonStr.Len() > 0)
+		{
+			AppendAnalyticsEventAttributeArray(Attrs, EZAttrib_FullCloseReasons, FJsonFragment(MoveTemp(FullCloseReasonsJsonStr)));
+		}
+
+		if (ClientCloseReasonsJsonStr.Len() > 0)
+		{
+			AppendAnalyticsEventAttributeArray(Attrs, EZAttrib_ClientCloseReasons, FJsonFragment(MoveTemp(ClientCloseReasonsJsonStr)));
+		}
+
+		if (FullClientCloseReasonsJsonStr.Len() > 0)
+		{
+			AppendAnalyticsEventAttributeArray(Attrs, EZAttrib_FullClientCloseReasons, FJsonFragment(MoveTemp(FullClientCloseReasonsJsonStr)));
+		}
+
+		if (RecoveredFaultsJsonStr.Len() > 0)
+		{
+			AppendAnalyticsEventAttributeArray(Attrs, EZAttrib_RecoveredFaults, FJsonFragment(MoveTemp(RecoveredFaultsJsonStr)));
+		}
+
+		if (FailedPingAddressesICMPJsonStr.Len() > 0)
+		{
+			AppendAnalyticsEventAttributeArray(Attrs, EZAttrib_FailedPingICMP, FJsonFragment(MoveTemp(FailedPingAddressesICMPJsonStr)));
+		}
+
+		if (FailedPingAddressesUDPJsonStr.Len() > 0)
+		{
+			AppendAnalyticsEventAttributeArray(Attrs, EZAttrib_FailedPingUDP, FJsonFragment(MoveTemp(FailedPingAddressesUDPJsonStr)));
+		}
+
+		AnalyticsProvider->RecordEvent(EZEventName, MoveTemp(Attrs));
 	}
 }
diff --git a/Engine/Source/Runtime/Engine/Private/DataChannel.cpp b/Engine/Source/Runtime/Engine/Private/DataChannel.cpp
--- a/Engine/Source/Runtime/Engine/Private/DataChannel.cpp
+++ b/Engine/Source/Runtime/Engine/Private/DataChannel.cpp
@@ -27,6 +27,7 @@
 #include "Misc/NetworkVersion.h"
 #include "Net/Core/PushModel/PushModel.h"
 #include "HAL/LowLevelMemStats.h"
+#include "Net/NetPing.h"
 
 DEFINE_LOG_CATEGORY(LogNet);
 DEFINE_LOG_CATEGORY(LogNetSubObject);
@@ -104,18 +105,15 @@ static TAutoConsoleVariable<bool> CVarEnableNetInitialSubObjects(
 	TEXT("Enables new SubObjects to set bNetInitial to true to make sure all replicated properties are replicated.")
 );
 
-namespace UE
+namespace UE::Net
 {
-	namespace Net
-	{
-		extern int32 FilterGuidRemapping;
+	extern int32 FilterGuidRemapping;
 
-		static float QueuedBunchTimeoutSeconds = 30.0f;
-		static FAutoConsoleVariableRef CVarQueuedBunchTimeoutSeconds(
-			TEXT("net.QueuedBunchTimeoutSeconds"),
-			QueuedBunchTimeoutSeconds,
-			TEXT("Time in seconds to wait for queued bunches on a channel to flush before logging a warning."));
-	};
+	static float QueuedBunchTimeoutSeconds = 30.0f;
+	static FAutoConsoleVariableRef CVarQueuedBunchTimeoutSeconds(
+		TEXT("net.QueuedBunchTimeoutSeconds"),
+		QueuedBunchTimeoutSeconds,
+		TEXT("Time in seconds to wait for queued bunches on a channel to flush before logging a warning."));
 };
 
 template<typename T>
@@ -1530,13 +1528,17 @@ IMPLEMENT_CONTROL_CHANNEL_MESSAGE(Abort);
 IMPLEMENT_CONTROL_CHANNEL_MESSAGE(PCSwap);
 IMPLEMENT_CONTROL_CHANNEL_MESSAGE(ActorChannelFailure);
 IMPLEMENT_CONTROL_CHANNEL_MESSAGE(DebugText);
+IMPLEMENT_CONTROL_CHANNEL_MESSAGE(NetGUIDAssign);
 IMPLEMENT_CONTROL_CHANNEL_MESSAGE(SecurityViolation);
+IMPLEMENT_CONTROL_CHANNEL_MESSAGE(GameSpecific);
+IMPLEMENT_CONTROL_CHANNEL_MESSAGE(EncryptionAck);
+IMPLEMENT_CONTROL_CHANNEL_MESSAGE(DestructionInfo);
+IMPLEMENT_CONTROL_CHANNEL_MESSAGE(CloseReason);
+IMPLEMENT_CONTROL_CHANNEL_MESSAGE(NetPing);
 IMPLEMENT_CONTROL_CHANNEL_MESSAGE(BeaconWelcome);
 IMPLEMENT_CONTROL_CHANNEL_MESSAGE(BeaconJoin);
 IMPLEMENT_CONTROL_CHANNEL_MESSAGE(BeaconAssignGUID);
 IMPLEMENT_CONTROL_CHANNEL_MESSAGE(BeaconNetGUIDAck);
-IMPLEMENT_CONTROL_CHANNEL_MESSAGE(EncryptionAck);
-IMPLEMENT_CONTROL_CHANNEL_MESSAGE(DestructionInfo);
 
 void UControlChannel::Init( UNetConnection* InConnection, int32 InChannelIndex, EChannelCreateFlags CreateFlags )
 {
@@ -1711,6 +1713,19 @@ void UControlChannel::ReceivedBunch( FInBunch& Bunch )
 		{
 			ReceiveDestructionInfo(Bunch);
 		}
+		else if (MessageType == NMT_NetPing)
+		{
+			ENetPingControlMessage NetPingMessageType;
+			FString MessageStr;
+
+			if (FNetControlMessage<NMT_NetPing>::Receive(Bunch, NetPingMessageType, MessageStr) &&
+				NetPingMessageType <= ENetPingControlMessage::Max)
+			{
+				FNetPing::HandleNetPingControlMessage(Connection, NetPingMessageType, MessageStr);
+			}
+
+			break;
+		}
 		else
 		{
 			// Process control message on client/server connection
@@ -1773,6 +1788,9 @@ void UControlChannel::ReceivedBunch( FInBunch& Bunch )
 				case NMT_CloseReason:
 					FNetControlMessage<NMT_CloseReason>::Discard(Bunch);
 					break;
+				case NMT_NetPing:
+					FNetControlMessage<NMT_NetPing>::Discard(Bunch);
+					break;
 				case NMT_BeaconWelcome:
 					//FNetControlMessage<NMT_BeaconWelcome>::Discard(Bunch);
 					break;
diff --git a/Engine/Source/Runtime/Engine/Private/Net/NetPing.cpp b/Engine/Source/Runtime/Engine/Private/Net/NetPing.cpp
new file mode 100644
--- /dev/null
+++ b/Engine/Source/Runtime/Engine/Private/Net/NetPing.cpp
@@ -0,0 +1,1141 @@
+// Copyright Epic Games, Inc. All Rights Reserved.
+
+#include "Net/NetPing.h"
+#include "Stats/Stats.h"
+#include "Icmp.h"
+#include "Net/DataChannel.h"
+
+
+#if STATS
+DECLARE_STATS_GROUP(TEXT("Ping"), STATGROUP_Ping, STATCAT_Advanced)
+
+// Special values: -1 for ping types which are not enabled, -2 for ping types which have been disabled due to failure/timeout
+DECLARE_DWORD_ACCUMULATOR_STAT(TEXT("Ping (Incl. Frame)"), STAT_Ping_RoundTripInclFrame, STATGROUP_Ping);
+DECLARE_DWORD_ACCUMULATOR_STAT(TEXT("Ping (Excl. Frame)"), STAT_Ping_RoundTripExclFrame, STATGROUP_Ping);
+DECLARE_DWORD_ACCUMULATOR_STAT(TEXT("Ping (ICMP)"), STAT_Ping_ICMP, STATGROUP_Ping);
+DECLARE_DWORD_ACCUMULATOR_STAT(TEXT("Ping Timeouts (ICMP)"), STAT_PingTimeouts_ICMP, STATGROUP_Ping);
+DECLARE_FLOAT_ACCUMULATOR_STAT(TEXT("Ping Timeout Percent (ICMP)"), STAT_PingTimeoutPercent_ICMP, STATGROUP_Ping);
+DECLARE_DWORD_ACCUMULATOR_STAT(TEXT("Ping (UDP QoS)"), STAT_Ping_UDPQoS, STATGROUP_Ping);
+DECLARE_DWORD_ACCUMULATOR_STAT(TEXT("Ping Timeouts (UDP QoS)"), STAT_PingTimeouts_UDPQoS, STATGROUP_Ping);
+DECLARE_FLOAT_ACCUMULATOR_STAT(TEXT("Ping Timeout Percent (UDP QoS)"), STAT_PingTimeoutPercent_UDPQoS, STATGROUP_Ping);
+
+DECLARE_CYCLE_STAT(TEXT("UpdatePing Time"), STAT_Ping_UpdatePingTime, STATGROUP_Ping);
+
+#endif
+
+
+namespace UE::Net
+{
+
+// CVars
+
+static TAutoConsoleVariable<int32> CVarNetPingEnabled(
+	TEXT("net.NetPingEnabled"), 0,
+	TEXT("Whether or not the NetPing ping handling interface is enabled. Used for centralized ping tracking, and ICMP/UDP ping. ")
+	TEXT("(Valid values: 0 = Off, 1 = Enabled for client, 2 = Enabled for server and client, 3 = Enabled for server only)"));
+
+static TAutoConsoleVariable<FString> CVarNetPingTypes(
+	TEXT("net.NetPingTypes"), TEXT(""),
+	TEXT("A comma-delimited list of EPingType pings to enable, and (optionally) the EPingAverageType averaging to apply to the ping ")
+	TEXT("(e.g: \"RoundTrip=None,RoundTripExclFrame=PlayerStateAvg,ICMP=MovingAverage\")."));
+
+static int32 GNetPingICMPInterval = 5.0;
+
+static FAutoConsoleVariableRef CVarNetPingICMPInterval(
+	TEXT("net.NetPingICMPInterval"), GNetPingICMPInterval,
+	TEXT("Specifies the interval (in seconds) for performing ICMP pings."));
+
+static int32 GNetPingUDPInterval = 5.0;
+
+static FAutoConsoleVariableRef CVarNetPingUDPInterval(
+	TEXT("net.NetPingUDPInterval"), GNetPingUDPInterval,
+	TEXT("Specifies the interval (in seconds) for performing UDP pings."));
+
+static int32 GNetPingUDPPort = 22222;
+
+static FAutoConsoleVariableRef CVarNetPingUDPPort(
+	TEXT("net.NetPingUDPPort"), GNetPingUDPPort,
+	TEXT("For 'UDPQoS' ping type, sets the port used for pinging."));
+
+static TAutoConsoleVariable<int32> CVarNetPingTimeoutDisableThreshold(
+	TEXT("net.NetPingTimeoutDisableThreshold"), 15,
+	TEXT("The number of times to send an ICMP/UDP ping when at a failure/timeout rate of 100%, before giving up and disabling pings."));
+
+#if !UE_BUILD_SHIPPING
+static int32 GNetPingDebugDump = 0;
+
+static FAutoConsoleVariableRef CVarNetPingDebugDump(
+	TEXT("net.NetPingDebugDump"), GNetPingDebugDump,
+	TEXT("Whether or not to dump NetPing ping values to log every 5 seconds."));
+#endif
+
+namespace Private
+{
+
+/**
+ * Base class for QoS-derived ping type handling
+ */
+class FNetPingQoSBase
+{
+protected:
+	FNetPingQoSBase() = delete;
+
+	/**
+	 * Base constructor
+	 *
+	 * @param InPingType	The QoS ping type this instance will handle
+	 * @param InOwner		The NetConnection which owns this ping type handler
+	 */
+	FNetPingQoSBase(EPingType InPingType, UNetConnection* InOwner);
+
+	/**
+	 * Resets the ping type handler, discarding any active pings
+	 */
+	void Reset();
+
+	/**
+	 * Handles receiving a QoS ping result.
+	 *
+	 * @param Result	The QoS ping result to handle
+	 */
+	void PingResult(FIcmpEchoResult Result);
+
+
+private:
+	/** The QoS ping type this instance handles */
+	const EPingType PingType;
+
+	/** The NetConnection which owns this ping type handler */
+	TObjectPtr<UNetConnection> Owner;
+
+	/** Whether or not to discard the next ping value result (typically after a reset) */
+	bool bDiscardNextPing = false;
+
+protected:
+	/** Whether or not a ping is in progress */
+	bool bPingInProgress = false;
+
+	/** The timestamp at which the last ping began */
+	double PingTimestamp = 0.0;
+};
+
+/**
+ * Handles QoS ICMP pings
+ */
+class FNetPingICMP : public FNetPingQoSBase, public TSharedFromThis<FNetPingICMP>
+{
+	template <typename ObjectType, ESPMode Mode>
+	friend class SharedPointerInternals::TIntrusiveReferenceController;
+
+private:
+	FNetPingICMP() = delete;
+
+	/**
+	 * Base constructor
+	 *
+	 * @param InOwner	The NetConnection which owns this ping type handler
+	 */
+	FNetPingICMP(UNetConnection* InOwner);
+
+public:
+	/**
+	 * Singleton for retrieving an instance to the sole ICMP ping handler.
+	 *
+	 * NOTE: Only the GameNetDriver client can obtain the ICMP ping handler,
+	 * to prevent conflicting results from multiple active ICMP pings.
+	 * This is a precautionary restriction, it might be possible to support this without conflicting results.
+	 *
+	 * @param InOwner	The NetConnection trying to obtain the ICMP ping handler.
+	 * @return			Returns a shared pointer to the ICMP ping handler, or nullptr if not available.
+	 */
+	static TSharedPtr<FNetPingICMP> GetNetPingICMP(UNetConnection* InOwner);
+
+	/**
+	 * Perform an ICMP ping.
+	 *
+	 * @param CurTimeSeconds	The current time, in seconds.
+	 * @param PingAddress		The address (without port) to ping.
+	 */
+	void Ping(double CurTimeSeconds, FString PingAddress);
+};
+
+/**
+ * Handles QoS UDP pings
+ */
+class FNetPingUDPQoS : public FNetPingQoSBase, public TSharedFromThis<FNetPingUDPQoS>
+{
+	template <typename ObjectType, ESPMode Mode>
+	friend class SharedPointerInternals::TIntrusiveReferenceController;
+
+private:
+	FNetPingUDPQoS() = delete;
+
+	/**
+	 * Base constructor
+	 *
+	 * @param InOwner		The NetConnection which owns this ping type handler
+	 */
+	FNetPingUDPQoS(UNetConnection* InOwner);
+
+public:
+	/**
+	 * Creates a UDPQoS ping handler.
+	 *
+	 * @param InOwner	The NetConnection that will own the UDPQoS ping handler.
+	 * @return			Returns a shared pointer to the UDPQoS ping handler.
+	 */
+	static TSharedPtr<FNetPingUDPQoS> CreateNetPingUDPQoS(UNetConnection* InOwner);
+
+	/**
+	 * Perform a UDPQoS ping.
+	 *
+	 * @param CurTimeSeconds	The current time, in seconds.
+	 * @param PingAddress		The address (with port) to ping.
+	 */
+	void Ping(double CurTimeSeconds, FString PingAddress);
+};
+
+}
+
+/**
+ * FNetPing::FMovingAverageStorage
+ */
+
+FNetPing::FMovingAverageStorage::FMovingAverageStorage()
+{
+	PingAccumulator.SetConsumer(&PingValues);
+}
+
+/**
+ * FNetPing::FPlayerStateAvgStorage
+ */
+
+void FNetPing::FPlayerStateAvgStorage::Reset()
+{
+	PingCount = 0;
+	PingTimeouts = 0;
+
+	PingValues.Reset();
+
+	CurPingBucket = 0;
+
+	for (int32 BucketIdx=0; BucketIdx<UE_ARRAY_COUNT(PingBucket); BucketIdx++)
+	{
+		PingBucket[BucketIdx] = {};
+	}
+
+	CurPingBucketTimestamp = 0.0;
+}
+
+void FNetPing::FPlayerStateAvgStorage::UpdatePing(double TimeVal, float InPing)
+{
+	// APlayerState::UpdatePing copy
+
+	// Limit the size of the ping, to avoid overflowing PingBucket values
+	InPing = FMath::Min(1.1f, InPing);
+
+	PingValues.PeekMeasurement(TimeVal, static_cast<double>(InPing));
+
+	float InPingInMs = InPing * 1000.f;
+
+	if ((TimeVal - CurPingBucketTimestamp) >= 1.0)
+	{
+		// Trigger ping recalculation now, while all buckets are 'full'
+		//	(misses the latest ping update, but averages a full 4 seconds data)
+		RecalculateAvgPing();
+
+		CurPingBucket = (CurPingBucket + 1) % UE_ARRAY_COUNT(PingBucket);
+		CurPingBucketTimestamp = TimeVal;
+
+
+		PingBucket[CurPingBucket].PingSum = FMath::FloorToInt(InPingInMs);
+		PingBucket[CurPingBucket].PingCount = 1;
+	}
+	// Limit the number of pings we accept per-bucket, to avoid overflowing PingBucket values
+	else if (PingBucket[CurPingBucket].PingCount < 7)
+	{
+		PingBucket[CurPingBucket].PingSum += FMath::FloorToInt(InPingInMs);
+		PingBucket[CurPingBucket].PingCount++;
+	}
+}
+
+void FNetPing::FPlayerStateAvgStorage::RecalculateAvgPing()
+{
+	// APlayerState::RecalculateAvgPing copy
+	int32 Sum = 0;
+	int32 Count = 0;
+
+	for (uint8 i=0; i<UE_ARRAY_COUNT(PingBucket); i++)
+	{
+		Sum += PingBucket[i].PingSum;
+		Count += PingBucket[i].PingCount;
+	}
+
+	const double FinalPingSeconds = (Count > 0 ? (static_cast<double>(Sum) / (static_cast<double>(Count) * 1000.0)) : 0.0);
+
+	PingValues.AddSample(FinalPingSeconds);
+}
+
+/**
+ * FNetPing
+ */
+
+FNetPing::FNetPing(UNetConnection* InOwner)
+	: Owner(InOwner)
+{
+}
+
+TPimplPtr<FNetPing> FNetPing::CreateNetPing(UNetConnection* InOwner)
+{
+	if (const UNetDriver* Driver = InOwner != nullptr ? InOwner->Driver : nullptr)
+	{
+		const int32 EnabledVal = CVarNetPingEnabled.GetValueOnAnyThread();
+		const bool bEnabledForClient = EnabledVal == 1 || EnabledVal == 2;
+		const bool bEnabledForServer = EnabledVal == 2 || EnabledVal == 3;
+		const bool bIsServer = Driver->IsServer();
+		const bool bIsGameDriver = (Driver->NetDriverName == NAME_GameNetDriver || Driver->NetDriverName == NAME_PendingNetDriver);
+
+		if (bIsGameDriver && ((!bIsServer && bEnabledForClient) || (bIsServer && bEnabledForServer)))
+		{
+			TPimplPtr<FNetPing> ReturnVal = MakePimpl<FNetPing>(InOwner);
+
+			ReturnVal->Init();
+
+			return ReturnVal;
+		}
+	}
+
+	return nullptr;
+}
+
+void FNetPing::Init()
+{
+	using namespace UE::Net::Private;
+
+	static FString LastPingTypesCVar;
+	static EPingType LastPingTypes = EPingType::None;
+	static EPingAverageType LastPingAverageTypes[static_cast<uint32>(EPingType::Count)] = { EPingAverageType::None };
+
+	FString PingTypesCVar = CVarNetPingTypes.GetValueOnAnyThread();
+
+	if (UNLIKELY(LastPingTypesCVar != PingTypesCVar))
+	{
+		LastPingTypesCVar = PingTypesCVar;
+
+		TArray<FString> PingTypeList;
+		UEnum* PingTypeEnum = StaticEnum<EPingType>();
+		UEnum* PingAvgEnum = StaticEnum<EPingAverageType>();
+
+		PingTypesCVar.ParseIntoArray(PingTypeList, TEXT(","));
+
+		for (FString CurPingAndAvg : PingTypeList)
+		{
+			FString CurPingType;
+			FString CurPingAvg;
+
+			if (CurPingAndAvg.Split(TEXT("="), &CurPingType, &CurPingAvg))
+			{
+				int64 PingTypeVal = PingTypeEnum->GetValueByNameString(CurPingType);
+				int64 PingAvgVal = PingAvgEnum->GetValueByNameString(CurPingAvg);
+
+				if (PingTypeVal != INDEX_NONE && PingAvgVal != INDEX_NONE)
+				{
+					const EPingType FinalPingTypeVal = static_cast<EPingType>(PingTypeVal);
+					const EPingAverageType FinalPingAvgVal = static_cast<EPingAverageType>(PingAvgVal);
+
+					PingTypes |= FinalPingTypeVal;
+
+					PingAverageTypes[PingTypeToIdx(FinalPingTypeVal)] = FinalPingAvgVal;
+				}
+				else
+				{
+					UE_LOG(LogNet, Error, TEXT("FNetPing: Invalid 'net.NetPingTypes' entry: PingType: %s (Valid: %i), ")
+							TEXT("PingAverageType: %s (Valid: %i)"), ToCStr(CurPingType), (int32)(PingTypeVal != INDEX_NONE),
+							ToCStr(CurPingAvg), (int32)(PingAvgVal != INDEX_NONE))
+				}
+			}
+			else
+			{
+				UE_LOG(LogNet, Error, TEXT("FNetPing: Invalid 'net.NetPingTypes' entry '%s' (Full CVar: %s)"),
+						ToCStr(CurPingAndAvg), ToCStr(PingTypesCVar));
+			}
+		}
+
+		LastPingTypes = PingTypes;
+
+		for (int32 AvgIdx=0; AvgIdx<UE_ARRAY_COUNT(PingAverageTypes); AvgIdx++)
+		{
+			LastPingAverageTypes[AvgIdx] = PingAverageTypes[AvgIdx];
+		}
+	}
+	else
+	{
+		PingTypes = LastPingTypes;
+
+		for (int32 AvgIdx=0; AvgIdx<UE_ARRAY_COUNT(PingAverageTypes); AvgIdx++)
+		{
+			PingAverageTypes[AvgIdx] = LastPingAverageTypes[AvgIdx];
+		}
+	}
+
+
+	auto GetNewStorageIdx = [this](EPingAverageType AvgType) -> uint8
+		{
+			if (AvgType == EPingAverageType::None)
+			{
+				return PingNoAverageStorage.AddDefaulted();
+			}
+			else if (AvgType == EPingAverageType::MovingAverage)
+			{
+				return PingMovingAverageStorage.AddDefaulted();
+			}
+			else //if (AvgType == EPingAverageType::PlayerStateAvg)
+			{
+				return PingPlayerStateAvgStorage.AddDefaulted();
+			}
+		};
+
+	if (EnumHasAnyFlags(PingTypes, EPingType::RoundTrip))
+	{
+		const int32 PingTypeIdx = PingTypeToIdx(EPingType::RoundTrip);
+
+		PingAverageStorageIdx[PingTypeIdx] = GetNewStorageIdx(PingAverageTypes[PingTypeIdx]);
+	}
+
+	if (EnumHasAnyFlags(PingTypes, EPingType::RoundTripExclFrame))
+	{
+		const int32 PingTypeIdx = PingTypeToIdx(EPingType::RoundTripExclFrame);
+
+		PingAverageStorageIdx[PingTypeIdx] = GetNewStorageIdx(PingAverageTypes[PingTypeIdx]);
+	}
+
+	if (UNetConnection* ServerConn = (Owner && Owner->Driver ? Owner->Driver->ServerConnection : nullptr))
+	{
+		if (EnumHasAnyFlags(PingTypes, EPingType::ICMP))
+		{
+			const int32 PingTypeIdx = PingTypeToIdx(EPingType::ICMP);
+
+			PingAverageStorageIdx[PingTypeIdx] = GetNewStorageIdx(PingAverageTypes[PingTypeIdx]);
+
+			NetPingICMP = FNetPingICMP::GetNetPingICMP(ServerConn);
+		}
+
+		if (EnumHasAnyFlags(PingTypes, EPingType::UDPQoS))
+		{
+			const int32 PingTypeIdx = PingTypeToIdx(EPingType::UDPQoS);
+
+			PingAverageStorageIdx[PingTypeIdx] = GetNewStorageIdx(PingAverageTypes[PingTypeIdx]);
+
+			NetPingUDPQoS = FNetPingUDPQoS::CreateNetPingUDPQoS(ServerConn);
+		}
+	}
+
+	// Negative ping to represent no values read
+#if STATS
+	SET_DWORD_STAT(STAT_Ping_RoundTripInclFrame, (EnumHasAnyFlags(PingTypes, EPingType::RoundTrip) ? 0 : -1));
+	SET_DWORD_STAT(STAT_Ping_RoundTripExclFrame, (EnumHasAnyFlags(PingTypes, EPingType::RoundTripExclFrame) ? 0 : -1));
+	SET_DWORD_STAT(STAT_Ping_ICMP, (EnumHasAnyFlags(PingTypes, EPingType::ICMP) ? 0 : -1));
+	SET_DWORD_STAT(STAT_Ping_UDPQoS, (EnumHasAnyFlags(PingTypes, EPingType::UDPQoS) ? 0 : -1));
+#endif
+}
+
+void FNetPing::UpdatePing(EPingType PingType, double TimeVal, double PingValue)
+{
+	using namespace UE::Net::Private;
+
+#if STATS
+	SCOPE_CYCLE_COUNTER(STAT_Ping_UpdatePingTime);
+#endif
+
+	if (EnumHasAnyFlags(PingTypes, PingType))
+	{
+		const int32 PingTypeIdx = PingTypeToIdxRuntime(PingType);
+		const uint8 StorageIdx = PingAverageStorageIdx[PingTypeIdx];
+		const EPingAverageType CurAvgType = PingAverageTypes[PingTypeIdx];
+		const int32 bUpdateCount = static_cast<int32>(!EnumHasAnyFlags(PingType, EPingType::ICMP | EPingType::UDPQoS));
+
+		if (CurAvgType == EPingAverageType::None)
+		{
+			FNoAverageStorage& CurStorage = PingNoAverageStorage[StorageIdx];
+
+			CurStorage.PingValues.AddSample(PingValue);
+
+			CurStorage.PingCount += bUpdateCount;
+
+#if STATS
+			UpdatePingStats(PingType, CurStorage.PingValues.GetCurrent());
+#endif
+		}
+		else if (CurAvgType == EPingAverageType::MovingAverage)
+		{
+			FMovingAverageStorage& CurStorage = PingMovingAverageStorage[StorageIdx];
+
+			CurStorage.PingAccumulator.AddMeasurement(TimeVal, PingValue);
+
+			CurStorage.PingCount += bUpdateCount;
+
+#if STATS
+			UpdatePingStats(PingType, CurStorage.PingValues.GetCurrent());
+#endif
+		}
+		else //if (CurAvgType == EPingAverageType::PlayerStateAvg)
+		{
+			FPlayerStateAvgStorage& CurStorage = PingPlayerStateAvgStorage[StorageIdx];
+
+			CurStorage.UpdatePing(TimeVal, static_cast<float>(PingValue));
+
+			CurStorage.PingCount += bUpdateCount;
+
+#if STATS
+			UpdatePingStats(PingType, CurStorage.PingValues.GetCurrent());
+#endif
+		}
+	}
+}
+
+void FNetPing::UpdatePingTimeout(EPingType PingType)
+{
+#if STATS
+	SCOPE_CYCLE_COUNTER(STAT_Ping_UpdatePingTime);
+#endif
+
+	if (EnumHasAnyFlags(PingTypes, PingType))
+	{
+		FAverageStorageBase& CurStorageBase = GetPingStorageBase(PingType);
+
+		CurStorageBase.PingTimeouts++;
+
+#if STATS
+		const double TimeoutPercent =
+			(static_cast<double>(CurStorageBase.PingTimeouts) / static_cast<double>(CurStorageBase.PingCount)) * 100.0;
+
+		UpdateTimeoutStats(PingType, CurStorageBase.PingTimeouts, TimeoutPercent);
+#endif
+
+		if (CurStorageBase.PingTimeouts == CurStorageBase.PingCount &&
+			CurStorageBase.PingTimeouts >= static_cast<uint32>(CVarNetPingTimeoutDisableThreshold.GetValueOnAnyThread()))
+		{
+			FString FailureAddress;
+
+			if (PingType == EPingType::ICMP)
+			{
+				FailureAddress = GetICMPAddr();
+
+				UE_LOG(LogNet, Warning, TEXT("FNetPing: ICMP ping failed with 100%% timeout after '%i' tries. Disabling. ")
+						TEXT("Ping address: %s"), CurStorageBase.PingTimeouts, ToCStr(FailureAddress));
+
+#if STATS
+				// -2 for failed pings
+				SET_DWORD_STAT(STAT_Ping_ICMP, -2);
+#endif
+
+				NetPingICMP.Reset();
+			}
+			else if (PingType == EPingType::UDPQoS)
+			{
+				FailureAddress = GetUDPQoSAddr();
+
+				UE_LOG(LogNet, Warning, TEXT("FNetPing: UDPQoS ping failed with 100%% timeout after '%i' tries. Disabling. ")
+						TEXT("Ping address: %s"), CurStorageBase.PingTimeouts, ToCStr(FailureAddress));
+
+#if STATS
+				// -2 for failed pings
+				SET_DWORD_STAT(STAT_Ping_UDPQoS, -2);
+#endif
+
+				NetPingUDPQoS.Reset();
+			}
+
+			PingTypes &= ~PingType;
+
+			// Notify the server of the failed ping
+			if (FailureAddress.Len() > 0 && Owner)
+			{
+				const uint32 PingTypeVal = static_cast<uint32>(PingType);
+				const ENetPingControlMessage MessageType = ENetPingControlMessage::PingFailure;
+				FString MessageStr = FString::Printf(TEXT("%i=%s"), PingTypeVal, ToCStr(FailureAddress));
+
+				FNetControlMessage<NMT_NetPing>::Send(Owner.Get(), MessageType, MessageStr);
+			}
+		}
+	}
+}
+
+#if STATS
+void FNetPing::UpdatePingStats(EPingType PingType, double CurrentValue)
+{
+	const int32 CurrentValueMS = FMath::TruncToInt(CurrentValue * 1000.0);
+
+	if (PingType == EPingType::RoundTrip)
+	{
+		SET_DWORD_STAT(STAT_Ping_RoundTripInclFrame, CurrentValueMS);
+	}
+	else if (PingType == EPingType::RoundTripExclFrame)
+	{
+		SET_DWORD_STAT(STAT_Ping_RoundTripExclFrame, CurrentValueMS);
+	}
+	else
+	{
+		const FAverageStorageBase& CurStorageBase = GetPingStorageBase(PingType);
+		const double TimeoutPercent = (CurStorageBase.PingTimeouts == 0 ? 0.0 :
+			(static_cast<double>(CurStorageBase.PingTimeouts) / static_cast<double>(CurStorageBase.PingCount)) * 100.0);
+
+		if (PingType == EPingType::ICMP)
+		{
+			SET_DWORD_STAT(STAT_Ping_ICMP, CurrentValueMS);
+			SET_FLOAT_STAT(STAT_PingTimeoutPercent_ICMP, TimeoutPercent);
+		}
+		else if (PingType == EPingType::UDPQoS)
+		{
+			SET_DWORD_STAT(STAT_Ping_UDPQoS, CurrentValueMS);
+			SET_FLOAT_STAT(STAT_PingTimeoutPercent_UDPQoS, TimeoutPercent);
+		}
+	}
+}
+
+void FNetPing::UpdateTimeoutStats(EPingType PingType, uint32 NumTimeouts, double TimeoutPercent)
+{
+	if (PingType == EPingType::ICMP)
+	{
+		SET_DWORD_STAT(STAT_PingTimeouts_ICMP, NumTimeouts);
+		SET_FLOAT_STAT(STAT_PingTimeoutPercent_ICMP, TimeoutPercent);
+	}
+	else if (PingType == EPingType::UDPQoS)
+	{
+		SET_DWORD_STAT(STAT_PingTimeouts_UDPQoS, NumTimeouts);
+		SET_FLOAT_STAT(STAT_PingTimeoutPercent_UDPQoS, TimeoutPercent);
+	}
+}
+#endif
+
+const FSampleMinMaxAvg& FNetPing::GetPingStorageValues(EPingType PingType) const
+{
+	using namespace UE::Net::Private;
+
+	check(EnumHasAnyFlags(PingTypes, PingType));
+
+	const int32 PingTypeIdx = PingTypeToIdxRuntime(PingType);
+	const uint8 StorageIdx = PingAverageStorageIdx[PingTypeIdx];
+	const EPingAverageType CurAvgType = PingAverageTypes[PingTypeIdx];
+
+	if (CurAvgType == EPingAverageType::None)
+	{
+		return PingNoAverageStorage[StorageIdx].PingValues;
+	}
+	else if (CurAvgType == EPingAverageType::MovingAverage)
+	{
+		return PingMovingAverageStorage[StorageIdx].PingValues;
+	}
+	else //if (CurAvgType == EPingAverageType::PlayerStateAvg)
+	{
+		return PingPlayerStateAvgStorage[StorageIdx].PingValues;
+	}
+}
+
+FNetPing::FAverageStorageBase& FNetPing::GetPingStorageBase(EPingType PingType)
+{
+	using namespace UE::Net::Private;
+
+	check(EnumHasAnyFlags(PingTypes, PingType));
+
+	const int32 PingTypeIdx = PingTypeToIdxRuntime(PingType);
+	const uint8 StorageIdx = PingAverageStorageIdx[PingTypeIdx];
+	const EPingAverageType CurAvgType = PingAverageTypes[PingTypeIdx];
+
+	if (CurAvgType == EPingAverageType::None)
+	{
+		return PingNoAverageStorage[StorageIdx];
+	}
+	else if (CurAvgType == EPingAverageType::MovingAverage)
+	{
+		return PingMovingAverageStorage[StorageIdx];
+	}
+	else //if (CurAvgType == EPingAverageType::PlayerStateAvg)
+	{
+		return PingPlayerStateAvgStorage[StorageIdx];
+	}
+}
+
+FPingValues FNetPing::GetPingValues(EPingType PingType) const
+{
+	FPingValues ReturnVal;
+
+	if (EnumHasAnyFlags(PingTypes, PingType))
+	{
+		const FSampleMinMaxAvg& CurStorageValues = GetPingStorageValues(PingType);
+		const double CurStorageMin = CurStorageValues.GetMin();
+		const double CurStorageMax = CurStorageValues.GetMax();
+
+		if (CurStorageValues.GetSampleCount() > 0)
+		{
+			ReturnVal.Current = CurStorageValues.GetCurrent();
+		}
+
+		if (CurStorageMin != std::numeric_limits<double>::max())
+		{
+			ReturnVal.Min = CurStorageMin;
+		}
+
+		if (CurStorageMax != std::numeric_limits<double>::min())
+		{
+			ReturnVal.Max = CurStorageMax;
+		}
+	}
+
+	return ReturnVal;
+}
+
+void FNetPing::ResetPingValues(EPingType PingType)
+{
+	using namespace UE::Net::Private;
+
+	if (EnumHasAnyFlags(PingTypes, PingType))
+	{
+		const int32 PingTypeIdx = PingTypeToIdxRuntime(PingType);
+		const uint8 StorageIdx = PingAverageStorageIdx[PingTypeIdx];
+		const EPingAverageType CurAvgType = PingAverageTypes[PingTypeIdx];
+
+		if (CurAvgType == EPingAverageType::None)
+		{
+			FNoAverageStorage& CurStorage = PingNoAverageStorage[StorageIdx];
+
+			CurStorage.PingCount = 0;
+			CurStorage.PingTimeouts = 0;
+			CurStorage.PingValues.Reset();
+		}
+		else if (CurAvgType == EPingAverageType::MovingAverage)
+		{
+			FMovingAverageStorage& CurStorage = PingMovingAverageStorage[StorageIdx];
+
+			CurStorage.PingCount = 0;
+			CurStorage.PingTimeouts = 0;
+			CurStorage.PingValues.Reset();
+			CurStorage.PingAccumulator.Reset();
+		}
+		else //if (CurAvgType == EPingAverageType::PlayerStateAvg)
+		{
+			PingPlayerStateAvgStorage[StorageIdx].Reset();
+		}
+	}
+}
+
+void FNetPing::TickRealtime(double CurTimeSeconds)
+{
+	if (CachedRemoteAddr.Len() > 0 && CurTimeSeconds > NextRemoteAddrCheckTimestamp)
+	{
+		UpdateRemoteAddr(CurTimeSeconds);
+	}
+
+	if (EnumHasAnyFlags(PingTypes, EPingType::ICMP) && CurTimeSeconds > NextPingTimerICMPTimestamp)
+	{
+		NextPingTimerICMPTimestamp = CurTimeSeconds + GNetPingICMPInterval;
+
+		PingTimerICMP(CurTimeSeconds);
+	}
+
+	if (EnumHasAnyFlags(PingTypes, EPingType::UDPQoS) && CurTimeSeconds > NextPingTimerUDPTimestamp)
+	{
+		NextPingTimerUDPTimestamp = CurTimeSeconds + GNetPingUDPInterval;
+
+		PingTimerUDP(CurTimeSeconds);
+	}
+
+#if !UE_BUILD_SHIPPING
+	static double LastDebugPingDump = 0.0;
+
+	if (!!GNetPingDebugDump && (CurTimeSeconds - LastDebugPingDump) > 5.0)
+	{
+		LastDebugPingDump = CurTimeSeconds;
+
+		UE_LOG(LogNet, Log, TEXT("FNetPing Debug Dump:"));
+
+		if (EnumHasAnyFlags(PingTypes, EPingType::RoundTrip))
+		{
+			FPingValues CurPingValues = GetPingValues(EPingType::RoundTrip);
+
+			UE_LOG(LogNet, Log, TEXT(" - RoundTrip: Min: %f, Max: %f, Current: %f"), CurPingValues.Min, CurPingValues.Max,
+					CurPingValues.Current);
+		}
+
+		if (EnumHasAnyFlags(PingTypes, EPingType::RoundTripExclFrame))
+		{
+			FPingValues CurPingValues = GetPingValues(EPingType::RoundTripExclFrame);
+
+			UE_LOG(LogNet, Log, TEXT(" - RoundTripExclFrame: Min: %f, Max: %f, Current: %f"), CurPingValues.Min, CurPingValues.Max,
+					CurPingValues.Current);
+		}
+
+		if (EnumHasAnyFlags(PingTypes, EPingType::ICMP))
+		{
+			FPingValues CurPingValues = GetPingValues(EPingType::ICMP);
+			const FAverageStorageBase& CurStorageBase = GetPingStorageBase(EPingType::ICMP);
+			const double TimeoutPercent =
+				(static_cast<double>(CurStorageBase.PingTimeouts) / static_cast<double>(CurStorageBase.PingCount)) * 100.0;
+
+			UE_LOG(LogNet, Log, TEXT(" - ICMP: Min: %f, Max: %f, Current: %f, Timeouts: %i, Timeout Percent: %f"),
+					CurPingValues.Min, CurPingValues.Max, CurPingValues.Current, CurStorageBase.PingTimeouts,
+					TimeoutPercent);
+		}
+
+		if (EnumHasAnyFlags(PingTypes, EPingType::UDPQoS))
+		{
+			FPingValues CurPingValues = GetPingValues(EPingType::UDPQoS);
+			const FAverageStorageBase& CurStorageBase = GetPingStorageBase(EPingType::UDPQoS);
+			const double TimeoutPercent =
+				(static_cast<double>(CurStorageBase.PingTimeouts) / static_cast<double>(CurStorageBase.PingCount)) * 100.0;
+
+			UE_LOG(LogNet, Log, TEXT(" - UDP: Min: %f, Max: %f, Current: %f, Timeouts: %i, Timeout Percent: %f"),
+					CurPingValues.Min, CurPingValues.Max, CurPingValues.Current, CurStorageBase.PingTimeouts,
+					TimeoutPercent);
+		}
+	}
+#endif
+}
+
+const TCHAR* FNetPing::GetRemoteAddr()
+{
+	if (LIKELY(CachedRemoteAddr.Len() > 0))
+	{
+		return *CachedRemoteAddr;
+	}
+
+	UpdateRemoteAddr(FPlatformTime::Seconds());
+
+	return CachedRemoteAddr.Len() > 0 ? *CachedRemoteAddr : TEXT("");
+}
+
+FString FNetPing::GetICMPAddr()
+{
+	return (OverridePingICMPAddress.Len() > 0 ? OverridePingICMPAddress : GetRemoteAddr());
+}
+
+FString FNetPing::GetUDPQoSAddr()
+{
+	TStringBuilder<256> PingIPStr;
+
+	if (OverridePingUDPQoSAddress.Len() > 0)
+	{
+		PingIPStr.Append(ToCStr(OverridePingUDPQoSAddress));
+	}
+	else
+	{
+		PingIPStr.Append(ToCStr(GetRemoteAddr()));
+	}
+
+
+	PingIPStr.Append(TEXT(':'));
+	PingIPStr.Appendf(TEXT("%i"), (OverridePingUDPQosPort != INDEX_NONE ? OverridePingUDPQosPort : GNetPingUDPPort));
+
+	return PingIPStr.ToString();
+}
+
+void FNetPing::UpdateRemoteAddr(double CurTimeSeconds)
+{
+	NextRemoteAddrCheckTimestamp = CurTimeSeconds + 1.0;
+
+	if (UNetConnection* ServerConn = (Owner && Owner->Driver ? Owner->Driver->ServerConnection : nullptr))
+	{
+		TSharedPtr<const FInternetAddr> CurRemoteAddr = ServerConn->GetRemoteAddr();
+
+		if (WeakRemoteAddr != CurRemoteAddr)
+		{
+			WeakRemoteAddr = CurRemoteAddr;
+			CachedRemoteAddr = ServerConn->LowLevelGetRemoteAddress();
+		}
+	}
+}
+
+void FNetPing::PingTimerICMP(double CurTimeSeconds)
+{
+	if (EnumHasAnyFlags(PingTypes, EPingType::ICMP) && NetPingICMP.IsValid() && Owner && Owner->Driver && Owner->Driver->ServerConnection)
+	{
+		FAverageStorageBase& CurStorageBase = GetPingStorageBase(EPingType::ICMP);
+
+		CurStorageBase.PingCount++;
+
+		NetPingICMP->Ping(CurTimeSeconds, GetICMPAddr());
+	}
+}
+
+void FNetPing::PingTimerUDP(double CurTimeSeconds)
+{
+	if (EnumHasAnyFlags(PingTypes, EPingType::UDPQoS) && NetPingUDPQoS.IsValid() && Owner && Owner->Driver && Owner->Driver->ServerConnection)
+	{
+		FAverageStorageBase& CurStorageBase = GetPingStorageBase(EPingType::UDPQoS);
+
+		CurStorageBase.PingCount++;
+
+		NetPingUDPQoS->Ping(CurTimeSeconds, GetUDPQoSAddr());
+	}
+}
+
+void FNetPing::HandleNetPingControlMessage(UNetConnection* Connection, ENetPingControlMessage MessageType, FString MessageStr)
+{
+	if (UNetDriver* Driver = Connection != nullptr ? Connection->Driver : nullptr)
+	{
+		bool bIsClient = Driver->ServerConnection == Connection;
+
+		auto IsAddrValid = [](const FString& InAddr) -> bool
+			{
+				if (InAddr.Len() < 32)
+				{
+					for (const TCHAR CurChar : InAddr)
+					{
+						// Valid characters
+						if ((CurChar >= TEXT('0') && CurChar <= TEXT('9')) || CurChar == TEXT(':') || CurChar == TEXT('.'))
+						{
+							continue;
+						}
+
+						return false;
+					}
+				}
+
+				return true;
+			};
+
+		if (bIsClient)
+		{
+			if (MessageType == ENetPingControlMessage::SetPingAddress)
+			{
+				if (FNetPing* NetPing = Driver->ServerConnection->GetNetPing())
+				{
+					TArray<FString> Args;
+
+					if (MessageStr.ParseIntoArray(Args, TEXT("=")) == 2 && Args[0].IsNumeric())
+					{
+						const EPingType PingType = static_cast<EPingType>(FCString::Atoi(*Args[0]));
+						FString& PingAddress = Args[1];
+
+						if (IsValidPingType(PingType) && IsAddrValid(PingAddress))
+						{
+							NetPing->ServerSetPingAddress(PingType, PingAddress);
+						}
+					}
+				}
+			}
+		}
+		else
+		{
+			if (MessageType == ENetPingControlMessage::PingFailure)
+			{
+				TArray<FString> Args;
+
+				if (MessageStr.ParseIntoArray(Args, TEXT("=")) == 2 && Args[0].IsNumeric())
+				{
+					const EPingType PingType = static_cast<EPingType>(FCString::Atoi(*Args[0]));
+					FString& PingAddress = Args[1];
+
+					if (IsValidPingType(PingType) && IsAddrValid(PingAddress))
+					{
+						if (PingType == EPingType::ICMP)
+						{
+							Connection->AnalyticsVars.AddFailedPingAddressICMP(PingAddress);
+						}
+						else if (PingType == EPingType::UDPQoS)
+						{
+							Connection->AnalyticsVars.AddFailedPingAddressUDP(PingAddress);
+						}
+					}
+				}
+			}
+		}
+	}
+}
+
+void FNetPing::ServerSetPingAddress(EPingType PingType, FString PingAddress)
+{
+	if (EnumHasAnyFlags(PingTypes, PingType) && OverridePingAddressCount < 8)
+	{
+		OverridePingAddressCount++;
+
+		if (PingType == EPingType::ICMP)
+		{
+			if (PingAddress.Find(TEXT(":")) == INDEX_NONE)
+			{
+				UE_LOG(LogNet, Log, TEXT("ServerSetPingAddress: EPingType::ICMP PingAddress: %s"), ToCStr(PingAddress));
+
+				OverridePingICMPAddress = PingAddress;
+
+				ResetPingValues(EPingType::ICMP);
+			}
+			else
+			{
+				UE_LOG(LogNet, Warning, TEXT("ServerSetPingAddress: EPingType::ICMP should not specify a port number: %s"),
+						ToCStr(PingAddress));
+			}
+		}
+		// For UDPQoS, it's valid to pass a string setting only the port: ":11111"
+		else if (PingType == EPingType::UDPQoS)
+		{
+			if (PingAddress.Len() > 0)
+			{
+				TArray<FString> IPPort;
+
+				if (PingAddress.ParseIntoArray(IPPort, TEXT(":")) <= 2)
+				{
+					if (IPPort[0].Len() > 0)
+					{
+						OverridePingUDPQoSAddress = IPPort[0];
+					}
+
+					if (IPPort.Num() == 2 && IPPort[1].IsNumeric())
+					{
+						OverridePingUDPQosPort = FCString::Atoi(ToCStr(IPPort[1]));
+					}
+
+					UE_LOG(LogNet, Log, TEXT("ServerSetPingAddress: EPingType::UDPQoS PingAddress: %s, PingPort: %i"),
+							ToCStr(OverridePingUDPQoSAddress), OverridePingUDPQosPort);
+
+					ResetPingValues(EPingType::UDPQoS);
+				}
+				else
+				{
+					UE_LOG(LogNet, Warning, TEXT("ServerSetPingAddress: EPingType::UDPQoS PingAddress split into too many components: %s"),
+							ToCStr(PingAddress));
+				}
+			}
+		}
+	}
+}
+
+namespace Private
+{
+
+/**
+ * FNetPingQoSBase
+ */
+FNetPingQoSBase::FNetPingQoSBase(EPingType InPingType, UNetConnection* InOwner)
+	: PingType(InPingType)
+	, Owner(InOwner)
+{
+}
+
+void FNetPingQoSBase::Reset()
+{
+	if (bPingInProgress)
+	{
+		bDiscardNextPing = true;
+	}
+
+	PingTimestamp = 0.0;
+}
+
+void FNetPingQoSBase::PingResult(FIcmpEchoResult Result)
+{
+	bPingInProgress = false;
+
+	if (!bDiscardNextPing)
+	{
+		if (FNetPing* NetPing = (Owner ? Owner->GetNetPing() : nullptr))
+		{
+			if (Result.Status == EIcmpResponseStatus::Success)
+			{
+				const double CurPing = static_cast<double>(Result.Time);
+				const double ApproxTimeSeconds = PingTimestamp + CurPing;
+
+				NetPing->UpdatePing(PingType, ApproxTimeSeconds, CurPing);
+			}
+			else
+			{
+				// Not always a timeout, but treat all failures as such anyway.
+				NetPing->UpdatePingTimeout(PingType);
+			}
+		}
+	}
+
+	bDiscardNextPing = false;
+}
+
+/**
+ * FNetPingICMP
+ */
+
+FNetPingICMP::FNetPingICMP(UNetConnection* InOwner)
+	: FNetPingQoSBase(EPingType::ICMP, InOwner)
+{
+}
+
+TSharedPtr<FNetPingICMP> FNetPingICMP::GetNetPingICMP(UNetConnection* InOwner)
+{
+	static TSharedPtr<FNetPingICMP> NetPingSingleton;
+
+	FName NetDriverName = (InOwner != nullptr && InOwner->Driver != nullptr ? InOwner->Driver->NetDriverName : NAME_None);
+
+	// The expected reference count, if no NetDriver has a handle to NetPingSingleton
+	const int32 IsUniqueRefCount = 1 + (int32)(NetPingSingleton.IsValid() && NetPingSingleton->bPingInProgress);
+
+	// Only provide an FNetPingICMP interface for GameNetDriver's, when no other NetDriver has a handle to one.
+	if ((NetDriverName == NAME_GameNetDriver || NetDriverName == NAME_PendingNetDriver) &&
+		(!NetPingSingleton.IsValid() || NetPingSingleton.GetSharedReferenceCount() == IsUniqueRefCount))
+	{
+		if (!NetPingSingleton.IsValid())
+		{
+			NetPingSingleton = MakeShared<FNetPingICMP>(InOwner);
+		}
+		else
+		{
+			NetPingSingleton->Reset();
+		}
+
+		return NetPingSingleton;
+	}
+
+	return TSharedPtr<FNetPingICMP>();
+}
+
+void FNetPingICMP::Ping(double CurTimeSeconds, FString PingAddress)
+{
+	if (!bPingInProgress)
+	{
+		bPingInProgress = true;
+		PingTimestamp = CurTimeSeconds;
+		TSharedPtr<FNetPingICMP> Self = AsShared();
+
+		FIcmp::IcmpEcho(PingAddress, static_cast<float>(GNetPingICMPInterval),
+			[Self = MoveTemp(Self)](FIcmpEchoResult Result)
+			{
+				if (Self.IsValid())
+				{
+					Self->PingResult(Result);
+				}
+			});
+	}
+}
+
+/**
+ * FNetPingUDPQoS
+ */
+
+FNetPingUDPQoS::FNetPingUDPQoS(UNetConnection* InOwner)
+	: FNetPingQoSBase(EPingType::UDPQoS, InOwner)
+{
+}
+
+TSharedPtr<FNetPingUDPQoS> FNetPingUDPQoS::CreateNetPingUDPQoS(UNetConnection* InOwner)
+{
+	return MakeShared<FNetPingUDPQoS>(InOwner);
+}
+
+void FNetPingUDPQoS::Ping(double CurTimeSeconds, FString PingAddress)
+{
+	if (!bPingInProgress)
+	{
+		bPingInProgress = true;
+		PingTimestamp = CurTimeSeconds;
+		TSharedPtr<FNetPingUDPQoS> Self = AsShared();
+
+		FUDPPing::UDPEcho(PingAddress, static_cast<float>(GNetPingUDPInterval),
+			[Self = MoveTemp(Self)](FIcmpEchoResult Result)
+			{
+				if (Self.IsValid())
+				{
+					Self->PingResult(Result);
+				}
+			});
+	}
+}
+
+}
+
+}
diff --git a/Engine/Source/Runtime/Engine/Private/NetConnection.cpp b/Engine/Source/Runtime/Engine/Private/NetConnection.cpp
--- a/Engine/Source/Runtime/Engine/Private/NetConnection.cpp
+++ b/Engine/Source/Runtime/Engine/Private/NetConnection.cpp
@@ -40,6 +40,7 @@
 #include "Math/NumericLimits.h"
 #include "UObject/UnrealNames.h"
 #include "HAL/LowLevelMemStats.h"
+#include "Net/NetPing.h"
 
 DECLARE_LLM_MEMORY_STAT(TEXT("NetConnection"), STAT_NetConnectionLLM, STATGROUP_LLMFULL);
 LLM_DEFINE_TAG(NetConnection, NAME_None, TEXT("Networking"), GET_STATFNAME(STAT_NetConnectionLLM), GET_STATFNAME(STAT_NetworkingSummaryLLM));
@@ -506,6 +507,8 @@ void UNetConnection::InitBase(UNetDriver* InDriver,class FSocket* InSocket, cons
 			});
 	}
 
+	NetPing = FNetPing::CreateNetPing(this);
+
 	UE_NET_TRACE_CONNECTION_CREATED(NetTraceId, GetConnectionId());
 	UE_NET_TRACE_CONNECTION_STATE_UPDATED(NetTraceId, GetConnectionId(), static_cast<uint8>(GetConnectionState()));
 }
@@ -986,6 +989,7 @@ void UNetConnection::Close(FNetResult&& CloseReason)
 		}
 
 		RPCDoS.NotifyClose();
+		NetPing.Reset();
 
 		if (const uint32 MyConnectionId = GetConnectionId())
 		{
@@ -2307,6 +2311,8 @@ void UNetConnection::WriteFinalPacketInfo(FBitWriter& Writer, const double Packe
 
 bool UNetConnection::ReadPacketInfo(FBitReader& Reader, bool bHasPacketInfoPayload)
 {
+	using namespace UE::Net;
+
 	// If this packet did not contain any packet info, nothing else to read
 	if (!bHasPacketInfoPayload)
 	{
@@ -2375,10 +2381,12 @@ bool UNetConnection::ReadPacketInfo(FBitReader& Reader, bool bHasPacketInfoPaylo
 		}
 
 
-		// use FApp's time because it is set closer to the beginning of the frame - we don't care about the time so far of the current frame to process the packet
-		const double CurrentTime = (PacketReceiveTime != 0.0 ? PacketReceiveTime : FApp::GetCurrentTime());
-		const double RTT		 = (CurrentTime - OutLagTime[Index] ) - ( CVarPingExcludeFrameTime.GetValueOnAnyThread() ? ServerFrameTime : 0.0 );
-		const double NewLag		 = FMath::Max( RTT, 0.0 );
+		// Use FApp's time because it is set closer to the beginning of the frame - we don't care about the time so far of the current frame to process the packet
+		const bool bExcludeFrameTime = !!CVarPingExcludeFrameTime.GetValueOnAnyThread();
+		const double CurrentTime	= (PacketReceiveTime != 0.0 ? PacketReceiveTime : FApp::GetCurrentTime());
+		const double RTT			= (CurrentTime - OutLagTime[Index]);
+		const double RTTExclFrame	= RTT - (bExcludeFrameTime ? ServerFrameTime : 0.0);
+		const double NewLag			= FMath::Max(RTTExclFrame, 0.0);
 
 		//UE_LOG( LogNet, Warning, TEXT( "Out: %i, InRemote: %i, Saturation: %f" ), OutBytesPerSecondHistory[Index], RemoteInKBytesPerSecond, RemoteSaturation );
 
@@ -2390,6 +2398,12 @@ bool UNetConnection::ReadPacketInfo(FBitReader& Reader, bool bHasPacketInfoPaylo
 			PlayerController->UpdatePing(NewLag);
 		}
 
+		if (NetPing.IsValid())
+		{
+			NetPing->UpdatePing(EPingType::RoundTrip, CurrentTime, NewLag);
+			NetPing->UpdatePing(EPingType::RoundTripExclFrame, CurrentTime, FMath::Max(RTT, 0.0));
+		}
+
 		if (NetworkCongestionControl.IsSet())
 		{
 			NetworkCongestionControl.GetValue().OnAck({ CurrentTime, OutAckPacketId });
@@ -4258,6 +4272,11 @@ void UNetConnection::Tick(float DeltaSeconds)
 	}
 
 	bFlushedNetThisFrame = false;
+
+	if (NetPing.IsValid())
+	{
+		NetPing->TickRealtime(CurrentRealtimeSeconds);
+	}
 }
 
 void UNetConnection::HandleConnectionTimeout(const FString& Error)
diff --git a/Engine/Source/Runtime/Engine/Private/World.cpp b/Engine/Source/Runtime/Engine/Private/World.cpp
--- a/Engine/Source/Runtime/Engine/Private/World.cpp
+++ b/Engine/Source/Runtime/Engine/Private/World.cpp
@@ -136,6 +136,7 @@
 #include "MoviePlayerProxy.h"
 #include "ObjectTrace.h"
 #include "ReplaySubsystem.h"
+#include "Net/NetPing.h"
 
 #if INCLUDE_CHAOS
 #include "ChaosSolversModule.h"
@@ -5734,6 +5735,29 @@ void UWorld::WelcomePlayer(UNetConnection* Connection)
 	}
 
 	FNetControlMessage<NMT_Welcome>::Send(Connection, LevelName, GameName, RedirectURL);
+
+	FString ClientNetPingICMPAddress;
+	FString ClientNetPingUDPAddress;
+
+	if (FParse::Value(FCommandLine::Get(), TEXT("ClientNetPingICMPAddress="), ClientNetPingICMPAddress) && ClientNetPingICMPAddress.Len() > 0)
+	{
+		const uint32 PingType = static_cast<uint32>(EPingType::ICMP);
+		const ENetPingControlMessage MessageType = ENetPingControlMessage::SetPingAddress;
+		FString MessageStr = FString::Printf(TEXT("%i=%s"), PingType, ToCStr(ClientNetPingICMPAddress));
+
+		FNetControlMessage<NMT_NetPing>::Send(Connection, MessageType, MessageStr);
+	}
+
+	if (FParse::Value(FCommandLine::Get(), TEXT("ClientNetPingUDPAddress="), ClientNetPingUDPAddress) && ClientNetPingUDPAddress.Len() > 0)
+	{
+		const uint32 PingType = static_cast<uint32>(EPingType::UDPQoS);
+		const ENetPingControlMessage MessageType = ENetPingControlMessage::SetPingAddress;
+		FString MessageStr = FString::Printf(TEXT("%i=%s"), PingType, ToCStr(ClientNetPingUDPAddress));
+
+		FNetControlMessage<NMT_NetPing>::Send(Connection, MessageType, MessageStr);
+	}
+
+
 	Connection->FlushNet();
 	// don't count initial join data for netspeed throttling
 	// as it's unnecessary, since connection won't be fully open until it all gets received, and this prevents later gameplay data from being delayed to "catch up"
diff --git a/Engine/Source/Runtime/Engine/Public/Analytics/EngineNetAnalytics.h b/Engine/Source/Runtime/Engine/Public/Analytics/EngineNetAnalytics.h
--- a/Engine/Source/Runtime/Engine/Public/Analytics/EngineNetAnalytics.h
+++ b/Engine/Source/Runtime/Engine/Public/Analytics/EngineNetAnalytics.h
@@ -42,6 +42,22 @@ public:
 		OutOfOrderPacketsDuplicateCount += Count;
 	}
 
+	void AddFailedPingAddressICMP(FString PingAddress)
+	{
+		if (FailedPingAddressesICMP.Num() < 32)
+		{
+			FailedPingAddressesICMP.AddUnique(PingAddress);
+		}
+	}
+
+	void AddFailedPingAddressUDP(FString PingAddress)
+	{
+		if (FailedPingAddressesICMP.Num() < 32)
+		{
+			FailedPingAddressesUDP.AddUnique(PingAddress);
+		}
+	}
+
 
 public:
 	/** The number of packets that were exclusively ack packets */
@@ -70,7 +86,14 @@ public:
 	/** NetConnection faults that were recovered from, and the number of times they were recovered from */
 	TMap<FString, int32> RecoveredFaults;
 
+private:
+	/** List of IP addresses a client failed to ping with ICMP. Should correlate against NetConnection count, to determine overall percent. */
+	TArray<FString> FailedPingAddressesICMP;
+
+	/** List of IP addresses a client failed to ping with UDP. Should correlate against NetConnection count, to determine overall percent. */
+	TArray<FString> FailedPingAddressesUDP;
 
+public:
 	/** Aggregation variables */
 
 	struct FPerNetConnData
diff --git a/Engine/Source/Runtime/Engine/Public/Net/DataChannel.h b/Engine/Source/Runtime/Engine/Public/Net/DataChannel.h
--- a/Engine/Source/Runtime/Engine/Public/Net/DataChannel.h
+++ b/Engine/Source/Runtime/Engine/Public/Net/DataChannel.h
@@ -13,6 +13,10 @@
 #include "Net/DataBunch.h"
 #include "Engine/ChildConnection.h"
 
+// Forward declarations
+enum class ENetPingControlMessage : uint8;
+
+
 /*-----------------------------------------------------------------------------
 	UControlChannel base class.
 -----------------------------------------------------------------------------*/
@@ -166,6 +170,7 @@ DEFINE_CONTROL_CHANNEL_MESSAGE(GameSpecific, 20, uint8, FString); // custom game
 DEFINE_CONTROL_CHANNEL_MESSAGE(EncryptionAck, 21);
 DEFINE_CONTROL_CHANNEL_MESSAGE(DestructionInfo, 22);
 DEFINE_CONTROL_CHANNEL_MESSAGE(CloseReason, 23, FString); // Reason for client NetConnection Close, for analytics/logging
+DEFINE_CONTROL_CHANNEL_MESSAGE(NetPing, 24, ENetPingControlMessage /* MessageType */, FString /* MessageStr */);
 
 // 			Beacon control channel flow
 // Client												Server
diff --git a/Engine/Source/Runtime/Engine/Public/Net/NetPing.h b/Engine/Source/Runtime/Engine/Public/Net/NetPing.h
new file mode 100644
--- /dev/null
+++ b/Engine/Source/Runtime/Engine/Public/Net/NetPing.h
@@ -0,0 +1,429 @@
+// Copyright Epic Games, Inc. All Rights Reserved.
+
+#pragma once
+
+#include "CoreMinimal.h"
+#include "Net/Core/Analytics/NetStatsUtils.h"
+#include "GameFramework/PlayerState.h"
+
+// Forward declarations
+enum class EPingType : uint32;
+
+namespace UE::Net
+{
+
+namespace Private
+{
+
+/**
+ * Converts an EPingType flag to an array index at compile time.
+ *
+ * @param PingType	The EPingType flag to convert
+ * @return			The index that the flag represents
+ */
+constexpr int32 PingTypeToIdx(EPingType PingType)
+{
+	// FGenericPlatformMath::FloorLog2 as constexpr
+
+	uint32 Value = static_cast<uint32>(PingType);
+	uint32 Pos = 0;
+	if (Value >= 1<<16) { Value >>= 16; Pos += 16; }
+	if (Value >= 1<< 8) { Value >>=  8; Pos +=  8; }
+	if (Value >= 1<< 4) { Value >>=  4; Pos +=  4; }
+	if (Value >= 1<< 2) { Value >>=  2; Pos +=  2; }
+	if (Value >= 1<< 1) {				Pos +=  1; }
+
+	return Pos;
+}
+
+/**
+ * Converts an EPingType flag to an array index at runtime.
+ *
+ * @param PingType	The EPingType flag to convert
+ * @return			The index that the flag represents
+ */
+static int32 PingTypeToIdxRuntime(EPingType PingType)
+{
+	return FMath::FloorLog2(static_cast<uint32>(PingType));
+}
+
+}
+}
+
+/**
+ * Flags specifying different types of ping.
+ */
+UENUM()
+enum class EPingType : uint32
+{
+	None					= 0x00000000,
+	RoundTrip				= 0x00000001,	// Round Trip ping which includes any server/client frame time delay
+	RoundTripExclFrame		= 0x00000002,	// Round Trip ping which may attempt to exclude server frame time delay (inaccurate)
+	ICMP					= 0x00000004,	// Standard ICMP ping (GameNetDriver client only, accuracy affected by network ICMP throttling)
+	UDPQoS					= 0x00000008,	// UDP based ping used with QoS ping host on server (more accurate than ICMP)
+
+	Max						= UDPQoS,
+	Count					= UE::Net::Private::PingTypeToIdx(static_cast<EPingType>(EPingType::Max)) + 1
+};
+
+ENUM_CLASS_FLAGS(EPingType);
+
+
+/**
+ * The type of averaging to perform on ping values
+ */
+UENUM()
+enum class EPingAverageType : uint8
+{
+	None,			// No averaging
+	MovingAverage,	// Moving average of the last 4 seconds of ping data
+	PlayerStateAvg	// Original PlayerState ping code, calculating a moving average of the last 4 seconds of ping data. To be deprecated.
+};
+
+/**
+ * NMT_NetPing control channel sub-messages
+ */
+enum class ENetPingControlMessage : uint8
+{
+	SetPingAddress,		// Sent from the server to the client, to override the ICMP/UDP ping address
+	PingFailure,		// Sent from the client to the server, to report failure to ping an ICMP/UDP address
+
+	Max = PingFailure
+};
+
+
+namespace UE::Net
+{
+/**
+ * Whether or not an EPingType value is valid
+ */
+static bool IsValidPingType(EPingType PingType)
+{
+	const uint32 PingTypeVal = static_cast<uint32>(PingType);
+
+	return FMath::IsPowerOfTwo(PingTypeVal) && PingTypeVal <= static_cast<uint32>(EPingType::Max);
+}
+
+namespace Private
+{
+	class FNetPingICMP;
+	class FNetPingUDPQoS;
+}
+
+/**
+ * Struct for returning ping values, in seconds.
+ */
+struct FPingValues
+{
+	/** The current/latest ping value, or -1.0 if not set */
+	double Current = -1.0;
+
+	/** The smallest measured ping value (no averaging), or -1.0 if not set */
+	double Min = -1.0;
+
+	/** The highest measured ping value (no averaging), or -1.0 if not set */
+	double Max = -1.0;
+};
+
+/**
+ * Central class for handling all Player/NetConnection ping collection and calculation.
+ * Can work with any NetDriver type (except ICMP), but restricted to GameNetDriver for now.
+ */
+class FNetPing
+{
+	template <typename T> friend struct UE::Core::Private::PimplPtr::TPimplHeapObjectImpl;
+
+private:
+	class FAverageStorageBase
+	{
+	public:
+		/** The total number of pings (including timeouts) */
+		uint32 PingCount = 0;
+
+		/** The number of ping timeouts, if applicable */
+		uint32 PingTimeouts = 0;
+	};
+
+	class FNoAverageStorage : public FAverageStorageBase
+	{
+	public:
+		/** Stores the final Min/Max/Current ping values */
+		TSampleMinMaxAvg<EMinMaxValueMode::PerSample> PingValues;
+	};
+
+	class FMovingAverageStorage : public FAverageStorageBase
+	{
+	public:
+		/** Stores the final Min/Max/Current ping values */
+		TSampleMinMaxAvg<EMinMaxValueMode::PerMeasurement> PingValues;
+
+		/** Accumulates ping readings and outputs the moving average to PingValues */
+		TBinnedMovingAvg<decltype(PingValues), TBinParms::NumBins(4)> PingAccumulator{TBinParms::TimePerBin(1.0)};
+
+
+	public:
+		FMovingAverageStorage();
+	};
+
+	/**
+	 * Implements storage and calculation duplicating original PlayerState ping code, for backwards-compatible/representative comparison.
+	 */
+	class FPlayerStateAvgStorage : public FAverageStorageBase
+	{
+	private:
+		/** The current PingBucket index that is being filled */
+		uint8 CurPingBucket = 0;
+
+		/**
+		 * Stores the last 4 seconds worth of ping data (one second per 'bucket').
+		 * It is stored in this manner, to allow calculating a moving average,
+		 * without using up a lot of space, while also being tolerant of changes in ping update frequency
+		 */
+		PingAvgData PingBucket[4];
+
+		/** The timestamp for when the current PingBucket began filling */
+		double CurPingBucketTimestamp = 0.0;
+
+	public:
+		/** Stores the final Min/Max/Current ping values */
+		TSampleMinMaxAvg<EMinMaxValueMode::PerMeasurement> PingValues;
+
+
+	public:
+		/**
+		 * Resets values to their base state
+		 */
+		void Reset();
+
+		/**
+		 * Receives ping updates for the client (both clientside and serverside), from the net driver
+		 * NOTE: This updates much more frequently clientside, thus the clientside ping will often be different to what the server displays
+		 */
+		void UpdatePing(double TimeVal, float InPing);
+
+	private:
+		/** Recalculates the replicated Ping value once per second (both clientside and serverside), based upon collected ping data */
+		void RecalculateAvgPing();
+	};
+
+
+private:
+	FNetPing() = delete;
+
+	FNetPing(UNetConnection* InOwner);
+
+public:
+	/**
+	 * Creates a new FNetPing instance, owned by the specified NetConnection.
+	 * NOTE: Presently only enabled for GameNetDriver, but can support other NetDriver's (except ICMP ping).
+	 *
+	 * @param InOwner	The owner of the new FNetPing instance.
+	 * @return			Returns the new FNetPing instance, or nullptr if disabled. Uses TPimplPtr for compatibility with forward-declare.
+	 */
+	static TPimplPtr<FNetPing> CreateNetPing(UNetConnection* InOwner);
+
+	/**
+	 * Updates the value of the specified ping type, with a new ping reading (in seconds).
+	 *
+	 * @param PingType		The type of ping being updated.
+	 * @param TimeVal		The time that the ping reading was calculated.
+	 * @param PingValue		The value of the new ping reading, in seconds.
+	 */
+	void UpdatePing(EPingType PingType, double TimeVal, double PingValue);
+
+	/**
+	 * Records a ping timeout, for the specified ping type.
+	 *
+	 * @param PingType		The type of ping which timed out.
+	 */
+	void UpdatePingTimeout(EPingType PingType);
+
+	/**
+	 * Retrieves the ping results for the specified ping type, in seconds.
+	 *
+	 * @param PingType		The type of ping to retrieve results for.
+	 * @return				Returns the ping results for the specified type, in seconds.
+	 */
+	ENGINE_API FPingValues GetPingValues(EPingType PingType) const;
+
+	/**
+	 * Ticks the ping handler.
+	 *
+	 * @param CurTimeSeconds	The current time in seconds.
+	 */
+	void TickRealtime(double CurTimeSeconds);
+
+	/**
+	 * Handles NMT_NetPing control channel messages, serverside/clientside.
+	 *
+	 * @param Connection	The connection that received the message
+	 * @param MessageType	The type of message received
+	 * @param MessageStr	The message content.
+	 */
+	ENGINE_API static void HandleNetPingControlMessage(UNetConnection* Connection, ENetPingControlMessage MessageType, FString MessageStr);
+
+	/**
+	 * Returns the list of enabled ping types.
+	 */
+	EPingType GetPingTypes() const
+	{
+		return PingTypes;
+	}
+
+private:
+	/**
+	 * Initialize the ping handler
+	 */
+	void Init();
+
+#if STATS
+	/**
+	 * Update the 'stat ping' ping value for the specified ping type.
+	 *
+	 * @param PingType		The type of ping having its stats value set.
+	 * @param CurrentValue	The value to assign to stats, for the specified ping value.
+	 */
+	void UpdatePingStats(EPingType PingType, double CurrentValue);
+
+	/**
+	 * Update the 'stat ping' timeout values for the specified ping type.
+	 *
+	 * @param PingType			The type of ping having its stats value set.
+	 * @param NumTimeouts		The number of times this ping type has timed out.
+	 * @param TimeoutPercent	The percentage of time this ping type has timed out.
+	 */
+	void UpdateTimeoutStats(EPingType PingType, uint32 NumTimeouts, double TimeoutPercent);
+#endif
+
+	/**
+	 * Returns the read-only stored values, for the specified ping type.
+	 * NOTE: Not safe to call without verifying PingType is enabled in PingTypes, first.
+	 *
+	 * @param PingType	The ping type to retrieve stored values for.
+	 * @return			Read-only reference to the stored ping values.
+	 */
+	const FSampleMinMaxAvg& GetPingStorageValues(EPingType PingType) const;
+
+	/**
+	 * Returns the modifiable base average storage (containing only common/shared stats), for the specified ping type.
+	 * NOTE: Not safe to call without verifying PingType is enabled in PingTypes, first.
+	 *
+	 * @param PingType	The ping type to retrieve the base storage for.
+	 * @return			Modifiable base storage for specified ping type.
+	 */
+	FAverageStorageBase& GetPingStorageBase(EPingType PingType);
+
+	/**
+	 * Resets the ping values for the specified ping type.
+	 *
+	 * @param PingType	The ping type to reset the values of
+	 */
+	void ResetPingValues(EPingType PingType);
+
+
+	/**
+	 * Returns the cached IP address of the remote connection (server, as this is clientside only), initializing the cached value if necessary.
+	 *
+	 * @return	The address of the remote connection.
+	 */
+	const TCHAR* GetRemoteAddr();
+
+	/**
+	 * Returns the address that ICMP pings should use
+	 */
+	FString GetICMPAddr();
+
+	/**
+	 * Returns the address that UDP pings should use
+	 */
+	FString GetUDPQoSAddr();
+
+	/**
+	 * Updates the cached remote address, if necessary.
+	 *
+	 * @param CurTimeSeconds	The current time in seconds, used for rate limiting.
+	 */
+	void UpdateRemoteAddr(double CurTimeSeconds);
+
+	/**
+	 * Performs regular ICMP pings.
+	 *
+	 * @param CurTimeSeconds	The current time in seconds.
+	 */
+	void PingTimerICMP(double CurTimeSeconds);
+
+	/**
+	 * Performs regular UDP pings.
+	 *
+	 * @param CurTimeSeconds	The current time in seconds.
+	 */
+	void PingTimerUDP(double CurTimeSeconds);
+
+
+	/**
+	 * Clientside implementation of the 'ENetPingControlMessage::SetPingAddress' control channel message.
+	 * Allows the server to specify the ICMP/UDPQoS ping address/port.
+	 *
+	 * @param PingType		The type of ping having its address assigned.
+	 * @param PingAddress	The address being assigned to the specified ping type.
+	 */
+	void ServerSetPingAddress(EPingType PingType, FString PingAddress);
+
+
+private:
+	/** The NetConnection which owns this ping handler */
+	TObjectPtr<UNetConnection> Owner;
+
+	/** Weak reference to the remote address, for checking if it has changed */
+	TWeakPtr<const FInternetAddr> WeakRemoteAddr;
+
+	/** Caches string value for the remote address */
+	FString CachedRemoteAddr = TEXT("");
+
+	/** The next time the remote address should be updated */
+	double NextRemoteAddrCheckTimestamp = 0.0;
+
+	/** Flags specifying the ping types that are enabled */
+	EPingType PingTypes = EPingType::None;
+
+	/** Specifies the type of averaging to use for each EPingType index */
+	EPingAverageType PingAverageTypes[static_cast<uint32>(EPingType::Count)] = { EPingAverageType::None };
+
+	/** Specifies the index into a storage TArray, for each EPingType index */
+	uint8 PingAverageStorageIdx[static_cast<uint32>(EPingType::Count)] = { 0 };
+
+	/** Storage used for EPingType's that don't use averaging */
+	TArray<FNoAverageStorage> PingNoAverageStorage;
+
+	/** Storage used for EPingType's that use a moving average */
+	TArray<FMovingAverageStorage> PingMovingAverageStorage;
+
+	/** Storage used for EPingType's that use PlayerState ping code averaging */
+	TArray<FPlayerStateAvgStorage> PingPlayerStateAvgStorage;
+
+	/** The next time to perform an ICMP ping */
+	double NextPingTimerICMPTimestamp = 0.0;
+
+	/** The next time to perform a UDP ping */
+	double NextPingTimerUDPTimestamp = 0.0;
+
+	/** Pointer to the ICMP ping handler, if set (only one ICMP handler instance can exist, held by the GameNetDriver client) */
+	TSharedPtr<Private::FNetPingICMP> NetPingICMP;
+
+	/** Pointer to a UDP ping handler, if set (no restrictions on the number of UDP handler instances nor supported drivers) */
+	TSharedPtr<Private::FNetPingUDPQoS> NetPingUDPQoS;
+
+	/** Overrides the address to use for ICMP pings (set by server) */
+	FString OverridePingICMPAddress;
+
+	/** Overrides the address to use for UDPQoS pings (set by server) */
+	FString OverridePingUDPQoSAddress;
+
+	/** Overrides the port to use for UDPQoS pings (set by server) */
+	int32 OverridePingUDPQosPort = INDEX_NONE;
+
+	/** The number of times a ping address was overridden. Used to limit the number of changes. */
+	int32 OverridePingAddressCount = 0;
+};
+
+}
diff --git a/Engine/Source/Runtime/Net/Core/Private/Net/Core/Analytics/NetStatsUtils.cpp b/Engine/Source/Runtime/Net/Core/Private/Net/Core/Analytics/NetStatsUtils.cpp
new file mode 100644
--- /dev/null
+++ b/Engine/Source/Runtime/Net/Core/Private/Net/Core/Analytics/NetStatsUtils.cpp
@@ -0,0 +1,164 @@
+// Copyright Epic Games, Inc. All Rights Reserved.
+
+// Includes
+#include "Net/Core/Analytics/NetStatsUtils.h"
+
+namespace UE
+{
+namespace Net
+{
+
+/**
+ * FBinnedMovingValueBase
+ */
+
+FBinnedMovingValueBase::FBinnedMovingValueBase(const TArrayView<FBin> InBins, double InTimePerBin, EBinnedValueMode InMode)
+	: BinsView(InBins)
+	, TimePerBin(InTimePerBin)
+	, BinRange(InBins.Num() * InTimePerBin)
+	, SampleMode(InMode)
+{
+}
+
+void FBinnedMovingValueBase::Reset()
+{
+	for (FBin& CurBin : BinsView)
+	{
+		CurBin = {};
+	}
+
+	FirstBinIndex = INDEX_NONE;
+	LastWrittenBinIndex = INDEX_NONE;
+	TotalSum = 0.0;
+	TotalCount = 0;
+}
+
+void FBinnedMovingValueBase::AddMeasurement_Implementation(double TimeVal, double Value, TOptional<double>& OutSample)
+{
+	const double BinnedTimeVal = TimeVal - (FMath::FloorToDouble(TimeVal / BinRange) * BinRange);
+	const int32 BinIdx = FMath::Min(static_cast<int32>(FMath::FloorToDouble(BinnedTimeVal / TimePerBin)), BinsView.Num()-1);
+
+	if (!BinsView.IsValidIndex(BinIdx))
+	{
+		return;
+	}
+
+	// If moving into a new bin, reset it and any skipped-over bins - and output a sample if bins reach the start again
+	if (BinIdx != LastWrittenBinIndex)
+	{
+		ResetNewAndSkippedBins(BinIdx, OutSample);
+	}
+
+	FBin& CurBin = BinsView[BinIdx];
+
+	CurBin.Sum += Value;
+	CurBin.Count++;
+
+	TotalSum += Value;
+	TotalCount++;
+}
+
+double FBinnedMovingValueBase::GetSample() const
+{
+	return (SampleMode == EBinnedValueMode::MovingSum ? TotalSum : (TotalSum / (double)TotalCount));
+}
+
+void FBinnedMovingValueBase::ResetNewAndSkippedBins(int32 BinIdx, TOptional<double>& OutSample)
+{
+	if (UNLIKELY(LastWrittenBinIndex == INDEX_NONE))
+	{
+		FirstBinIndex = BinIdx;
+	}
+	else
+	{
+		auto BinDecrement = [&](int32 BinDecIndex) -> int32
+		{
+			return (BinDecIndex == 0 ? BinsView.Num() : BinDecIndex) - 1;
+		};
+
+		// Wipe skipped over bins, that received no write
+		for (int32 ResetIdx=BinDecrement(BinIdx); ResetIdx!=LastWrittenBinIndex; ResetIdx=BinDecrement(ResetIdx))
+		{
+			if (ResetIdx == FirstBinIndex)
+			{
+				OutSample = GetSample();
+			}
+
+			FBin& ResetBin = BinsView[ResetIdx];
+
+			TotalSum -= ResetBin.Sum;
+			TotalCount -= ResetBin.Count;
+			ResetBin = {};
+		}
+
+		if (BinIdx == FirstBinIndex)
+		{
+			OutSample = GetSample();
+		}
+	}
+
+	FBin& CurBin = BinsView[BinIdx];
+
+	TotalSum -= CurBin.Sum;
+	TotalCount -= CurBin.Count;
+	CurBin = {};
+
+	LastWrittenBinIndex = BinIdx;
+}
+
+
+/**
+ * FSampleMinMaxAvg
+ */
+
+void FSampleMinMaxAvg::AddSample_Internal(double Value)
+{
+	CurrentValue = Value;
+	TotalSum += Value;
+	TotalCount++;
+}
+
+void FSampleMinMaxAvg::Reset()
+{
+	MinValue = std::numeric_limits<double>::max();
+	MaxValue = std::numeric_limits<double>::min();
+	CurrentValue = 0.0;
+	TotalSum = 0.0;
+	TotalCount = 0;
+}
+
+}
+}
+
+
+#if !(UE_BUILD_SHIPPING || UE_BUILD_TEST)
+#include "Misc/AutomationTest.h"
+
+IMPLEMENT_SIMPLE_AUTOMATION_TEST(FBinnedMovingValueTest, "System.Core.Networking.BinnedMovingValue",
+									EAutomationTestFlags::ApplicationContextMask | EAutomationTestFlags::SmokeFilter
+
+);
+
+bool FBinnedMovingValueTest::RunTest(const FString& Parameters)
+{
+	using namespace UE::Net;
+
+	TSampleMinMaxAvg<EMinMaxValueMode::PerSample> DeltaStats;
+	TDeltaBinnedMovingSum<decltype(DeltaStats), TBinParms::NumBins(6)> DeltaTest{TBinParms::TimePerBin(10.0)};
+
+	// 250 per minute, with 6 bins
+	const double MarginOfError = 250.0 / 6.0;
+
+	DeltaTest.SetConsumer(&DeltaStats);
+
+	for (double TimeVal = 60.0, TotalVal = 0.0; TimeVal <= 360.0; TimeVal += 3.0, TotalVal += 12.5)
+	{
+		DeltaTest.AddMeasurement(TimeVal, TotalVal);
+	}
+
+	TestTrue(TEXT("Moving sum test of 250 delta per minute failed"), (250.0 - DeltaTest.GetSample()) < MarginOfError);
+	TestTrue(TEXT("DeltaStats output from moving sum test of 250 delta per minute failed"), (250.0 - DeltaStats.GetAvg()) < MarginOfError);
+
+	return true;
+}
+#endif
diff --git a/Engine/Source/Runtime/Net/Core/Public/Net/Core/Analytics/NetStatsUtils.h b/Engine/Source/Runtime/Net/Core/Public/Net/Core/Analytics/NetStatsUtils.h
new file mode 100644
--- /dev/null
+++ b/Engine/Source/Runtime/Net/Core/Public/Net/Core/Analytics/NetStatsUtils.h
@@ -0,0 +1,602 @@
+// Copyright Epic Games, Inc. All Rights Reserved.
+
+#pragma once
+
+// Includes
+
+#include "CoreMinimal.h"
+#include "Templates/Models.h"
+#include <limits>
+
+
+namespace UE::Net
+{
+
+/**
+ * Inherited class for a sample producer, which outputs samples to a consumer of the specified class.
+ *
+ * The consumer must implement: void AddSample(double Val)
+ *
+ * If the consumer wants to peek at each individual measurement that makes up the sample (for e.g. determining Min/Max measurements),
+ * then the consumer must implement: void PeekMeasurement(double TimeVal, double Val)
+ *
+ * @param ConsumerType	The type which will be used for consuming samples.
+ */
+template<typename ConsumerType>
+class TSampleProducer
+{
+	struct CSampleConsumer
+	{
+		template<typename T>
+		auto Requires(T& A) -> decltype(
+			A.AddSample(0.0)
+		);
+	};
+
+	struct CPeekMeasurements
+	{
+		template<typename T>
+		auto Requires(T& A) -> decltype(
+			A.PeekMeasurement(0.0, 0.0)
+		);
+	};
+
+	static_assert(TModels<CSampleConsumer, ConsumerType>::Value, "ConsumerType must implement AddSample(double Val)");
+
+public:
+	/**
+	 * Sets the consumer object to output samples to.
+	 *
+	 * @param InConsumer	The consumer object
+	 */
+	void SetConsumer(ConsumerType* InConsumer)
+	{
+		Consumer = InConsumer;
+	}
+
+protected:
+	/**
+	 * Called by subclasses to output a sample to the consumer.
+	 *
+	 * @param Value		The sample value being output
+	 */
+	void OutputSample(double Value)
+	{
+		if (Consumer != nullptr)
+		{
+			Consumer->AddSample(Value);
+		}
+	}
+
+	/**
+	 * Called by subclasses to let the consumer peek at individual measurements
+	 *
+	 * @param Value		The measurement value being collected for sampling
+	 */
+
+	template<typename InnerConsumerType=ConsumerType, typename = typename TEnableIf<TModels<CPeekMeasurements, InnerConsumerType>::Value>::Type>
+	inline void PeekMeasurement(double TimeVal, double Value)
+	{
+		if (Consumer != nullptr)
+		{
+			Consumer->PeekMeasurement(TimeVal, Value);
+		}
+	}
+
+	template<typename InnerConsumerType=ConsumerType, typename = typename TEnableIf<!TModels<CPeekMeasurements, InnerConsumerType>::Value>::Type, int32 UnusedParamForODR=0>
+	inline void PeekMeasurement(double TimeVal, double Value)
+	{
+	}
+
+private:
+	/** The consumer object that receives samples */
+	ConsumerType* Consumer;
+};
+
+/**
+ * The type of moving value to calculate/track
+ */
+enum class EBinnedValueMode : uint8
+{
+	/** Tracks a moving average */
+	MovingAvg,
+
+	/** Tracks a moving sum */
+	MovingSum
+};
+
+/**
+ * Base class for moving values, implementing non-templatized data-structure/methods.
+ * See TBinnedMovingValue.
+ */
+class NETCORE_API FBinnedMovingValueBase
+{
+protected:
+	struct FBin
+	{
+		/** Sum of all values in this bin */
+		double Sum = 0.0;
+
+		/** Number of values in this bin */
+		int32 Count = 0;
+	};
+
+protected:
+	FBinnedMovingValueBase() = delete;
+	explicit FBinnedMovingValueBase(const TArrayView<FBin> InBins, double InTimePerBin, EBinnedValueMode InMode);
+
+	/**
+	 * Internal implementation for AddMeasurement, which outputs a sample if the bins have filled and rolled back over to FirstBinIndex.
+	 *
+	 * @param TimeVal		The time value associated with the measurement, for time-based binning
+	 * @param Value			The measurement value
+	 * @param OutSample		Outputs a sample, if all bins have filled and rolled back over to FirstBinIndex
+	 */
+	void AddMeasurement_Implementation(double TimeVal, double Value, TOptional<double>& OutSample);
+
+public:
+	/**
+	 * Gets a sample of the calculated moving value, based on current bin data.
+	 *
+	 * @return	Outputs the calculated sample value
+	 */
+	double GetSample() const;
+
+	/**
+	 * Resets the accumulated bin/sample data.
+	 */
+	void Reset();
+
+private:
+	void ResetNewAndSkippedBins(int32 BinIdx, TOptional<double>& OutSample);
+
+
+private:
+	/** A view of the allocated bins, for collecting moving values */
+	const TArrayView<FBin> BinsView;
+
+	/** The amount of time tracked by each individual bin */
+	const double TimePerBin = 0.0;
+
+	/** The total range of all bins */
+	const double BinRange = 0.0;
+
+	/** The mode for sample calculations - typically moving average, or moving sum */
+	const EBinnedValueMode SampleMode = EBinnedValueMode::MovingAvg;
+
+	/** The index of the first bin that a value was stored in, to detect when all bins are rolled over */
+	int32 FirstBinIndex = INDEX_NONE;
+
+	/** The index of the last bin to be written to */
+	int32 LastWrittenBinIndex = INDEX_NONE;
+
+	/** The total summed value of all bins */
+	double TotalSum = 0.0;
+
+	/** The total count of all bins */
+	int32 TotalCount = 0;
+};
+
+/**
+ * Tracks/stores a moving average or moving sum over a period of time, binning the accumulated values into different time periods,
+ * for smoothing and balancing the memory/accuracy cost of dropping old values as they fall out of the measured time period.
+ *
+ * Use TBinnedMovingAvg/TBinnedMovingSum aliases, instead of using this directly.
+ *
+ * Measurements are input using AddMeasurement, and samples/values are output to the specified consumer, which is set using SetConsumer.
+ *
+ * Samples are automatically output when all of the data bins fill and are rolled over - this guarantees perfect accuracy over the full time period.
+ * Samples can be manually output using Flush, but this can be less accurate depending on the use case - see Flush for more information.
+ *
+ * @param ConsumerType	The type which will be used for consuming samples.
+ * @param NumBins		The number of bins to use.
+ * @param InMode		Internal
+ */
+template<typename ConsumerType, int32 NumBins, EBinnedValueMode InMode>
+class TBinnedMovingValue : public FBinnedMovingValueBase, public TSampleProducer<ConsumerType>
+{
+public:
+	/**
+	 * Main constructor
+	 *
+	 * @param TimePerBin	The amount of time that each bin represents.
+	 */
+	explicit TBinnedMovingValue(double TimePerBin)
+		: FBinnedMovingValueBase(MakeArrayView(Bins, NumBins), TimePerBin, InMode)
+	{
+	}
+
+	/**
+	 * Inputs a measurement value and the approximate time it was measured.
+	 *
+	 * @param TimeVal		The time of the measurement.
+	 * @param Value			The measurement value
+	 */
+	void AddMeasurement(double TimeVal, double Value)
+	{
+		TOptional<double> FinalSample;
+
+		TSampleProducer<ConsumerType>::PeekMeasurement(TimeVal, Value);
+		AddMeasurement_Implementation(TimeVal, Value, FinalSample);
+
+		if (FinalSample.IsSet())
+		{
+			this->OutputSample(FinalSample.GetValue());
+		}
+	}
+
+	/**
+	 * Flushes a sample from the current bin values to the consumer.
+	 *
+	 * Since this is a moving value, none of the bins are reset when a sample is output (so the flushed sample can partially reflect old values),
+	 * and the current bin will have been reset recently so it can be filled (so the flushed value may not reflect the full time period).
+	 *
+	 * The user must consider whether this is accurate/suitable for their use case.
+	 */
+	void Flush()
+	{
+		this->OutputSample(GetSample());
+	}
+
+private:
+	/** Allocation for the bins which collect moving values */
+	FBin Bins[NumBins] = {};
+};
+
+/**
+ * Version of TBinnedMovingValue which takes permanently incrementing values (e.g. frame count),
+ * and uses the delta since the previous value as the measurement (this means the very first value is not measured).
+ *
+ * Use TDeltaBinnedMovingAvg/TDeltaBinnedMovingSum aliases, instead of using this directly.
+ *
+ * @param ConsumerType	The type which will be used for consuming samples.
+ * @param NumBins		The number of bins to use.
+ * @param InMode		Internal
+ */
+template<typename ConsumerType, int32 NumBins, EBinnedValueMode InMode>
+class TBinnedDeltaMovingValue : private TBinnedMovingValue<ConsumerType, NumBins, InMode>
+{
+	using Super = TBinnedMovingValue<ConsumerType, NumBins, InMode>;
+
+public:
+	// Using Super for constructors can break clang
+	using TBinnedMovingValue<ConsumerType, NumBins, InMode>::TBinnedMovingValue;
+	using Super::SetConsumer;
+	using Super::GetSample;
+	using Super::Flush;
+
+	/**
+	 * Inputs a value whose delta is calculated for measurement, and the approximate time it was measured.
+	 *
+	 * @param TimeVal		The time of the measurement.
+	 * @param Value			The value to calculate the delta of, for measurement
+	 */
+	void AddMeasurement(double TimeVal, double Value)
+	{
+		if (LastValue != std::numeric_limits<double>::min())
+		{
+			Super::AddMeasurement(TimeVal, (Value - LastValue));
+		}
+
+		LastValue = Value;
+	}
+
+	/**
+	 * Resets the accumulated bin/sample data (the very first value after this is not measured).
+	 */
+	void Reset()
+	{
+		Super::Reset();
+
+		LastValue = std::numeric_limits<double>::min();
+	}
+
+private:
+	/** The last input value, for tracking deltas */
+	double LastValue = std::numeric_limits<double>::min();
+};
+
+/**
+ * Alias for tracking a moving average. See TBinnedMovingValue.
+ */
+template<typename ConsumerType, int32 NumBins>
+using TBinnedMovingAvg = TBinnedMovingValue<ConsumerType, NumBins, EBinnedValueMode::MovingAvg>;
+
+/**
+ * Alias for tracking a moving average using the delta of the input value. See TBinnedDeltaMovingValue/TBinnedMovingValue.
+ */
+template<typename ConsumerType, int32 NumBins>
+using TDeltaBinnedMovingAvg = TBinnedDeltaMovingValue<ConsumerType, NumBins, EBinnedValueMode::MovingAvg>;
+
+/**
+ * Alias for tracking a moving sum. See TBinnedMovingValue.
+ */
+template<typename ConsumerType, int32 NumBins>
+using TBinnedMovingSum = TBinnedMovingValue<ConsumerType, NumBins, EBinnedValueMode::MovingSum>;
+
+/**
+ * Alias for tracking a moving sum using the delta of the input value. See TBinnedDeltaMovingValue/TBinnedMovingValue.
+ */
+template<typename ConsumerType, int32 NumBins>
+using TDeltaBinnedMovingSum = TBinnedDeltaMovingValue<ConsumerType, NumBins, EBinnedValueMode::MovingSum>;
+
+}
+
+/**
+ * Quick named parameter emulation, e.g:
+ *	TDeltaBinnedMovingSum<FSampleMinMaxAvg, TBinParms::NumBins(6)> PageFaultsPerMinTracking{TBinParms::TimePerBin(10.0)}
+ */
+namespace TBinParms
+{
+	static constexpr int32 NumBins(int32 InNumBins)
+	{
+		return InNumBins;
+	}
+
+	static constexpr double TimePerBin(double InBinSize)
+	{
+		return InBinSize;
+	}
+}
+
+
+namespace UE::Net
+{
+
+/**
+ * Determines how Min/Max values are calculated, for FSampleMinMaxAvg
+ */
+enum EMinMaxValueMode
+{
+	/** Min/Max values are calculated based on the final (possibly averaged/summed) sample */
+	PerSample,
+
+	/** Min/Max values are calculated based on every individual measurement accumulated in the sample (TSampleProducer only!) */
+	PerMeasurement,
+
+	/** Min/Max values are calculated based on the average of measurements accumulated over the specified time period (TSampleProducer only!) */
+	PerTimeSeconds
+};
+
+/**
+ * Base class for tracking the minimum, maximum and average of the passed in measurement values.
+ *
+ * This base class can handle reading.
+ * The TSampleMinMaxAvg subclass is compatible as a consumer for TSampleProducer, and is used for writing.
+ */
+class NETCORE_API FSampleMinMaxAvg
+{
+	template<EMinMaxValueMode, int32> friend class TSampleMinMaxAvg;
+
+private:
+	void AddSample_Internal(double Value);
+
+public:
+	/**
+	 * Resets the accumulated sample data
+	 */
+	void Reset();
+
+	/**
+	 * Gets the minimum sample value encountered
+	 *
+	 * @return	The minimum sample value
+	 */
+	double GetMin() const
+	{
+		return MinValue;
+	}
+
+	/**
+	 * Gets the maximum sample value encountered
+	 *
+	 * @return	The maximum sample value
+	 */
+	double GetMax() const
+	{
+		return MaxValue;
+	}
+
+	/**
+	 * Gets the current/latest sample value
+	 *
+	 * @return	The current sample value
+	 */
+	double GetCurrent() const
+	{
+		return CurrentValue;
+	}
+
+	/**
+	 * Gets the average sample value
+	 *
+	 * @return	The average sample value
+	 */
+	double GetAvg() const
+	{
+		return TotalCount == 0 ? 0.0 : (TotalSum / static_cast<double>(TotalCount));
+	}
+
+	/**
+	 * Gets the sum of all sample values tracked
+	 *
+	 * @return	The sum of samples
+	 */
+	double GetSum() const
+	{
+		return TotalSum;
+	}
+
+	/**
+	 * Gets the number of samples tracked
+	 *
+	 * @return	The number of samples
+	 */
+	int32 GetSampleCount() const
+	{
+		return TotalCount;
+	}
+
+private:
+	/** The minimum sample value encountered */
+	double MinValue = std::numeric_limits<double>::max();
+
+	/** The maximum sample value encountered */
+	double MaxValue = std::numeric_limits<double>::min();
+
+	/** The current sample value */
+	double CurrentValue = 0.0;
+
+	/** The total sum of all samples */
+	double TotalSum = 0.0;
+
+	/** The number of samples tracked */
+	int32 TotalCount = 0;
+};
+
+/**
+ * Writable subclass of FSampleMinMaxAvg. Can be cast to FSampleMinMaxAvg for read-only usage.
+ *
+ * TSampleMinMaxAvg<EMinMaxValueMode::PerSample>
+ * TSampleMinMaxAvg<EMinMaxValueMode::PerMeasurement>
+ * TSampleMinMaxAvg<EMinMaxValueMode::PerTimeSeconds, TSampleParms::TimeSeconds(1)>
+ */
+template<EMinMaxValueMode Mode, int32 MinMaxTimeSeconds=1>
+class TSampleMinMaxAvg;
+
+template<>
+class TSampleMinMaxAvg<EMinMaxValueMode::PerSample> : public FSampleMinMaxAvg
+{
+public:
+	/**
+	 * Inputs a sample (usually from a TSampleProducer)
+	 */
+	void AddSample(double Value)
+	{
+		if (Value < MinValue)
+		{
+			MinValue = Value;
+		}
+
+		if (Value > MaxValue)
+		{
+			MaxValue = Value;
+		}
+
+		AddSample_Internal(Value);
+	}
+};
+
+template<>
+class TSampleMinMaxAvg<EMinMaxValueMode::PerMeasurement> : public FSampleMinMaxAvg
+{
+public:
+	/**
+	 * Inputs a sample (usually from a TSampleProducer)
+	 */
+	inline void AddSample(double Value)
+	{
+		AddSample_Internal(Value);
+	}
+
+	void PeekMeasurement(double TimeVal, double Value)
+	{
+		if (Value < MinValue)
+		{
+			MinValue = Value;
+		}
+
+		if (Value > MaxValue)
+		{
+			MaxValue = Value;
+		}
+	}
+};
+
+template<int32 MinMaxTimeSeconds>
+class TSampleMinMaxAvg<EMinMaxValueMode::PerTimeSeconds, MinMaxTimeSeconds> : public FSampleMinMaxAvg
+{
+public:
+	/**
+	 * Inputs a sample (usually from a TSampleProducer)
+	 */
+	inline void AddSample(double Value)
+	{
+		AddSample_Internal(Value);
+	}
+
+	void PeekMeasurement(double TimeVal, double Value)
+	{
+		if (UNLIKELY(MinMaxPeriodStart == 0.0))
+		{
+			MinMaxPeriodStart = TimeVal;
+		}
+		else if ((TimeVal - MinMaxPeriodStart) >= static_cast<double>(MinMaxTimeSeconds))
+		{
+			const double NewMinMaxSample = MinMaxPeriodSum / MinMaxPeriodCount;
+
+			if (NewMinMaxSample < MinValue)
+			{
+				MinValue = NewMinMaxSample;
+			}
+
+			if (NewMinMaxSample > MaxValue)
+			{
+				MaxValue = NewMinMaxSample;
+			}
+
+			MinMaxPeriodStart = TimeVal;
+			MinMaxPeriodSum = 0.0;
+			MinMaxPeriodCount = 0;
+		}
+
+		MinMaxPeriodSum += Value;
+		MinMaxPeriodCount++;
+	}
+
+private:
+	/** The time at which the last Min/Max averaging period started */
+	double MinMaxPeriodStart = 0.0;
+
+	/** The sum of measurements during the current Min/Max averaging period */
+	double MinMaxPeriodSum = 0.0;
+
+	/** The number of measurements during the current Min/Max averaging period */
+	int32 MinMaxPeriodCount = 0;
+};
+
+/**
+ * Class level 'using namespace UE::Net', for NetStatsUtils.h types. Inherit as protected.
+ */
+class FUsingNetStatsUtils
+{
+public:
+	template<typename ConsumerType> using TSampleProducer = TSampleProducer<ConsumerType>;
+	using EBinnedValueMode = EBinnedValueMode;
+	template<typename ConsumerType, int32 NumBins, EBinnedValueMode InMode> using TBinnedMovingValue =
+		TBinnedMovingValue<ConsumerType, NumBins, InMode>;
+	template<typename ConsumerType, int32 NumBins, EBinnedValueMode InMode> using TBinnedDeltaMovingValue =
+		TBinnedDeltaMovingValue<ConsumerType, NumBins, InMode>;
+	template<typename ConsumerType, int32 NumBins> using TBinnedMovingAvg = TBinnedMovingAvg<ConsumerType, NumBins>;
+	template<typename ConsumerType, int32 NumBins> using TDeltaBinnedMovingAvg = TDeltaBinnedMovingAvg<ConsumerType, NumBins>;
+	template<typename ConsumerType, int32 NumBins> using TBinnedMovingSum = TBinnedMovingSum<ConsumerType, NumBins>;
+	template<typename ConsumerType, int32 NumBins> using TDeltaBinnedMovingSum = TDeltaBinnedMovingSum<ConsumerType, NumBins>;
+	using EMinMaxValueMode = EMinMaxValueMode;
+	using FSampleMinMaxAvg = FSampleMinMaxAvg;
+	template<EMinMaxValueMode Mode, int32 MinMaxTimeSeconds=1> using TSampleMinMaxAvg = TSampleMinMaxAvg<Mode, MinMaxTimeSeconds>;
+};
+
+}
+
+/**
+ * Quick named parameter emulation, e.g:
+ *	TSampleMinMaxAvg<EMinMaxValueMode::PerTimeSeconds, TSampleParms::TimeSeconds(1)> GameThreadCPUPercentPerMinStats;
+ */
+namespace TSampleParms
+{
+	static constexpr int32 TimeSeconds(int32 InTimeSeconds)
+	{
+		return InTimeSeconds;
+	}
+}
+
