CIS fix.

diff --git a/Engine/Source/Runtime/Engine/Public/Net/NetPing.h b/Engine/Source/Runtime/Engine/Public/Net/NetPing.h
--- a/Engine/Source/Runtime/Engine/Public/Net/NetPing.h
+++ b/Engine/Source/Runtime/Engine/Public/Net/NetPing.h
@@ -8,6 +8,7 @@
 
 // Forward declarations
 enum class EPingType : uint32;
+class FInternetAddr;
 
 namespace UE::Net
 {
