Fix memory stomps caused by lack of synchronization on FUpdateStreamingRenderAssetsTask

#localization none
#tests ran a replay several times on a windows client

diff --git a/Engine/Source/Runtime/Engine/Private/Streaming/StreamingManagerTexture.cpp b/Engine/Source/Runtime/Engine/Private/Streaming/StreamingManagerTexture.cpp
--- a/Engine/Source/Runtime/Engine/Private/Streaming/StreamingManagerTexture.cpp
+++ b/Engine/Source/Runtime/Engine/Private/Streaming/StreamingManagerTexture.cpp
@@ -1169,8 +1169,6 @@ void FRenderAssetStreamingManager::FastForceFullyResident(UStreamableRenderAsset
  */
 void FRenderAssetStreamingManager::UpdateStreamingRenderAssets(int32 StageIndex, int32 NumUpdateStages, bool bWaitForMipFading, bool bAsync)
 {
-	TArray<FStreamingRenderAsset>& StreamingRenderAssets = GetStreamingRenderAssetsAsyncSafe();
-
 	if ( StageIndex == 0 )
 	{
 		CurrentUpdateStreamingRenderAssetIndex = 0;
@@ -1178,7 +1176,7 @@ void FRenderAssetStreamingManager::UpdateStreamingRenderAssets(int32 StageIndex,
 	}
 
 	int32 StartIndex = CurrentUpdateStreamingRenderAssetIndex;
-	int32 EndIndex = StreamingRenderAssets.Num() * (StageIndex + 1) / NumUpdateStages;
+	int32 EndIndex = AsyncUnsafeStreamingRenderAssets.Num() * (StageIndex + 1) / NumUpdateStages;
 
 #if !STATS
 	if (GAllowParallelUpdateStreamingRenderAssets)
@@ -1215,7 +1213,7 @@ void FRenderAssetStreamingManager::UpdateStreamingRenderAssets(int32 StageIndex,
 		for (int32 i = 0; i < NumThreadTasks; ++i)
 		{
 			int32 NumAssetsToProcess = FMath::Min<int32>(NumRemaining, NumItemsPerGroup);
-			Packets.Add(FPacket(Start, Start + NumAssetsToProcess, StreamingRenderAssets));
+			Packets.Add(FPacket(Start, Start + NumAssetsToProcess, AsyncUnsafeStreamingRenderAssets));
 			Start += NumAssetsToProcess;
 			NumRemaining -= NumAssetsToProcess;
 			if (NumRemaining <= 0)
@@ -1262,7 +1260,7 @@ void FRenderAssetStreamingManager::UpdateStreamingRenderAssets(int32 StageIndex,
 	{
 		for (int32 Index = StartIndex; Index < EndIndex; ++Index)
 		{
-			FStreamingRenderAsset& StreamingRenderAsset = StreamingRenderAssets[Index];
+			FStreamingRenderAsset& StreamingRenderAsset = AsyncUnsafeStreamingRenderAssets[Index];
 			FPlatformMisc::Prefetch(&StreamingRenderAsset + 1);
 
 			// Is this texture/mesh marked for removal? Will get cleanup once the async task is done.
@@ -1607,7 +1605,7 @@ public:
 	}
 	static FORCEINLINE ESubsequentsMode::Type GetSubsequentsMode()
 	{
-		return ESubsequentsMode::FireAndForget;
+		return ESubsequentsMode::TrackSubsequents;
 	}
 	void DoTask(ENamedThreads::Type CurrentThread, const FGraphEventRef& MyCompletionGraphEvent)
 	{
@@ -1727,12 +1725,6 @@ void FRenderAssetStreamingManager::UpdateResourceStreaming( float DeltaTime, boo
 		IncrementalUpdate(1.f / (float)FMath::Max(NumRenderAssetProcessingStages - 1, 1), true); // -1 since we don't want to do anything at stage 0.
 		++ProcessingStage;
 
-		// TODO: figure out whether letting the task to run past is causing memory stomps
-		if (bOverlappedExecution && StreamingRenderAssetsSyncEvent.IsValid())
-		{
-			StreamingRenderAssetsSyncEvent->Wait(ENamedThreads::GameThread);
-		}
-
 		STAT(GatheredStats.UpdateStreamingDataCycles = FMath::Max<uint32>(ProcessingStage > 2 ? GatheredStats.UpdateStreamingDataCycles : 0, FPlatformTime::Cycles() - StartTime);)
 	}
 	else if (AsyncWork->IsDone())
