MaterialInstanceDynamic::SetDoubleVectorParameter should take an FVector4 rather than FVector3.

diff --git a/Engine/Source/Runtime/Engine/Classes/Materials/MaterialInstanceDynamic.h b/Engine/Source/Runtime/Engine/Classes/Materials/MaterialInstanceDynamic.h
--- a/Engine/Source/Runtime/Engine/Classes/Materials/MaterialInstanceDynamic.h
+++ b/Engine/Source/Runtime/Engine/Classes/Materials/MaterialInstanceDynamic.h
@@ -81,7 +81,7 @@ class ENGINE_API UMaterialInstanceDynamic : public UMaterialInstance
 
 	/** Set an MID vector parameter value */
 	UFUNCTION(BlueprintCallable, meta = (Keywords = "SetVectorParameterValue"), Category = "Rendering|Material")
-	void SetDoubleVectorParameterValue(FName ParameterName, FVector Value);
+	void SetDoubleVectorParameterValue(FName ParameterName, FVector4 Value);
 
 	/** Set an MID vector parameter value, using MPI (to allow access to layer parameters) */
 	UFUNCTION(BlueprintCallable, meta = (Keywords = "SetColorParameterValue"), Category = "Rendering|Material")
diff --git a/Engine/Source/Runtime/Engine/Private/Materials/MaterialInstanceDynamic.cpp b/Engine/Source/Runtime/Engine/Private/Materials/MaterialInstanceDynamic.cpp
--- a/Engine/Source/Runtime/Engine/Private/Materials/MaterialInstanceDynamic.cpp
+++ b/Engine/Source/Runtime/Engine/Private/Materials/MaterialInstanceDynamic.cpp
@@ -126,7 +126,7 @@ void UMaterialInstanceDynamic::SetVectorParameterValue(FName ParameterName, FLin
 	SetVectorParameterValueInternal(ParameterInfo,Value);
 }
 
-void UMaterialInstanceDynamic::SetDoubleVectorParameterValue(FName ParameterName, FVector Value)
+void UMaterialInstanceDynamic::SetDoubleVectorParameterValue(FName ParameterName, FVector4 Value)
 {
 	FMaterialParameterInfo ParameterInfo(ParameterName);
 	SetDoubleVectorParameterValueInternal(ParameterInfo, Value);
