CIS fixes.

diff --git a/Engine/Source/Runtime/Engine/Private/Net/NetPing.cpp b/Engine/Source/Runtime/Engine/Private/Net/NetPing.cpp
--- a/Engine/Source/Runtime/Engine/Private/Net/NetPing.cpp
+++ b/Engine/Source/Runtime/Engine/Private/Net/NetPing.cpp
@@ -823,7 +823,7 @@ FString FNetPing::GetUDPQoSAddr()
 	}
 
 
-	PingIPStr.Append(TEXT(':'));
+	PingIPStr.AppendChar(TEXT(':'));
 	PingIPStr.Appendf(TEXT("%i"), (OverridePingUDPQosPort != INDEX_NONE ? OverridePingUDPQosPort : GNetPingUDPPort));
 
 	return PingIPStr.ToString();
