- Notify about changes to min lod so system can respond

diff --git a/Engine/Source/Runtime/Engine/Classes/Engine/StaticMesh.h b/Engine/Source/Runtime/Engine/Classes/Engine/StaticMesh.h
--- a/Engine/Source/Runtime/Engine/Classes/Engine/StaticMesh.h
+++ b/Engine/Source/Runtime/Engine/Classes/Engine/StaticMesh.h
@@ -748,6 +748,9 @@ public:
 	void SetQualityLevelMinLOD(FPerQualityLevelInt InMinLOD)
 	{
 		MinQualityLevelLOD = MoveTemp(InMinLOD);
+#if WITH_EDITOR
+		OnMeshChanged.Broadcast();
+#endif
 	}
 
 	UFUNCTION(BlueprintPure, Category = StaticMesh)
@@ -800,6 +803,9 @@ public:
 		PRAGMA_DISABLE_DEPRECATION_WARNINGS
 		MinLOD = MoveTemp(InMinLOD);
 		PRAGMA_ENABLE_DEPRECATION_WARNINGS
+#if WITH_EDITOR
+		OnMeshChanged.Broadcast();
+#endif
 	}
 
 	UFUNCTION(BlueprintPure, Category=StaticMesh)
