[Backout] - CL19276151
Original CL Desc
-----------------------------------------------------------------
- Notify about changes to min lod so system can respond

diff --git a/Engine/Source/Runtime/Engine/Classes/Engine/StaticMesh.h b/Engine/Source/Runtime/Engine/Classes/Engine/StaticMesh.h
--- a/Engine/Source/Runtime/Engine/Classes/Engine/StaticMesh.h
+++ b/Engine/Source/Runtime/Engine/Classes/Engine/StaticMesh.h
@@ -748,9 +748,6 @@ public:
 	void SetQualityLevelMinLOD(FPerQualityLevelInt InMinLOD)
 	{
 		MinQualityLevelLOD = MoveTemp(InMinLOD);
-#if WITH_EDITOR
-		OnMeshChanged.Broadcast();
-#endif
 	}
 
 	UFUNCTION(BlueprintPure, Category = StaticMesh)
@@ -803,9 +800,6 @@ public:
 		PRAGMA_DISABLE_DEPRECATION_WARNINGS
 		MinLOD = MoveTemp(InMinLOD);
 		PRAGMA_ENABLE_DEPRECATION_WARNINGS
-#if WITH_EDITOR
-		OnMeshChanged.Broadcast();
-#endif
 	}
 
 	UFUNCTION(BlueprintPure, Category=StaticMesh)
