Shrink USoundCue from 1376 to 1336
#rnx

diff --git a/Engine/Source/Runtime/Engine/Classes/Sound/SoundConcurrency.h b/Engine/Source/Runtime/Engine/Classes/Sound/SoundConcurrency.h
--- a/Engine/Source/Runtime/Engine/Classes/Sound/SoundConcurrency.h
+++ b/Engine/Source/Runtime/Engine/Classes/Sound/SoundConcurrency.h
@@ -80,7 +80,13 @@ struct ENGINE_API FSoundConcurrencySettings
 
 	/* Whether or not to limit the concurrency to per sound owner (i.e. the actor that plays the sound). If the sound doesn't have an owner, it falls back to global concurrency. */
 	UPROPERTY(EditAnywhere, BlueprintReadWrite, Category = Concurrency)
-	uint32 bLimitToOwner:1;
+	uint8 bLimitToOwner:1;
+
+	/**
+	 * Whether or not volume scaling can recover volume ducking behavior when concurrency group sounds stop (default scale mode only).
+	 */
+	UPROPERTY(EditAnywhere, BlueprintReadWrite, Category = "Volume Scaling", meta = (DisplayName = "Can Recover", EditCondition = "VolumeScaleMode == EConcurrencyVolumeScaleMode::Default"))
+	uint8 bVolumeScaleCanRelease:1;
 
 	/** Which concurrency resolution policy to use if max voice count is reached. */
 	UPROPERTY(EditAnywhere, BlueprintReadWrite, Category = Concurrency)
@@ -113,12 +119,6 @@ public:
 	UPROPERTY(EditAnywhere, BlueprintReadWrite, Category = "Volume Scaling", meta = (DisplayName = "Duck Time", UIMin = "0.0", ClampMin = "0.0", UIMax = "10.0", ClampMax = "1000000.0"))
 	float VolumeScaleAttackTime;
 
-	/**
-	 * Whether or not volume scaling can recover volume ducking behavior when concurrency group sounds stop (default scale mode only).
-	 */
-	UPROPERTY(EditAnywhere, BlueprintReadWrite, Category = "Volume Scaling", meta = (DisplayName = "Can Recover", EditCondition = "VolumeScaleMode == EConcurrencyVolumeScaleMode::Default"))
-	uint32 bVolumeScaleCanRelease:1;
-
 	/**
 	 * Time taken to recover volume scalar duck.
 	 */
@@ -135,12 +135,12 @@ public:
 	FSoundConcurrencySettings()
 		: MaxCount(16)
 		, bLimitToOwner(0)
+		, bVolumeScaleCanRelease(0)
 		, ResolutionRule(EMaxConcurrentResolutionRule::StopFarthestThenOldest)
 		, RetriggerTime(0.0f)
 		, VolumeScale(1.0f)
 		, VolumeScaleMode(EConcurrencyVolumeScaleMode::Default)
 		, VolumeScaleAttackTime(0.01f)
-		, bVolumeScaleCanRelease(0)
 		, VolumeScaleReleaseTime(0.5f)
 		, VoiceStealReleaseTime(0.0f)
 	{}
diff --git a/Engine/Source/Runtime/Engine/Classes/Sound/SoundCue.h b/Engine/Source/Runtime/Engine/Classes/Sound/SoundCue.h
--- a/Engine/Source/Runtime/Engine/Classes/Sound/SoundCue.h
+++ b/Engine/Source/Runtime/Engine/Classes/Sound/SoundCue.h
@@ -89,10 +89,6 @@ class ENGINE_API USoundCue : public USoundBase
 {
 	GENERATED_UCLASS_BODY()
 
-	/* Makes this sound cue automatically load any sound waves it can play into the cache when it is loaded. */
-	UPROPERTY(EditAnywhere, Category = Memory)
-	uint32 bPrimeOnLoad : 1;
-
 	UPROPERTY()
 	TObjectPtr<USoundNode> FirstNode;
 
@@ -125,6 +121,11 @@ private:
 	float MaxAudibleDistance;
 
 public:
+
+	/* Makes this sound cue automatically load any sound waves it can play into the cache when it is loaded. */
+	UPROPERTY(EditAnywhere, Category = Memory)
+	uint8 bPrimeOnLoad : 1;
+
 	/* Indicates whether attenuation should use the Attenuation Overrides or the Attenuation Settings asset */
 	UPROPERTY(EditAnywhere, Category = Attenuation)
 	uint8 bOverrideAttenuation : 1;
@@ -134,9 +135,7 @@ public:
 	uint8 bExcludeFromRandomNodeBranchCulling : 1;
 
 private:
-	UPROPERTY()
-	int32 CookedQualityIndex = INDEX_NONE;
-	
+
 	/** Whether a sound has play when silent enabled (i.e. for a sound cue, if any sound wave player has it enabled). */
 	UPROPERTY()
 	uint8 bHasPlayWhenSilent : 1;
@@ -147,6 +146,9 @@ private:
 	uint8 bShouldApplyInteriorVolumesCached : 1;
 	uint8 bIsRetainingAudio : 1;
 
+	UPROPERTY()
+	int32 CookedQualityIndex = INDEX_NONE;
+
 public:
 
 	//~ Begin UObject Interface.
