Fix a bug where translucent material couldn't get reflection on mobile deferred.

diff --git a/Engine/Source/Runtime/Engine/Private/Components/ReflectionCaptureComponent.cpp b/Engine/Source/Runtime/Engine/Private/Components/ReflectionCaptureComponent.cpp
--- a/Engine/Source/Runtime/Engine/Private/Components/ReflectionCaptureComponent.cpp
+++ b/Engine/Source/Runtime/Engine/Private/Components/ReflectionCaptureComponent.cpp
@@ -778,7 +778,7 @@ void UReflectionCaptureComponent::SendRenderTransform_Concurrent()
 void UReflectionCaptureComponent::OnRegister()
 {
 	const ERHIFeatureLevel::Type FeatureLevel = GetWorld()->FeatureLevel;
-	const bool bEncodedDataRequired = (FeatureLevel == ERHIFeatureLevel::ES3_1) && !IsMobileDeferredShadingEnabled(GMaxRHIShaderPlatform);
+	const bool bEncodedDataRequired = (FeatureLevel == ERHIFeatureLevel::ES3_1);
 
 	if (bEncodedDataRequired)
 	{
