Chaos: Fix for materials in RBAN not applying restitution in low to mid speed collisions

diff --git a/Engine/Source/Runtime/Engine/Private/PhysicsEngine/ImmediatePhysics/ImmediatePhysicsChaos/ImmediatePhysicsSimulation_Chaos.cpp b/Engine/Source/Runtime/Engine/Private/PhysicsEngine/ImmediatePhysics/ImmediatePhysicsChaos/ImmediatePhysicsSimulation_Chaos.cpp
--- a/Engine/Source/Runtime/Engine/Private/PhysicsEngine/ImmediatePhysics/ImmediatePhysicsChaos/ImmediatePhysicsSimulation_Chaos.cpp
+++ b/Engine/Source/Runtime/Engine/Private/PhysicsEngine/ImmediatePhysics/ImmediatePhysicsChaos/ImmediatePhysicsSimulation_Chaos.cpp
@@ -79,7 +79,7 @@ int32 ChaosImmediate_Collision_NumVelocityShockPropagationIterations = 0;
 int32 ChaosImmediate_Collision_Priority = 1;
 Chaos::FRealSingle ChaosImmediate_Collision_CullDistance = -1.0f;
 Chaos::FRealSingle ChaosImmediate_Collision_MaxDepenetrationVelocity = -1.0f;
-Chaos::FRealSingle ChaosImmediate_Collision_RestitutionThresholdMultiplier = 3.0f;
+Chaos::FRealSingle ChaosImmediate_Collision_RestitutionThresholdMultiplier = 1.5f;
 int32 ChaosImmediate_Collision_RestitutionEnabled = true;
 int32 ChaosImmediate_Collision_DeferNarrowPhase = 1;
 int32 ChaosImmediate_Collision_UseManifolds = 0;
