Added a helper class for logging amortized checkpoint steps
Added "demo.ForcePersistentLevelPriority" to always record the persistent level first when using level streaming fixes and prioritization
Added "demo.DestructionInfoPriority" for allowing destruction infos to be prioritized lower than max int32
Added "demo.LateDestructionInfoPrioritize" to prioritize destruction infos after the active actor set
Added "demo.ViewTargetPriorityScale" to scale view target replay priority
Added "Demo.MaxDesiredRecordTimeMS", "Demo.CheckpointSaveMaxMSPerFrame", "Demo.ActorPrioritizationEnabled", and "Demo.SetLocalViewerOverride" for updating those settings on the current world's demo driver
Added "demo.MaximumRecDestructionInfoTime" to cap the percent of a budgeted frame that can be spent replicating destruction infos
Removing STAT_RepLayout_AddPropertyCmd, it has served its purpose for now and spams stat captures when a rep layout is initialized
Clamping replay priority into a valid range before rounding
Misc iterator cleanup
Adjusting recording CSV stats around actor vs destruction infos

diff --git a/Engine/Source/Runtime/Engine/Classes/Engine/DemoNetDriver.h b/Engine/Source/Runtime/Engine/Classes/Engine/DemoNetDriver.h
--- a/Engine/Source/Runtime/Engine/Classes/Engine/DemoNetDriver.h
+++ b/Engine/Source/Runtime/Engine/Classes/Engine/DemoNetDriver.h
@@ -655,7 +655,7 @@ private:
 	 * @return True if there is time remaining to replicate more actors. False otherwise.
 	 */
 	bool ReplicatePrioritizedActors(const FDemoActorPriority* ActorsToReplicate, uint32 Count, class FRepActorsParams& Params);
-	bool ReplicatePrioritizedActor(const FActorPriority& ActorPriority, const class FRepActorsParams& Params);
+	bool ReplicatePrioritizedActor(const FActorPriority& ActorPriority, class FRepActorsParams& Params);
 
 	friend class FPendingTaskHelper;
 
@@ -688,6 +688,9 @@ private:
 	// for demo recording. Only used if MaxDesiredRecordTimeMS > 0.
 	float RecordBuildConsiderAndPrioritizeTimeSlice;
 
+	// Max percent of time to spend replicating prioritized destruction infos. Only used if MaxDesiredRecordTimeMS > 0.
+	float RecordDestructionInfoReplicationTimeSlice;
+
 	void AdjustConsiderTime(const float ReplicatedPercent);
 
 	bool ProcessFastForwardPackets(TArrayView<FPlaybackPacket> Packets, const TSet<int32>& LevelIndices);
diff --git a/Engine/Source/Runtime/Engine/Private/DemoNetDriver.cpp b/Engine/Source/Runtime/Engine/Private/DemoNetDriver.cpp
--- a/Engine/Source/Runtime/Engine/Private/DemoNetDriver.cpp
+++ b/Engine/Source/Runtime/Engine/Private/DemoNetDriver.cpp
@@ -91,7 +91,7 @@ TAutoConsoleVariable<int32> CVarWithGameSpecificFrameData(TEXT("demo.WithGameSpe
 static TAutoConsoleVariable<float> CVarDemoIncreaseRepPrioritizeThreshold(TEXT("demo.IncreaseRepPrioritizeThreshold"), 0.9, TEXT("The % of Replicated to Prioritized actors at which prioritize time will be decreased."));
 static TAutoConsoleVariable<float> CVarDemoDecreaseRepPrioritizeThreshold(TEXT("demo.DecreaseRepPrioritizeThreshold"), 0.7, TEXT("The % of Replicated to Prioritized actors at which prioritize time will be increased."));
 static TAutoConsoleVariable<float> CVarDemoMinimumRepPrioritizeTime(TEXT("demo.MinimumRepPrioritizePercent"), 0.3, TEXT("Minimum percent of time that must be spent prioritizing actors, regardless of throttling."));
-static TAutoConsoleVariable<float> CVarDemoMaximumRepPrioritizeTime(TEXT("demo.MaximumRepPrioritizePercent"), 0.8, TEXT("Maximum percent of time that may be spent prioritizing actors, regardless of throttling."));
+static TAutoConsoleVariable<float> CVarDemoMaximumRepPrioritizeTime(TEXT("demo.MaximumRepPrioritizePercent"), 0.7, TEXT("Maximum percent of time that may be spent prioritizing actors, regardless of throttling."));
 
 static TAutoConsoleVariable<int32> CVarFastForwardLevelsPausePlayback(TEXT("demo.FastForwardLevelsPausePlayback"), 0, TEXT("If true, pause channels and playback while fast forward levels task is running."));
 
@@ -140,6 +140,90 @@ namespace DemoNetDriverRecordingPrivate
 		RecordUnicastRPCs,
 		TEXT("When true, also record unicast client rpcs on actors that share a net driver name with the demo driver.")
 	);
+
+	static TAutoConsoleVariable<int32> CVarDemoForcePersistentLevelPriority(TEXT("demo.ForcePersistentLevelPriority"), 0, TEXT("If true, force persistent level to record first when prioritizing and using streaming level fixes."));
+	static TAutoConsoleVariable<int32> CVarDemoDestructionInfoPriority(TEXT("demo.DestructionInfoPriority"), MAX_int32, TEXT("Replay net priority assigned to destruction infos during recording."));
+	static TAutoConsoleVariable<int32> CVarDemoLateDestructionInfoPrioritize(TEXT("demo.LateDestructionInfoPrioritize"), 0, TEXT("If true, process destruction infos at the end of the prioritization phase."));
+	static TAutoConsoleVariable<float> CVarDemoViewTargetPriorityScale(TEXT("demo.ViewTargetPriorityScale"), 3.0, TEXT("Scale view target priority by this value when prioritization is enabled."));
+	static TAutoConsoleVariable<float> CVarDemoMaximumRecDestructionInfoTime(TEXT("demo.MaximumRecDestructionInfoTime"), 0.2, TEXT("Maximum percentage of frame to use replicating destruction infos, if per frame limit is enabled."));
+	
+	static FAutoConsoleCommandWithWorldAndArgs DemoMaxDesiredRecordTimeMS(
+		TEXT("Demo.MaxDesiredRecordTimeMS"),
+		TEXT("Set max desired record time in MS on demo driver of the current world."),
+		FConsoleCommandWithWorldAndArgsDelegate::CreateStatic(
+			[](const TArray<FString>& Params, UWorld* World)
+			{
+				if (World)
+				{
+					if (UDemoNetDriver* Driver = World->GetDemoNetDriver())
+					{
+						if (Params.Num() > 0)
+						{
+							const float TimeInMS = FCString::Atof(*Params[0]);
+
+							Driver->SetMaxDesiredRecordTimeMS(TimeInMS);
+						}
+					}
+				}
+			}));
+
+	static FAutoConsoleCommandWithWorldAndArgs DemoCheckpointSaveMaxMSPerFrame(
+		TEXT("Demo.CheckpointSaveMaxMSPerFrame"),
+		TEXT("Set max checkpoint record time in MS on demo driver of the current world."),
+		FConsoleCommandWithWorldAndArgsDelegate::CreateStatic(
+			[](const TArray<FString>& Params, UWorld* World)
+			{
+				if (World)
+				{
+					if (UDemoNetDriver* Driver = World->GetDemoNetDriver())
+					{
+						if (Params.Num() > 0)
+						{
+							const float TimeInMS = FCString::Atof(*Params[0]);
+
+							Driver->SetCheckpointSaveMaxMSPerFrame(TimeInMS);
+						}
+					}
+				}
+			}));
+
+	static FAutoConsoleCommandWithWorldAndArgs DemoActorPrioritizationEnabled(
+		TEXT("Demo.ActorPrioritizationEnabled"),
+		TEXT("Set whether or not actor prioritization is enabled on demo driver of the current world."),
+		FConsoleCommandWithWorldAndArgsDelegate::CreateStatic(
+			[](const TArray<FString>& Params, UWorld* World)
+			{
+				if (World)
+				{
+					if (UDemoNetDriver* Driver = World->GetDemoNetDriver())
+					{
+						if (Params.Num() > 0)
+						{
+							const bool bPrioritize = FCString::ToBool(*Params[0]);
+
+							Driver->SetActorPrioritizationEnabled(bPrioritize);
+						}
+					}
+				}
+			}));
+
+	static FAutoConsoleCommandWithWorldAndArgs DemoSetLocalViewerOverride(
+		TEXT("Demo.SetLocalViewerOverride"),
+		TEXT("Set first local player controller as the viewer override on demo driver of the current world."),
+		FConsoleCommandWithWorldAndArgsDelegate::CreateStatic(
+			[](const TArray<FString>& Params, UWorld* World)
+			{
+				if (World)
+				{
+					if (UDemoNetDriver* Driver = World->GetDemoNetDriver())
+					{
+						if (APlayerController* ViewerPC = GEngine->GetFirstLocalPlayerController(World))
+						{
+							Driver->SetViewerOverride(ViewerPC);
+						}
+					}
+				}
+			}));
 }
 
 struct FDemoBudgetLogHelper
@@ -619,6 +703,7 @@ void UDemoNetDriver::InitDefaults()
 	}
 
 	RecordBuildConsiderAndPrioritizeTimeSlice = CVarDemoMaximumRepPrioritizeTime.GetValueOnGameThread();
+	RecordDestructionInfoReplicationTimeSlice = DemoNetDriverRecordingPrivate::CVarDemoMaximumRecDestructionInfoTime.GetValueOnAnyThread();
 }
 
 UDemoNetDriver::UDemoNetDriver(const FObjectInitializer& ObjectInitializer)
@@ -742,6 +827,7 @@ bool UDemoNetDriver::InitBase(bool bInitAsClient, FNetworkNotify* InNotify, cons
 		CheckpointSaveMaxMSPerFrame		= -1.0f;
 
 		RecordBuildConsiderAndPrioritizeTimeSlice = CVarDemoMaximumRepPrioritizeTime.GetValueOnAnyThread();
+		RecordDestructionInfoReplicationTimeSlice = DemoNetDriverRecordingPrivate::CVarDemoMaximumRecDestructionInfoTime.GetValueOnAnyThread();
 
 		if (RelevantTimeout == 0.0f)
 		{
@@ -877,11 +963,11 @@ bool UDemoNetDriver::InitConnect(FNetworkNotify* InNotify, const FURL& ConnectUR
 	ReplayHelper.ActiveReplayName = ConnectURL.Map;
 
 	TArray<int32> UserIndices;
-	for (FLocalPlayerIterator It(GEngine, World); It; ++It)
+	for (FLocalPlayerIterator LocalPlayerIt(GEngine, World); LocalPlayerIt; ++LocalPlayerIt)
 	{
-		if (*It)
+		if (*LocalPlayerIt)
 		{
-			UserIndices.Add(It->GetControllerId());
+			UserIndices.Add(LocalPlayerIt->GetControllerId());
 		}
 	}
 	
@@ -1015,11 +1101,11 @@ void UDemoNetDriver::NotifyStreamingLevelUnload( ULevel* InLevel )
 	if (InLevel && !InLevel->bClientOnlyVisible && HasLevelStreamingFixes() && IsPlaying())
 	{
 		// We can't just iterate over the levels actors, because the ones in the queue will already have been destroyed.
-		for (TMap<FString, FRollbackNetStartupActorInfo>::TIterator It = RollbackNetStartupActors.CreateIterator(); It; ++It)
+		for (TMap<FString, FRollbackNetStartupActorInfo>::TIterator RollbackIt = RollbackNetStartupActors.CreateIterator(); RollbackIt; ++RollbackIt)
 		{
-			if (It.Value().Level == InLevel)
+			if (RollbackIt.Value().Level == InLevel)
 			{
-				It.RemoveCurrent();
+				RollbackIt.RemoveCurrent();
 			}
 		}
 	}
@@ -1492,30 +1578,38 @@ public:
 	FRepActorsParams& operator=(const FRepActorsParams&) = delete;
 	FRepActorsParams& operator=(FRepActorsParams&&) = delete;
 
-	FRepActorsParams(const bool bInUseAdaptiveNetFrequency, const bool bInDoFindActorChannel, const bool bInDoCheckDormancy,
+	FRepActorsParams(UDemoNetConnection* InConnection, const bool bInUseAdaptiveNetFrequency, const bool bInDoFindActorChannel, const bool bInDoCheckDormancy,
 					const float InMinRecordHz, const float InMaxRecordHz, const float InServerTickTime,
-					const double InReplicationStartTimeSeconds, const double InTimeLimitSeconds):
+					const double InReplicationStartTimeSeconds, const double InTimeLimitSeconds, const double InDestructionInfoTimeLimitSeconds):
+		Connection(InConnection),
 		bUseAdapativeNetFrequency(bInUseAdaptiveNetFrequency),
 		bDoFindActorChannel(bInDoFindActorChannel),
 		bDoCheckDormancy(bInDoCheckDormancy),
 		NumActorsReplicated(0),
+		NumDestructionInfosReplicated(0),
 		MinRecordHz(InMinRecordHz),
 		MaxRecordHz(InMaxRecordHz),
 		ServerTickTime(InServerTickTime),
 		ReplicationStartTimeSeconds(InReplicationStartTimeSeconds),
-		TimeLimitSeconds(InTimeLimitSeconds)
+		TimeLimitSeconds(InTimeLimitSeconds),
+		DestructionInfoTimeLimitSeconds(InDestructionInfoTimeLimitSeconds),
+		TotalDestructionInfoRecordTime(0.0)
 	{
 	}
 
+	UDemoNetConnection* Connection;
 	const bool bUseAdapativeNetFrequency;
 	const bool bDoFindActorChannel;
 	const bool bDoCheckDormancy;
 	int32 NumActorsReplicated;
+	int32 NumDestructionInfosReplicated;
 	const float MinRecordHz;
 	const float MaxRecordHz;
 	const float ServerTickTime;
 	const double ReplicationStartTimeSeconds;
 	const double TimeLimitSeconds;
+	const double DestructionInfoTimeLimitSeconds;
+	double TotalDestructionInfoRecordTime;
 };
 
 void UDemoNetDriver::TickDemoRecord(float DeltaSeconds)
@@ -1563,21 +1657,32 @@ void UDemoNetDriver::BuildSortedLevelPriorityOnLevels(const TArray<FDemoActorPri
 	const int32 Count = PrioritizedActorList.Num();
 	const FDemoActorPriority* Priorities = PrioritizedActorList.GetData();
 
-	for (int32 It = 0; It < Count;)
+	const bool bHighPriorityPersistentLevel = DemoNetDriverRecordingPrivate::CVarDemoForcePersistentLevelPriority.GetValueOnAnyThread() != 0;
+
+	for (int32 Index = 0; Index < Count;)
 	{
-		const UObject* CurrentLevel = Priorities[It].Level;
+		const UObject* CurrentLevel = Priorities[Index].Level;
 
 		FLevelnterval Interval;
-		Interval.StartIndex = It;
-		Interval.Priority = Priorities[It].ActorPriority.Priority;
+		Interval.StartIndex = Index;
+
+		if (bHighPriorityPersistentLevel && World && (CurrentLevel == World->PersistentLevel))
+		{
+			Interval.Priority = MAX_int32;
+		}
+		else
+		{
+			Interval.Priority = Priorities[Index].ActorPriority.Priority;
+		}
+
 		Interval.LevelIndex = (CurrentLevel != nullptr ? ReplayHelper.FindOrAddLevelStatus(*Cast<ULevel>(CurrentLevel)).LevelIndex + 1 : 0);
 
-		while (It < Count && Priorities[It].Level == CurrentLevel)
+		while (Index < Count && Priorities[Index].Level == CurrentLevel)
 		{
-			++It;
+			++Index;
 		}
 
-		Interval.Count = It - Interval.StartIndex;
+		Interval.Count = Index - Interval.StartIndex;
 
 		OutLevelIntervals.Add(Interval);
 	}
@@ -1660,32 +1765,36 @@ void UDemoNetDriver::TickDemoRecordFrame(float DeltaSeconds)
 	}
 
 	const bool bDoCheckDormancyEarly = CVarDemoLateActorDormancyCheck.GetValueOnAnyThread() == 0;
+	const bool bLateDestructionInfos = DemoNetDriverRecordingPrivate::CVarDemoLateDestructionInfoPrioritize.GetValueOnAnyThread() != 0;
 	const bool bDoPrioritizeActors = bPrioritizeActors;
 	const bool bDoFindActorChannelEarly = bDoPrioritizeActors || bDoCheckDormancyEarly;
 
+	int32 ActorsPrioritized = 0;
+	int32 DestructionInfosPrioritized = 0;
+
 	{
 		DECLARE_SCOPE_CYCLE_COUNTER(TEXT("Replay prioritize time"), STAT_ReplayPrioritizeTime, STATGROUP_Net);
 
 		const double ConsiderTimeLimit = RecordTimeLimit * RecordBuildConsiderAndPrioritizeTimeSlice;
 		auto HasConsiderTimeBeenExhausted = [ConsiderTimeLimit, RecordFrameStartTime, RecordTimeLimit]() -> bool
 		{
-			return RecordTimeLimit > 0.f && (FPlatformTime::Seconds() - RecordFrameStartTime) > ConsiderTimeLimit;
+			return RecordTimeLimit > 0.f && ((FPlatformTime::Seconds() - RecordFrameStartTime) > ConsiderTimeLimit);
 		};
 
+		auto PrioritizeDestructionInfos = [ClientConnection, this, &HasConsiderTimeBeenExhausted]()
 		{
 			SCOPED_NAMED_EVENT(UDemoNetDriver_PrioritizeDestroyedOrDormantActors, FColor::Green);
 
 			// Add destroyed actors that the client may not have a channel for
-			// We add these first so they get more of the prioritize time slice.
-			// This is because they are marked top priority anyway, and won't need to be prioritized
-			// which should decrease overall time spent next frame.
 			FDemoActorPriority DestroyedActorPriority;
-			DestroyedActorPriority.ActorPriority.Priority = 0x7FFFFFFF;
-			for (auto It = ClientConnection->GetDestroyedStartupOrDormantActorGUIDs().CreateIterator(); It; ++It)
+			DestroyedActorPriority.ActorPriority.Priority = DemoNetDriverRecordingPrivate::CVarDemoDestructionInfoPriority.GetValueOnAnyThread();
+
+			for (auto DestroyedOrDormantGUID = ClientConnection->GetDestroyedStartupOrDormantActorGUIDs().CreateIterator(); DestroyedOrDormantGUID; ++DestroyedOrDormantGUID)
 			{
-				TUniquePtr<FActorDestructionInfo>& DInfo = DestroyedStartupOrDormantActors.FindChecked(*It);
+				TUniquePtr<FActorDestructionInfo>& DInfo = DestroyedStartupOrDormantActors.FindChecked(*DestroyedOrDormantGUID);
 				DestroyedActorPriority.ActorPriority.DestructionInfo = DInfo.Get();
 				DestroyedActorPriority.Level = HasLevelStreamingFixes() ? DestroyedActorPriority.ActorPriority.DestructionInfo->Level.Get() : nullptr;
+
 				PrioritizedActors.Add(DestroyedActorPriority);
 
 				if (HasConsiderTimeBeenExhausted())
@@ -1694,6 +1803,13 @@ void UDemoNetDriver::TickDemoRecordFrame(float DeltaSeconds)
 					break;
 				}
 			}
+		};
+
+		if (!bLateDestructionInfos)
+		{
+			PrioritizeDestructionInfos();
+
+			DestructionInfosPrioritized = PrioritizedActors.Num();
 		}
 
 		if (!HasConsiderTimeBeenExhausted())
@@ -1813,7 +1929,15 @@ void UDemoNetDriver::TickDemoRecordFrame(float DeltaSeconds)
 					if (bDoPrioritizeActors) // implies bDoFindActorChannelEarly is true
 					{
 						const double LastReplicationTime = Channel ? (GetElapsedTime() - Channel->LastUpdateTime) : SpawnPrioritySeconds;
-						ActorPriority.Priority = FMath::RoundToInt(65536.0f * Actor->GetReplayPriority(ViewLocation, ViewDirection, Viewer, ViewTarget, Channel, LastReplicationTime));
+						float ReplayPriority = 65536.0f * Actor->GetReplayPriority(ViewLocation, ViewDirection, Viewer, ViewTarget, Channel, LastReplicationTime);
+						
+						if (Actor == ViewTarget)
+						{
+							ReplayPriority = ReplayPriority * DemoNetDriverRecordingPrivate::CVarDemoViewTargetPriorityScale.GetValueOnAnyThread();
+						}
+
+						// clamp into a valid range prior to rounding to avoid potential undefined behavior
+						ActorPriority.Priority = FMath::RoundToInt(FMath::Clamp(ReplayPriority, (float)(MIN_int32 + 10), (float)(MAX_int32 - 10)));
 					}
 
 					PrioritizedActors.Add(DemoActorPriority);
@@ -1843,6 +1967,25 @@ void UDemoNetDriver::TickDemoRecordFrame(float DeltaSeconds)
 				}
 			}
 		}
+
+		if (bLateDestructionInfos)
+		{
+			ActorsPrioritized = PrioritizedActors.Num();
+
+			if (!HasConsiderTimeBeenExhausted())
+			{
+				PrioritizeDestructionInfos();
+				DestructionInfosPrioritized = PrioritizedActors.Num();
+			}
+			else
+			{
+				UE_LOG(LogDemo, Verbose, TEXT("Consider time exhaused without processing destruction infos"));
+			}
+		}
+		else
+		{
+			ActorsPrioritized = PrioritizedActors.Num() - DestructionInfosPrioritized;
+		}
 	}
 
 	if (HasLevelStreamingFixes())
@@ -1872,11 +2015,13 @@ void UDemoNetDriver::TickDemoRecordFrame(float DeltaSeconds)
 	const double PrioritizeEndTime = FPlatformTime::Seconds();
 	const double TotalPrioritizeActorsTime = (PrioritizeEndTime - RecordFrameStartTime);
 	const float TotalPrioritizeActorsTimeMS = TotalPrioritizeActorsTime * 1000.f;
-	const int32 NumPrioritizedActors = PrioritizedActors.Num();
 
-	CSV_CUSTOM_STAT(Demo, DemoRecPrioritizeTime, TotalPrioritizeActorsTimeMS, ECsvCustomStatOp::Set);
-	CSV_CUSTOM_STAT(Demo, DemoRecPriotizedActors, NumPrioritizedActors, ECsvCustomStatOp::Set);
+	CSV_CUSTOM_STAT(Demo, DemoPrioritizeTime, TotalPrioritizeActorsTimeMS, ECsvCustomStatOp::Set);
 	CSV_CUSTOM_STAT(Demo, DemoNumActiveObjects, NumActiveObjects, ECsvCustomStatOp::Set);
+	CSV_CUSTOM_STAT(Demo, DemoPrioritizedActors, ActorsPrioritized, ECsvCustomStatOp::Set);
+	CSV_CUSTOM_STAT(Demo, DemoPrioritizedDestInfos, DestructionInfosPrioritized, ECsvCustomStatOp::Set);
+
+	const int32 NumPrioritizedActors = PrioritizedActors.Num();
 
 	// Make sure we're under the desired recording time quota, if any.
 	// See ReplicatePriorizeActor.
@@ -1898,6 +2043,7 @@ void UDemoNetDriver::TickDemoRecordFrame(float DeltaSeconds)
 
 	FRepActorsParams Params
 	(
+		ClientConnection,
 		CVarUseAdaptiveReplayUpdateFrequency.GetValueOnAnyThread() > 0,
 		!bDoFindActorChannelEarly,
 		!bDoCheckDormancyEarly,
@@ -1905,7 +2051,8 @@ void UDemoNetDriver::TickDemoRecordFrame(float DeltaSeconds)
 		MaxRecordHz,
 		ServerTickTime,
 		RecordFrameStartTime,
-		RecordTimeLimit
+		RecordTimeLimit,
+		RecordTimeLimit * RecordDestructionInfoReplicationTimeSlice
 	);
 
 	{
@@ -1931,17 +2078,18 @@ void UDemoNetDriver::TickDemoRecordFrame(float DeltaSeconds)
 	}
 
 	CSV_CUSTOM_STAT(Demo, DemoNumReplicatedActors, Params.NumActorsReplicated, ECsvCustomStatOp::Set);
+	CSV_CUSTOM_STAT(Demo, DemoNumReplicatedDestructionInfos, Params.NumDestructionInfosReplicated, ECsvCustomStatOp::Set);
 
 	FReplayHelper::FlushNetChecked(*ClientConnection);
 
 	WriteDemoFrameFromQueuedDemoPackets(*FileAr, ReplayHelper.QueuedDemoPackets, GetDemoCurrentTime(), EWriteDemoFrameFlags::None);
 
-	const float ReplicatedPercent = NumPrioritizedActors != 0 ? (float)Params.NumActorsReplicated / (float)NumPrioritizedActors : 1.0f;
+	const float ReplicatedPercent = NumPrioritizedActors != 0 ? (float)(Params.NumActorsReplicated + Params.NumDestructionInfosReplicated) / (float)NumPrioritizedActors : 1.0f;
 	AdjustConsiderTime(ReplicatedPercent);
 	LastReplayFrameFidelity = ReplicatedPercent;
 }
 
-bool UDemoNetDriver::ReplicatePrioritizedActor(const FActorPriority& ActorPriority, const FRepActorsParams& Params)
+bool UDemoNetDriver::ReplicatePrioritizedActor(const FActorPriority& ActorPriority, FRepActorsParams& Params)
 {
 	FNetworkObjectInfo* ActorInfo = ActorPriority.ActorInfo;
 	FActorDestructionInfo* DestructionInfo = ActorPriority.DestructionInfo;
@@ -1950,35 +2098,46 @@ bool UDemoNetDriver::ReplicatePrioritizedActor(const FActorPriority& ActorPriori
 
 	const bool bDoFindActorChannel = Params.bDoFindActorChannel;
 	const bool bDoCheckDormancy = Params.bDoCheckDormancy;
-
-	UDemoNetConnection* Connection = CastChecked<UDemoNetConnection>(ClientConnections[0]);
+	const bool bDestructionInfo = DestructionInfo != nullptr && ActorInfo == nullptr;
+	const bool bActorInfo = ActorInfo != nullptr && DestructionInfo == nullptr;
 
 	// Deletion entry
-	if (ActorInfo == nullptr && DestructionInfo != nullptr)
+	if (bDestructionInfo)
 	{
-		UActorChannel* Channel = (UActorChannel*)Connection->CreateChannelByName(NAME_Actor, EChannelCreateFlags::OpenedLocally);
-		if (Channel)
+		// only process destruction infos if we're below the time limit
+		if (Params.TotalDestructionInfoRecordTime < Params.DestructionInfoTimeLimitSeconds)
 		{
-			UE_LOG(LogDemo, Verbose, TEXT("TickDemoRecord creating destroy channel for NetGUID <%s,%s> Priority: %d"), *DestructionInfo->NetGUID.ToString(), *DestructionInfo->PathName, ActorPriority.Priority);
+			++Params.NumDestructionInfosReplicated;
+
+			UActorChannel* Channel = (UActorChannel*)Params.Connection->CreateChannelByName(NAME_Actor, EChannelCreateFlags::OpenedLocally);
+			if (Channel)
+			{
+				UE_LOG(LogDemo, Verbose, TEXT("TickDemoRecord creating destroy channel for NetGUID <%s,%s> Priority: %d"), *DestructionInfo->NetGUID.ToString(), *DestructionInfo->PathName, ActorPriority.Priority);
+
+				FScopedRepContext LevelContext(Params.Connection, DestructionInfo->Level.Get());
 
-			FScopedRepContext LevelContext(Connection, DestructionInfo->Level.Get());
+				// Send a close bunch on the new channel
+				PRAGMA_DISABLE_DEPRECATION_WARNINGS
+				Channel->SetChannelActorForDestroy(DestructionInfo);
+				PRAGMA_ENABLE_DEPRECATION_WARNINGS
 
-			// Send a close bunch on the new channel
-			PRAGMA_DISABLE_DEPRECATION_WARNINGS
-			Channel->SetChannelActorForDestroy(DestructionInfo);
-			PRAGMA_ENABLE_DEPRECATION_WARNINGS
+				// Remove from connection's to-be-destroyed list (close bunch is reliable, so it will make it there)
+				Params.Connection->RemoveDestructionInfo(DestructionInfo);
 
-			// Remove from connection's to-be-destroyed list (close bunch is reliable, so it will make it there)
-			Connection->RemoveDestructionInfo(DestructionInfo);
+				// calling conditional cleanup now allows the channel to be returned to any pools and reused immediately
+				Channel->ConditionalCleanUp(false, DestructionInfo->Reason);
+			}
 		}
 	}
-	else if (ActorInfo != nullptr && DestructionInfo == nullptr)
+	else if (bActorInfo)
 	{
+		++Params.NumActorsReplicated;
+
 		AActor* Actor = ActorInfo->Actor;
 		
 		if (bDoCheckDormancy)
 		{
-			UActorChannel* Channel = (bDoFindActorChannel ? Connection->FindActorChannelRef(Actor) : ActorPriority.Channel);
+			UActorChannel* Channel = (bDoFindActorChannel ? Params.Connection->FindActorChannelRef(Actor) : ActorPriority.Channel);
 			if (Channel && ShouldActorGoDormantForDemo(Actor, Channel))
 			{
 				// Either shouldn't go dormant, or is already dormant
@@ -1999,14 +2158,10 @@ bool UDemoNetDriver::ReplicatePrioritizedActor(const FActorPriority& ActorPriori
 
 		const float LastReplicateDelta = static_cast<float>(GetDemoCurrentTime() - ActorInfo->LastNetReplicateTime);
 
-		if (Actor->MinNetUpdateFrequency == 0.0f)
-		{
-			Actor->MinNetUpdateFrequency = 2.0f;
-		}
-
 		// Calculate min delta (max rate actor will update), and max delta (slowest rate actor will update)
 		const float MinOptimalDelta = NetUpdateDelay;										// Don't go faster than NetUpdateFrequency
-		const float MaxOptimalDelta = FMath::Max(1.0f / Actor->MinNetUpdateFrequency, MinOptimalDelta);	// Don't go slower than MinNetUpdateFrequency (or NetUpdateFrequency if it's slower)
+		const float MinNetUpdateFrequency = (Actor->MinNetUpdateFrequency == 0.0f) ? 2.0f : Actor->MinNetUpdateFrequency;
+		const float MaxOptimalDelta = FMath::Max(1.0f / MinNetUpdateFrequency, MinOptimalDelta);	// Don't go slower than MinNetUpdateFrequency (or NetUpdateFrequency if it's slower)
 
 		const float ScaleDownStartTime = 2.0f;
 		const float ScaleDownTimeRange = 5.0f;
@@ -2028,9 +2183,9 @@ bool UDemoNetDriver::ReplicatePrioritizedActor(const FActorPriority& ActorPriori
 		// Try to spread the updates across multiple frames to smooth out spikes.
 		ActorInfo->NextUpdateTime = (GetDemoCurrentTime() + NextUpdateDelta - ClampedExtraTime + ((UpdateDelayRandomStream.FRand() - 0.5) * Params.ServerTickTime));
 
-		const bool bDidReplicateActor = DemoReplicateActor(Actor, Connection, false);
+		const bool bDidReplicateActor = DemoReplicateActor(Actor, Params.Connection, false);
 
-		const bool bUpdatedExternalData = ReplayHelper.UpdateExternalDataForObject(Connection, Actor);
+		const bool bUpdatedExternalData = ReplayHelper.UpdateExternalDataForObject(Params.Connection, Actor);
 
 		if (bDidReplicateActor || bUpdatedExternalData)
 		{
@@ -2050,6 +2205,11 @@ bool UDemoNetDriver::ReplicatePrioritizedActor(const FActorPriority& ActorPriori
 		const double RecordEndTimeSeconds = FPlatformTime::Seconds();
 		const double RecordTimeSeconds = RecordEndTimeSeconds - RecordStartTimeSeconds;
 
+		if (bDestructionInfo)
+		{
+			Params.TotalDestructionInfoRecordTime += RecordTimeSeconds;
+		}
+
 		if ((ActorInfo && ActorInfo->Actor) && (RecordTimeSeconds > (Params.TimeLimitSeconds * 0.95f)))
 		{
 			UE_LOG(LogDemo, Verbose, TEXT("Actor %s took more than 95%% of maximum desired recording time. Actor: %.3fms. Max: %.3fms."),
@@ -2075,19 +2235,18 @@ bool UDemoNetDriver::ReplicatePrioritizedActor(const FActorPriority& ActorPriori
 bool UDemoNetDriver::ReplicatePrioritizedActors(const FDemoActorPriority* ActorsToReplicate, uint32 Count, FRepActorsParams& Params)
 {
 	bool bTimeRemaining = true;
-	uint32 It = 0;
-	for (; It < Count; ++It)
+	uint32 NumProcessed = 0;
+	for (; NumProcessed < Count; ++NumProcessed)
 	{
-		const FActorPriority& ActorPriority = ActorsToReplicate[It].ActorPriority;
+		const FActorPriority& ActorPriority = ActorsToReplicate[NumProcessed].ActorPriority;
 		bTimeRemaining = ReplicatePrioritizedActor(ActorPriority, Params);
 		if (!bTimeRemaining)
 		{
-			++It;
+			++NumProcessed;
 			break;
 		}
 	}
 
-	Params.NumActorsReplicated += It;
 	return bTimeRemaining;
 }
 
@@ -2673,16 +2832,16 @@ void UDemoNetDriver::TickDemoPlayback(float DeltaSeconds)
 		}
 
 		// Process playback frames
-		for (auto It = ReplayHelper.PlaybackFrames.CreateIterator(); It; ++It)
+		for (auto FrameIt = ReplayHelper.PlaybackFrames.CreateIterator(); FrameIt; ++FrameIt)
 		{
-			if (It.Key() <= GetDemoCurrentTime())
+			if (FrameIt.Key() <= GetDemoCurrentTime())
 			{
 				if (!bIsFastForwarding)
 				{
-					FNetworkReplayDelegates::OnProcessGameSpecificFrameData.Broadcast(World, It.Key(), It.Value());
+					FNetworkReplayDelegates::OnProcessGameSpecificFrameData.Broadcast(World, FrameIt.Key(), FrameIt.Value());
 				}
 	
-				It.RemoveCurrent();
+				FrameIt.RemoveCurrent();
 			}
 		}
 	}
@@ -3026,15 +3185,15 @@ void UDemoNetDriver::RespawnNecessaryNetStartupActors(TArray<AActor*>& SpawnedAc
 {
 	TGuardValue<bool> RestoringStartupActors(bIsRestoringStartupActors, true);
 
-	for (auto It = RollbackNetStartupActors.CreateIterator(); It; ++It)
+	for (auto RollbackIt = RollbackNetStartupActors.CreateIterator(); RollbackIt; ++RollbackIt)
 	{
-		if (ReplayHelper.PlaybackDeletedNetStartupActors.Contains(It.Key()))
+		if (ReplayHelper.PlaybackDeletedNetStartupActors.Contains(RollbackIt.Key()))
 		{
 			// We don't want to re-create these since they should no longer exist after the current checkpoint
 			continue;
 		}
 
-		FRollbackNetStartupActorInfo& RollbackActor = It.Value();
+		FRollbackNetStartupActorInfo& RollbackActor = RollbackIt.Value();
 
 		// filter to a specific level
 		if ((Level != nullptr) && (RollbackActor.Level != Level))
@@ -3085,9 +3244,9 @@ void UDemoNetDriver::RespawnNecessaryNetStartupActors(TArray<AActor*>& SpawnedAc
 		AActor* Actor = World->SpawnActorAbsolute(RollbackActor.Archetype->GetClass(), SpawnTransform, SpawnInfo);
 		if (Actor)
 		{
-			if (!ensure( Actor->GetFullName() == It.Key()))
+			if (!ensure( Actor->GetFullName() == RollbackIt.Key()))
 			{
-				UE_LOG(LogDemo, Log, TEXT("RespawnNecessaryNetStartupActors: NetStartupRollbackActor name doesn't match original: %s, %s"), *Actor->GetFullName(), *It.Key());
+				UE_LOG(LogDemo, Log, TEXT("RespawnNecessaryNetStartupActors: NetStartupRollbackActor name doesn't match original: %s, %s"), *Actor->GetFullName(), *RollbackIt.Key());
 			}
 
 			bool bSanityCheckReferences = true;
@@ -3173,7 +3332,7 @@ void UDemoNetDriver::RespawnNecessaryNetStartupActors(TArray<AActor*>& SpawnedAc
 			SpawnedActors.Add(Actor);
 		}
 
-		It.RemoveCurrent();
+		RollbackIt.RemoveCurrent();
 	}
 
 	RollbackNetStartupActors.Compact();
@@ -3683,9 +3842,9 @@ bool UDemoNetDriver::LoadCheckpoint(const FGotoResult& GotoResult)
 	{
 		World->GetGameInstance()->OnSeamlessTravelDuringReplay();
 
-		for (FActorIterator It(World); It; ++It)
+		for (FActorIterator ActorIt(World); ActorIt; ++ActorIt)
 		{
-			World->DestroyActor(*It, true);
+			World->DestroyActor(*ActorIt, true);
 		}
 
 		// Clean package map to prepare to restore it to the checkpoint state
@@ -3770,11 +3929,11 @@ bool UDemoNetDriver::LoadCheckpoint(const FGotoResult& GotoResult)
 
 	// Determine if an Actor has a reference to a spectator in some way.
 	// This prevents garbage collection on splitscreen playercontrollers
-	auto HasPlayerSpectatorRef = [this](const FActorIterator& InActorIterator) {
+	auto HasPlayerSpectatorRef = [this](const AActor* InActor) {
 		for (const APlayerController* CurSpectator : SpectatorControllers)
 		{
-			if (*InActorIterator == CurSpectator || *InActorIterator == CurSpectator->GetSpectatorPawn()
-				|| InActorIterator->GetOwner() == CurSpectator)
+			if (InActor == CurSpectator || InActor == CurSpectator->GetSpectatorPawn()
+				|| InActor->IsOwnedBy(CurSpectator))
 			{
 				return true;
 			}
@@ -3784,31 +3943,33 @@ bool UDemoNetDriver::LoadCheckpoint(const FGotoResult& GotoResult)
 
 
 	// Destroy all non startup actors. They will get restored with the checkpoint
-	for (FActorIterator It(World); It; ++It)
+	for (FActorIterator ActorIt(World); ActorIt; ++ActorIt)
 	{
+		AActor* CurrentActor = *ActorIt;
+
 		// If there are any existing actors that are bAlwaysRelevant, don't queue their bunches.
 		// Actors that do queue their bunches might not appear immediately after the checkpoint is loaded,
 		// and missing bAlwaysRelevant actors are more likely to cause noticeable artifacts.
 		// NOTE - We are adding the actor guid here, under the assumption that the actor will reclaim the same guid when we load the checkpoint
 		// This is normally the case, but could break if actors get destroyed and re-created with different guids during recording
-		if (It->bAlwaysRelevant)
+		if (CurrentActor->bAlwaysRelevant)
 		{
-			AddNonQueuedActorForScrubbing(*It);
+			AddNonQueuedActorForScrubbing(CurrentActor);
 		}
 		
-		const bool bShouldPreserveForPlayerController = HasPlayerSpectatorRef(It);
-		const bool bShouldPreserveForRewindability = (It->bReplayRewindable && !It->IsNetStartupActor());										
+		const bool bShouldPreserveForPlayerController = HasPlayerSpectatorRef(CurrentActor);
+		const bool bShouldPreserveForRewindability = (CurrentActor->bReplayRewindable && !CurrentActor->IsNetStartupActor());
 
 		if (bShouldPreserveForPlayerController || bShouldPreserveForRewindability)
 		{
 			// If an non-startup actor that we don't destroy has an entry in the GuidCache, preserve that entry so
 			// that the object will be re-used after loading the checkpoint. Otherwise, a new copy
 			// of the object will be created each time a checkpoint is loaded, causing a leak.
-			const FNetworkGUID FoundGUID = GuidCache->NetGUIDLookup.FindRef( *It );
+			const FNetworkGUID FoundGUID = GuidCache->NetGUIDLookup.FindRef(CurrentActor);
 				
 			if (FoundGUID.IsValid())
 			{
-				NetGUIDsToPreserve.Emplace(FoundGUID, *It);
+				NetGUIDsToPreserve.Emplace(FoundGUID, CurrentActor);
 				
 				if (bShouldPreserveForRewindability)
 				{
@@ -3816,24 +3977,24 @@ bool UDemoNetDriver::LoadCheckpoint(const FGotoResult& GotoResult)
 				}
 			}
 
-			KeepAliveActors.Add(*It);
+			KeepAliveActors.Add(CurrentActor);
 			continue;
 		}
 
 		// Prevent NetStartupActors from being destroyed.
 		// NetStartupActors that can't have properties directly re-applied should use QueueNetStartupActorForRollbackViaDeletion.
-		if (It->IsNetStartupActor())
+		if (CurrentActor->IsNetStartupActor())
 		{
 			// Go ahead and rewind this now, since we won't be destroying it later.
-			if (It->bReplayRewindable)
+			if (CurrentActor->bReplayRewindable)
 			{
-				It->RewindForReplay();
+				CurrentActor->RewindForReplay();
 			}
-			KeepAliveActors.Add(*It);
+			KeepAliveActors.Add(CurrentActor);
 			continue;
 		}
 
-		World->DestroyActor(*It, true);
+		World->DestroyActor(CurrentActor, true);
 	}
 
 	// Destroy all particle FX attached to the WorldSettings (the WorldSettings actor persists but the particle FX spawned at runtime shouldn't)
@@ -3879,11 +4040,11 @@ bool UDemoNetDriver::LoadCheckpoint(const FGotoResult& GotoResult)
 	ReplayHelper.PlaybackFrames.Empty();
 
 	// Destroy startup actors that need to rollback via being destroyed and re-created
-	for (FActorIterator It(World); It; ++It)
+	for (FActorIterator ActorIt(World); ActorIt; ++ActorIt)
 	{
-		if (RollbackNetStartupActors.Contains(It->GetFullName()))
+		if (RollbackNetStartupActors.Contains(ActorIt->GetFullName()))
 		{
-			World->DestroyActor(*It, true);
+			World->DestroyActor(*ActorIt, true);
 		}
 	}
 
@@ -4138,13 +4299,15 @@ bool UDemoNetDriver::LoadCheckpoint(const FGotoResult& GotoResult)
 	if (World != nullptr)
 	{
 		// Destroy startup actors that shouldn't exist past this checkpoint
-		for (FActorIterator It( World ); It; ++It)
+		for (FActorIterator ActorIt( World ); ActorIt; ++ActorIt)
 		{
-			const FString FullName = It->GetFullName();
+			AActor* CurrentActor = *ActorIt;
+
+			const FString FullName = CurrentActor->GetFullName();
 
 			if (ReplayHelper.PlaybackDeletedNetStartupActors.Contains(FullName))
 			{
-				if (It->bReplayRewindable)
+				if (CurrentActor->bReplayRewindable)
 				{
 					// Log and skip. We can't queue Rewindable actors and we can't destroy them.
 					// This actor may still get destroyed during cleanup.
@@ -4153,12 +4316,12 @@ bool UDemoNetDriver::LoadCheckpoint(const FGotoResult& GotoResult)
 				}
 
 				// Put this actor on the rollback list so we can undelete it during future scrubbing
-				QueueNetStartupActorForRollbackViaDeletion(*It);
+				QueueNetStartupActorForRollbackViaDeletion(CurrentActor);
 
 				UE_LOG(LogDemo, Verbose, TEXT("LoadCheckpoint: deleting startup actor %s"), *FullName);
 
 				// Delete the actor
-				World->DestroyActor(*It, true);
+				World->DestroyActor(CurrentActor, true);
 			}
 		}
 
@@ -4577,11 +4740,11 @@ void UDemoNetConnection::HandleClientPlayer(APlayerController* PC, UNetConnectio
 	ULocalPlayer* LocalPlayer = nullptr;
 	uint8 PlayerIndex = 0;
 	// Attempt to find the player that doesn't already have a connection.
-	for (FLocalPlayerIterator It(GEngine, Driver->GetWorld()); It; ++It, PlayerIndex++)
+	for (FLocalPlayerIterator LocalPlayerIt(GEngine, Driver->GetWorld()); LocalPlayerIt; ++LocalPlayerIt, PlayerIndex++)
 	{
 		if (PC->NetPlayerIndex == PlayerIndex)
 		{
-			LocalPlayer = *It;
+			LocalPlayer = *LocalPlayerIt;
 			break;
 		}
 	}
@@ -4594,17 +4757,15 @@ void UDemoNetConnection::HandleClientPlayer(APlayerController* PC, UNetConnectio
 	{
 		DemoDriver->RestoreConnectionPostScrub(PC, NetConnection);
 	}
-	
+
 	// This is very likely our main demo controller.
 	DemoDriver->SetSpectatorController(PC);
 
-	for (FActorIterator It(Driver->World); It; ++It)
+	// Find a player start, if one exists
+	for (TActorIterator<APlayerStart> PlayerStartIt(Driver->World); PlayerStartIt; ++PlayerStartIt)
 	{
-		if (It->IsA(APlayerStart::StaticClass()))
-		{
-			PC->SetInitialLocationAndRotation(It->GetActorLocation(), It->GetActorRotation());
-			break;
-		}
+		PC->SetInitialLocationAndRotation(PlayerStartIt->GetActorLocation(), PlayerStartIt->GetActorRotation());
+		break;
 	}
 }
 
@@ -4738,10 +4899,9 @@ void UDemoNetDriver::InitDestroyedStartupActors()
 		check(ReplayHelper.RecordingDeltaCheckpointData.RecordingDeletedNetStartupActors.Num() == 0);
 
 		// add startup actors destroyed before the creation of this net driver
-		for (auto LevelIt(World->GetLevelIterator()); LevelIt; ++LevelIt)
+		for (FConstLevelIterator LevelIt(World->GetLevelIterator()); LevelIt; ++LevelIt)
 		{
-			ULevel* Level = *LevelIt;
-			if (Level)
+			if (const ULevel* Level = *LevelIt)
 			{
 				const TArray<FReplicatedStaticActorDestructionInfo>& DestroyedReplicatedStaticActors = Level->GetDestroyedReplicatedStaticActors();
 				for (const FReplicatedStaticActorDestructionInfo& Info : DestroyedReplicatedStaticActors)
@@ -5286,4 +5446,4 @@ void UDemoNetDriver::RequestCheckpoint()
 	{
 		ReplayHelper.RequestCheckpoint();
 	}
-}
\ No newline at end of file
+}
diff --git a/Engine/Source/Runtime/Engine/Private/RepLayout.cpp b/Engine/Source/Runtime/Engine/Private/RepLayout.cpp
--- a/Engine/Source/Runtime/Engine/Private/RepLayout.cpp
+++ b/Engine/Source/Runtime/Engine/Private/RepLayout.cpp
@@ -27,7 +27,6 @@
 #include "UObject/EnumProperty.h"
 #include "UObject/UnrealType.h"
 
-DECLARE_CYCLE_STAT(TEXT("RepLayout AddPropertyCmd"), STAT_RepLayout_AddPropertyCmd, STATGROUP_Game);
 DECLARE_CYCLE_STAT(TEXT("RepLayout InitFromObjectClass"), STAT_RepLayout_InitFromObjectClass, STATGROUP_Game);
 DECLARE_CYCLE_STAT(TEXT("RepLayout BuildShadowOffsets"), STAT_RepLayout_BuildShadowOffsets, STATGROUP_Game);
 DECLARE_CYCLE_STAT(TEXT("RepLayout DeltaSerializeFastArray"), STAT_RepLayout_DeltaSerializeFastArray, STATGROUP_Game);
@@ -5292,8 +5291,6 @@ static uint32 AddPropertyCmd(
 	FInitFromPropertySharedParams& SharedParams,
 	const FInitFromPropertyStackParams& StackParams)
 {
-	SCOPE_CYCLE_COUNTER(STAT_RepLayout_AddPropertyCmd);
-
 	FRepLayoutCmd & Cmd = SharedParams.Cmds.AddZeroed_GetRef();
 
 	Cmd.Property = StackParams.Property;
diff --git a/Engine/Source/Runtime/Engine/Private/ReplayHelper.cpp b/Engine/Source/Runtime/Engine/Private/ReplayHelper.cpp
--- a/Engine/Source/Runtime/Engine/Private/ReplayHelper.cpp
+++ b/Engine/Source/Runtime/Engine/Private/ReplayHelper.cpp
@@ -999,7 +999,8 @@ bool FReplayHelper::SerializeGuidCache(UNetConnection* Connection, const FRepAct
 		*CheckpointArchive << CheckpointSaveContext.NextAmortizedItem;
 	}
 
-	const double StartTime = FPlatformTime::Seconds();
+	FCheckpointStepHelper StepHelper(ECheckpointSaveState::SerializeGuidCache, Params.StartCheckpointTime, &CheckpointSaveContext.NextAmortizedItem, CheckpointSaveContext.NetGuidCacheSnapshot.Num());
+
 	const double Deadline = Params.StartCheckpointTime + Params.CheckpointMaxUploadTimePerFrame;
 
 	check(CheckpointSaveContext.NetGuidCacheSnapshot.Num() == 0 || CheckpointSaveContext.NetGuidCacheSnapshot.IsValidIndex(CheckpointSaveContext.NextAmortizedItem));
@@ -1067,8 +1068,6 @@ bool FReplayHelper::SerializeGuidCache(UNetConnection* Connection, const FRepAct
 		CheckpointArchive->Seek(Pos);
 	}
 
-	UE_LOG(LogDemo, Log, TEXT("Checkpoint. SerializeGuidCache: %i/%i (total %i), took %.3f (%.3f)"), CheckpointSaveContext.NextAmortizedItem, CheckpointSaveContext.NetGuidCacheSnapshot.Num(), CheckpointSaveContext.NumNetGuidsForRecording, FPlatformTime::Seconds() - Params.StartCheckpointTime, FPlatformTime::Seconds() - StartTime);
-
 	return bCompleted;
 }
 
@@ -1086,7 +1085,8 @@ bool FReplayHelper::SerializeDeletedStartupActors(UNetConnection* Connection, co
 		*CheckpointArchive << DeletedCount;
 	}
 
-	const double StartTime = FPlatformTime::Seconds();
+	FCheckpointStepHelper StepHelper(ECheckpointSaveState::SerializeDeletedStartupActors, Params.StartCheckpointTime, &CheckpointSaveContext.NextAmortizedItem, DeletedActors.Num());
+
 	const double Deadline = Params.StartCheckpointTime + Params.CheckpointMaxUploadTimePerFrame;
 
 	check(DeletedActors.Num() == 0 || DeletedActors.IsValidIndex(CheckpointSaveContext.NextAmortizedItem));
@@ -1109,8 +1109,6 @@ bool FReplayHelper::SerializeDeletedStartupActors(UNetConnection* Connection, co
 
 	const bool bCompleted = (CheckpointSaveContext.NextAmortizedItem == DeletedActors.Num());
 
-	UE_LOG(LogDemo, Log, TEXT("Checkpoint. SerializeDeletedStartupActors: %i/%i, took %.3f (%.3f)"), CheckpointSaveContext.NextAmortizedItem, DeletedActors.Num(), FPlatformTime::Seconds() - Params.StartCheckpointTime, FPlatformTime::Seconds() - StartTime);
-
 	return bCompleted;
 }
 
@@ -1125,7 +1123,8 @@ bool FReplayHelper::SerializeDeltaDynamicDestroyed(UNetConnection* Connection, c
 		*CheckpointArchive << TotalCount;
 	}
 
-	const double StartTime = FPlatformTime::Seconds();
+	FCheckpointStepHelper StepHelper(ECheckpointSaveState::SerializeDeltaDynamicDestroyed, Params.StartCheckpointTime, &CheckpointSaveContext.NextAmortizedItem, TotalCount);
+
 	const double Deadline = Params.StartCheckpointTime + Params.CheckpointMaxUploadTimePerFrame;
 
 	check(TotalCount == 0 || CheckpointSaveContext.DeltaCheckpointData.DestroyedDynamicActors.IsValidId(FSetElementId::FromInteger(CheckpointSaveContext.NextAmortizedItem)));
@@ -1145,8 +1144,6 @@ bool FReplayHelper::SerializeDeltaDynamicDestroyed(UNetConnection* Connection, c
 
 	const bool bCompleted = (CheckpointSaveContext.NextAmortizedItem == TotalCount);
 
-	UE_LOG(LogDemo, Log, TEXT("Checkpoint. SerializeDeltaDynamicDestroyed: %i/%i, took %.3f (%.3f)"), CheckpointSaveContext.NextAmortizedItem, TotalCount, FPlatformTime::Seconds() - Params.StartCheckpointTime, FPlatformTime::Seconds() - StartTime);
-
 	return bCompleted;
 }
 
@@ -1161,7 +1158,8 @@ bool FReplayHelper::SerializeDeltaClosedChannels(UNetConnection* Connection, con
 		*CheckpointArchive << TotalCount;
 	}
 
-	const double StartTime = FPlatformTime::Seconds();
+	FCheckpointStepHelper StepHelper(ECheckpointSaveState::SerializeDeltaClosedChannels, Params.StartCheckpointTime, &CheckpointSaveContext.NextAmortizedItem, TotalCount);
+
 	const double Deadline = Params.StartCheckpointTime + Params.CheckpointMaxUploadTimePerFrame;
 
 	check(TotalCount == 0 || CheckpointSaveContext.DeltaChannelCloseKeys.IsValidIndex(CheckpointSaveContext.NextAmortizedItem));
@@ -1184,8 +1182,6 @@ bool FReplayHelper::SerializeDeltaClosedChannels(UNetConnection* Connection, con
 
 	const bool bCompleted = (CheckpointSaveContext.NextAmortizedItem == TotalCount);
 
-	UE_LOG(LogDemo, Log, TEXT("Checkpoint. SerializeDeltaClosedChannels: %i/%i, took %.3f (%.3f)"), CheckpointSaveContext.NextAmortizedItem, TotalCount, FPlatformTime::Seconds() - Params.StartCheckpointTime, FPlatformTime::Seconds() - StartTime);
-
 	return bCompleted;
 }
 
diff --git a/Engine/Source/Runtime/Engine/Public/ReplayHelper.h b/Engine/Source/Runtime/Engine/Public/ReplayHelper.h
--- a/Engine/Source/Runtime/Engine/Public/ReplayHelper.h
+++ b/Engine/Source/Runtime/Engine/Public/ReplayHelper.h
@@ -249,6 +249,57 @@ private:
 		Finalize,
 	};
 
+	struct FCheckpointStepHelper
+	{
+		FCheckpointStepHelper() = delete;
+		FCheckpointStepHelper(ECheckpointSaveState InCheckpointState, const double InCheckpointStartTime, int32* InCurrentIndex, int32 InTotalCount)
+			: CheckpointState(InCheckpointState)
+			, CheckpointStartTime(InCheckpointStartTime)
+			, CurrentIndex(InCurrentIndex)
+			, TotalCount(InTotalCount)
+		{
+			check(InCurrentIndex);
+			StartTime = FPlatformTime::Seconds();
+		}
+
+		~FCheckpointStepHelper()
+		{
+			const double EndTime = FPlatformTime::Seconds();
+			const double TotalTimeInMS = (EndTime - CheckpointStartTime) * 1000.0;
+			const double StepTimeInMS = (EndTime - StartTime) * 1000.0;
+
+			const TCHAR* StateStr = TEXT("Unknown");
+
+			switch (CheckpointState)
+			{
+			case ECheckpointSaveState::SerializeDeletedStartupActors:
+				StateStr = TEXT("SerializeDeletedStartupActors");
+				break;
+			case ECheckpointSaveState::SerializeDeltaDynamicDestroyed:
+				StateStr = TEXT("SerializeDeltaDynamicDestroyed");
+				break;
+			case ECheckpointSaveState::SerializeDeltaClosedChannels:
+				StateStr = TEXT("SerializeDeltaClosedChannels");
+				break;
+			case ECheckpointSaveState::SerializeGuidCache:
+				StateStr = TEXT("SerializeGuidCache");
+				break;
+			default:
+				ensureMsgf(false, TEXT("FCheckpointStepHelper: Unsupported checkpoint state: %d"), CheckpointState);
+				break;
+			}
+
+			UE_LOG(LogDemo, Log, TEXT("Checkpoint. %s: %i/%i, took %.2fms (Total this frame: %.2fms)"), StateStr, *CurrentIndex, TotalCount, StepTimeInMS, TotalTimeInMS);
+		}
+
+	private:
+		ECheckpointSaveState CheckpointState;
+		double StartTime = 0.0;
+		double CheckpointStartTime = 0.0;
+		int32* CurrentIndex = nullptr;
+		int32 TotalCount = 0;
+	};
+
 	/** When we save a checkpoint, we remember all of the actors that need a checkpoint saved out by adding them to this list */
 	struct FPendingCheckPointActor
 	{
