Fix crash in EDL builds caused by accessing sparse class data in Link() before it is serialized

Only calls the Link() time stuff in WITH_EDITORONLY_DATA builds (as it was added to address cook-time issues anyways the impact should be minimal)

diff --git a/Engine/Source/Runtime/Engine/Private/Animation/AnimBlueprintGeneratedClass.cpp b/Engine/Source/Runtime/Engine/Private/Animation/AnimBlueprintGeneratedClass.cpp
--- a/Engine/Source/Runtime/Engine/Private/Animation/AnimBlueprintGeneratedClass.cpp
+++ b/Engine/Source/Runtime/Engine/Private/Animation/AnimBlueprintGeneratedClass.cpp
@@ -543,17 +543,25 @@ void UAnimBlueprintGeneratedClass::Link(FArchive& Ar, bool bRelinkExistingProper
 		}
 	}
 
-	// Must build constant properties to be able to iterate subsystems
-	BuildConstantProperties();
-	
-	ForEachSubsystem([this](const FAnimSubsystemContext& InContext)
+#if WITH_EDITORONLY_DATA
+	// With cooked data, we skip this as Link() is called on serialization and sparse class data
+	// is not properly serialized yet
+	if(!GetPackage()->bIsCookedForEditor)
 	{
-		Subsystems.Add(&InContext.Subsystem);
+		// Must build constant properties to be able to iterate subsystems and call OnLink()
+		// Subsystems in non-editor data builds should be initialized in OnPostLoadDefaults.
+		BuildConstantProperties();
+		
+		ForEachSubsystem([this](const FAnimSubsystemContext& InContext)
+		{
+			Subsystems.Add(&InContext.Subsystem);
 
-		FAnimSubsystemLinkContext Context(InContext, *this);
-		const_cast<FAnimSubsystem&>(InContext.Subsystem).OnLink(Context);
-		return EAnimSubsystemEnumeration::Continue;
-	});	
+			FAnimSubsystemLinkContext Context(InContext, *this);
+			const_cast<FAnimSubsystem&>(InContext.Subsystem).OnLink(Context);
+			return EAnimSubsystemEnumeration::Continue;
+		});
+	}
+#endif
 }
 
 void UAnimBlueprintGeneratedClass::PurgeClass(bool bRecompilingOnLoad)
diff --git a/Engine/Source/Runtime/Engine/Public/Animation/AnimSubsystem.h b/Engine/Source/Runtime/Engine/Public/Animation/AnimSubsystem.h
--- a/Engine/Source/Runtime/Engine/Public/Animation/AnimSubsystem.h
+++ b/Engine/Source/Runtime/Engine/Public/Animation/AnimSubsystem.h
@@ -122,6 +122,8 @@ struct FAnimSubsystem
 	/** Override point to perform subsystem class data initialization post-load */
 	virtual void OnPostLoad(FAnimSubsystemPostLoadContext& InContext) {}
 
+#if WITH_EDITORONLY_DATA
 	/** Override point to perform subsystem class data initialization on class link */
 	virtual void OnLink(FAnimSubsystemLinkContext& InContext) {}
+#endif
 };
\ No newline at end of file
