SubObject API
* Added cheat to validate the old and new methods of replicating subobjects. Enabled via net.SubObjects.CompareWithLegacy
* The comparison is run only for actors that converted to the new list.
* This calls the ReplicateSubObject virtual so the method still needs to be implemented for the validation to work.

diff --git a/Engine/Source/Runtime/Engine/Classes/Components/ActorComponent.h b/Engine/Source/Runtime/Engine/Classes/Components/ActorComponent.h
--- a/Engine/Source/Runtime/Engine/Classes/Components/ActorComponent.h
+++ b/Engine/Source/Runtime/Engine/Classes/Components/ActorComponent.h
@@ -475,7 +475,8 @@ public:
 	}
 
 	/** 
-	* Allows a component to replicate other subobject on the actor
+	* Allows a component to replicate other subobject on the actor.
+	* Must return true if any data gets serialized into the bunch.
 	* This method is used only when bReplicateUsingRegisteredSubObjectList is false.
 	* Otherwise this function is not called and only the ReplicatedSubObjects list is used.
 	*/
diff --git a/Engine/Source/Runtime/Engine/Classes/Engine/ActorChannel.h b/Engine/Source/Runtime/Engine/Classes/Engine/ActorChannel.h
--- a/Engine/Source/Runtime/Engine/Classes/Engine/ActorChannel.h
+++ b/Engine/Source/Runtime/Engine/Classes/Engine/ActorChannel.h
@@ -395,6 +395,8 @@ private:
 	/** Find the replicated subobjects of the component and write them into the bunch */
 	bool WriteSubObjects(UActorComponent* Component, FOutBunch& Bunch, FReplicationFlags RepFlags, const TStaticBitArray<COND_Max>& ConditionMap);
 
+	bool ValidateReplicatedSubObjects() const;
+
 private:
 
 	// TODO: It would be nice to merge the tracking of these with PendingGuidResolves, to not duplicate memory,
diff --git a/Engine/Source/Runtime/Engine/Classes/GameFramework/Actor.h b/Engine/Source/Runtime/Engine/Classes/GameFramework/Actor.h
--- a/Engine/Source/Runtime/Engine/Classes/GameFramework/Actor.h
+++ b/Engine/Source/Runtime/Engine/Classes/GameFramework/Actor.h
@@ -718,7 +718,12 @@ public:
     */
 	bool IsUsingRegisteredSubObjectList() const { return bReplicateUsingRegisteredSubObjectList; }
 
-	/** Method that allows an actor to replicate subobjects on its actor channel */
+	/** 
+	* Method that allows an actor to replicate subobjects on its actor channel. 
+	* Must return true if any data was serialized into the bunch.
+	* This method is used only when bReplicateUsingRegisteredSubObjectList is false.
+	* Otherwise this function is not called and only the ReplicatedSubObjects list is used.
+	*/
 	virtual bool ReplicateSubobjects(class UActorChannel *Channel, class FOutBunch *Bunch, FReplicationFlags *RepFlags);
 
 	/** Called on the actor when a new subobject is dynamically created via replication */
diff --git a/Engine/Source/Runtime/Engine/Private/ActorReplication.cpp b/Engine/Source/Runtime/Engine/Private/ActorReplication.cpp
--- a/Engine/Source/Runtime/Engine/Private/ActorReplication.cpp
+++ b/Engine/Source/Runtime/Engine/Private/ActorReplication.cpp
@@ -478,6 +478,7 @@ bool AActor::ReplicateSubobjects(UActorChannel *Channel, FOutBunch *Bunch, FRepl
 		{
 			UActorChannel::SetCurrentSubObjectOwner(ActorComp);
 			WroteSomething |= ActorComp->ReplicateSubobjects(Channel, Bunch, RepFlags);		// Lets the component add subobjects before replicating its own properties.
+			UActorChannel::SetCurrentSubObjectOwner(this);
 			WroteSomething |= Channel->ReplicateSubobject(ActorComp, *Bunch, *RepFlags);	// (this makes those subobjects 'supported', and from here on those objects may have reference replicated)		
 		}
 	}
diff --git a/Engine/Source/Runtime/Engine/Private/DataChannel.cpp b/Engine/Source/Runtime/Engine/Private/DataChannel.cpp
--- a/Engine/Source/Runtime/Engine/Private/DataChannel.cpp
+++ b/Engine/Source/Runtime/Engine/Private/DataChannel.cpp
@@ -114,7 +114,19 @@ namespace UE::Net
 		TEXT("net.QueuedBunchTimeoutSeconds"),
 		QueuedBunchTimeoutSeconds,
 		TEXT("Time in seconds to wait for queued bunches on a channel to flush before logging a warning."));
-};
+
+	int32 GCVarCompareSubObjectsReplicated = 0;
+	FAutoConsoleVariableRef CVarCompareSubObjectsReplicated(
+		TEXT("net.SubObjects.CompareWithLegacy"),
+		GCVarCompareSubObjectsReplicated,
+		TEXT("When turned on we will collect the subobjects replicated by the ReplicateSubObjects method and compare them with the ones replicated via the Actor's registered list. If a divergence is detected it will trigger an ensure."));
+
+	TAutoConsoleVariable<int32> CVarLogAllSubObjectComparisonErrors(
+		TEXT("net.SubObjects.LogAllComparisonErrors"),
+		0,
+		TEXT("If enabled log all the errors detected by the CompareWithLegacy cheat. Otherwise only the first ensure triggered gets logged."));
+
+}; // namespace UE::Net
 
 template<typename T>
 static const bool IsBunchTooLarge(UNetConnection* Connection, T* Bunch)
@@ -158,9 +170,50 @@ namespace DataChannelInternal
 	/**
 	* When enabled we will ignore requests to replicate a subobject via the public ReplicateSubObject function.
 	* This occurs when the subobjects belong to an ActorComponent converted to the registered list, but the owner does not support it yet so we still call ReplicateSubObjects on the actor.
-	* This flag is controlled via UActorChannel::SetCurrentSubObjectOwner, so it is important for users set the owner properly if they are controlling component replication themselves.
+	* This flag is controlled via UActorChannel::SetCurrentSubObjectOwner, so it is important for users to set the owner properly if they are controlling component replication themselves.
 	*/
 	bool bIgnoreLegacyReplicateSubObject = false;
+
+	/** When enabled store all subobjects written into the bunch in the ReplicatedSubObjectsTracker array */
+	bool bTrackReplicatedSubObjects = false;
+
+	/** 
+	* This flag is true when an actor using the new subobject list is calling the legacy function to look for divergences between the old and new method. 
+	* This ensures that the subobjects replicated the old way are only collected and not replicated.
+	*/
+	bool bTestingLegacyMethodForComparison = false;
+
+    struct FSubObjectReplicatedInfo
+	{
+		// The replicated subobject
+		UObject* SubObject = nullptr;
+		// The owner of the subobject, usually an actor or actor component
+		UObject* SubObjectOwner = nullptr;
+
+		FSubObjectReplicatedInfo(UObject* InSubObject) : SubObject(InSubObject), SubObjectOwner(CurrentSubObjectOwner) {}
+
+		bool operator==(const FSubObjectReplicatedInfo& rhs) const { return SubObject == rhs.SubObject; }
+
+		FString Describe(AActor* ReplicatedActor) const
+		{ 
+			return FString::Printf(TEXT("Subobject %s::%s (0x%p) [owner %s]"), *ReplicatedActor->GetName(), *SubObject->GetName(), SubObject, *GetNameSafe(SubObjectOwner));
+		}
+	};
+
+	/**
+	
+	* List of replicated subobjects collected via Actor::ReplicateSubObjects.
+	* This list is static to prevent memory overhead on every actor channel.
+	* It is emptied at the end of every call to UActorChannel::ReplicateActor()
+	*/
+	TArray<FSubObjectReplicatedInfo> LegacySubObjectsCollected;
+
+    /**
+	* The subobjects actually replicated into the bunch.
+	* This list is static to prevent memory overhead on every actor channel. 
+    * It is emptied at the end of every call to UActorChannel::ReplicateActor()
+	*/
+	TArray<FSubObjectReplicatedInfo> ReplicatedSubObjectsTracker;
 }
 
 /*-----------------------------------------------------------------------------
@@ -3471,6 +3524,15 @@ bool UActorChannel::DoSubObjectReplication(FOutBunch& Bunch, FReplicationFlags&
 	{
 		// If the actor replicates it's subobjects and those of it's components via the list.
 		bWroteSomethingImportant |= ReplicateRegisteredSubObjects(Bunch, OutRepFlags);
+
+		if (UE::Net::GCVarCompareSubObjectsReplicated)
+		{
+			ValidateReplicatedSubObjects();
+
+			DataChannelInternal::bTrackReplicatedSubObjects = false;
+			DataChannelInternal::LegacySubObjectsCollected.Reset();
+			DataChannelInternal::ReplicatedSubObjectsTracker.Reset();
+		}
 	}
 	else
 	{
@@ -3478,6 +3540,9 @@ bool UActorChannel::DoSubObjectReplication(FOutBunch& Bunch, FReplicationFlags&
 		bWroteSomethingImportant |= Actor->ReplicateSubobjects(this, &Bunch, &OutRepFlags);
 	}
 
+	AActor* TempNull(nullptr);
+	SetCurrentSubObjectOwner(TempNull);
+
 	return bWroteSomethingImportant;
 }
 
@@ -3485,6 +3550,21 @@ bool UActorChannel::ReplicateRegisteredSubObjects(FOutBunch& Bunch, FReplication
 {
 	check(Actor->IsUsingRegisteredSubObjectList());
 
+	if (UE::Net::GCVarCompareSubObjectsReplicated)
+	{
+		// Call the legacy method but simply collect the subobjects it would have replicated
+		DataChannelInternal::bTestingLegacyMethodForComparison = true;
+
+		SetCurrentSubObjectOwner(Actor);
+		FReplicationFlags TempRepFlags = RepFlags; // Mirror the flags in case someone actually dirties them.
+		Actor->ReplicateSubobjects(this, &Bunch, &TempRepFlags);
+		
+		DataChannelInternal::bTestingLegacyMethodForComparison = false;
+
+		// From here track all subobjects that write into the bunch.
+		DataChannelInternal::bTrackReplicatedSubObjects = true;
+	}
+
 	const TStaticBitArray<COND_Max> ConditionMap = FSendingRepState::BuildConditionMapFromRepFlags(RepFlags);
 
 	bool bWroteSomethingImportant = false;
@@ -3508,7 +3588,6 @@ bool UActorChannel::ReplicateRegisteredSubObjects(FOutBunch& Bunch, FReplication
 		check(ReplicatedComponent);
 
 		SetCurrentSubObjectOwner(ReplicatedComponent);
-
 		if (ReplicatedComponent->IsUsingRegisteredSubObjectList())
 		{
 			bWroteSomethingImportant |= WriteSubObjects(ReplicatedComponent, Bunch, RepFlags, ConditionMap);
@@ -3547,6 +3626,12 @@ bool UActorChannel::ReplicateSubobject(UObject* SubObj, FOutBunch& Bunch, FRepli
 		return false;
 	}
 
+	if (DataChannelInternal::bTestingLegacyMethodForComparison)
+	{
+		DataChannelInternal::LegacySubObjectsCollected.Emplace(DataChannelInternal::FSubObjectReplicatedInfo(SubObj));
+		return false;
+	}
+
 	bool bWroteSomethingImportant = false;
 
 	if (!DataChannelInternal::bIgnoreLegacyReplicateSubObject)
@@ -3577,31 +3662,29 @@ bool UActorChannel::ReplicateSubobject(UActorComponent* ReplicatedComponent, FOu
 		return false;
 	}
 
-	bool bWroteSomethingImportant = false;
-
-	if (ReplicatedComponent->IsUsingRegisteredSubObjectList())
+	// This traps actor components using the registration list even if their owning Actor isn't using the list itself.
+	if (ReplicatedComponent->IsUsingRegisteredSubObjectList() && !DataChannelInternal::bTestingLegacyMethodForComparison)
 	{
+		bool bWroteSomethingImportant = false;
+
 		const TStaticBitArray<COND_Max> ConditionMap = FSendingRepState::BuildConditionMapFromRepFlags(RepFlags);
 
 		checkf(Actor->IsUsingRegisteredSubObjectList() == false, TEXT("This code should only be hit when an Actor that does NOT support the SubObjectList is replicating an ActorComponent that does: %s replicating %s."), 
 			*Actor->GetName(), *ReplicatedComponent->GetName());
 
-		ensureMsgf(DataChannelInternal::CurrentSubObjectOwner == ReplicatedComponent, TEXT("%s (owner %s) was not the CurrentSubObjectOwner prior to replicating his subobjects. Make sure to call UActorChannel::SetCurrentSubObjectOwner BEFORE replicating any of the component's subobjects."),
-			*ReplicatedComponent->GetName(), *Actor->GetName());
-
         // Write all the subobjects of the component
 		bWroteSomethingImportant |= WriteSubObjects(ReplicatedComponent, Bunch, RepFlags, ConditionMap);
 
         // Finally write the component itself
 		SetCurrentSubObjectOwner(Actor);
 		bWroteSomethingImportant |= WriteSubObjectInBunch(ReplicatedComponent, Bunch, RepFlags);
+
+		return bWroteSomethingImportant;
 	}
 	else
 	{
 		return ReplicateSubobject(StaticCast<UObject*>(ReplicatedComponent), Bunch, RepFlags);
 	}
-
-	return bWroteSomethingImportant;
 }
 
 bool UActorChannel::WriteSubObjects(UActorComponent* ReplicatedComponent, FOutBunch& Bunch, FReplicationFlags RepFlags, const TStaticBitArray<COND_Max>& ConditionMap)
@@ -3623,6 +3706,35 @@ bool UActorChannel::WriteSubObjects(UActorComponent* ReplicatedComponent, FOutBu
 	return bWroteSomethingImportant;
 }
 
+bool UActorChannel::ValidateReplicatedSubObjects() const
+{
+	bool bErrorDetected = false;
+	const bool bLogAllErrors = UE::Net::CVarLogAllSubObjectComparisonErrors.GetValueOnAnyThread() != 0;
+	
+	// Make sure both lists contain the same subobjects and trigger an ensure if an element is missing from the other.
+	for (const DataChannelInternal::FSubObjectReplicatedInfo& Info : DataChannelInternal::LegacySubObjectsCollected)
+	{
+		constexpr bool bNoShrinking = false;
+		const bool bWasFound = DataChannelInternal::ReplicatedSubObjectsTracker.RemoveSingleSwap(Info, bNoShrinking) != 0;
+
+		ensureMsgf(bWasFound, TEXT("%s was not replicated by the subobject list. Only by the legacy ReplicateSubObjects method."), *Info.Describe(Actor));
+		UE_CLOG(!bWasFound && bLogAllErrors, LogNetSubObject, Warning, TEXT("%s was not replicated by the subobject list."), *Info.Describe(Actor));
+
+		bErrorDetected |= !bWasFound;
+	}
+
+	// Leftover subobjects in this list means we replicated more elements than the legacy method would.
+	for (const DataChannelInternal::FSubObjectReplicatedInfo& Info : DataChannelInternal::ReplicatedSubObjectsTracker)
+	{
+		ensureMsgf(false, TEXT("%s was replicated only in the registered subobject list. Not in the legacy path."), *Info.Describe(Actor));
+		UE_CLOG(bLogAllErrors, LogNet, Warning, TEXT("%s was replicated only in the registered subobject list. Not in the legacy path."), *Info.Describe(Actor));
+	}
+
+	bErrorDetected |= DataChannelInternal::ReplicatedSubObjectsTracker.Num() != 0;
+
+	return bErrorDetected;
+}
+
 FString UActorChannel::Describe()
 {
 	if (!Actor)
@@ -4627,6 +4739,8 @@ bool UActorChannel::WriteSubObjectInBunch(UObject* Obj, FOutBunch& Bunch, FRepli
 		return false;
 	}
 
+	check(DataChannelInternal::bTestingLegacyMethodForComparison == false);
+
     SCOPE_CYCLE_UOBJECT(ActorChannelRepSubObj, Obj);
 
     if (RepFlags.bUseCustomSubobjectReplication)
@@ -4634,6 +4748,11 @@ bool UActorChannel::WriteSubObjectInBunch(UObject* Obj, FOutBunch& Bunch, FRepli
 		return ReplicateSubobjectCustom(Obj, Bunch, RepFlags);
 	}
 
+	if (DataChannelInternal::bTrackReplicatedSubObjects)
+	{
+		DataChannelInternal::ReplicatedSubObjectsTracker.Emplace(DataChannelInternal::FSubObjectReplicatedInfo(Obj));
+	}
+
 	bool bFoundInvalidReplicator = false;
 	TSharedRef<FObjectReplicator>* FoundReplicator = FindReplicator(Obj, &bFoundInvalidReplicator);
 	const bool bFoundReplicator = (FoundReplicator != nullptr);
