Fix an edge case with pushing and popping device profiles. If a cvar was set externally between a push and pop, it would be restored with the new priority rather than the original pushed priority. Not only would the cvar value get overridden, the popped value would then have high priority and could never change. This could manifest when setting cvars via execcmds, for example.

This change detects if the priority has changed to a value higher than SetByDeviceProfile since the last push and skips in those cases.

diff --git a/Engine/Source/Runtime/Engine/Classes/DeviceProfiles/DeviceProfileManager.h b/Engine/Source/Runtime/Engine/Classes/DeviceProfiles/DeviceProfileManager.h
--- a/Engine/Source/Runtime/Engine/Classes/DeviceProfiles/DeviceProfileManager.h
+++ b/Engine/Source/Runtime/Engine/Classes/DeviceProfiles/DeviceProfileManager.h
@@ -11,6 +11,7 @@
 #include "UObject/Object.h"
 #include "DeviceProfileMatching.h"
 #include "CoreGlobals.h"
+#include "HAL/IConsoleManager.h"
 #include "DeviceProfileManager.generated.h"
 
 class UDeviceProfile;
@@ -21,6 +22,18 @@ DECLARE_MULTICAST_DELEGATE( FOnDeviceProfileManagerUpdated );
 // Delegate used to notify systems when the active device profile changes
 DECLARE_MULTICAST_DELEGATE( FOnActiveDeviceProfileChanged );
 
+struct FPushedCVarSetting
+{
+	FPushedCVarSetting() : SetBy(ECVF_Default)	{}
+	FPushedCVarSetting(const FString& InValue, EConsoleVariableFlags InFlags) 
+		: Value(InValue)
+		, SetBy(EConsoleVariableFlags(InFlags & ECVF_SetByMask))
+	{}
+
+	FString Value;
+	EConsoleVariableFlags SetBy;
+};
+
 /**
  * Implements a helper class that manages all profiles in the Device
  */
@@ -272,11 +285,11 @@ private:
 
 	// Original values of all the CVars modified by the DP.
 	// Used to undo the DP before applying new state.
-	static TMap<FString, FString> PushedSettings;
+	static TMap<FString, FPushedCVarSetting> PushedSettings;
 
 #if ALLOW_OTHER_PLATFORM_CONFIG
 	// Original values of the CVars modified by SetPreviewDeviceProfile.
-	TMap<FString, FString> PreviewPushedSettings;
+	TMap<FString, FPushedCVarSetting> PreviewPushedSettings;
 #endif
 	// Holds the device profile that has been overridden, null no override active.
 	UDeviceProfile* BaseDeviceProfile = nullptr;
diff --git a/Engine/Source/Runtime/Engine/Private/DeviceProfiles/DeviceProfileManager.cpp b/Engine/Source/Runtime/Engine/Private/DeviceProfiles/DeviceProfileManager.cpp
--- a/Engine/Source/Runtime/Engine/Private/DeviceProfiles/DeviceProfileManager.cpp
+++ b/Engine/Source/Runtime/Engine/Private/DeviceProfiles/DeviceProfileManager.cpp
@@ -41,7 +41,7 @@ static TAutoConsoleVariable<int32> CVarAllowScalabilityGroupsToChangeAtRuntime(
 TMap<FString, FString> UDeviceProfileManager::DeviceProfileScalabilityCVars;
 
 FString UDeviceProfileManager::BackupSuffix = TEXT("_Backup");
-TMap<FString, FString> UDeviceProfileManager::PushedSettings;
+TMap<FString, FPushedCVarSetting> UDeviceProfileManager::PushedSettings;
 TArray<FSelectedFragmentProperties> UDeviceProfileManager::PlatformFragmentsSelected;
 
 UDeviceProfileManager* UDeviceProfileManager::DeviceProfileManagerSingleton = nullptr;
@@ -385,12 +385,12 @@ void UDeviceProfileManager::ProcessDeviceProfileIniSettings(const FString& Devic
 							{
 								if (CVar)
 								{
-									// remember the previous value
-									FString OldValue = CVar->GetString();
-									PushedSettings.Add(CVarKey, OldValue);
+									// remember the previous value and priority
+									FPushedCVarSetting OldSetting = FPushedCVarSetting(CVar->GetString(), CVar->GetFlags());
+									PushedSettings.Add(CVarKey, OldSetting);
 
 									// indicate we are pushing, not setting
-									UE_LOG(LogDeviceProfileManager, Log, TEXT("Pushing Device Profile CVar: [[%s:%s -> %s]]"), *CVarKey, *OldValue, *CVarValue);
+									UE_LOG(LogDeviceProfileManager, Log, TEXT("Pushing Device Profile CVar: [[%s:%s -> %s]]"), *CVarKey, *OldSetting.Value, *CVarValue);
 								}
 								else
 								{
@@ -480,17 +480,29 @@ void UDeviceProfileManager::ProcessDeviceProfileIniSettings(const FString& Devic
 /**
 * Set the cvar state to PushedSettings.
 */
-static void RestorePushedState(TMap<FString, FString>& PushedSettings)
+static void RestorePushedState(TMap<FString, FPushedCVarSetting>& PushedSettings)
 {
 	// restore pushed settings
-	for (TMap<FString, FString>::TIterator It(PushedSettings); It; ++It)
+	for (TMap<FString, FPushedCVarSetting>::TIterator It(PushedSettings); It; ++It)
 	{
 		IConsoleVariable* CVar = IConsoleManager::Get().FindConsoleVariable(*It.Key());
 		if (CVar)
 		{
-			// restore it!
-			CVar->SetWithCurrentPriority(*It.Value());
-			UE_LOG(LogDeviceProfileManager, Log, TEXT("Popping Device Profile CVar: [[%s:%s]]"), *It.Key(), *It.Value());
+			const FPushedCVarSetting& Setting = It.Value();
+
+			// Check if the priority has increased since we pushed it and is now higher than SetByDeviceProfile
+			EConsoleVariableFlags CurrentPriority = EConsoleVariableFlags(int32(CVar->GetFlags()) & (int32)ECVF_SetByMask);
+			if (CurrentPriority > ECVF_SetByDeviceProfile && CurrentPriority > Setting.SetBy )
+			{
+				UE_LOG(LogDeviceProfileManager, Warning, TEXT("Popping Device Profile CVar skipped because priority has been overridden to higher than ECVF_SetByDeviceProfile since the last push [[%s:%s]]"), *It.Key(), *Setting.Value);
+			}
+			else
+			{
+				// restore it!
+				CVar->SetWithCurrentPriority(*Setting.Value);
+				UE_LOG(LogDeviceProfileManager, Log, TEXT("Popping Device Profile CVar: [[%s:%s]]"), *It.Key(), *Setting.Value);
+			}
+
 		}
 	}
 
@@ -954,9 +966,9 @@ void UDeviceProfileManager::SetPreviewDeviceProfile(UDeviceProfile* DeviceProfil
 		// skip over scalability group cvars (maybe they shouldn't be left in the AllExpandedCVars?)
 		if (CVar != nullptr && !CVar->TestFlags(EConsoleVariableFlags::ECVF_ScalabilityGroup))
 		{
-			// remember the previous value so we can restore
-			FString OldValue = CVar->GetString();
-			PreviewPushedSettings.Add(Pair.Key, OldValue);
+			// remember the previous value and priority so we can restore
+			FPushedCVarSetting OldSetting = FPushedCVarSetting(CVar->GetString(), CVar->GetFlags());
+			PreviewPushedSettings.Add(Pair.Key, OldSetting);
 			//UE_LOG(LogDeviceProfileManager, Log, TEXT("Pushing Device Profile CVar: [[%s:%s]]"), *Pair.Key, *Pair.Value);
 			// cheat CVar can only be set in ConsoleVariables.ini
 			if (!CVar->TestFlags(EConsoleVariableFlags::ECVF_Cheat))
