Add define UE_ANIM_REMOVE_DEPRECATED_ANCESTOR_TRACKER to remove large deprecated member from FAnimationUpdateSharedContext.

diff --git a/Engine/Source/Runtime/Engine/Classes/Animation/AnimNodeBase.h b/Engine/Source/Runtime/Engine/Classes/Animation/AnimNodeBase.h
--- a/Engine/Source/Runtime/Engine/Classes/Animation/AnimNodeBase.h
+++ b/Engine/Source/Runtime/Engine/Classes/Animation/AnimNodeBase.h
@@ -27,6 +27,10 @@
 #define DECLARE_SCOPE_HIERARCHICAL_COUNTER_ANIMNODE(Method) \
 	DECLARE_SCOPE_HIERARCHICAL_COUNTER_FUNC()
 
+#ifndef UE_ANIM_REMOVE_DEPRECATED_ANCESTOR_TRACKER
+	#define UE_ANIM_REMOVE_DEPRECATED_ANCESTOR_TRACKER 0
+#endif
+
 class IAnimClassInterface;
 class UAnimBlueprint;
 class UAnimInstance;
@@ -135,17 +139,21 @@ struct FAnimationUpdateSharedContext
 	FAnimationUpdateSharedContext(FAnimationUpdateSharedContext& ) = delete;
 	FAnimationUpdateSharedContext& operator=(const FAnimationUpdateSharedContext&) = delete;
 
+#if !UE_ANIM_REMOVE_DEPRECATED_ANCESTOR_TRACKER
 	UE_DEPRECATED(5.0, "Please use the message & tagging system in UE::Anim::FMessageStack")
 	FAnimNodeTracker AncestorTracker;
+#endif
 
 	// Message stack used for storing scoped messages and tags during execution
 	UE::Anim::FMessageStack MessageStack;
 
 	void CopyForCachedUpdate(FAnimationUpdateSharedContext& Source)
 	{
+#if !UE_ANIM_REMOVE_DEPRECATED_ANCESTOR_TRACKER
 PRAGMA_DISABLE_DEPRECATION_WARNINGS
 		AncestorTracker.CopyTopsOnly(Source.AncestorTracker);
 PRAGMA_ENABLE_DEPRECATION_WARNINGS
+#endif
 		MessageStack.CopyForCachedUpdate(Source.MessageStack);
 	}
 };
@@ -184,6 +192,7 @@ public:
 	ENGINE_API UAnimBlueprint* GetAnimBlueprint() const;
 #endif //WITH_EDITORONLY_DATA
 
+#if !UE_ANIM_REMOVE_DEPRECATED_ANCESTOR_TRACKER
 	template<typename NodeType>
 	UE_DEPRECATED(5.0, "Please use the message & tagging system in UE::Anim::FMessageStack")
 	FScopedAnimNodeTracker TrackAncestor(NodeType* Node) const {
@@ -195,7 +204,9 @@ public:
 
 		return FScopedAnimNodeTracker();
 	}
+#endif
 
+#if !UE_ANIM_REMOVE_DEPRECATED_ANCESTOR_TRACKER
 	template<typename NodeType>
 	UE_DEPRECATED(5.0, "Please use the message & tagging system in UE::Anim::FMessageStack")
 	NodeType* GetAncestor() const {
@@ -207,6 +218,7 @@ public:
 		
 		return nullptr;
 	}
+#endif
 
 	// Get the innermost scoped message of the specified type
 	template<typename TGraphMessageType>
