Optimizing padding in AActor to reduce size from 664 bytes to 640.

diff --git a/Engine/Source/Runtime/Engine/Classes/GameFramework/Actor.h b/Engine/Source/Runtime/Engine/Classes/GameFramework/Actor.h
--- a/Engine/Source/Runtime/Engine/Classes/GameFramework/Actor.h
+++ b/Engine/Source/Runtime/Engine/Classes/GameFramework/Actor.h
@@ -478,6 +478,9 @@ private:
 	 */
 	EActorBeginPlayState ActorHasBegunPlay:2;
 
+	/** Set while actor is being constructed. Used to ensure that construction is not re-entrant. */
+	uint8 bActorIsBeingConstructed : 1;
+
 	static uint32 BeginPlayCallDepth;
 
 protected:
@@ -527,11 +530,6 @@ private:
 	/** Internal helper to update Overlaps during Actor initialization/BeginPlay correctly based on the UpdateOverlapsMethodDuringLevelStreaming and bGenerateOverlapEventsDuringLevelStreaming settings. */
 	void UpdateInitialOverlaps(bool bFromLevelStreaming);
 
-	/** Describes how much control the remote machine has over the actor. */
-	UPROPERTY(Replicated, Transient)
-	TEnumAsByte<enum ENetRole> RemoteRole;	
-
-
 public:
 	/**
 	 * Set whether this actor replicates to network clients. When this actor is spawned on the server it will be sent to clients as well.
@@ -570,11 +568,6 @@ public:
 	 */
 	void SetNetAddressable();
 
-private:
-	/** Used for replication of our RootComponent's position and velocity */
-	UPROPERTY(EditDefaultsOnly, ReplicatedUsing=OnRep_ReplicatedMovement, Category=Replication, AdvancedDisplay)
-	struct FRepMovement ReplicatedMovement;
-
 public:
 	/** How long this Actor lives before dying, 0=forever. Note this is the INITIAL value and should not be modified once play has begun. */
 	UPROPERTY(EditAnywhere, BlueprintReadOnly, Category=Actor)
@@ -583,12 +576,15 @@ public:
 	/** Allow each actor to run at a different time speed. The DeltaTime for a frame is multiplied by the global TimeDilation (in WorldSettings) and this CustomTimeDilation for this actor's tick.  */
 	UPROPERTY(BlueprintReadWrite, AdvancedDisplay, Category=Actor)
 	float CustomTimeDilation;
+	
+private:
+	/** Describes how much control the remote machine has over the actor. */
+	UPROPERTY(Replicated, Transient)
+	TEnumAsByte<enum ENetRole> RemoteRole;
 
-	/**
-	 * The time this actor was created, relative to World->GetTimeSeconds().
-	 * @see UWorld::GetTimeSeconds()
-	 */
-	float CreationTime;
+	/** The RayTracingGroupId this actor and its components belong to. (For components that did not specify any) */
+	UPROPERTY()
+	int32 RayTracingGroupId;
 
 protected:
 #if WITH_EDITORONLY_DATA
@@ -613,6 +609,12 @@ protected:
 	UPROPERTY(Transient, ReplicatedUsing=OnRep_AttachmentReplication)
 	struct FRepAttachment AttachmentReplication;
 
+private:
+	/** Used for replication of our RootComponent's position and velocity */
+	UPROPERTY(EditDefaultsOnly, ReplicatedUsing=OnRep_ReplicatedMovement, Category=Replication, AdvancedDisplay)
+	struct FRepMovement ReplicatedMovement;
+
+public:
 	/**
 	 * Owner of this Actor, used primarily for replication (bNetUseOwnerRelevancy & bOnlyRelevantToOwner) and visibility (PrimitiveComponent bOwnerNoSee and bOnlyOwnerSee)
 	 * @see SetOwner(), GetOwner()
@@ -667,6 +669,12 @@ public:
 	/** The priority of this input component when pushed in to the stack. */
 	UPROPERTY(EditAnywhere, Category=Input)
 	int32 InputPriority;
+	
+	/**
+	 * The time this actor was created, relative to World->GetTimeSeconds().
+	 * @see UWorld::GetTimeSeconds()
+	 */
+	float CreationTime;
 
 	/** Component that handles input for this actor, if input is enabled. */
 	UPROPERTY(DuplicateTransient)
@@ -808,10 +816,6 @@ private:
 	TObjectPtr<class UHLODLayer> HLODLayer;
 #endif
 
-	/** The RayTracingGroupId this actor and its components belong to. (For components that did not specify any) */
-	UPROPERTY()
-	int32 RayTracingGroupId;
-
 public:
 	/** Return the value of bAllowReceiveTickEventOnDedicatedServer, indicating whether the Blueprint ReceiveTick() event will occur on dedicated servers. */
 	FORCEINLINE bool AllowReceiveTickEventOnDedicatedServer() const { return bAllowReceiveTickEventOnDedicatedServer; }
@@ -1030,9 +1034,6 @@ private:
 	uint8 bForceExternalActorLevelReferenceForPIE : 1;
 #endif // WITH_EDITORONLY_DATA
 
-	/** Set while actor is being constructed. Used to ensure that construction is not re-entrant. */
-	uint8 bActorIsBeingConstructed : 1;
-
 public:
 	/** Array of tags that can be used for grouping and categorizing. */
 	UPROPERTY(EditAnywhere, BlueprintReadWrite, AdvancedDisplay, Category=Actor)
