BugFix: "Quality and Random Nodes do not evaluate on cook and instead load all variations in the SoundCue bypassing the Quality level"

diff --git a/Engine/Source/Runtime/Engine/Private/SoundCue.cpp b/Engine/Source/Runtime/Engine/Private/SoundCue.cpp
--- a/Engine/Source/Runtime/Engine/Private/SoundCue.cpp
+++ b/Engine/Source/Runtime/Engine/Private/SoundCue.cpp
@@ -212,8 +212,8 @@ void USoundCue::PostLoad()
 #endif // WITH_EDITOR
 
 	// Warn if the Quality index is set to something that we can't support.
-	UE_CLOG(USoundCue::GetCachedQualityLevel() != CookedQualityIndex && CookedQualityIndex != INDEX_NONE, LogAudio, Warning,
-		TEXT("'%s' is ingoring Quality Setting '%s'(%d) as it was cooked with '%s'(%d)"),
+	UE_CLOG(USoundCue::GetCachedQualityLevel() != CookedQualityIndex && CookedQualityIndex != INDEX_NONE, LogAudio, Verbose,
+		TEXT("'%s' is igoring Quality Setting '%s'(%d) as it was cooked with '%s'(%d)"),
 		*GetFullNameSafe(this),
 		*GetDefault<UAudioSettings>()->FindQualityNameByIndex(USoundCue::GetCachedQualityLevel()),
 		USoundCue::GetCachedQualityLevel(),
diff --git a/Engine/Source/Runtime/Engine/Private/SoundNodeQualityLevel.cpp b/Engine/Source/Runtime/Engine/Private/SoundNodeQualityLevel.cpp
--- a/Engine/Source/Runtime/Engine/Private/SoundNodeQualityLevel.cpp
+++ b/Engine/Source/Runtime/Engine/Private/SoundNodeQualityLevel.cpp
@@ -117,7 +117,7 @@ void USoundNodeQualityLevel::ParseNodes( FAudioDevice* AudioDevice, const UPTRIN
 	int32 QualityLevel = USoundCue::GetCachedQualityLevel();
 	
 	// If CookedQualityLevelIndex has been set, we will have a *single* quality level.
-	if (CookedQualityLevelIndex >= 0)
+	if (CookedQualityLevelIndex >= 0 && ChildNodes.Num() == 1)
 	{	
 		// Remap to index 0 (as all other levels have been removed by cooker).
 		QualityLevel = 0;
@@ -148,15 +148,27 @@ void USoundNodeQualityLevel::Serialize(FArchive& Ar)
 				CookedQualityLevelIndex = CookOverrides->SoundCueCookQualityIndex;
 
 				// Move out all nodes.
-				TArray<USoundNode*> ChildNodesBackup;
+				TArray<TObjectPtr<class USoundNode>> ChildNodesBackup;
 				ChildNodesBackup = MoveTemp(ChildNodes);
-
+				check(ChildNodes.Num() == 0);
+			
 				// Put *just* the node we care about in our child array to be serialized by the Super
 				if (ChildNodesBackup.IsValidIndex(CookedQualityLevelIndex))
 				{
 					ChildNodes.Add(ChildNodesBackup[CookedQualityLevelIndex]);
 				}
 
+				int32 BranchesPruned = ChildNodesBackup.Num() - ChildNodes.Num();
+				UE_CLOG(
+					BranchesPruned > 0,
+					LogAudio,
+					Display, 
+					TEXT("Pruning '%s' of '%d' quality branches, as it's cooked at '%s' quality."),
+					*GetFullNameSafe(this),
+					BranchesPruned,
+					*GetDefault<UAudioSettings>()->FindQualityNameByIndex(CookedQualityLevelIndex)
+				);
+
 				// Call base serialize that will walk all properties and serialize them.
 				Super::Serialize(Ar);
 
@@ -190,7 +202,7 @@ void USoundNodeQualityLevel::ForCurrentQualityLevel(TFunction<void(USoundNode*)>
 #endif
 
 	// If CookedQualityLevelIndex has been set, we will have a *single* quality level.
-	if (CookedQualityLevelIndex >= 0)
+	if (CookedQualityLevelIndex >= 0 && ChildNodes.Num() == 1)
 	{
 		// Remap to index 0 (as all other levels have been removed by cooker).
 		QualityLevel = 0;
