Parameter Update Perf Regression Fix-ups
- Use cached transmissable inputs when initializing/sanitizing parameters
- Update internal API to promote batching & keeping AudioParameters sorted/using faster merge logic where possible
- Swap CitySample vehicle controller to use batched parameter update call
- Clarify When RuntimeData is cached for safety (in InitResources or InitParameters, whichever comes first)
- Alleviate extreme confusion around AudioParameter setters requiring type by adding static conversion functions
#rnx

diff --git a/Engine/Plugins/Runtime/Metasound/Source/MetasoundEngine/Private/MetasoundSource.cpp b/Engine/Plugins/Runtime/Metasound/Source/MetasoundEngine/Private/MetasoundSource.cpp
--- a/Engine/Plugins/Runtime/Metasound/Source/MetasoundEngine/Private/MetasoundSource.cpp
+++ b/Engine/Plugins/Runtime/Metasound/Source/MetasoundEngine/Private/MetasoundSource.cpp
@@ -273,26 +273,27 @@ void UMetaSoundSource::SetRegistryAssetClassInfo(const Metasound::Frontend::FNod
 
 void UMetaSoundSource::InitParameters(TArray<FAudioParameter>& InParametersToInit, FName InFeatureName)
 {
-	TMap<FName, const FMetasoundFrontendClassInput*> TransmittableInputMap;
-	const TArray<const FMetasoundFrontendClassInput*> TransmittableInputs = GetTransmittableClassInputs();
-	Algo::Transform(TransmittableInputs, TransmittableInputMap, [](const FMetasoundFrontendClassInput* Input)
-	{
-		return TPair<FName, const FMetasoundFrontendClassInput*>(Input->Name, Input);
-	});
-	
+	// Have to call cache vs a simple get as the source may have yet to start playing/has not been registered
+	// via InitResources. If it has, this call is fast and returns the already cached RuntimeData.
+	const FRuntimeData& RuntimeData = CacheRuntimeData();
+	const TArray<FMetasoundFrontendClassInput>& TransmittableInputs = RuntimeData.TransmittableInputs;
+
 	// Removes values that are not explicitly defined by the ParamType and returns
 	// whether or not the parameter is a valid input and should be included.
-	auto Sanitize = [&TransmittableInputMap](FAudioParameter& Parameter) -> bool
+	auto Sanitize = [&TransmittableInputs](FAudioParameter& Parameter) -> bool
 	{
-		if (const FMetasoundFrontendClassInput* Input = TransmittableInputMap.FindRef(Parameter.ParamName))
+		auto FindInput = [&Parameter](const FMetasoundFrontendClassInput& Input)
 		{
-			// Only remove if parameter typename is specified and does not match
-			if (Parameter.TypeName != FName() && Parameter.TypeName != Input->TypeName)
+			if (Parameter.ParamName == Input.Name)
 			{
-				return false;
+				return Parameter.TypeName.IsNone() || Parameter.TypeName == Input.TypeName;
 			}
-		}
-		else
+
+			return false;
+		};
+
+		const FMetasoundFrontendClassInput* Input = TransmittableInputs.FindByPredicate(FindInput);
+		if (!Input)
 		{
 			return false;
 		}
@@ -426,12 +427,11 @@ void UMetaSoundSource::InitParameters(TArray<FAudioParameter>& InParametersToIni
 
 	for (int32 i = InParametersToInit.Num() - 1; i >= 0; --i)
 	{
-		FAudioParameter& Parameter = InParametersToInit[i];		
+		FAudioParameter& Parameter = InParametersToInit[i];
 		
 #if !NO_LOGGING
 		// For logging in case of failure
-		FString AssetName;
-		GetName(AssetName); 
+		const FString AssetName = GetName();
 #endif // !NO_LOGGING
 		
 		if (Sanitize(Parameter))
@@ -440,19 +440,21 @@ void UMetaSoundSource::InitParameters(TArray<FAudioParameter>& InParametersToIni
 			{
 				ConstructProxies(Parameter);
 			}
-#if !NO_LOGGING
 			else
 			{
+#if !NO_LOGGING
 				UE_LOG(LogMetaSound, Error, TEXT("Failed to set invalid parameter '%s' in asset '%s': Either does not exist or is unsupported type"), *Parameter.ParamName.ToString(), *AssetName);
-			}
 #endif // !NO_LOGGING
+				constexpr bool bAllowShrinking = false;
+				InParametersToInit.RemoveAtSwap(i, 1, bAllowShrinking);
+			}
 		}
 		else
 		{
 			constexpr bool bAllowShrinking = false;
 			InParametersToInit.RemoveAtSwap(i, 1, bAllowShrinking);
 #if !NO_LOGGING
-			UE_LOG(LogMetaSound, Error, TEXT("Failed to set parameter '%s' in asset '%s': No name specified, no matching input found, or type mismatch."), *Parameter.ParamName.ToString(), *AssetName);
+			UE_LOG(LogMetaSound, Error, TEXT("Failed to set parameter '%s' in asset '%s': No name specified, no transmittable input found, or type mismatch."), *Parameter.ParamName.ToString(), *AssetName);
 #endif // !NO_LOGGING
 		}
 	}
@@ -469,7 +471,9 @@ void UMetaSoundSource::InitResources()
 	{
 		RegOptions.bAutoUpdateLogWarningOnDroppedConnection = Settings->bAutoUpdateLogWarningOnDroppedConnection;
 	}
+
 	RegisterGraphWithFrontend(RegOptions);
+	CacheRuntimeData();
 }
 
 Metasound::Frontend::FNodeClassInfo UMetaSoundSource::GetAssetClassInfo() const
@@ -512,7 +516,7 @@ ISoundGeneratorPtr UMetaSoundSource::CreateSoundGenerator(const FSoundGeneratorI
 	FOperatorSettings InSettings = GetOperatorSettings(static_cast<FSampleRate>(SampleRate));
 	FMetasoundEnvironment Environment = CreateEnvironment(InParams);
 
-	TSharedPtr<const IGraph, ESPMode::ThreadSafe> MetasoundGraph = GetMetasoundCoreGraph();
+	TSharedPtr<const IGraph, ESPMode::ThreadSafe> MetasoundGraph = GetRuntimeData().Graph;
 	if (!MetasoundGraph.IsValid())
 	{
 		return ISoundGeneratorPtr(nullptr);
@@ -536,70 +540,69 @@ bool UMetaSoundSource::GetAllDefaultParameters(TArray<FAudioParameter>& OutParam
 	using namespace Metasound::Frontend;
 	using namespace Metasound::Engine;
 
-	TArray<const FMetasoundFrontendClassInput*> TransmittableInputs = GetTransmittableClassInputs();
-	for(const FMetasoundFrontendClassInput* Input : TransmittableInputs)
+	// TODO: Make this use the cached runtime data's input copy. This call can become expensive if called repeatedly.
+	TArray<FMetasoundFrontendClassInput> TransmittableInputs = GetTransmittableClassInputs();
+	for(const FMetasoundFrontendClassInput& Input : TransmittableInputs)
 	{
-		check(Input);
-
 		FAudioParameter Params;
-		Params.ParamName = Input->Name;
-		Params.TypeName = Input->TypeName;
+		Params.ParamName = Input.Name;
+		Params.TypeName = Input.TypeName;
 
-		switch (Input->DefaultLiteral.GetType())
+		switch (Input.DefaultLiteral.GetType())
 		{
 			case EMetasoundFrontendLiteralType::Boolean:
 			{
 				Params.ParamType = EAudioParameterType::Boolean;
-				ensure(Input->DefaultLiteral.TryGet(Params.BoolParam));
+				ensure(Input.DefaultLiteral.TryGet(Params.BoolParam));
 			}
 			break;
 
 			case EMetasoundFrontendLiteralType::BooleanArray:
 			{
 				Params.ParamType = EAudioParameterType::BooleanArray;
-				ensure(Input->DefaultLiteral.TryGet(Params.ArrayBoolParam));
+				ensure(Input.DefaultLiteral.TryGet(Params.ArrayBoolParam));
 			}
 			break;
 
 			case EMetasoundFrontendLiteralType::Integer:
 			{
 				Params.ParamType = EAudioParameterType::Integer;
-				ensure(Input->DefaultLiteral.TryGet(Params.IntParam));
+				ensure(Input.DefaultLiteral.TryGet(Params.IntParam));
 			}
 			break;
 
 			case EMetasoundFrontendLiteralType::IntegerArray:
 			{
 				Params.ParamType = EAudioParameterType::IntegerArray;
-				ensure(Input->DefaultLiteral.TryGet(Params.ArrayIntParam));
+				ensure(Input.DefaultLiteral.TryGet(Params.ArrayIntParam));
 			}
 			break;
 
 			case EMetasoundFrontendLiteralType::Float:
 			{
 				Params.ParamType = EAudioParameterType::Float;
-				ensure(Input->DefaultLiteral.TryGet(Params.FloatParam));
+				ensure(Input.DefaultLiteral.TryGet(Params.FloatParam));
 			}
 			break;
 
 			case EMetasoundFrontendLiteralType::FloatArray:
 			{
 				Params.ParamType = EAudioParameterType::FloatArray;
-				ensure(Input->DefaultLiteral.TryGet(Params.ArrayFloatParam));
+				ensure(Input.DefaultLiteral.TryGet(Params.ArrayFloatParam));
 			}
 			break;
 
 			case EMetasoundFrontendLiteralType::String:
 			{
 				Params.ParamType = EAudioParameterType::String;
-				ensure(Input->DefaultLiteral.TryGet(Params.StringParam));
+				ensure(Input.DefaultLiteral.TryGet(Params.StringParam));
 			}
 			break;
 
 			case EMetasoundFrontendLiteralType::StringArray:
 			{
 				Params.ParamType = EAudioParameterType::StringArray;
-				ensure(Input->DefaultLiteral.TryGet(Params.ArrayStringParam));
+				ensure(Input.DefaultLiteral.TryGet(Params.ArrayStringParam));
 			}
 			break;
 
@@ -607,7 +610,7 @@ bool UMetaSoundSource::GetAllDefaultParameters(TArray<FAudioParameter>& OutParam
 			{
 				Params.ParamType = EAudioParameterType::Object;
 				UObject* Object = nullptr;
-				ensure(Input->DefaultLiteral.TryGet(Object));
+				ensure(Input.DefaultLiteral.TryGet(Object));
 				Params.ObjectParam = Object;
 			}
 			break;
@@ -615,7 +618,7 @@ bool UMetaSoundSource::GetAllDefaultParameters(TArray<FAudioParameter>& OutParam
 			case EMetasoundFrontendLiteralType::UObjectArray:
 			{
 				Params.ParamType = EAudioParameterType::ObjectArray;
-				ensure(Input->DefaultLiteral.TryGet(Params.ArrayObjectParam));
+				ensure(Input.DefaultLiteral.TryGet(Params.ArrayObjectParam));
 			}
 			break;
 
@@ -646,6 +649,11 @@ bool UMetaSoundSource::IsParameterValid(const FAudioParameter& InParameter, cons
 	using namespace Metasound;
 	using namespace Metasound::Frontend;
 
+	if (InParameter.ParamName.IsNone())
+	{
+		return false;
+	}
+
 	const FName* TypeName = InInputNameTypeMap.Find(InParameter.ParamName);
 	if (!TypeName)
 	{
@@ -791,11 +799,7 @@ TUniquePtr<Audio::IParameterTransmitter> UMetaSoundSource::CreateParameterTransm
 	}
 
 	TUniquePtr<Audio::IParameterTransmitter> NewTransmitter = MakeUnique<Metasound::FMetaSoundParameterTransmitter>(InitParams);
-	for (FAudioParameter& AudioParam : InParams.DefaultParams)
-	{
-		NewTransmitter->SetParameter(MoveTemp(AudioParam));
-	}
-	InParams.DefaultParams.Reset();
+	NewTransmitter->SetParameters(MoveTemp(InParams.DefaultParams));
 
 	return NewTransmitter;
 }
diff --git a/Engine/Plugins/Runtime/Metasound/Source/MetasoundFrontend/Private/MetasoundAssetBase.cpp b/Engine/Plugins/Runtime/Metasound/Source/MetasoundFrontend/Private/MetasoundAssetBase.cpp
--- a/Engine/Plugins/Runtime/Metasound/Source/MetasoundFrontend/Private/MetasoundAssetBase.cpp
+++ b/Engine/Plugins/Runtime/Metasound/Source/MetasoundFrontend/Private/MetasoundAssetBase.cpp
@@ -539,7 +539,7 @@ void FMetasoundAssetBase::ResetSynchronizationState()
 }
 #endif // WITH_EDITOR
 
-TSharedPtr<Metasound::IGraph, ESPMode::ThreadSafe> FMetasoundAssetBase::BuildMetasoundDocument() const
+TSharedPtr<Metasound::IGraph, ESPMode::ThreadSafe> FMetasoundAssetBase::BuildMetasoundDocument(const TArray<FMetasoundFrontendClassInput>& InTransmittableInputs) const
 {
 	using namespace Metasound;
 	using namespace Metasound::Frontend;
@@ -550,10 +550,8 @@ TSharedPtr<Metasound::IGraph, ESPMode::ThreadSafe> FMetasoundAssetBase::BuildMet
 	TUniquePtr<FFrontendGraph> FrontendGraph = FFrontendGraphBuilder::CreateGraph(GetDocumentChecked(), GetOwningAssetName());
 	if (FrontendGraph.IsValid())
 	{
-		const TArray<const FMetasoundFrontendClassInput*> TransmittableInputs = GetTransmittableClassInputs();
-
 		TSet<FVertexName> TransmittableInputNames;
-		Algo::Transform(TransmittableInputs, TransmittableInputNames, [](const FMetasoundFrontendClassInput* Input) { return Input->Name; });
+		Algo::Transform(InTransmittableInputs, TransmittableInputNames, [](const FMetasoundFrontendClassInput& Input) { return Input.Name; });
 
 		bool bSuccessfullyInjectedReceiveNodes = InjectReceiveNodes(*FrontendGraph, FMetaSoundParameterTransmitter::CreateSendAddressFromEnvironment, TransmittableInputNames);
 		if (!bSuccessfullyInjectedReceiveNodes)
@@ -668,6 +666,7 @@ TArray<FMetasoundAssetBase::FSendInfoAndVertexName> FMetasoundAssetBase::GetSend
 	check(IsInGameThread() || IsInAudioThread());
 
 	const FRuntimeData& RuntimeData = GetRuntimeData();
+	checkf(CurrentCachedRuntimeDataChangeID == RuntimeData.ChangeID, TEXT("Asset must have up-to-date cached RuntimeData prior to calling GetSendInfos"));
 
 	TArray<FSendInfoAndVertexName> SendInfos;
 
@@ -812,18 +811,15 @@ FString FMetasoundAssetBase::GetOwningAssetName() const
 	return FString();
 }
 
-TSharedPtr<const Metasound::IGraph, ESPMode::ThreadSafe> FMetasoundAssetBase::GetMetasoundCoreGraph() const
-{
-	return FMetasoundAssetBase::GetRuntimeData().Graph;
-}
-
-TArray<const FMetasoundFrontendClassInput*> FMetasoundAssetBase::GetTransmittableClassInputs() const
+TArray<FMetasoundFrontendClassInput> FMetasoundAssetBase::GetTransmittableClassInputs() const
 {
 	using namespace Metasound;
 	using namespace Metasound::Frontend;
 
+	METASOUND_TRACE_CPUPROFILER_EVENT_SCOPE(FMetasoundAssetBase::GetTransmittableClassInputs);
+
 	check(IsInGameThread() || IsInAudioThread());
-	TArray<const FMetasoundFrontendClassInput*> Inputs;
+	TArray<FMetasoundFrontendClassInput> Inputs;
 
 	const FMetasoundFrontendDocument& Doc = GetDocumentChecked();
 	auto GetInputName = [](const FMetasoundFrontendClassInput& InInput) { return InInput.Name; };
@@ -860,7 +856,7 @@ TArray<const FMetasoundFrontendClassInput*> FMetasoundAssetBase::GetTransmittabl
 
 		return false;
 	};
-	Algo::TransformIf(Doc.RootGraph.Interface.Inputs, Inputs, IsTransmittable, [] (const FMetasoundFrontendClassInput& Input) { return &Input; });
+	Algo::TransformIf(Doc.RootGraph.Interface.Inputs, Inputs, IsTransmittable, [] (const FMetasoundFrontendClassInput& Input) { return Input; });
 
 	return Inputs;
 }
@@ -873,9 +869,10 @@ void FMetasoundAssetBase::RebuildReferencedAssetClassKeys()
 	SetReferencedAssetClassKeys(MoveTemp(ReferencedKeys));
 }
 
-const FMetasoundAssetBase::FRuntimeData& FMetasoundAssetBase::GetRuntimeData() const
+const FMetasoundAssetBase::FRuntimeData& FMetasoundAssetBase::CacheRuntimeData()
 {
-	
+	METASOUND_TRACE_CPUPROFILER_EVENT_SCOPE(MetaSoundAssetBase::CacheRuntimeData);
+
 	// Check if a ChangeID has been generated before.
 	if (!CurrentCachedRuntimeDataChangeID.IsValid())
 	{
@@ -886,15 +883,25 @@ const FMetasoundAssetBase::FRuntimeData& FMetasoundAssetBase::GetRuntimeData() c
 	if (CachedRuntimeData.ChangeID != CurrentCachedRuntimeDataChangeID)
 	{
 		// Update CachedRuntimeData.
-		CachedRuntimeData.TransmittableInputs.Reset();
-		TArray<const FMetasoundFrontendClassInput*> ClassInputs = GetTransmittableClassInputs();
-		Algo::Transform(ClassInputs, CachedRuntimeData.TransmittableInputs, [] (const FMetasoundFrontendClassInput* Input) { return *Input; });
+		TArray<FMetasoundFrontendClassInput> ClassInputs = GetTransmittableClassInputs();
+		TSharedPtr<Metasound::IGraph, ESPMode::ThreadSafe> Graph = BuildMetasoundDocument(ClassInputs);
 
-		CachedRuntimeData.Graph = BuildMetasoundDocument();
-		CachedRuntimeData.ChangeID = CurrentCachedRuntimeDataChangeID;
+		CachedRuntimeData =
+		{
+			CurrentCachedRuntimeDataChangeID,
+			MoveTemp(ClassInputs),
+			MoveTemp(Graph)
+		};
 	}
 
 	return CachedRuntimeData;
 }
 
+const FMetasoundAssetBase::FRuntimeData& FMetasoundAssetBase::GetRuntimeData() const
+{
+	ensureMsgf(CurrentCachedRuntimeDataChangeID == CachedRuntimeData.ChangeID, TEXT("Accessing out-of-date runtime data: MetaSound asset '%s'."), *GetOwningAssetName());
+
+	return CachedRuntimeData;
+}
+
 #undef LOCTEXT_NAMESPACE // "MetaSound"
diff --git a/Engine/Plugins/Runtime/Metasound/Source/MetasoundFrontend/Private/MetasoundParameterTransmitter.cpp b/Engine/Plugins/Runtime/Metasound/Source/MetasoundFrontend/Private/MetasoundParameterTransmitter.cpp
--- a/Engine/Plugins/Runtime/Metasound/Source/MetasoundFrontend/Private/MetasoundParameterTransmitter.cpp
+++ b/Engine/Plugins/Runtime/Metasound/Source/MetasoundFrontend/Private/MetasoundParameterTransmitter.cpp
@@ -177,10 +177,18 @@ namespace Metasound
 		return InstanceID;
 	}
 
-	bool FMetaSoundParameterTransmitter::SetParameter(FAudioParameter&& InParameter)
+	bool FMetaSoundParameterTransmitter::SetParameters(TArray<FAudioParameter>&& InParameters)
 	{
-		const FName ParamName = InParameter.ParamName;
-		return SetParameterWithLiteral(ParamName, Frontend::ConvertParameterToLiteral(MoveTemp(InParameter)));
+		bool bSuccess = true;
+
+		for (FAudioParameter& InParameter : InParameters)
+		{
+			const FName ParamName = InParameter.ParamName;
+			bSuccess &= SetParameterWithLiteral(ParamName, Frontend::ConvertParameterToLiteral(MoveTemp(InParameter)));
+		}
+
+		InParameters.Reset();
+		return bSuccess;
 	}
 
 	bool FMetaSoundParameterTransmitter::SetParameterWithLiteral(FName InParameterName, const FLiteral& InLiteral)
diff --git a/Engine/Plugins/Runtime/Metasound/Source/MetasoundFrontend/Public/MetasoundAssetBase.h b/Engine/Plugins/Runtime/Metasound/Source/MetasoundFrontend/Public/MetasoundAssetBase.h
--- a/Engine/Plugins/Runtime/Metasound/Source/MetasoundFrontend/Public/MetasoundAssetBase.h
+++ b/Engine/Plugins/Runtime/Metasound/Source/MetasoundFrontend/Public/MetasoundAssetBase.h
@@ -178,20 +178,12 @@ protected:
 	// Returns an access pointer to the document.
 	virtual Metasound::Frontend::FConstDocumentAccessPtr GetDocument() const = 0;
 
-	// Returns a shared instance of the core metasound graph.
-	TSharedPtr<const Metasound::IGraph, ESPMode::ThreadSafe> GetMetasoundCoreGraph() const;
-
 #if WITH_EDITORONLY_DATA
 	bool bSynchronizationRequired = true;
 	bool bSynchronizationUpdateDetails = false;
 #endif // WITH_EDITORONLY_DATA
 
 protected:
-	TArray<const FMetasoundFrontendClassInput*> GetTransmittableClassInputs() const;
-
-private:
-	Metasound::Frontend::FNodeRegistryKey RegistryKey;
-
 	// Container for runtime data of MetaSound graph.
 	struct FRuntimeData
 	{
@@ -204,15 +196,25 @@ private:
 		// Core graph.
 		TSharedPtr<Metasound::IGraph, ESPMode::ThreadSafe> Graph;
 	};
-	// Cache ID is used to determine whether CachedRuntimeData is out-of-date.
-	mutable FGuid CurrentCachedRuntimeDataChangeID;
-	mutable FRuntimeData CachedRuntimeData;
 
-	// Returns the current runtime data. If the cached data is out-of-date, it will
-	// be updated in this call.
+	// Returns the cached runtime data. Call updates cached data if out-of-date.
+	const FRuntimeData& CacheRuntimeData();
+
+	// Returns the cached runtime data.
 	const FRuntimeData& GetRuntimeData() const;
 
-	TSharedPtr<Metasound::IGraph, ESPMode::ThreadSafe> BuildMetasoundDocument() const;
+	// Returns all transmissible class inputs.  This is a potentially expensive.
+	// Prefer accessing transmissible class inputs using CacheRuntimeData.
+	TArray<FMetasoundFrontendClassInput> GetTransmittableClassInputs() const;
+
+private:
+	Metasound::Frontend::FNodeRegistryKey RegistryKey;
+
+	// Cache ID is used to determine whether CachedRuntimeData is out-of-date.
+	FGuid CurrentCachedRuntimeDataChangeID;
+	FRuntimeData CachedRuntimeData;
+
+	TSharedPtr<Metasound::IGraph, ESPMode::ThreadSafe> BuildMetasoundDocument(const TArray<FMetasoundFrontendClassInput>& InTransmittableInputs) const;
 	Metasound::FSendAddress CreateSendAddress(uint64 InInstanceID, const Metasound::FVertexName& InVertexName, const FName& InDataTypeName) const;
 	Metasound::Frontend::FNodeHandle AddInputPinForSendAddress(const Metasound::FMetaSoundParameterTransmitter::FSendInfo& InSendInfo, Metasound::Frontend::FGraphHandle InGraph) const;
 };
diff --git a/Engine/Plugins/Runtime/Metasound/Source/MetasoundFrontend/Public/MetasoundParameterTransmitter.h b/Engine/Plugins/Runtime/Metasound/Source/MetasoundFrontend/Public/MetasoundParameterTransmitter.h
--- a/Engine/Plugins/Runtime/Metasound/Source/MetasoundFrontend/Public/MetasoundParameterTransmitter.h
+++ b/Engine/Plugins/Runtime/Metasound/Source/MetasoundFrontend/Public/MetasoundParameterTransmitter.h
@@ -85,11 +85,11 @@ namespace Metasound
 		/** Returns ID of the MetaSound instance associated with this transmitter. */
 		uint64 GetInstanceID() const override;
 
-		/** Sets a parameter using an AudioParameter struct
+		/** Sets parameters using array of AudioParameter structs
 		 *
 		 * @param InParameter - Parameter to set.
 		 */
-		bool SetParameter(FAudioParameter&& InParameter) override;
+		virtual bool SetParameters(TArray<FAudioParameter>&& InParameters) override;
 
 		/** Set a parameter using a literal.
 		 *
diff --git a/Engine/Source/Runtime/AudioExtensions/Public/AudioParameterControllerInterface.h b/Engine/Source/Runtime/AudioExtensions/Public/AudioParameterControllerInterface.h
--- a/Engine/Source/Runtime/AudioExtensions/Public/AudioParameterControllerInterface.h
+++ b/Engine/Source/Runtime/AudioExtensions/Public/AudioParameterControllerInterface.h
@@ -75,6 +75,10 @@ public:
 	UFUNCTION(BlueprintCallable, meta = (DisplayName = "Set Object Array Parameter"), Category = "Audio|Parameter")
 	virtual void SetObjectArrayParameter(FName InName, const TArray<UObject*>& InValue) = 0;
 
+	// Sets an array of parameters as a batch
+	UFUNCTION(BlueprintCallable, meta = (DisplayName = "Set Parameters"), Category = "Audio|Parameter")
+	virtual void SetParameters_Blueprint(const TArray<FAudioParameter>& InParameters) = 0;
+
 	// Sets a named parameter to the given parameter structure value
 	virtual void SetParameter(FAudioParameter&& InValue) = 0;
 
@@ -83,14 +87,14 @@ public:
 
 	// Template Specialization for non-script clients.
 	template<typename DataType> void SetParameter(FName InName, DataType&&) = delete;
-	template<> void SetParameter(FName InName, bool&& InBool) { SetBoolParameter(InName, InBool); }
-	template<> void SetParameter(FName InName, float&& InFloat) { SetFloatParameter(InName, InFloat); }
-	template<> void SetParameter(FName InName, int32&& InInteger) { SetIntParameter(InName, InInteger); }
-	template<> void SetParameter(FName InName, FString&& InString) { SetStringParameter(InName, InString); }
-	template<> void SetParameter(FName InName, UObject*&& InObject) { SetObjectParameter(InName, InObject); }
-	template<> void SetParameter(FName InName, TArray<bool>&& InBools) { SetBoolArrayParameter(InName, InBools); }
-	template<> void SetParameter(FName InName, TArray<float>&& InFloats) { SetFloatArrayParameter(InName, InFloats); }
-	template<> void SetParameter(FName InName, TArray<int32>&& InIntegers) { SetIntArrayParameter(InName, InIntegers); }
-	template<> void SetParameter(FName InName, TArray<FString>&& InStrings) { SetStringArrayParameter(InName, InStrings); }
-	template<> void SetParameter(FName InName, TArray<UObject*>&& InObjects) { SetObjectArrayParameter(InName, InObjects); }
+	template<> void SetParameter(FName InName, bool&& InBool) { SetParameter({ InName, MoveTemp(InBool) }); }
+	template<> void SetParameter(FName InName, float&& InFloat) { SetParameter({ InName, MoveTemp(InFloat) }); }
+	template<> void SetParameter(FName InName, int32&& InInteger) { SetParameter({ InName, MoveTemp(InInteger) }); }
+	template<> void SetParameter(FName InName, FString&& InString) { SetParameter({ InName, MoveTemp(InString) }); }
+	template<> void SetParameter(FName InName, UObject*&& InObject) { SetParameter({ InName, MoveTemp(InObject) }); }
+	template<> void SetParameter(FName InName, TArray<bool>&& InBools) { SetParameter({ InName, MoveTemp(InBools)}); }
+	template<> void SetParameter(FName InName, TArray<float>&& InFloats) { SetParameter({ InName, MoveTemp(InFloats) }); }
+	template<> void SetParameter(FName InName, TArray<int32>&& InIntegers) { SetParameter({ InName, MoveTemp(InIntegers) }); }
+	template<> void SetParameter(FName InName, TArray<FString>&& InStrings) { SetParameter({ InName, MoveTemp(InStrings) }); }
+	template<> void SetParameter(FName InName, TArray<UObject*>&& InObjects) { SetParameter({ InName, MoveTemp(InObjects) }); }
 };
diff --git a/Engine/Source/Runtime/AudioExtensions/Public/IAudioParameterInterfaceRegistry.h b/Engine/Source/Runtime/AudioExtensions/Public/IAudioParameterInterfaceRegistry.h
--- a/Engine/Source/Runtime/AudioExtensions/Public/IAudioParameterInterfaceRegistry.h
+++ b/Engine/Source/Runtime/AudioExtensions/Public/IAudioParameterInterfaceRegistry.h
@@ -110,7 +110,7 @@ namespace Audio
 		TArray<FOutput> Outputs;
 		TArray<FEnvironmentVariable> Environment;
 	};
-	using FParameterInterfacePtr = TSharedPtr<FParameterInterface>;
+	using FParameterInterfacePtr = TSharedPtr<FParameterInterface, ESPMode::ThreadSafe>;
 
 	class AUDIOEXTENSIONS_API IAudioParameterInterfaceRegistry
 	{
diff --git a/Engine/Source/Runtime/AudioExtensions/Public/IAudioParameterTransmitter.h b/Engine/Source/Runtime/AudioExtensions/Public/IAudioParameterTransmitter.h
--- a/Engine/Source/Runtime/AudioExtensions/Public/IAudioParameterTransmitter.h
+++ b/Engine/Source/Runtime/AudioExtensions/Public/IAudioParameterTransmitter.h
@@ -54,7 +54,7 @@ namespace Audio
 			virtual uint64 GetInstanceID() const = 0;
 
 			// Parameter Setters
-			virtual bool SetParameter(FAudioParameter&& InValue) = 0;
+			virtual bool SetParameters(TArray<FAudioParameter>&& InParameters) = 0;
 
 			// Create a copy of the instance transmitter.
 			virtual TUniquePtr<Audio::IParameterTransmitter> Clone() const = 0;
diff --git a/Engine/Source/Runtime/Engine/Private/ActiveSound.cpp b/Engine/Source/Runtime/Engine/Private/ActiveSound.cpp
--- a/Engine/Source/Runtime/Engine/Private/ActiveSound.cpp
+++ b/Engine/Source/Runtime/Engine/Private/ActiveSound.cpp
@@ -607,12 +607,14 @@ void FActiveSound::UpdateInterfaceParameters(const TArray<FListener>& InListener
 		return;
 	}
 
+	TArray<FAudioParameter> ParamsToUpdate;
+
 	const FListener& Listener = InListeners[ClosestListenerIndex];
 	const FVector SourceDirection = Transform.GetLocation() - Listener.Transform.GetLocation();
 	if (bImplementsAttenuation)
 	{
 		const float Distance = SourceDirection.Size();
-		InstanceTransmitter->SetParameter({ AttenuationInterface::Inputs::Distance, Distance });
+		ParamsToUpdate.Add({ AttenuationInterface::Inputs::Distance, Distance });
 	}
 
 	if (bImplementsSpatialization)
@@ -621,9 +623,14 @@ void FActiveSound::UpdateInterfaceParameters(const TArray<FListener>& InListener
 		const FVector2D SourceAzimuthAndElevation = FMath::GetAzimuthAndElevation(SourceDirectionNormal, FVector::ForwardVector, FVector::RightVector, FVector::UpVector);
 		const float Azimuth = FMath::RadiansToDegrees(SourceAzimuthAndElevation.X);
 		const float Elevation = FMath::RadiansToDegrees(SourceAzimuthAndElevation.Y);
-		InstanceTransmitter->SetParameter({ SpatializationInterface::Inputs::Azimuth, Azimuth });
-		InstanceTransmitter->SetParameter({ SpatializationInterface::Inputs::Elevation, Elevation });
+		ParamsToUpdate.Append(
+		{
+			{ SpatializationInterface::Inputs::Azimuth, Azimuth },
+			{ SpatializationInterface::Inputs::Elevation, Elevation }
+		});
 	}
+
+	InstanceTransmitter->SetParameters(MoveTemp(ParamsToUpdate));
 }
 
 void FActiveSound::UpdateWaveInstances(TArray<FWaveInstance*> &InWaveInstances, const float DeltaTime)
diff --git a/Engine/Source/Runtime/Engine/Private/Audio/SoundParameterControllerInterface.cpp b/Engine/Source/Runtime/Engine/Private/Audio/SoundParameterControllerInterface.cpp
--- a/Engine/Source/Runtime/Engine/Private/Audio/SoundParameterControllerInterface.cpp
+++ b/Engine/Source/Runtime/Engine/Private/Audio/SoundParameterControllerInterface.cpp
@@ -6,6 +6,7 @@
 #include "AudioDevice.h"
 #include "AudioThread.h"
 #include "IAudioExtensionPlugin.h"
+#include "ProfilingDebugging/CpuProfilerTrace.h"
 
 
 namespace SoundParameterControllerInterfacePrivate
@@ -70,8 +71,8 @@ void ISoundParameterControllerInterface::SetTriggerParameter(FName InName)
 					const FName ParamName = Param.ParamName;
 
 					// Must be copied as original version must be preserved in case command is called on multiple ActiveSounds.
-					FAudioParameter TempParam = Param;
-					if (!Transmitter->SetParameter(MoveTemp(TempParam)))
+					TArray<FAudioParameter> TempParam = { Param };
+					if (!Transmitter->SetParameters(MoveTemp(TempParam)))
 					{
 						UE_LOG(LogAudio, Warning, TEXT("Failed to execute trigger parameter '%s'"), *ParamName.ToString());
 					}
@@ -83,102 +84,94 @@ void ISoundParameterControllerInterface::SetTriggerParameter(FName InName)
 
 void ISoundParameterControllerInterface::SetBoolParameter(FName InName, bool InValue)
 {
-	SetParameterInternal(FAudioParameter(InName, InValue));
+	SetParameters({ FAudioParameter(InName, InValue) });
 }
 
 void ISoundParameterControllerInterface::SetBoolArrayParameter(FName InName, const TArray<bool>& InValue)
 {
-	SetParameterInternal(FAudioParameter(InName, InValue));
+	SetParameters( { FAudioParameter(InName, InValue) });
 }
 
 void ISoundParameterControllerInterface::SetIntParameter(FName InName, int32 InValue)
 {
-	SetParameterInternal(FAudioParameter(InName, InValue));
+	SetParameters( { FAudioParameter(InName, InValue) });
 }
 
 void ISoundParameterControllerInterface::SetIntArrayParameter(FName InName, const TArray<int32>& InValue)
 {
-	SetParameterInternal(FAudioParameter(InName, InValue));
+	SetParameters( { FAudioParameter(InName, InValue) });
 }
 
 void ISoundParameterControllerInterface::SetFloatParameter(FName InName, float InValue)
 {
-	SetParameterInternal(FAudioParameter(InName, InValue));
+	SetParameters( { FAudioParameter(InName, InValue) });
 }
 
 void ISoundParameterControllerInterface::SetFloatArrayParameter(FName InName, const TArray<float>& InValue)
 {
-	SetParameterInternal(FAudioParameter(InName, InValue));
+	SetParameters( { FAudioParameter(InName, InValue) });
 }
 
 void ISoundParameterControllerInterface::SetStringParameter(FName InName, const FString& InValue)
 {
-	SetParameterInternal(FAudioParameter(InName, InValue));
+	SetParameters( { FAudioParameter(InName, InValue) });
 }
 
 void ISoundParameterControllerInterface::SetStringArrayParameter(FName InName, const TArray<FString>& InValue)
 {
-	SetParameterInternal(FAudioParameter(InName, InValue));
+	SetParameters( { FAudioParameter(InName, InValue) });
 }
 
 void ISoundParameterControllerInterface::SetObjectParameter(FName InName, UObject* InValue)
 {
-	SetParameterInternal(FAudioParameter(InName, InValue));
+	SetParameters( { FAudioParameter(InName, InValue) });
 }
 
 void ISoundParameterControllerInterface::SetObjectArrayParameter(FName InName, const TArray<UObject*>& InValue)
 {
-	SetParameterInternal(FAudioParameter(InName, InValue));
+	SetParameters( { FAudioParameter(InName, InValue) });
 }
 
 void ISoundParameterControllerInterface::SetParameter(FAudioParameter&& InValue)
 {
-	SetParameterInternal(MoveTemp(InValue));
+	SetParameters({ MoveTemp(InValue) });
 }
 
 void ISoundParameterControllerInterface::SetParameters(TArray<FAudioParameter>&& InValues)
 {
-	for (const FAudioParameter& Value : InValues)
+	TRACE_CPUPROFILER_EVENT_SCOPE(ISoundParameterControllerInterface::SetParameters);
+
+	const bool bUpdateActiveSound = IsPlaying() && !GetDisableParameterUpdatesWhilePlaying();
+
+	TArray<FAudioParameter> ParamsToSet;
+	if (bUpdateActiveSound)
 	{
-		TArray<FAudioParameter>& InstanceParameters = GetInstanceParameters();
-		if (FAudioParameter* CurrentParam = FAudioParameter::FindOrAddParam(InstanceParameters, Value.ParamName))
-		{
-			CurrentParam->Merge(Value, false /* bInTakeName */);
-		}
+		ParamsToSet = InValues;
 	}
 
-	if (IsPlaying() && !GetDisableParameterUpdatesWhilePlaying())
+	TArray<FAudioParameter>& InstanceParameters = GetInstanceParameters();
+	FAudioParameter::Merge(MoveTemp(InValues), InstanceParameters);
+
+	if (bUpdateActiveSound)
 	{
 		if (FAudioDevice* AudioDevice = GetAudioDevice())
 		{
-			TArray<FAudioParameter> ParamsToSet;
 			if (USoundBase* Sound = GetSound())
 			{
-				Sound->InitParameters(InValues, SoundParameterControllerInterfacePrivate::ProxyFeatureName);
+				Sound->InitParameters(ParamsToSet, SoundParameterControllerInterfacePrivate::ProxyFeatureName);
 			}
 
-			ParamsToSet = MoveTemp(InValues);
-
-			if (ParamsToSet.Num() > 0)
+			// Prior call to InitParameters can prune parameters if they are
+			// invalid, so check here to avoid unnecessary pass of empty array.
+			if (!ParamsToSet.IsEmpty())
 			{
 				DECLARE_CYCLE_STAT(TEXT("FAudioThreadTask.SoundParameterControllerInterface.SetParameters"), STAT_AudioSetParameters, STATGROUP_AudioThreadCommands);
 				AudioDevice->SendCommandToActiveSounds(GetInstanceOwnerID(), [AudioDevice, Params = MoveTemp(ParamsToSet)](FActiveSound& ActiveSound)
 				{
 					if (Audio::IParameterTransmitter* Transmitter = ActiveSound.GetTransmitter())
 					{
-						for (const FAudioParameter& Param : Params)
-						{
-							const FName ParamName = Param.ParamName;
-							if (!ParamName.IsNone())
-							{
-								// Must be copied as original version must be preserved in case command is called on multiple ActiveSounds.
-								FAudioParameter TempParam = Param;
-								if (!Transmitter->SetParameter(MoveTemp(TempParam)))
-								{
-									UE_LOG(LogAudio, Warning, TEXT("Failed to set parameter '%s'"), *ParamName.ToString());
-								}
-							}
-						}
+						TArray<FAudioParameter> TempParams = Params;
+						Transmitter->SetParameters(MoveTemp(TempParams));
 					}
 				}, GET_STATID(STAT_AudioSetParameters));
 			}
@@ -186,57 +179,28 @@ void ISoundParameterControllerInterface::SetParameters(TArray<FAudioParameter>&&
 	}
 }
 
-void ISoundParameterControllerInterface::SetParameterInternal(FAudioParameter&& InParam)
+void ISoundParameterControllerInterface::SetParameters_Blueprint(const TArray<FAudioParameter>& InValues)
 {
-	if (InParam.ParamName.IsNone())
-	{
-		return;
-	}
+	TArray<FAudioParameter> Values = InValues;
+	SetParameters(MoveTemp(Values));
+}
 
-	TArray<FAudioParameter>& InstanceParameters = GetInstanceParameters();
-	if (FAudioParameter* CurrentParam = FAudioParameter::FindOrAddParam(InstanceParameters, InParam.ParamName))
-	{
-		CurrentParam->Merge(InParam, false /* bInTakeName */);
-	}
+FAudioParameter UAudioParameterConversionStatics::BooleanToAudioParameter(FName Name, bool Bool) { return { Name, Bool }; }
 
-	if (IsPlaying() && !GetDisableParameterUpdatesWhilePlaying())
-	{
-		if (FAudioDevice* AudioDevice = GetAudioDevice())
-		{
-			FAudioParameter ParamToSet;
-			if (USoundBase* Sound = GetSound())
-			{
-				TArray<FAudioParameter> Params = { MoveTemp(InParam) };
-				Sound->InitParameters(Params, SoundParameterControllerInterfacePrivate::ProxyFeatureName);
-				if (Params.Num() == 0)
-				{
-					// USoundBase::InitParameters(...) can remove parameters. 
-					// Exit early if the parameter is removed.
-					return;
-				}
-				ParamToSet = MoveTemp(Params[0]);
-			}
-			else
-			{
-				ParamToSet = MoveTemp(InParam);
-			}
+FAudioParameter UAudioParameterConversionStatics::FloatToAudioParameter(FName Name, float Float) { return { Name, Float }; }
 
-			DECLARE_CYCLE_STAT(TEXT("FAudioThreadTask.SoundParameterControllerInterface.SetParameter"), STAT_AudioSetParameter, STATGROUP_AudioThreadCommands);
+FAudioParameter UAudioParameterConversionStatics::IntegerToAudioParameter(FName Name, int32 Integer) { return { Name, Integer }; }
 
-			AudioDevice->SendCommandToActiveSounds(GetInstanceOwnerID(), [AudioDevice, Param = MoveTemp(ParamToSet)](FActiveSound& ActiveSound)
-			{
-				if (Audio::IParameterTransmitter* Transmitter = ActiveSound.GetTransmitter())
-				{
-					const FName ParamName = Param.ParamName;
+FAudioParameter UAudioParameterConversionStatics::StringToAudioParameter(FName Name, FString String) { return { Name, MoveTemp(String) }; }
 
-					// Must be copied as original version must be preserved in case command is called on multiple ActiveSounds.
-					FAudioParameter TempParam = Param;
-					if (!Transmitter->SetParameter(MoveTemp(TempParam)))
-					{
-						UE_LOG(LogAudio, Warning, TEXT("Failed to set parameter '%s'"), *ParamName.ToString());
-					}
-				}
-			}, GET_STATID(STAT_AudioSetParameter));
-		}
-	}
-}
+FAudioParameter UAudioParameterConversionStatics::ObjectToAudioParameter(FName Name, UObject* Object) { return { Name, Object }; }
+
+FAudioParameter UAudioParameterConversionStatics::BooleanArrayToAudioParameter(FName Name, TArray<bool> Bools) { return { Name, MoveTemp(Bools) }; }
+
+FAudioParameter UAudioParameterConversionStatics::FloatArrayToAudioParameter(FName Name, TArray<float> Floats) { return { Name, MoveTemp(Floats) }; }
+
+FAudioParameter UAudioParameterConversionStatics::IntegerArrayToAudioParameter(FName Name, TArray<int32> Integers) { return { Name, MoveTemp(Integers) }; }
+
+FAudioParameter UAudioParameterConversionStatics::StringArrayToAudioParameter(FName Name, TArray<FString> Strings) { return { Name, MoveTemp(Strings) }; }
+
+FAudioParameter UAudioParameterConversionStatics::ObjectArrayToAudioParameter(FName Name, TArray<UObject*> Objects) { return { Name, MoveTemp(Objects) }; }
\ No newline at end of file
diff --git a/Engine/Source/Runtime/Engine/Private/SoundBase.cpp b/Engine/Source/Runtime/Engine/Private/SoundBase.cpp
--- a/Engine/Source/Runtime/Engine/Private/SoundBase.cpp
+++ b/Engine/Source/Runtime/Engine/Private/SoundBase.cpp
@@ -308,96 +308,10 @@ TUniquePtr<Audio::IParameterTransmitter> USoundBase::CreateParameterTransmitter(
 			return true;
 		}
 
-		bool SetParameter(FAudioParameter&& InValue) override
+		bool SetParameters(TArray<FAudioParameter>&& InParameters) override
 		{
-			if (FAudioParameter* Param = FAudioParameter::FindOrAddParam(AudioParameters, InValue.ParamName))
-			{
-				// Cannot do wholesale move because legacy system (i.e. SoundCues) support
-				// multiple value types mapped to a single parameter FName
-				switch (InValue.ParamType)
-				{
-					case EAudioParameterType::Boolean:
-					{
-						Param->BoolParam = InValue.BoolParam;
-					}
-					break;
-
-					case EAudioParameterType::BooleanArray:
-					{
-						Param->ArrayBoolParam = MoveTemp(InValue.ArrayBoolParam);
-					}
-					break;
-
-					case EAudioParameterType::Float:
-					{
-						Param->FloatParam = InValue.FloatParam;
-					}
-					break;
-
-					case EAudioParameterType::FloatArray:
-					{
-						Param->ArrayFloatParam = MoveTemp(InValue.ArrayFloatParam);
-					}
-					break;
-
-					case EAudioParameterType::Integer:
-					case EAudioParameterType::NoneArray:
-					{
-						Param->IntParam = InValue.IntParam;
-					}
-					break;
-
-					case EAudioParameterType::IntegerArray:
-					{
-						Param->ArrayIntParam = MoveTemp(InValue.ArrayIntParam);
-					}
-					break;
-
-					case EAudioParameterType::None:
-					{
-						*Param = MoveTemp(InValue);
-					}
-					break;
-
-					case EAudioParameterType::Object:
-					{
-						Param->ObjectParam = MoveTemp(InValue.ObjectParam);
-						Param->ObjectProxies = MoveTemp(InValue.ObjectProxies);
-					}
-					break;
-
-					case EAudioParameterType::ObjectArray:
-					{
-						Param->ArrayObjectParam = MoveTemp(InValue.ArrayObjectParam);
-						Param->ObjectProxies = MoveTemp(InValue.ObjectProxies);
-					}
-					break;
-
-					case EAudioParameterType::String:
-					{
-						Param->StringParam = MoveTemp(InValue.StringParam);
-					}
-					break;
-
-					case EAudioParameterType::StringArray:
-					{
-						Param->ArrayStringParam = MoveTemp(InValue.ArrayStringParam);
-					}
-					break;
-
-					default:
-					{
-						checkNoEntry();
-					}
-				}
-
-				InValue = FAudioParameter();
-
-				return true;
-			}
-
-			InValue = FAudioParameter();
-			return false;
+			FAudioParameter::Merge(MoveTemp(InParameters), AudioParameters);
+			return true;
 		}
 
 		uint64 GetInstanceID() const override
@@ -460,6 +374,11 @@ void USoundBase::InitParameters(TArray<FAudioParameter>& InParametersToInit, FNa
 
 bool USoundBase::IsParameterValid(const FAudioParameter& InParameter) const
 {
+	if (InParameter.ParamName.IsNone())
+	{
+		return false;
+	}
+
 	switch (InParameter.ParamType)
 	{
 		case EAudioParameterType::BooleanArray:
diff --git a/Engine/Source/Runtime/Engine/Public/Audio/SoundParameterControllerInterface.h b/Engine/Source/Runtime/Engine/Public/Audio/SoundParameterControllerInterface.h
--- a/Engine/Source/Runtime/Engine/Public/Audio/SoundParameterControllerInterface.h
+++ b/Engine/Source/Runtime/Engine/Public/Audio/SoundParameterControllerInterface.h
@@ -3,6 +3,7 @@
 
 #include "AudioParameterControllerInterface.h"
 #include "UObject/Object.h"
+#include "Kismet/BlueprintFunctionLibrary.h"
 
 #include "SoundParameterControllerInterface.generated.h"
 
@@ -42,6 +43,7 @@ public:
 
 	void SetParameter(FAudioParameter&& InValue) override;
 	void SetParameters(TArray<FAudioParameter>&& InValues) override;
+	void SetParameters_Blueprint(const TArray<FAudioParameter>& InValues) override;
 
 	/** Returns the active audio device to use for this component based on whether or not the component is playing in a world. */
 	virtual FAudioDevice* GetAudioDevice() const = 0;
@@ -58,7 +60,41 @@ public:
 	virtual bool IsPlaying() const = 0;
 
 	virtual bool GetDisableParameterUpdatesWhilePlaying() const = 0;
+};
+
+UCLASS()
+class ENGINE_API UAudioParameterConversionStatics : public UBlueprintFunctionLibrary
+{
+	GENERATED_BODY()
+
+public:
+	UFUNCTION(BlueprintPure, Category = "Audio|Parameter", meta = (Keywords = "make construct convert create"))
+	static UPARAM(DisplayName = "Parameter") FAudioParameter BooleanToAudioParameter(FName Name, bool Bool);
+
+	UFUNCTION(BlueprintPure, Category = "Audio|Parameter", meta = (Keywords = "make construct convert create"))
+	static UPARAM(DisplayName = "Parameter") FAudioParameter FloatToAudioParameter(FName Name, float Float);
+
+	UFUNCTION(BlueprintPure, Category = "Audio|Parameter", meta = (Keywords = "make construct convert create"))
+	static UPARAM(DisplayName = "Parameter") FAudioParameter IntegerToAudioParameter(FName Name, int32 Integer);
+
+	UFUNCTION(BlueprintPure, Category = "Audio|Parameter", meta = (Keywords = "make construct convert create"))
+	static UPARAM(DisplayName = "Parameter") FAudioParameter StringToAudioParameter(FName Name, FString String);
+
+	UFUNCTION(BlueprintPure, Category = "Audio|Parameter", meta = (Keywords = "make construct convert create"))
+	static UPARAM(DisplayName = "Parameter") FAudioParameter ObjectToAudioParameter(FName Name, UObject* Object);
+
+	UFUNCTION(BlueprintPure, Category = "Audio|Parameter", meta = (Keywords = "make construct convert create"))
+	static UPARAM(DisplayName = "Parameter") FAudioParameter BooleanArrayToAudioParameter(FName Name, TArray<bool> Bools);
+
+	UFUNCTION(BlueprintPure, Category = "Audio|Parameter", meta = (Keywords = "make construct convert create"))
+	static UPARAM(DisplayName = "Parameter") FAudioParameter FloatArrayToAudioParameter(FName Name, TArray<float> Floats);
+
+	UFUNCTION(BlueprintPure, Category = "Audio|Parameter", meta = (Keywords = "make construct convert create"))
+	static UPARAM(DisplayName = "Parameter") FAudioParameter IntegerArrayToAudioParameter(FName Name, TArray<int32> Integers);
+
+	UFUNCTION(BlueprintPure, Category = "Audio|Parameter", meta = (Keywords = "make construct convert create"))
+	static UPARAM(DisplayName = "Parameter") FAudioParameter StringArrayToAudioParameter(FName Name, TArray<FString> Strings);
 
-private:
-	void SetParameterInternal(FAudioParameter&& InValue);
+	UFUNCTION(BlueprintPure, Category = "Audio|Parameter", meta = (Keywords = "make construct convert create"))
+	static UPARAM(DisplayName = "Parameter") FAudioParameter ObjectArrayToAudioParameter(FName Name, TArray<UObject*> Objects);
 };
