[Re-submitting] Added an "Overlay Material" to a mesh component.
Mesh with assigned overlay material will be rendered two times once with original material and once with an overlay material. Supports static mesh, instanced static mesh and skeletal mesh

diff --git a/Engine/Source/Runtime/Engine/Classes/Components/MeshComponent.h b/Engine/Source/Runtime/Engine/Classes/Components/MeshComponent.h
--- a/Engine/Source/Runtime/Engine/Classes/Components/MeshComponent.h
+++ b/Engine/Source/Runtime/Engine/Classes/Components/MeshComponent.h
@@ -42,6 +42,18 @@ public:
 	/** Returns override Materials count */
 	virtual int32 GetNumOverrideMaterials() const;
 
+	/** Translucent material to blend on top of this mesh. Mesh will be rendered twice - once with a base material and once with overlay material */
+	UPROPERTY(EditAnywhere, BlueprintReadOnly, AdvancedDisplay, Category=Rendering)
+	TObjectPtr<class UMaterialInterface> OverlayMaterial;
+
+	/** Get the overlay material used by this instance */
+	UFUNCTION(BlueprintCallable, Category="Rendering|Material")
+	class UMaterialInterface* GetOverlayMaterial() const;
+
+	/** Change the overlay material used by this instance */
+	UFUNCTION(BlueprintCallable, Category="Rendering|Material")
+	void SetOverlayMaterial(class UMaterialInterface* NewOverlayMaterial);
+
 #if WITH_EDITOR
 	/*
 	 * Make sure the Override array is using only the space it should use.
diff --git a/Engine/Source/Runtime/Engine/Private/Components/MeshComponent.cpp b/Engine/Source/Runtime/Engine/Private/Components/MeshComponent.cpp
--- a/Engine/Source/Runtime/Engine/Private/Components/MeshComponent.cpp
+++ b/Engine/Source/Runtime/Engine/Private/Components/MeshComponent.cpp
@@ -114,6 +114,12 @@ FMaterialRelevance UMeshComponent::GetMaterialRelevance(ERHIFeatureLevel::Type I
 		}
 		Result |= MaterialInterface->GetRelevance_Concurrent(InFeatureLevel);
 	}
+
+	if (OverlayMaterial != nullptr)
+	{
+		Result |= OverlayMaterial->GetRelevance_Concurrent(InFeatureLevel);
+	}
+
 	return Result;
 }
 
@@ -179,6 +185,25 @@ void UMeshComponent::GetUsedMaterials(TArray<UMaterialInterface*>& OutMaterials,
 			OutMaterials.Add(MaterialInterface);
 		}
 	}
+
+	if (OverlayMaterial != nullptr)
+	{
+		OutMaterials.Add(OverlayMaterial);
+	}
+}
+
+UMaterialInterface* UMeshComponent::GetOverlayMaterial() const
+{
+	return OverlayMaterial;
+}
+
+void UMeshComponent::SetOverlayMaterial(UMaterialInterface* NewOverlayMaterial)
+{
+	if (OverlayMaterial != NewOverlayMaterial)
+	{
+		OverlayMaterial = NewOverlayMaterial;
+		MarkRenderStateDirty();
+	}
 }
 
 void UMeshComponent::PrestreamTextures( float Seconds, bool bPrioritizeCharacterTextures, int32 CinematicTextureGroups )
diff --git a/Engine/Source/Runtime/Engine/Private/Components/SkinnedMeshComponent.cpp b/Engine/Source/Runtime/Engine/Private/Components/SkinnedMeshComponent.cpp
--- a/Engine/Source/Runtime/Engine/Private/Components/SkinnedMeshComponent.cpp
+++ b/Engine/Source/Runtime/Engine/Private/Components/SkinnedMeshComponent.cpp
@@ -2699,6 +2699,11 @@ void USkinnedMeshComponent::GetUsedMaterials( TArray<UMaterialInterface*>& OutMa
 			UMaterialInterface* MaterialInterface = GetMaterial( MatIdx );
 			OutMaterials.Add( MaterialInterface );
 		}
+
+		if (OverlayMaterial != nullptr)
+		{
+			OutMaterials.Add(OverlayMaterial);
+		}
 	}
 
 	if (bGetDebugMaterials)
diff --git a/Engine/Source/Runtime/Engine/Private/Components/StaticMeshComponent.cpp b/Engine/Source/Runtime/Engine/Private/Components/StaticMeshComponent.cpp
--- a/Engine/Source/Runtime/Engine/Private/Components/StaticMeshComponent.cpp
+++ b/Engine/Source/Runtime/Engine/Private/Components/StaticMeshComponent.cpp
@@ -2523,6 +2523,11 @@ void UStaticMeshComponent::GetUsedMaterials(TArray<UMaterialInterface*>& OutMate
 			{
 				OutMaterials.Add(GetMaterial(MaterialIndex));
 			}
+
+			if (OverlayMaterial != nullptr)
+			{
+				OutMaterials.Add(OverlayMaterial);
+			}
 		}
 	}
 }
diff --git a/Engine/Source/Runtime/Engine/Private/HierarchicalInstancedStaticMesh.cpp b/Engine/Source/Runtime/Engine/Private/HierarchicalInstancedStaticMesh.cpp
--- a/Engine/Source/Runtime/Engine/Private/HierarchicalInstancedStaticMesh.cpp
+++ b/Engine/Source/Runtime/Engine/Private/HierarchicalInstancedStaticMesh.cpp
@@ -1335,6 +1335,16 @@ void FHierarchicalStaticMeshSceneProxy::FillDynamicMeshElements(FMeshElementColl
 					if (TotalTriangles < (int64)CVarMaxTrianglesToRender.GetValueOnRenderThread())
 					{
 						Collector.AddMesh(ElementParams.ViewIndex, MeshBatch);
+
+						if (OverlayMaterial != nullptr)
+						{
+							FMeshBatch& OverlayMeshBatch = Collector.AllocateMesh();
+							OverlayMeshBatch = MeshBatch;
+							OverlayMeshBatch.CastShadow = false;
+							OverlayMeshBatch.bSelectable = false;
+							OverlayMeshBatch.MaterialRenderProxy = OverlayMaterial->GetRenderProxy();
+							Collector.AddMesh(ElementParams.ViewIndex, OverlayMeshBatch);
+						}
 					}
 				}
 				else
@@ -1452,6 +1462,16 @@ void FHierarchicalStaticMeshSceneProxy::FillDynamicMeshElements(FMeshElementColl
 						if (TotalTriangles < (int64)CVarMaxTrianglesToRender.GetValueOnRenderThread())
 						{
 							Collector.AddMesh(ElementParams.ViewIndex, MeshElement);
+
+							if (OverlayMaterial != nullptr)
+							{
+								FMeshBatch& OverlayMeshBatch = Collector.AllocateMesh();
+								OverlayMeshBatch = MeshElement;
+								OverlayMeshBatch.CastShadow = false;
+								OverlayMeshBatch.bSelectable = false;
+								OverlayMeshBatch.MaterialRenderProxy = OverlayMaterial->GetRenderProxy();
+								Collector.AddMesh(ElementParams.ViewIndex, OverlayMeshBatch);
+							}
 						}
 					}
 				}
@@ -1459,6 +1479,7 @@ void FHierarchicalStaticMeshSceneProxy::FillDynamicMeshElements(FMeshElementColl
 		}
 	}
 #if STATS
+	TotalTriangles*= (OverlayMaterial != nullptr ? 2 : 1);
 	TotalTriangles = FMath::Min<int64>(TotalTriangles, MAX_int32);
 	INC_DWORD_STAT_BY(STAT_FoliageTriangles, (uint32)TotalTriangles);
 	INC_DWORD_STAT_BY(STAT_StaticMeshTriangles, (uint32)TotalTriangles);
diff --git a/Engine/Source/Runtime/Engine/Private/InstancedStaticMesh.cpp b/Engine/Source/Runtime/Engine/Private/InstancedStaticMesh.cpp
--- a/Engine/Source/Runtime/Engine/Private/InstancedStaticMesh.cpp
+++ b/Engine/Source/Runtime/Engine/Private/InstancedStaticMesh.cpp
@@ -1306,6 +1306,18 @@ void FInstancedStaticMeshSceneProxy::GetDynamicMeshElements(const TArray<const F
 
 							Collector.AddMesh(ViewIndex, MeshElement);
 							INC_DWORD_STAT_BY(STAT_StaticMeshTriangles, MeshElement.GetNumPrimitives());
+
+							if (OverlayMaterial != nullptr)
+							{
+								FMeshBatch& OverlayMeshBatch = Collector.AllocateMesh();
+								OverlayMeshBatch = MeshElement;
+								OverlayMeshBatch.CastShadow = false;
+								OverlayMeshBatch.bSelectable = false;
+								OverlayMeshBatch.MaterialRenderProxy = OverlayMaterial->GetRenderProxy();
+								Collector.AddMesh(ViewIndex, OverlayMeshBatch);
+								
+								INC_DWORD_STAT_BY(STAT_StaticMeshTriangles, OverlayMeshBatch.GetNumPrimitives());
+							}
 						}
 					}
 				}
@@ -1370,6 +1382,15 @@ void FInstancedStaticMeshSceneProxy::SetupProxy(UInstancedStaticMeshComponent* I
 		}
 	}
 
+	if (OverlayMaterial != nullptr)
+	{
+		if (!OverlayMaterial->CheckMaterialUsage_Concurrent(MATUSAGE_InstancedStaticMeshes))
+		{
+			OverlayMaterial = UMaterial::GetDefaultMaterial(MD_Surface);
+			UE_LOG(LogStaticMesh, Error, TEXT("Overlay material with missing usage flag was applied to instanced static mesh %s"),	*InComponent->GetStaticMesh()->GetPathName());
+		}
+	}
+
 	// Copy the parameters for LOD - all instances
 	UserData_AllInstances.MeshRenderData = InComponent->GetStaticMesh()->GetRenderData();
 	UserData_AllInstances.StartCullDistance = InComponent->InstanceStartCullDistance;
diff --git a/Engine/Source/Runtime/Engine/Private/SkeletalMesh.cpp b/Engine/Source/Runtime/Engine/Private/SkeletalMesh.cpp
--- a/Engine/Source/Runtime/Engine/Private/SkeletalMesh.cpp
+++ b/Engine/Source/Runtime/Engine/Private/SkeletalMesh.cpp
@@ -5991,6 +5991,7 @@ FSkeletalMeshSceneProxy::FSkeletalMeshSceneProxy(const USkinnedMeshComponent* Co
 		,	SkeletalMeshRenderData(InSkelMeshRenderData)
 		,	SkeletalMeshForDebug(Component->SkeletalMesh)
 		,	PhysicsAssetForDebug(Component->GetPhysicsAsset())
+		,	OverlayMaterial(Component->OverlayMaterial)
 #if RHI_RAYTRACING
 		,	bAnySegmentUsesWorldPositionOffset(false)
 #endif
@@ -6114,6 +6115,15 @@ FSkeletalMeshSceneProxy::FSkeletalMeshSceneProxy(const USkinnedMeshComponent* Co
 		}
 	}
 
+	if (OverlayMaterial != nullptr)
+	{
+		if (!OverlayMaterial->CheckMaterialUsage_Concurrent(MATUSAGE_SkeletalMesh))
+		{
+			OverlayMaterial = UMaterial::GetDefaultMaterial(MD_Surface);
+			UE_LOG(LogSkeletalMesh, Error, TEXT("Overlay material with missing usage flag was applied to skeletal mesh %s"),	*Component->SkeletalMesh->GetPathName());
+		}
+	}
+
 	bCastDynamicShadow = bCastDynamicShadow && bAnySectionCastsShadow;
 
 	// Try to find a color for level coloration.
@@ -6415,6 +6425,15 @@ void FSkeletalMeshSceneProxy::DrawStaticElements(FStaticPrimitiveDrawInterface*
 					BatchElement.IndexBuffer = LODData.MultiSizeIndexContainer.GetIndexBuffer();
 													
 					PDI->DrawMesh(MeshElement, ScreenSize);
+
+					if (OverlayMaterial != nullptr)
+					{
+						FMeshBatch OverlayMeshBatch(MeshElement);
+						OverlayMeshBatch.CastShadow = false;
+						OverlayMeshBatch.bSelectable = false;
+						OverlayMeshBatch.MaterialRenderProxy = OverlayMaterial->GetRenderProxy();
+						PDI->DrawMesh(OverlayMeshBatch, ScreenSize);
+					}
 				}
 			}
 		}
@@ -6693,6 +6712,19 @@ void FSkeletalMeshSceneProxy::GetDynamicElementsSection(const TArray<const FScen
 			INC_DWORD_STAT_BY(STAT_GPUSkinVertices,(uint32)(bIsCPUSkinned ? 0 : NumVertices));
 			INC_DWORD_STAT_BY(STAT_SkelMeshTriangles,Mesh.GetNumPrimitives());
 			INC_DWORD_STAT(STAT_SkelMeshDrawCalls);
+
+			if (OverlayMaterial != nullptr)
+			{
+				FMeshBatch& OverlayMeshBatch = Collector.AllocateMesh();
+				OverlayMeshBatch = Mesh;
+				OverlayMeshBatch.CastShadow = false;
+				OverlayMeshBatch.bSelectable = false;
+				OverlayMeshBatch.MaterialRenderProxy = OverlayMaterial->GetRenderProxy();
+				Collector.AddMesh(ViewIndex, OverlayMeshBatch);
+				
+				INC_DWORD_STAT_BY(STAT_SkelMeshTriangles, OverlayMeshBatch.GetNumPrimitives());
+				INC_DWORD_STAT(STAT_SkelMeshDrawCalls);
+			}
 		}
 	}
 }
diff --git a/Engine/Source/Runtime/Engine/Private/StaticMeshRender.cpp b/Engine/Source/Runtime/Engine/Private/StaticMeshRender.cpp
--- a/Engine/Source/Runtime/Engine/Private/StaticMeshRender.cpp
+++ b/Engine/Source/Runtime/Engine/Private/StaticMeshRender.cpp
@@ -193,6 +193,7 @@ static TAutoConsoleVariable<float> CVarRayTracingStaticMeshesWPOCullingRadius(
 FStaticMeshSceneProxy::FStaticMeshSceneProxy(UStaticMeshComponent* InComponent, bool bForceLODsShareStaticLighting)
 	: FPrimitiveSceneProxy(InComponent, InComponent->GetStaticMesh()->GetFName())
 	, RenderData(InComponent->GetStaticMesh()->GetRenderData())
+	, OverlayMaterial(InComponent->OverlayMaterial)
 	, ForcedLodModel(InComponent->ForcedLodModel)
 	, bCastShadow(InComponent->CastShadow)
 	, bReverseCulling(InComponent->bReverseCulling)
@@ -1188,6 +1189,16 @@ void FStaticMeshSceneProxy::DrawStaticElements(FStaticPrimitiveDrawInterface* PD
 								PDI->DrawMesh(MeshBatch, FLT_MAX);
 							}
 						}
+
+						if (OverlayMaterial != nullptr)
+						{
+							FMeshBatch OverlayMeshBatch(BaseMeshBatch);
+							OverlayMeshBatch.CastShadow = false;
+							OverlayMeshBatch.bSelectable = false;
+							OverlayMeshBatch.MaterialRenderProxy = OverlayMaterial->GetRenderProxy();
+							PDI->DrawMesh(OverlayMeshBatch, FLT_MAX);
+						}
+
 						{
 							PDI->DrawMesh(BaseMeshBatch, FLT_MAX);
 						}
@@ -1305,6 +1316,15 @@ void FStaticMeshSceneProxy::DrawStaticElements(FStaticPrimitiveDrawInterface* PD
 								}
 							}
 
+							if (OverlayMaterial != nullptr)
+							{
+								FMeshBatch OverlayMeshBatch(BaseMeshBatch);
+								OverlayMeshBatch.CastShadow = false;
+								OverlayMeshBatch.bSelectable = false;
+								OverlayMeshBatch.MaterialRenderProxy = OverlayMaterial->GetRenderProxy();
+								PDI->DrawMesh(OverlayMeshBatch, ScreenSize);
+							}
+
 							{
 								// Standard mesh elements.
 								// If we have submitted an optimized shadow-only mesh, remaining mesh elements must not cast shadows.
diff --git a/Engine/Source/Runtime/Engine/Public/SkeletalMeshTypes.h b/Engine/Source/Runtime/Engine/Public/SkeletalMeshTypes.h
--- a/Engine/Source/Runtime/Engine/Public/SkeletalMeshTypes.h
+++ b/Engine/Source/Runtime/Engine/Public/SkeletalMeshTypes.h
@@ -241,7 +241,8 @@ protected:
 	/** The points to the skeletal mesh and physics assets are purely for debug purposes. Access is NOT thread safe! */
 	const USkeletalMesh* SkeletalMeshForDebug;
 	class UPhysicsAsset* PhysicsAssetForDebug;
-
+	
+	UMaterialInterface* OverlayMaterial;
 public:
 #if RHI_RAYTRACING
 	bool bAnySegmentUsesWorldPositionOffset : 1;
diff --git a/Engine/Source/Runtime/Engine/Public/StaticMeshResources.h b/Engine/Source/Runtime/Engine/Public/StaticMeshResources.h
--- a/Engine/Source/Runtime/Engine/Public/StaticMeshResources.h
+++ b/Engine/Source/Runtime/Engine/Public/StaticMeshResources.h
@@ -986,6 +986,8 @@ protected:
 	const FDistanceFieldVolumeData* DistanceFieldData;	
 	const FCardRepresentationData* CardRepresentationData;	
 
+	UMaterialInterface* OverlayMaterial;
+
 #if RHI_RAYTRACING
 	bool bSupportRayTracing;
 	bool bDynamicRayTracingGeometry;
