Added max draw distance for an overlay material. Supported on static and skeletal meshes

diff --git a/Engine/Source/Runtime/Engine/Classes/Components/MeshComponent.h b/Engine/Source/Runtime/Engine/Classes/Components/MeshComponent.h
--- a/Engine/Source/Runtime/Engine/Classes/Components/MeshComponent.h
+++ b/Engine/Source/Runtime/Engine/Classes/Components/MeshComponent.h
@@ -45,6 +45,10 @@ public:
 	/** Translucent material to blend on top of this mesh. Mesh will be rendered twice - once with a base material and once with overlay material */
 	UPROPERTY(EditAnywhere, BlueprintReadOnly, AdvancedDisplay, Category=Rendering)
 	TObjectPtr<class UMaterialInterface> OverlayMaterial;
+	
+	/** The max draw distance for overlay material. A distance of 0 indicates that overlay will be culled using primitive max distance. */
+	UPROPERTY(EditAnywhere, BlueprintReadOnly, AdvancedDisplay, Category=Rendering)
+	float OverlayMaterialMaxDrawDistance;
 
 	/** Get the overlay material used by this instance */
 	UFUNCTION(BlueprintCallable, Category="Rendering|Material")
@@ -54,6 +58,10 @@ public:
 	UFUNCTION(BlueprintCallable, Category="Rendering|Material")
 	void SetOverlayMaterial(class UMaterialInterface* NewOverlayMaterial);
 
+	/** Change the overlay material max draw distance used by this instance */
+	UFUNCTION(BlueprintCallable, Category="Rendering|Material")
+	void SetOverlayMaterialMaxDrawDistance(float InMaxDrawDistance);
+
 #if WITH_EDITOR
 	/*
 	 * Make sure the Override array is using only the space it should use.
diff --git a/Engine/Source/Runtime/Engine/Private/Components/MeshComponent.cpp b/Engine/Source/Runtime/Engine/Private/Components/MeshComponent.cpp
--- a/Engine/Source/Runtime/Engine/Private/Components/MeshComponent.cpp
+++ b/Engine/Source/Runtime/Engine/Private/Components/MeshComponent.cpp
@@ -206,6 +206,15 @@ void UMeshComponent::SetOverlayMaterial(UMaterialInterface* NewOverlayMaterial)
 	}
 }
 
+void UMeshComponent::SetOverlayMaterialMaxDrawDistance(float InMaxDrawDistance)
+{
+	if (OverlayMaterialMaxDrawDistance != InMaxDrawDistance)
+	{
+		OverlayMaterialMaxDrawDistance = InMaxDrawDistance;
+		MarkRenderStateDirty();
+	}
+}
+
 void UMeshComponent::PrestreamTextures( float Seconds, bool bPrioritizeCharacterTextures, int32 CinematicTextureGroups )
 {
 	// If requested, tell the streaming system to only process character textures for 30 frames.
diff --git a/Engine/Source/Runtime/Engine/Private/HierarchicalInstancedStaticMesh.cpp b/Engine/Source/Runtime/Engine/Private/HierarchicalInstancedStaticMesh.cpp
--- a/Engine/Source/Runtime/Engine/Private/HierarchicalInstancedStaticMesh.cpp
+++ b/Engine/Source/Runtime/Engine/Private/HierarchicalInstancedStaticMesh.cpp
@@ -1340,6 +1340,7 @@ void FHierarchicalStaticMeshSceneProxy::FillDynamicMeshElements(FMeshElementColl
 						{
 							FMeshBatch& OverlayMeshBatch = Collector.AllocateMesh();
 							OverlayMeshBatch = MeshBatch;
+							OverlayMeshBatch.bOverlayMaterial = true;
 							OverlayMeshBatch.CastShadow = false;
 							OverlayMeshBatch.bSelectable = false;
 							OverlayMeshBatch.MaterialRenderProxy = OverlayMaterial->GetRenderProxy();
@@ -1467,6 +1468,7 @@ void FHierarchicalStaticMeshSceneProxy::FillDynamicMeshElements(FMeshElementColl
 							{
 								FMeshBatch& OverlayMeshBatch = Collector.AllocateMesh();
 								OverlayMeshBatch = MeshElement;
+								OverlayMeshBatch.bOverlayMaterial = true;
 								OverlayMeshBatch.CastShadow = false;
 								OverlayMeshBatch.bSelectable = false;
 								OverlayMeshBatch.MaterialRenderProxy = OverlayMaterial->GetRenderProxy();
diff --git a/Engine/Source/Runtime/Engine/Private/InstancedStaticMesh.cpp b/Engine/Source/Runtime/Engine/Private/InstancedStaticMesh.cpp
--- a/Engine/Source/Runtime/Engine/Private/InstancedStaticMesh.cpp
+++ b/Engine/Source/Runtime/Engine/Private/InstancedStaticMesh.cpp
@@ -1311,6 +1311,7 @@ void FInstancedStaticMeshSceneProxy::GetDynamicMeshElements(const TArray<const F
 							{
 								FMeshBatch& OverlayMeshBatch = Collector.AllocateMesh();
 								OverlayMeshBatch = MeshElement;
+								OverlayMeshBatch.bOverlayMaterial = true;
 								OverlayMeshBatch.CastShadow = false;
 								OverlayMeshBatch.bSelectable = false;
 								OverlayMeshBatch.MaterialRenderProxy = OverlayMaterial->GetRenderProxy();
diff --git a/Engine/Source/Runtime/Engine/Private/SkeletalMesh.cpp b/Engine/Source/Runtime/Engine/Private/SkeletalMesh.cpp
--- a/Engine/Source/Runtime/Engine/Private/SkeletalMesh.cpp
+++ b/Engine/Source/Runtime/Engine/Private/SkeletalMesh.cpp
@@ -5992,6 +5992,7 @@ FSkeletalMeshSceneProxy::FSkeletalMeshSceneProxy(const USkinnedMeshComponent* Co
 		,	SkeletalMeshForDebug(Component->SkeletalMesh)
 		,	PhysicsAssetForDebug(Component->GetPhysicsAsset())
 		,	OverlayMaterial(Component->OverlayMaterial)
+		,	OverlayMaterialMaxDrawDistance(Component->OverlayMaterialMaxDrawDistance)
 #if RHI_RAYTRACING
 		,	bAnySegmentUsesWorldPositionOffset(false)
 #endif
@@ -6429,10 +6430,13 @@ void FSkeletalMeshSceneProxy::DrawStaticElements(FStaticPrimitiveDrawInterface*
 					if (OverlayMaterial != nullptr)
 					{
 						FMeshBatch OverlayMeshBatch(MeshElement);
+						OverlayMeshBatch.bOverlayMaterial = true;
 						OverlayMeshBatch.CastShadow = false;
 						OverlayMeshBatch.bSelectable = false;
 						OverlayMeshBatch.MaterialRenderProxy = OverlayMaterial->GetRenderProxy();
-						PDI->DrawMesh(OverlayMeshBatch, ScreenSize);
+						// Reuse mesh ScreenSize as cull distance for an overlay. Overlay does not need to compute LOD so we can avoid adding new members into MeshBatch or MeshRelevance
+						float OverlayMeshScreenSize = OverlayMaterialMaxDrawDistance;
+						PDI->DrawMesh(OverlayMeshBatch, OverlayMeshScreenSize);
 					}
 				}
 			}
@@ -6715,15 +6719,37 @@ void FSkeletalMeshSceneProxy::GetDynamicElementsSection(const TArray<const FScen
 
 			if (OverlayMaterial != nullptr)
 			{
-				FMeshBatch& OverlayMeshBatch = Collector.AllocateMesh();
-				OverlayMeshBatch = Mesh;
-				OverlayMeshBatch.CastShadow = false;
-				OverlayMeshBatch.bSelectable = false;
-				OverlayMeshBatch.MaterialRenderProxy = OverlayMaterial->GetRenderProxy();
-				Collector.AddMesh(ViewIndex, OverlayMeshBatch);
+				const bool bHasOverlayCullDistance = 
+					OverlayMaterialMaxDrawDistance > 1.f && 
+					OverlayMaterialMaxDrawDistance != FLT_MAX && 
+					!ViewFamily.EngineShowFlags.DistanceCulledPrimitives;
 				
-				INC_DWORD_STAT_BY(STAT_SkelMeshTriangles, OverlayMeshBatch.GetNumPrimitives());
-				INC_DWORD_STAT(STAT_SkelMeshDrawCalls);
+				bool bAddOverlay = true;
+				if (bHasOverlayCullDistance)
+				{
+					float MaxDrawDistanceScale = GetCachedScalabilityCVars().ViewDistanceScale;
+					MaxDrawDistanceScale *= GetCachedScalabilityCVars().CalculateFieldOfViewDistanceScale(View->DesiredFOV);
+					float DistanceSquared = (GetBounds().Origin - View->ViewMatrices.GetViewOrigin()).SizeSquared();
+					if (DistanceSquared > FMath::Square(OverlayMaterialMaxDrawDistance * MaxDrawDistanceScale))
+					{
+						// distance culled
+						bAddOverlay = false;
+					}
+				}
+				
+				if (bAddOverlay)
+				{
+					FMeshBatch& OverlayMeshBatch = Collector.AllocateMesh();
+					OverlayMeshBatch = Mesh;
+					OverlayMeshBatch.bOverlayMaterial = true;
+					OverlayMeshBatch.CastShadow = false;
+					OverlayMeshBatch.bSelectable = false;
+					OverlayMeshBatch.MaterialRenderProxy = OverlayMaterial->GetRenderProxy();
+					Collector.AddMesh(ViewIndex, OverlayMeshBatch);
+				
+					INC_DWORD_STAT_BY(STAT_SkelMeshTriangles, OverlayMeshBatch.GetNumPrimitives());
+					INC_DWORD_STAT(STAT_SkelMeshDrawCalls);
+				}
 			}
 		}
 	}
diff --git a/Engine/Source/Runtime/Engine/Private/StaticMeshRender.cpp b/Engine/Source/Runtime/Engine/Private/StaticMeshRender.cpp
--- a/Engine/Source/Runtime/Engine/Private/StaticMeshRender.cpp
+++ b/Engine/Source/Runtime/Engine/Private/StaticMeshRender.cpp
@@ -194,6 +194,7 @@ FStaticMeshSceneProxy::FStaticMeshSceneProxy(UStaticMeshComponent* InComponent,
 	: FPrimitiveSceneProxy(InComponent, InComponent->GetStaticMesh()->GetFName())
 	, RenderData(InComponent->GetStaticMesh()->GetRenderData())
 	, OverlayMaterial(InComponent->OverlayMaterial)
+	, OverlayMaterialMaxDrawDistance(InComponent->OverlayMaterialMaxDrawDistance)
 	, ForcedLodModel(InComponent->ForcedLodModel)
 	, bCastShadow(InComponent->CastShadow)
 	, bReverseCulling(InComponent->bReverseCulling)
@@ -1190,18 +1191,19 @@ void FStaticMeshSceneProxy::DrawStaticElements(FStaticPrimitiveDrawInterface* PD
 							}
 						}
 
+						{
+							PDI->DrawMesh(BaseMeshBatch, FLT_MAX);
+						}
+
 						if (OverlayMaterial != nullptr)
 						{
 							FMeshBatch OverlayMeshBatch(BaseMeshBatch);
+							OverlayMeshBatch.bOverlayMaterial = true;
 							OverlayMeshBatch.CastShadow = false;
 							OverlayMeshBatch.bSelectable = false;
 							OverlayMeshBatch.MaterialRenderProxy = OverlayMaterial->GetRenderProxy();
 							PDI->DrawMesh(OverlayMeshBatch, FLT_MAX);
 						}
-
-						{
-							PDI->DrawMesh(BaseMeshBatch, FLT_MAX);
-						}
 					}
 				}
 			}
@@ -1316,15 +1318,6 @@ void FStaticMeshSceneProxy::DrawStaticElements(FStaticPrimitiveDrawInterface* PD
 								}
 							}
 
-							if (OverlayMaterial != nullptr)
-							{
-								FMeshBatch OverlayMeshBatch(BaseMeshBatch);
-								OverlayMeshBatch.CastShadow = false;
-								OverlayMeshBatch.bSelectable = false;
-								OverlayMeshBatch.MaterialRenderProxy = OverlayMaterial->GetRenderProxy();
-								PDI->DrawMesh(OverlayMeshBatch, ScreenSize);
-							}
-
 							{
 								// Standard mesh elements.
 								// If we have submitted an optimized shadow-only mesh, remaining mesh elements must not cast shadows.
@@ -1334,6 +1327,18 @@ void FStaticMeshSceneProxy::DrawStaticElements(FStaticPrimitiveDrawInterface* PD
 								MeshBatch.bUseForDepthPass &= !bUseUnifiedMeshForDepth;
 								PDI->DrawMesh(MeshBatch, ScreenSize);
 							}
+
+							if (OverlayMaterial != nullptr)
+							{
+								FMeshBatch OverlayMeshBatch(BaseMeshBatch);
+								OverlayMeshBatch.bOverlayMaterial = true;
+								OverlayMeshBatch.CastShadow = false;
+								OverlayMeshBatch.bSelectable = false;
+								OverlayMeshBatch.MaterialRenderProxy = OverlayMaterial->GetRenderProxy();
+								// Reuse mesh ScreenSize as cull distance for an overlay. Overlay does not need to compute LOD so we can avoid adding new members into MeshBatch or MeshRelevance
+								float OverlayMeshScreenSize = OverlayMaterialMaxDrawDistance;
+								PDI->DrawMesh(OverlayMeshBatch, OverlayMeshScreenSize);
+							}
 						}
 					}
 				}
diff --git a/Engine/Source/Runtime/Engine/Public/MeshBatch.h b/Engine/Source/Runtime/Engine/Public/MeshBatch.h
--- a/Engine/Source/Runtime/Engine/Public/MeshBatch.h
+++ b/Engine/Source/Runtime/Engine/Public/MeshBatch.h
@@ -366,6 +366,9 @@ struct FMeshBatch
 	uint32 bRenderToVirtualTexture : 1;
 	/** What virtual texture material type this mesh batch should be rendered with. */
 	uint32 RuntimeVirtualTextureMaterialType : RuntimeVirtualTexture::MaterialType_NumBits;
+	
+	/** Whether mesh is rendered with overlay material. */
+	uint32 bOverlayMaterial	: 1;
 
 #if RHI_RAYTRACING
 	uint32 CastRayTracedShadow : 1;	// Whether it casts ray traced shadow.
@@ -476,6 +479,7 @@ struct FMeshBatch
 	,	bDitheredLODTransition(false)
 	,	bRenderToVirtualTexture(false)
 	,	RuntimeVirtualTextureMaterialType(0)
+	,	bOverlayMaterial(false)
 #if RHI_RAYTRACING
 	,	CastRayTracedShadow(true)
 #endif
diff --git a/Engine/Source/Runtime/Engine/Public/SkeletalMeshTypes.h b/Engine/Source/Runtime/Engine/Public/SkeletalMeshTypes.h
--- a/Engine/Source/Runtime/Engine/Public/SkeletalMeshTypes.h
+++ b/Engine/Source/Runtime/Engine/Public/SkeletalMeshTypes.h
@@ -243,6 +243,7 @@ protected:
 	class UPhysicsAsset* PhysicsAssetForDebug;
 	
 	UMaterialInterface* OverlayMaterial;
+	float OverlayMaterialMaxDrawDistance;
 public:
 #if RHI_RAYTRACING
 	bool bAnySegmentUsesWorldPositionOffset : 1;
diff --git a/Engine/Source/Runtime/Engine/Public/StaticMeshResources.h b/Engine/Source/Runtime/Engine/Public/StaticMeshResources.h
--- a/Engine/Source/Runtime/Engine/Public/StaticMeshResources.h
+++ b/Engine/Source/Runtime/Engine/Public/StaticMeshResources.h
@@ -987,6 +987,7 @@ protected:
 	const FCardRepresentationData* CardRepresentationData;	
 
 	UMaterialInterface* OverlayMaterial;
+	float OverlayMaterialMaxDrawDistance;
 
 #if RHI_RAYTRACING
 	bool bSupportRayTracing;
diff --git a/Engine/Source/Runtime/Renderer/Private/SceneVisibility.cpp b/Engine/Source/Runtime/Renderer/Private/SceneVisibility.cpp
--- a/Engine/Source/Runtime/Renderer/Private/SceneVisibility.cpp
+++ b/Engine/Source/Runtime/Renderer/Private/SceneVisibility.cpp
@@ -2520,6 +2520,8 @@ struct FRelevancePacket
 		const bool bMobileBasePassAlwaysUsesCSM = (ShadingPath == EShadingPath::Mobile) && MobileBasePassAlwaysUsesCSM(Scene->GetShaderPlatform());
 		const bool bHLODActive = Scene->SceneLODHierarchy.IsActive();
 		const FHLODVisibilityState* const HLODState = bHLODActive && ViewState ? &ViewState->HLODVisibilityState : nullptr;
+		float MaxDrawDistanceScale = GetCachedScalabilityCVars().ViewDistanceScale;
+		MaxDrawDistanceScale *= GetCachedScalabilityCVars().CalculateFieldOfViewDistanceScale(View.DesiredFOV);
 
 		
 		for (int32 StaticPrimIndex = 0, Num = RelevantStaticPrimitives.NumPrims; StaticPrimIndex < Num; ++StaticPrimIndex)
@@ -2554,6 +2556,20 @@ struct FRelevancePacket
 				const FStaticMeshBatchRelevance& StaticMeshRelevance = PrimitiveSceneInfo->StaticMeshRelevances[MeshIndex];
 				const FStaticMeshBatch& StaticMesh = PrimitiveSceneInfo->StaticMeshes[MeshIndex];
 
+				if (StaticMesh.bOverlayMaterial && !View.Family->EngineShowFlags.DistanceCulledPrimitives)
+				{
+					// Overlay mesh can have a his own cull distance that is shorter than primitive cull distance
+					float OverlayMaterialMaxDrawDistance = StaticMeshRelevance.ScreenSize;
+					if (OverlayMaterialMaxDrawDistance > 1.f && OverlayMaterialMaxDrawDistance != FLT_MAX)
+					{
+						if (DistanceSquared > FMath::Square(OverlayMaterialMaxDrawDistance * MaxDrawDistanceScale))
+						{
+							// distance culled
+							continue;
+						}
+					}
+				}
+
 				if (LODToRender.ContainsLOD(StaticMeshRelevance.LODIndex))
 				{
 					uint8 MarkMask = 0;
@@ -2805,7 +2821,9 @@ struct FRelevancePacket
 							BatchAndProxy.Proxy = PrimitiveSceneInfo->Proxy;
 						}
 
-						if (ViewRelevance.bRenderInMainPass && ViewRelevance.bDecal)
+						// FIXME: Now if a primitive has one batch with a decal material all primitive mesh batches will be added as decals
+						// Because ViewRelevance is a sum of all material relevances in the primitive
+						if (ViewRelevance.bRenderInMainPass && ViewRelevance.bDecal && StaticMeshRelevance.bUseForMaterial)
 						{
 							MeshDecalBatches.AddUninitialized(1);
 							FMeshDecalBatch& BatchAndProxy = MeshDecalBatches.Last();
