HLOD Actor Desc - Duplicate MID referenced by HLOD imposters as they are outered to the level

diff --git a/Engine/Source/Runtime/Engine/Private/HLOD/HLODProxyDesc.cpp b/Engine/Source/Runtime/Engine/Private/HLOD/HLODProxyDesc.cpp
--- a/Engine/Source/Runtime/Engine/Private/HLOD/HLODProxyDesc.cpp
+++ b/Engine/Source/Runtime/Engine/Private/HLOD/HLODProxyDesc.cpp
@@ -8,6 +8,7 @@
 #include "GameFramework/WorldSettings.h"
 #include "LevelUtils.h"
 #include "Engine/LevelStreaming.h"
+#include "Materials/MaterialInstanceDynamic.h"
 #endif
 
 #if WITH_EDITOR
@@ -34,7 +35,16 @@ bool FHLODISMComponentDesc::operator==(const FHLODISMComponentDesc& Other) const
 		return false;
 	}
 
-	if (Material != Other.Material)
+	UMaterialInstance* MaterialInstance = Cast<UMaterialInstance>(Material);
+	UMaterialInstance* OtherMaterialInstance = Cast<UMaterialInstance>(Other.Material);
+	if (MaterialInstance && OtherMaterialInstance)
+	{
+		if (!MaterialInstance->Equivalent(OtherMaterialInstance))
+		{
+			return false;
+		}
+	}
+	else if (Material != Other.Material)
 	{
 		return false;
 	}
@@ -103,7 +113,14 @@ bool UHLODProxyDesc::UpdateFromLODActor(const ALODActor* InLODActor)
 	{
 		if (Pair.Key.IsValid() && Pair.Value->GetInstanceCount() != 0)
 		{
-			ISMComponentsDesc.Emplace(Pair.Value);
+			FHLODISMComponentDesc& ISMComponentDesc = ISMComponentsDesc.Emplace_GetRef(Pair.Value);
+
+			// MIDs are not assets and are normally outered to their owner component.
+			// We need to duplicate them here to make sure we don't create references to actors in the source level.
+			if (UMaterialInstanceDynamic* MID = Cast<UMaterialInstanceDynamic>(ISMComponentDesc.Material))
+			{ 
+				ISMComponentDesc.Material = DuplicateObject<UMaterialInstanceDynamic>(MID, this);
+			}
 		}
 	}
 
@@ -270,10 +287,18 @@ ALODActor* UHLODProxyDesc::SpawnLODActor(ULevel* InLevel) const
 
 	for (const FHLODISMComponentDesc& ISMComponentDesc : ISMComponentsDesc)
 	{
-		if (!ISMComponentDesc.StaticMesh || !ISMComponentDesc.Material || ISMComponentDesc.Instances.Num() == 0)
+		UStaticMesh* ISMStaticMesh = ISMComponentDesc.StaticMesh;
+		UMaterialInterface* ISMMaterial = ISMComponentDesc.Material;
+
+		if (!ISMStaticMesh || !ISMMaterial || ISMComponentDesc.Instances.IsEmpty())
 		{
 			continue;
 		}
+
+		if (UMaterialInstanceDynamic* MID = Cast<UMaterialInstanceDynamic>(ISMMaterial))
+		{
+			ISMMaterial = DuplicateObject<UMaterialInstanceDynamic>(MID, LODActor);
+		}
 		
 		// Apply transform to HISM instances
 		const bool bTransformInstances = !ActorTransform.Equals(FTransform::Identity);
@@ -285,11 +310,11 @@ ALODActor* UHLODProxyDesc::SpawnLODActor(ULevel* InLevel) const
 				Transform *= ActorTransform;
 			}
 
-			LODActor->AddInstances(ISMComponentDesc.StaticMesh, ISMComponentDesc.Material, Transforms);
+			LODActor->AddInstances(ISMStaticMesh, ISMMaterial, Transforms);
 		}
 		else
 		{
-			LODActor->AddInstances(ISMComponentDesc.StaticMesh, ISMComponentDesc.Material, ISMComponentDesc.Instances);
+			LODActor->AddInstances(ISMStaticMesh, ISMMaterial, ISMComponentDesc.Instances);
 		}
 	}
 
diff --git a/Engine/Source/Runtime/Engine/Public/HLOD/HLODProxyDesc.h b/Engine/Source/Runtime/Engine/Public/HLOD/HLODProxyDesc.h
--- a/Engine/Source/Runtime/Engine/Public/HLOD/HLODProxyDesc.h
+++ b/Engine/Source/Runtime/Engine/Public/HLOD/HLODProxyDesc.h
@@ -28,7 +28,7 @@ public:
 	TObjectPtr<UStaticMesh> StaticMesh = nullptr;
 	
 	UPROPERTY()
-	TObjectPtr<const UMaterialInterface> Material = nullptr;
+	TObjectPtr<UMaterialInterface> Material = nullptr;
 
 	UPROPERTY()
 	TArray<FTransform> Instances;
