Don't allocate ISM previous transforms when there are no instances.
Zero out buffer when it is allocated in case it is sparsely filled.
#rnx

diff --git a/Engine/Source/Runtime/Engine/Private/InstancedStaticMesh.cpp b/Engine/Source/Runtime/Engine/Private/InstancedStaticMesh.cpp
--- a/Engine/Source/Runtime/Engine/Private/InstancedStaticMesh.cpp
+++ b/Engine/Source/Runtime/Engine/Private/InstancedStaticMesh.cpp
@@ -1431,8 +1431,8 @@ void FInstancedStaticMeshSceneProxy::SetupProxy(UInstancedStaticMeshComponent* I
 		bSupportsInstanceDataBuffer = true;
 		InstanceSceneData.SetNumZeroed(NumRenderInstances);
 
-		bHasPerInstanceDynamicData = InComponent->PerInstancePrevTransform.Num() == InComponent->GetInstanceCount();
-		InstanceDynamicData.SetNumUninitialized(bHasPerInstanceDynamicData ? NumRenderInstances : 0);
+		bHasPerInstanceDynamicData = InComponent->PerInstancePrevTransform.Num() > 0 && InComponent->PerInstancePrevTransform.Num() == InComponent->GetInstanceCount();
+		InstanceDynamicData.SetNumZeroed(bHasPerInstanceDynamicData ? NumRenderInstances : 0);
 
 		static const auto AllowStaticLightingVar = IConsoleManager::Get().FindTConsoleVariableDataInt(TEXT("r.AllowStaticLighting"));
 		const bool bAllowStaticLighting = (!AllowStaticLightingVar || AllowStaticLightingVar->GetValueOnAnyThread() != 0);
