Gate the fix to LOD around a CVar so we can temporarily avoid a performance hit on platforms that aren't experiencing the bug.

diff --git a/Engine/Source/Runtime/Engine/Private/Components/SkinnedMeshComponent.cpp b/Engine/Source/Runtime/Engine/Private/Components/SkinnedMeshComponent.cpp
--- a/Engine/Source/Runtime/Engine/Private/Components/SkinnedMeshComponent.cpp
+++ b/Engine/Source/Runtime/Engine/Private/Components/SkinnedMeshComponent.cpp
@@ -58,6 +58,14 @@ static TAutoConsoleVariable<int32> CVarAnimVisualizeLODs(
 	0,
 	TEXT("Visualize SkelMesh LODs"));
 
+int32 GUpdateBoundsNotifyStreaming = 0;
+FAutoConsoleVariableRef CVarUpdateBoundsNotifyStreaming(
+	TEXT("r.SkinnedMesh.UpdateBoundsNotifyStreaming"),
+	GUpdateBoundsNotifyStreaming,
+	TEXT("Update the streaming manager when the bounds change significantly"),
+	ECVF_Default
+);
+
 namespace FAnimUpdateRateManager
 {
 	static float TargetFrameTimeForUpdateRate = 1.f / 30.f; //Target frame rate for lookahead URO
@@ -1169,14 +1177,21 @@ FBoxSphereBounds USkinnedMeshComponent::CalcBounds(const FTransform& LocalToWorl
 
 void USkinnedMeshComponent::UpdateBounds()
 {
-	const auto OldRadius = Bounds.SphereRadius;
+	if (GUpdateBoundsNotifyStreaming)
+	{
+		const auto OldRadius = Bounds.SphereRadius;
 
-	Super::UpdateBounds();
+		Super::UpdateBounds();
 
-	// Avoid updating the streamer for small changes in size
-	if (!FMath::IsNearlyEqual(OldRadius, Bounds.SphereRadius, 0.1f) && GetForcedLOD() == 0)
+		// Avoid updating the streamer for small changes in size
+		if (!FMath::IsNearlyEqual(OldRadius, Bounds.SphereRadius, 0.1f) && GetForcedLOD() == 0)
+		{
+			IStreamingManager::Get().NotifyPrimitiveUpdated_Concurrent(this);
+		}
+	}
+	else
 	{
-		IStreamingManager::Get().NotifyPrimitiveUpdated_Concurrent(this);
+		Super::UpdateBounds();
 	}
 }
 
