WorldPartition:
- Fix UWorldPartitionLevelStreamingPolicy::RemapSoftObjectPath not properly remapping Actor Sub Objects in PIE
- Fix UWorldPartitionLevelStreamingPolicy::GetSubObject not properly remapping Actor Sub Objects at Runtime
#rnx

diff --git a/Engine/Source/Runtime/Engine/Private/WorldPartition/WorldPartitionLevelStreamingPolicy.cpp b/Engine/Source/Runtime/Engine/Private/WorldPartition/WorldPartitionLevelStreamingPolicy.cpp
--- a/Engine/Source/Runtime/Engine/Private/WorldPartition/WorldPartitionLevelStreamingPolicy.cpp
+++ b/Engine/Source/Runtime/Engine/Private/WorldPartition/WorldPartitionLevelStreamingPolicy.cpp
@@ -99,6 +99,7 @@ void UWorldPartitionLevelStreamingPolicy::RemapSoftObjectPath(FSoftObjectPath& O
 	// Make sure to work on non-PIE path (can happen for modified actors in PIE)
 	int32 PIEInstanceID = INDEX_NONE;
 	FString SrcPath = UWorld::RemovePIEPrefix(ObjectPath.ToString(), &PIEInstanceID);
+	const FSoftObjectPath SrcObjectPath(SrcPath);
 
 	FName* CellName = ActorToCellRemapping.Find(FName(*SrcPath));
 	if (!CellName)
@@ -112,7 +113,7 @@ void UWorldPartitionLevelStreamingPolicy::RemapSoftObjectPath(FSoftObjectPath& O
 			if (SubObjectPos != INDEX_NONE)
 			{
 				const FString ActorSubPathString = SubPathString.Left(SubObjectPos);
-				FSoftObjectPath ActorObjectPath(ObjectPath);
+				FSoftObjectPath ActorObjectPath(SrcObjectPath);
 				ActorObjectPath.SetSubPathString(ActorSubPathString);
 				CellName = ActorToCellRemapping.Find(FName(*ActorObjectPath.ToString()));
 			}
@@ -121,7 +122,6 @@ void UWorldPartitionLevelStreamingPolicy::RemapSoftObjectPath(FSoftObjectPath& O
 
 	if (CellName)
 	{
-		const FSoftObjectPath SrcObjectPath(SrcPath);
 		if (!SrcObjectPath.GetSubPathString().IsEmpty())
 		{
 			UWorld* World = WorldPartition->GetWorld();
@@ -156,10 +156,10 @@ UObject* UWorldPartitionLevelStreamingPolicy::GetSubObject(const TCHAR* SubObjec
 	FString SubObjectContext;	
 	if (!FString(SubObjectPath).Split(TEXT("."), &SubObjectContext, &SubObjectName))
 	{
-		SubObjectName = SubObjectPath;
+		SubObjectContext = SubObjectPath;
 	}
 
-	const FString SrcPath = UWorld::RemovePIEPrefix(*SubObjectName);
+	const FString SrcPath = UWorld::RemovePIEPrefix(*SubObjectContext);
 	if (FName* CellName = SubObjectsToCellRemapping.Find(FName(*SrcPath)))
 	{
 		if (UWorldPartitionRuntimeLevelStreamingCell* Cell = (UWorldPartitionRuntimeLevelStreamingCell*)StaticFindObject(UWorldPartitionRuntimeLevelStreamingCell::StaticClass(), GetOuterUWorldPartition(), *(CellName->ToString())))
