UE-145026 - BP Precision issues occur in CM_Character_Jumping (CharacterMovement  EngineTests) when comparing FP literals with serialized floats

We need to ensure that literal default values are interpreted as float strings instead of doubles. We've observed that some older Blueprint content might be using exact comparisons with floats. While this is generally frowned upon, some game logic may depend on this behavior. These issues can be difficult to debug if we covertly interpret literal float data as doubles.

As for the fix, anytime we serialize old pin data that represented a single precision float, we set a flag on the type that ensures we parse its default value string as a single precision float. This flag is also serialized, and isn't changed until the user actually changes the current value to something different.

diff --git a/Engine/Source/Editor/BlueprintGraph/Private/EdGraphSchema_K2.cpp b/Engine/Source/Editor/BlueprintGraph/Private/EdGraphSchema_K2.cpp
--- a/Engine/Source/Editor/BlueprintGraph/Private/EdGraphSchema_K2.cpp
+++ b/Engine/Source/Editor/BlueprintGraph/Private/EdGraphSchema_K2.cpp
@@ -4751,6 +4751,12 @@ void UEdGraphSchema_K2::TrySetDefaultValue(UEdGraphPin& Pin, const FString& NewD
 		Pin.DefaultValue = UseDefaultValue;
 		Pin.DefaultTextValue = UseDefaultText;
 
+		// Legacy float data will continue to serialize as a single precision float until we explicitly change the default value
+		if (Pin.PinType.PinCategory == PC_Real)
+		{
+			Pin.PinType.bSerializeAsSinglePrecisionFloat = false;
+		}
+
 		UEdGraphNode* Node = Pin.GetOwningNode();
 		Node->PinDefaultValueChanged(&Pin);
 
diff --git a/Engine/Source/Editor/KismetCompiler/Private/KismetCompilerVMBackend.cpp b/Engine/Source/Editor/KismetCompiler/Private/KismetCompilerVMBackend.cpp
--- a/Engine/Source/Editor/KismetCompiler/Private/KismetCompilerVMBackend.cpp
+++ b/Engine/Source/Editor/KismetCompiler/Private/KismetCompilerVMBackend.cpp
@@ -721,7 +721,16 @@ public:
 			}
 			else if (FLiteralTypeHelper::IsDouble(&Term->Type, CoerceProperty))
 			{
-				double Value = FCString::Atod(*(Term->Name));
+				double Value = 0.0;
+				if (Term->Type.bSerializeAsSinglePrecisionFloat)
+				{
+					Value = FCString::Atof(*(Term->Name));
+				}
+				else
+				{
+					Value = FCString::Atod(*(Term->Name));
+				}
+				
 				Writer << EX_DoubleConst;
 				Writer << Value;
 			}
diff --git a/Engine/Source/Runtime/Core/Public/UObject/UE5ReleaseStreamObjectVersion.h b/Engine/Source/Runtime/Core/Public/UObject/UE5ReleaseStreamObjectVersion.h
--- a/Engine/Source/Runtime/Core/Public/UObject/UE5ReleaseStreamObjectVersion.h
+++ b/Engine/Source/Runtime/Core/Public/UObject/UE5ReleaseStreamObjectVersion.h
@@ -118,6 +118,9 @@ struct CORE_API FUE5ReleaseStreamObjectVersion
 		// Storing inertia tensor as vec3 instead of matrix.
 		ChaosInertiaConvertedToVec3,
 
+		// For Blueprint real numbers, ensure that legacy float data is serialized as single-precision
+		SerializeFloatPinDefaultValuesAsSinglePrecision,
+
 		// -----<new versions can be added above this line>-------------------------------------------------
 		VersionPlusOne,
 		LatestVersion = VersionPlusOne - 1
diff --git a/Engine/Source/Runtime/Engine/Classes/EdGraph/EdGraphPin.h b/Engine/Source/Runtime/Engine/Classes/EdGraph/EdGraphPin.h
--- a/Engine/Source/Runtime/Engine/Classes/EdGraph/EdGraphPin.h
+++ b/Engine/Source/Runtime/Engine/Classes/EdGraph/EdGraphPin.h
@@ -136,6 +136,10 @@ public:
 	UPROPERTY()
 	uint8 bIsUObjectWrapper:1;
 
+	/** Set to true if the type was serialized prior to BlueprintPinsUseRealNumbers */
+	UPROPERTY()
+	uint8 bSerializeAsSinglePrecisionFloat:1;
+
 	FORCEINLINE bool IsContainer() const { return (ContainerType != EPinContainerType::None); }
 	FORCEINLINE bool IsArray() const { return (ContainerType == EPinContainerType::Array); }
 	FORCEINLINE bool IsSet() const { return (ContainerType == EPinContainerType::Set); }
@@ -150,6 +154,7 @@ public:
 		, bIsConst(false)
 		, bIsWeakPointer(false)
 		, bIsUObjectWrapper(false)
+		, bSerializeAsSinglePrecisionFloat(false)
 	{
 	}
 
@@ -164,6 +169,7 @@ public:
 		, bIsConst(false)
 		, bIsWeakPointer(false)
 		, bIsUObjectWrapper(false)
+		, bSerializeAsSinglePrecisionFloat(false)
 	{
 	}
 
@@ -176,7 +182,6 @@ public:
 			&& (ContainerType == Other.ContainerType)
 			&& (bIsReference == Other.bIsReference)
 			&& (bIsWeakPointer == Other.bIsWeakPointer)
-			//&& (bIsUObjectWrapper == Other.bIsUObjectWrapper)
 			&& (PinSubCategoryMemberReference == Other.PinSubCategoryMemberReference)
 			&& (bIsConst == Other.bIsConst);
 	}
@@ -196,6 +201,7 @@ public:
 		bIsReference = false;
 		bIsWeakPointer = false;
 		bIsUObjectWrapper = false;
+		bSerializeAsSinglePrecisionFloat = false;
 		bIsConst = false;
 	}
 
diff --git a/Engine/Source/Runtime/Engine/Private/EdGraph/EdGraphPin.cpp b/Engine/Source/Runtime/Engine/Private/EdGraph/EdGraphPin.cpp
--- a/Engine/Source/Runtime/Engine/Private/EdGraph/EdGraphPin.cpp
+++ b/Engine/Source/Runtime/Engine/Private/EdGraph/EdGraphPin.cpp
@@ -274,6 +274,8 @@ bool FEdGraphPinType::Serialize(FArchive& Ar)
 		Ar << bIsUObjectWrapperBool;
 	}
 
+	FName OldPinCategory = PinCategory;
+
 	if (Ar.IsLoading())
 	{
 		bIsReference = bIsReferenceBool;
@@ -314,6 +316,24 @@ bool FEdGraphPinType::Serialize(FArchive& Ar)
 		bIsUObjectWrapper = bIsUObjectWrapperBool;
 	}
 
+#if WITH_EDITOR
+	bool bSerializeAsSinglePrecisionFloatBool = bSerializeAsSinglePrecisionFloat;
+
+	if (Ar.CustomVer(FUE5ReleaseStreamObjectVersion::GUID) >= FUE5ReleaseStreamObjectVersion::SerializeFloatPinDefaultValuesAsSinglePrecision)
+	{
+		Ar << bSerializeAsSinglePrecisionFloatBool;
+	}
+	else
+	{
+		if (OldPinCategory == TEXT("float"))
+		{
+			bSerializeAsSinglePrecisionFloatBool = true;
+		}
+	}
+
+	bSerializeAsSinglePrecisionFloat = bSerializeAsSinglePrecisionFloatBool;
+#endif
+
 	return true;
 }
 
@@ -579,6 +599,7 @@ void TransferPersistentDataFromOldPin(UEdGraphPin& DestPin, T& SourcePin, const
 			DestPin.DefaultObject = SourcePin.DefaultObject;
 			DestPin.DefaultValue = MoveTempIfPossible(SourcePin.DefaultValue);
 			DestPin.DefaultTextValue = MoveTempIfPossible(SourcePin.DefaultTextValue);
+			DestPin.PinType.bSerializeAsSinglePrecisionFloat = SourcePin.PinType.bSerializeAsSinglePrecisionFloat;
 		}
 		else
 		{
