Fix ICMP/UDP ping handler holding an invalid object reference.

diff --git a/Engine/Source/Runtime/Engine/Private/Net/NetPing.cpp b/Engine/Source/Runtime/Engine/Private/Net/NetPing.cpp
--- a/Engine/Source/Runtime/Engine/Private/Net/NetPing.cpp
+++ b/Engine/Source/Runtime/Engine/Private/Net/NetPing.cpp
@@ -84,9 +84,9 @@ protected:
 	 * Base constructor
 	 *
 	 * @param InPingType	The QoS ping type this instance will handle
-	 * @param InOwner		The NetConnection which owns this ping type handler
+	 * @param InOwner		The FNetPing instance which owns this ping type handler
 	 */
-	FNetPingQoSBase(EPingType InPingType, UNetConnection* InOwner);
+	FNetPingQoSBase(EPingType InPingType, FNetPing* InOwner);
 
 	/**
 	 * Resets the ping type handler, discarding any active pings
@@ -100,13 +100,19 @@ protected:
 	 */
 	void PingResult(FIcmpEchoResult Result);
 
+public:
+	/**
+	 * Release an in-progress QoS ping handler from its Owner
+	 */
+	void Release();
+
 
 private:
 	/** The QoS ping type this instance handles */
 	const EPingType PingType;
 
-	/** The NetConnection which owns this ping type handler */
-	TObjectPtr<UNetConnection> Owner;
+	/** The FNetPing instance which owns this ping type handler */
+	FNetPing* Owner = nullptr;
 
 	/** Whether or not to discard the next ping value result (typically after a reset) */
 	bool bDiscardNextPing = false;
@@ -133,9 +139,9 @@ private:
 	/**
 	 * Base constructor
 	 *
-	 * @param InOwner	The NetConnection which owns this ping type handler
+	 * @param InOwner	The FNetPing instance which owns this ping type handler
 	 */
-	FNetPingICMP(UNetConnection* InOwner);
+	FNetPingICMP(FNetPing* InOwner);
 
 public:
 	/**
@@ -145,10 +151,10 @@ public:
 	 * to prevent conflicting results from multiple active ICMP pings.
 	 * This is a precautionary restriction, it might be possible to support this without conflicting results.
 	 *
-	 * @param InOwner	The NetConnection trying to obtain the ICMP ping handler.
-	 * @return			Returns a shared pointer to the ICMP ping handler, or nullptr if not available.
+	 * @param InNetConn		The NetConnection whose FNetPing instance is trying to obtain the ICMP ping handler.
+	 * @return				Returns a shared pointer to the ICMP ping handler, or nullptr if not available.
 	 */
-	static TSharedPtr<FNetPingICMP> GetNetPingICMP(UNetConnection* InOwner);
+	static TSharedPtr<FNetPingICMP> GetNetPingICMP(UNetConnection* InNetConn);
 
 	/**
 	 * Perform an ICMP ping.
@@ -173,18 +179,18 @@ private:
 	/**
 	 * Base constructor
 	 *
-	 * @param InOwner		The NetConnection which owns this ping type handler
+	 * @param InOwner		The FNetPing instance which owns this ping type handler
 	 */
-	FNetPingUDPQoS(UNetConnection* InOwner);
+	FNetPingUDPQoS(FNetPing* InOwner);
 
 public:
 	/**
 	 * Creates a UDPQoS ping handler.
 	 *
-	 * @param InOwner	The NetConnection that will own the UDPQoS ping handler.
+	 * @param InOwner	The FNetPing instance that will own the UDPQoS ping handler.
 	 * @return			Returns a shared pointer to the UDPQoS ping handler.
 	 */
-	static TSharedPtr<FNetPingUDPQoS> CreateNetPingUDPQoS(UNetConnection* InOwner);
+	static TSharedPtr<FNetPingUDPQoS> CreateNetPingUDPQoS(FNetPing* InOwner);
 
 	/**
 	 * Perform a UDPQoS ping.
@@ -285,6 +291,19 @@ FNetPing::FNetPing(UNetConnection* InOwner)
 {
 }
 
+FNetPing::~FNetPing()
+{
+	if (NetPingICMP.IsValid())
+	{
+		NetPingICMP->Release();
+	}
+
+	if (NetPingUDPQoS.IsValid())
+	{
+		NetPingUDPQoS->Release();
+	}
+}
+
 TPimplPtr<FNetPing> FNetPing::CreateNetPing(UNetConnection* InOwner)
 {
 	if (const UNetDriver* Driver = InOwner != nullptr ? InOwner->Driver : nullptr)
@@ -426,7 +445,7 @@ void FNetPing::Init()
 
 			PingAverageStorageIdx[PingTypeIdx] = GetNewStorageIdx(PingAverageTypes[PingTypeIdx]);
 
-			NetPingUDPQoS = FNetPingUDPQoS::CreateNetPingUDPQoS(ServerConn);
+			NetPingUDPQoS = FNetPingUDPQoS::CreateNetPingUDPQoS(this);
 		}
 	}
 
@@ -1005,12 +1024,17 @@ namespace Private
 /**
  * FNetPingQoSBase
  */
-FNetPingQoSBase::FNetPingQoSBase(EPingType InPingType, UNetConnection* InOwner)
+FNetPingQoSBase::FNetPingQoSBase(EPingType InPingType, FNetPing* InOwner)
 	: PingType(InPingType)
 	, Owner(InOwner)
 {
 }
 
+void FNetPingQoSBase::Release()
+{
+	Owner = nullptr;
+}
+
 void FNetPingQoSBase::Reset()
 {
 	if (bPingInProgress)
@@ -1025,22 +1049,19 @@ void FNetPingQoSBase::PingResult(FIcmpEchoResult Result)
 {
 	bPingInProgress = false;
 
-	if (!bDiscardNextPing)
+	if (!bDiscardNextPing && Owner != nullptr)
 	{
-		if (FNetPing* NetPing = (Owner ? Owner->GetNetPing() : nullptr))
+		if (Result.Status == EIcmpResponseStatus::Success)
 		{
-			if (Result.Status == EIcmpResponseStatus::Success)
-			{
-				const double CurPing = static_cast<double>(Result.Time);
-				const double ApproxTimeSeconds = PingTimestamp + CurPing;
+			const double CurPing = static_cast<double>(Result.Time);
+			const double ApproxTimeSeconds = PingTimestamp + CurPing;
 
-				NetPing->UpdatePing(PingType, ApproxTimeSeconds, CurPing);
-			}
-			else
-			{
-				// Not always a timeout, but treat all failures as such anyway.
-				NetPing->UpdatePingTimeout(PingType);
-			}
+			Owner->UpdatePing(PingType, ApproxTimeSeconds, CurPing);
+		}
+		else
+		{
+			// Not always a timeout, but treat all failures as such anyway.
+			Owner->UpdatePingTimeout(PingType);
 		}
 	}
 
@@ -1051,34 +1072,37 @@ void FNetPingQoSBase::PingResult(FIcmpEchoResult Result)
  * FNetPingICMP
  */
 
-FNetPingICMP::FNetPingICMP(UNetConnection* InOwner)
+FNetPingICMP::FNetPingICMP(FNetPing* InOwner)
 	: FNetPingQoSBase(EPingType::ICMP, InOwner)
 {
 }
 
-TSharedPtr<FNetPingICMP> FNetPingICMP::GetNetPingICMP(UNetConnection* InOwner)
+TSharedPtr<FNetPingICMP> FNetPingICMP::GetNetPingICMP(UNetConnection* InNetConn)
 {
 	static TSharedPtr<FNetPingICMP> NetPingSingleton;
 
-	FName NetDriverName = (InOwner != nullptr && InOwner->Driver != nullptr ? InOwner->Driver->NetDriverName : NAME_None);
+	if (FNetPing* InOwner = InNetConn != nullptr ? InNetConn->GetNetPing() : nullptr)
+	{
+		FName NetDriverName = (InNetConn->Driver != nullptr ? InNetConn->Driver->NetDriverName : NAME_None);
 
-	// The expected reference count, if no NetDriver has a handle to NetPingSingleton
-	const int32 IsUniqueRefCount = 1 + (int32)(NetPingSingleton.IsValid() && NetPingSingleton->bPingInProgress);
+		// The expected reference count, if no NetDriver has a handle to NetPingSingleton
+		const int32 IsUniqueRefCount = 1 + (int32)(NetPingSingleton.IsValid() && NetPingSingleton->bPingInProgress);
 
-	// Only provide an FNetPingICMP interface for GameNetDriver's, when no other NetDriver has a handle to one.
-	if ((NetDriverName == NAME_GameNetDriver || NetDriverName == NAME_PendingNetDriver) &&
-		(!NetPingSingleton.IsValid() || NetPingSingleton.GetSharedReferenceCount() == IsUniqueRefCount))
-	{
-		if (!NetPingSingleton.IsValid())
+		// Only provide an FNetPingICMP interface for GameNetDriver's, when no other NetDriver has a handle to one.
+		if ((NetDriverName == NAME_GameNetDriver || NetDriverName == NAME_PendingNetDriver) &&
+			(!NetPingSingleton.IsValid() || NetPingSingleton.GetSharedReferenceCount() == IsUniqueRefCount))
 		{
-			NetPingSingleton = MakeShared<FNetPingICMP>(InOwner);
-		}
-		else
-		{
-			NetPingSingleton->Reset();
-		}
+			if (!NetPingSingleton.IsValid())
+			{
+				NetPingSingleton = MakeShared<FNetPingICMP>(InOwner);
+			}
+			else
+			{
+				NetPingSingleton->Reset();
+			}
 
-		return NetPingSingleton;
+			return NetPingSingleton;
+		}
 	}
 
 	return TSharedPtr<FNetPingICMP>();
@@ -1107,12 +1131,12 @@ void FNetPingICMP::Ping(double CurTimeSeconds, FString PingAddress)
  * FNetPingUDPQoS
  */
 
-FNetPingUDPQoS::FNetPingUDPQoS(UNetConnection* InOwner)
+FNetPingUDPQoS::FNetPingUDPQoS(FNetPing* InOwner)
 	: FNetPingQoSBase(EPingType::UDPQoS, InOwner)
 {
 }
 
-TSharedPtr<FNetPingUDPQoS> FNetPingUDPQoS::CreateNetPingUDPQoS(UNetConnection* InOwner)
+TSharedPtr<FNetPingUDPQoS> FNetPingUDPQoS::CreateNetPingUDPQoS(FNetPing* InOwner)
 {
 	return MakeShared<FNetPingUDPQoS>(InOwner);
 }
diff --git a/Engine/Source/Runtime/Engine/Public/Net/NetPing.h b/Engine/Source/Runtime/Engine/Public/Net/NetPing.h
--- a/Engine/Source/Runtime/Engine/Public/Net/NetPing.h
+++ b/Engine/Source/Runtime/Engine/Public/Net/NetPing.h
@@ -214,6 +214,8 @@ private:
 	FNetPing(UNetConnection* InOwner);
 
 public:
+	~FNetPing();
+
 	/**
 	 * Creates a new FNetPing instance, owned by the specified NetConnection.
 	 * NOTE: Presently only enabled for GameNetDriver, but can support other NetDriver's (except ICMP ping).
