WorldPartition: sub WorldPartition support (Game Worlds only for now) (no DataLayer support yet as this will come in after some other DataLayer refactors)

#rnx

diff --git a/Engine/Source/Runtime/Engine/Private/Components/WorldPartitionStreamingSourceComponent.cpp b/Engine/Source/Runtime/Engine/Private/Components/WorldPartitionStreamingSourceComponent.cpp
--- a/Engine/Source/Runtime/Engine/Private/Components/WorldPartitionStreamingSourceComponent.cpp
+++ b/Engine/Source/Runtime/Engine/Private/Components/WorldPartitionStreamingSourceComponent.cpp
@@ -37,10 +37,9 @@ void UWorldPartitionStreamingSourceComponent::OnRegister()
 	}
 #endif
 
-	if (UWorldPartition* WorldPartition = World->GetWorldPartition())
-	{
-		WorldPartition->RegisterStreamingSourceProvider(this);
-	}
+	UWorldPartitionSubsystem* WorldPartitionSubsystem = GetWorld()->GetSubsystem<UWorldPartitionSubsystem>();
+	check(WorldPartitionSubsystem);
+	WorldPartitionSubsystem->RegisterStreamingSourceProvider(this);
 }
 
 void UWorldPartitionStreamingSourceComponent::OnUnregister()
@@ -56,10 +55,9 @@ void UWorldPartitionStreamingSourceComponent::OnUnregister()
 	}
 #endif
 
-	if (UWorldPartition* WorldPartition = World->GetWorldPartition())
-	{
-		verify(WorldPartition->UnregisterStreamingSourceProvider(this));
-	}
+	UWorldPartitionSubsystem* WorldPartitionSubsystem = GetWorld()->GetSubsystem<UWorldPartitionSubsystem>();
+	check(WorldPartitionSubsystem);
+	verify(WorldPartitionSubsystem->UnregisterStreamingSourceProvider(this));
 }
 
 bool UWorldPartitionStreamingSourceComponent::GetStreamingSource(FWorldPartitionStreamingSource& OutStreamingSource)
diff --git a/Engine/Source/Runtime/Engine/Private/Level.cpp b/Engine/Source/Runtime/Engine/Private/Level.cpp
--- a/Engine/Source/Runtime/Engine/Private/Level.cpp
+++ b/Engine/Source/Runtime/Engine/Private/Level.cpp
@@ -2415,11 +2415,20 @@ void ULevel::OnLevelLoaded()
 
 	// 1. Cook commandlet does it's own UWorldPartition::Initialize call in FWorldPartitionCookPackageSplitter::GetGenerateList
 	// 2. Do not Initialize if World doesn't have a UWorldPartitionSubsystem (Known case is when WorldType == EWorldType::Inactive)
-	if (!IsRunningCookCommandlet() && IsPersistentLevel() && GetWorld()->HasSubsystem<UWorldPartitionSubsystem>())
+	if (!IsRunningCookCommandlet() && GetWorld()->HasSubsystem<UWorldPartitionSubsystem>())
 	{
-		if (UWorldPartition* WorldPartition = GetWorldPartition())
+		if (IsPersistentLevel() || GetWorld()->IsGameWorld())
 		{
-			WorldPartition->Initialize(GetWorld(), FTransform::Identity);
+			if (UWorldPartition* WorldPartition = GetWorldPartition())
+			{
+				FTransform Transform = FTransform::Identity;
+				if (ULevelStreaming* LevelStreaming = FLevelUtils::FindStreamingLevel(this))
+				{
+					Transform = LevelStreaming->LevelTransform;
+				}
+
+				WorldPartition->Initialize(GetWorld(), Transform);
+			}
 		}
 	}
 }
diff --git a/Engine/Source/Runtime/Engine/Private/WorldPartition/DataLayer/WorldDataLayers.cpp b/Engine/Source/Runtime/Engine/Private/WorldPartition/DataLayer/WorldDataLayers.cpp
--- a/Engine/Source/Runtime/Engine/Private/WorldPartition/DataLayer/WorldDataLayers.cpp
+++ b/Engine/Source/Runtime/Engine/Private/WorldPartition/DataLayer/WorldDataLayers.cpp
@@ -780,7 +780,7 @@ void AWorldDataLayers::PostLoad()
 
 	GetLevel()->ConditionalPostLoad();
 
-	GetWorld()->SetWorldDataLayers(this);
+	GetTypedOuter<UWorld>()->SetWorldDataLayers(this);
 
 #if WITH_EDITOR
 	// Remove all Editor Data Layers when cooking or when in a game world
diff --git a/Engine/Source/Runtime/Engine/Private/WorldPartition/HLOD/HLODSubsystem.cpp b/Engine/Source/Runtime/Engine/Private/WorldPartition/HLOD/HLODSubsystem.cpp
--- a/Engine/Source/Runtime/Engine/Private/WorldPartition/HLOD/HLODSubsystem.cpp
+++ b/Engine/Source/Runtime/Engine/Private/WorldPartition/HLOD/HLODSubsystem.cpp
@@ -8,6 +8,7 @@
 #include "WorldPartition/WorldPartitionSubsystem.h"
 #include "WorldPartition/WorldPartitionRuntimeCell.h"
 #include "WorldPartition/WorldPartitionRuntimeHash.h"
+#include "WorldPartition/WorldPartitionLevelStreamingDynamic.h"
 #include "UObject/UObjectHash.h"
 #include "UObject/UObjectGlobals.h"
 #include "Engine/Engine.h"
@@ -15,6 +16,7 @@
 #include "Engine/World.h"
 #include "EngineModule.h"
 #include "EngineUtils.h"
+#include "LevelUtils.h"
 #include "Components/StaticMeshComponent.h"
 
 #define LOCTEXT_NAMESPACE "HLODSubsystem"
@@ -56,6 +58,24 @@ static TAutoConsoleVariable<int32> CVarHLODWarmupVTSizeClamp(
 	ECVF_Default
 );
 
+namespace FHLODSubsystem
+{
+	UWorldPartition* GetWorldPartition(AWorldPartitionHLOD* InWorldPartitionHLOD)
+	{
+		// Alwaysloaded Cell level will have a WorldPartition
+		if (UWorldPartition* WorldPartition = InWorldPartitionHLOD->GetLevel()->GetWorldPartition())
+		{
+			return WorldPartition;
+		} // If not find it through the cell
+		else if (UWorldPartitionLevelStreamingDynamic* LevelStreaming = Cast<UWorldPartitionLevelStreamingDynamic>(FLevelUtils::FindStreamingLevel(InWorldPartitionHLOD->GetLevel())))
+		{
+			return LevelStreaming->GetWorldPartitionRuntimeCell()->GetTypedOuter<UWorldPartition>();
+		}
+		
+		return nullptr;
+	}
+};
+
 UHLODSubsystem::UHLODSubsystem()
 	: UWorldSubsystem()
 {
@@ -79,15 +99,19 @@ FAutoConsoleCommand UHLODSubsystem::EnableHLODCommand(
 			if (World && World->IsGameWorld())
 			{
 				UHLODSubsystem* HLODSubSystem = World->GetSubsystem<UHLODSubsystem>();
-				for (const auto& CellHLODMapping : HLODSubSystem->CellsData)
+				for (const auto& KeyValuePair : HLODSubSystem->WorldPartitionsHLODRuntimeData)
 				{
-					const FCellData& CellData = CellHLODMapping.Value;
-					bool bIsHLODVisible = UHLODSubsystem::WorldPartitionHLODEnabled && !CellData.bIsCellVisible;
-					for (AWorldPartitionHLOD* HLODActor : CellData.LoadedHLODs)
+					for (const auto& CellHLODMapping : KeyValuePair.Value.CellsData)
 					{
-						HLODActor->SetVisibility(bIsHLODVisible);
+						const FCellData& CellData = CellHLODMapping.Value;
+						bool bIsHLODVisible = UHLODSubsystem::WorldPartitionHLODEnabled && !CellData.bIsCellVisible;
+						for (AWorldPartitionHLOD* HLODActor : CellData.LoadedHLODs)
+						{
+							HLODActor->SetVisibility(bIsHLODVisible);
+						}
 					}
 				}
+				
 			}
 		}
 	})
@@ -109,13 +133,10 @@ void UHLODSubsystem::Initialize(FSubsystemCollectionBase& Collection)
 
 	if (World->IsGameWorld())
 	{
-		if (UWorldPartition* WorldPartition = GetWorld()->GetWorldPartition())
-		{
-			GetWorld()->OnWorldPartitionInitialized().AddUObject(this, &UHLODSubsystem::OnWorldPartitionInitialized);
-			GetWorld()->OnWorldPartitionUninitialized().AddUObject(this, &UHLODSubsystem::OnWorldPartitionUninitialized);
+		GetWorld()->OnWorldPartitionInitialized().AddUObject(this, &UHLODSubsystem::OnWorldPartitionInitialized);
+		GetWorld()->OnWorldPartitionUninitialized().AddUObject(this, &UHLODSubsystem::OnWorldPartitionUninitialized);
 
-			SceneViewExtension = FSceneViewExtensions::NewExtension<FHLODResourcesResidencySceneViewExtension>(World);
-		}
+		SceneViewExtension = FSceneViewExtensions::NewExtension<FHLODResourcesResidencySceneViewExtension>(World);
 	}
 }
 
@@ -129,31 +150,37 @@ void UHLODSubsystem::Deinitialize()
 
 void UHLODSubsystem::OnWorldPartitionInitialized(UWorldPartition* InWorldPartition)
 {
-	check(InWorldPartition == GetWorld()->GetWorldPartition());
-	check(CellsData.IsEmpty());
+	check(!WorldPartitionsHLODRuntimeData.Contains(InWorldPartition));
 
+	FWorldPartitionHLODRuntimeData& WorldPartitionHLODRuntimeData = WorldPartitionsHLODRuntimeData.Add(InWorldPartition, FWorldPartitionHLODRuntimeData());
+	
 	TSet<const UWorldPartitionRuntimeCell*> StreamingCells;
 	InWorldPartition->RuntimeHash->GetAllStreamingCells(StreamingCells, /*bAllDataLayers*/ true);
 
 	// Build cell to HLOD mapping
 	for (const UWorldPartitionRuntimeCell* Cell : StreamingCells)
 	{
-		CellsData.Emplace(Cell->GetFName());
+		WorldPartitionHLODRuntimeData.CellsData.Emplace(Cell->GetFName());
 	}
 }
 
 void UHLODSubsystem::OnWorldPartitionUninitialized(UWorldPartition* InWorldPartition)
 {
-	check(InWorldPartition == GetWorld()->GetWorldPartition());
-	CellsData.Reset();
+	check(WorldPartitionsHLODRuntimeData.Contains(InWorldPartition));
+	WorldPartitionsHLODRuntimeData.Remove(InWorldPartition);
 }
 
 void UHLODSubsystem::RegisterHLODActor(AWorldPartitionHLOD* InWorldPartitionHLOD)
 {
 	TRACE_CPUPROFILER_EVENT_SCOPE(UHLODSubsystem::RegisterHLODActor);
 
+	const UWorldPartition* WorldPartition = FHLODSubsystem::GetWorldPartition(InWorldPartitionHLOD);
+	check(WorldPartition);
+
+	FWorldPartitionHLODRuntimeData& WorldPartitionHLODRuntimeData = WorldPartitionsHLODRuntimeData.FindChecked(WorldPartition);
+
 	const FName CellName = InWorldPartitionHLOD->GetSourceCellName();
-	FCellData* CellData = CellsData.Find(CellName);
+	FCellData* CellData = WorldPartitionHLODRuntimeData.CellsData.Find(CellName);
 
 #if WITH_EDITOR
 	UE_LOG(LogHLODSubsystem, Verbose, TEXT("Registering HLOD %s (%s) for cell %s"), *InWorldPartitionHLOD->GetActorLabel(), *InWorldPartitionHLOD->GetActorGuid().ToString(), *CellName.ToString());
@@ -175,8 +202,13 @@ void UHLODSubsystem::UnregisterHLODActor(AWorldPartitionHLOD* InWorldPartitionHL
 {
 	TRACE_CPUPROFILER_EVENT_SCOPE(UHLODSubsystem::UnregisterHLODActor);
 
+	const UWorldPartition* WorldPartition = FHLODSubsystem::GetWorldPartition(InWorldPartitionHLOD);
+	check(WorldPartition);
+
+	FWorldPartitionHLODRuntimeData& WorldPartitionHLODRuntimeData = WorldPartitionsHLODRuntimeData.FindChecked(WorldPartition);
+
 	const FName CellName = InWorldPartitionHLOD->GetSourceCellName();
-	FCellData* CellData = CellsData.Find(CellName);
+	FCellData* CellData = WorldPartitionHLODRuntimeData.CellsData.Find(CellName);
 
 #if WITH_EDITOR
 	UE_LOG(LogHLODSubsystem, Verbose, TEXT("Unregistering HLOD %s (%s) for cell %s"), *InWorldPartitionHLOD->GetActorLabel(), *InWorldPartitionHLOD->GetActorGuid().ToString(), *CellName.ToString());
@@ -191,7 +223,10 @@ void UHLODSubsystem::UnregisterHLODActor(AWorldPartitionHLOD* InWorldPartitionHL
 
 void UHLODSubsystem::OnCellShown(const UWorldPartitionRuntimeCell* InCell)
 {
-	FCellData& CellData = CellsData.FindChecked(InCell->GetFName());
+	const UWorldPartition* WorldPartition = InCell->GetTypedOuter<UWorldPartition>();
+	FWorldPartitionHLODRuntimeData& WorldPartitionHLODRuntimeData = WorldPartitionsHLODRuntimeData.FindChecked(WorldPartition);
+
+	FCellData& CellData = WorldPartitionHLODRuntimeData.CellsData.FindChecked(InCell->GetFName());
 	CellData.bIsCellVisible = true;
 
 #if WITH_EDITOR
@@ -209,7 +244,10 @@ void UHLODSubsystem::OnCellShown(const UWorldPartitionRuntimeCell* InCell)
 
 void UHLODSubsystem::OnCellHidden(const UWorldPartitionRuntimeCell* InCell)
 {
-	FCellData& CellData = CellsData.FindChecked(InCell->GetFName());
+	const UWorldPartition* WorldPartition = InCell->GetTypedOuter<UWorldPartition>();
+	FWorldPartitionHLODRuntimeData& WorldPartitionHLODRuntimeData = WorldPartitionsHLODRuntimeData.FindChecked(WorldPartition);
+
+	FCellData& CellData = WorldPartitionHLODRuntimeData.CellsData.FindChecked(InCell->GetFName());
 	CellData.bIsCellVisible = false;
 
 #if WITH_EDITOR
@@ -348,7 +386,10 @@ bool UHLODSubsystem::RequestUnloading(const UWorldPartitionRuntimeCell* InCell)
 		return true;
 	}
 
-	FCellData& CellData = CellsData.FindChecked(InCell->GetFName());
+	const UWorldPartition* WorldPartition = InCell->GetTypedOuter<UWorldPartition>();
+	FWorldPartitionHLODRuntimeData& WorldPartitionHLODRuntimeData = WorldPartitionsHLODRuntimeData.FindChecked(WorldPartition);
+
+	FCellData& CellData = WorldPartitionHLODRuntimeData.CellsData.FindChecked(InCell->GetFName());
 
 	// If cell wasn't even visible yet or doesn't have HLOD actors, skip warmup.
 	bPerformWarmpup = !CellData.LoadedHLODs.IsEmpty() && CellData.bIsCellVisible;
@@ -375,7 +416,7 @@ bool UHLODSubsystem::RequestUnloading(const UWorldPartitionRuntimeCell* InCell)
 		// Warmup will be triggered in the next BeginRenderView() call, at which point the frame number will have been incremented.
 		CellData.WarmupStartFrame = CurrentFrameNumber + 1; 
 		CellData.WarmupEndFrame = CellData.WarmupStartFrame + CVarHLODWarmupNumFrames.GetValueOnGameThread();
-		CellsToWarmup.Add(&CellData);
+		WorldPartitionHLODRuntimeData.CellsToWarmup.Add(&CellData);
 	}
 
 	// Test if warmup is completed
@@ -391,16 +432,19 @@ bool UHLODSubsystem::RequestUnloading(const UWorldPartitionRuntimeCell* InCell)
 
 void UHLODSubsystem::OnBeginRenderViews(const FSceneViewFamily& InViewFamily)
 {
-	for (TSet<FCellData*>::TIterator CellIt(CellsToWarmup); CellIt; ++CellIt)
+	for (auto& KeyValuePair : WorldPartitionsHLODRuntimeData)
 	{
-		FCellData* CellPendingUnload = *CellIt;
-	
-		MakeRenderResourcesResident(*CellPendingUnload, InViewFamily);
-
-		// Stop processing this cell if warmup is done.
-		if (InViewFamily.FrameNumber >= CellPendingUnload->WarmupEndFrame)
+		for (TSet<FCellData*>::TIterator CellIt(KeyValuePair.Value.CellsToWarmup); CellIt; ++CellIt)
 		{
-			CellIt.RemoveCurrent();
+			FCellData* CellPendingUnload = *CellIt;
+
+			MakeRenderResourcesResident(*CellPendingUnload, InViewFamily);
+
+			// Stop processing this cell if warmup is done.
+			if (InViewFamily.FrameNumber >= CellPendingUnload->WarmupEndFrame)
+			{
+				CellIt.RemoveCurrent();
+			}
 		}
 	}
 }
diff --git a/Engine/Source/Runtime/Engine/Private/WorldPartition/WorldPartition.cpp b/Engine/Source/Runtime/Engine/Private/WorldPartition/WorldPartition.cpp
--- a/Engine/Source/Runtime/Engine/Private/WorldPartition/WorldPartition.cpp
+++ b/Engine/Source/Runtime/Engine/Private/WorldPartition/WorldPartition.cpp
@@ -409,8 +409,11 @@ void UWorldPartition::Initialize(UWorld* InWorld, const FTransform& InTransform)
 
 	RegisterDelegates();
 
-	AWorldPartitionReplay::Initialize(World);
-	
+	if (IsMainWorldPartition())
+	{
+		AWorldPartitionReplay::Initialize(World);
+	}
+
 #if WITH_EDITOR
 	const bool bIsGame = IsRunningGame();
 	const bool bIsEditor = !World->IsGameWorld();
@@ -427,7 +430,7 @@ void UWorldPartition::Initialize(UWorld* InWorld, const FTransform& InTransform)
 	}
 	else if (bIsEditor)
 	{
-		CreateOrRepairWorldPartition(World->GetWorldSettings());
+		CreateOrRepairWorldPartition(OuterWorld->GetWorldSettings());
 
 		check(!StreamingPolicy);
 		check(EditorHash);
@@ -448,7 +451,7 @@ void UWorldPartition::Initialize(UWorld* InWorld, const FTransform& InTransform)
 		// - World Partition map template (New Level)
 		// - PIE World Travel
 		FString SourceWorldPath, RemappedWorldPath;
-		const bool bIsInstanced = World->GetSoftObjectPathMapping(SourceWorldPath, RemappedWorldPath);
+		const bool bIsInstanced = OuterWorld->GetSoftObjectPathMapping(SourceWorldPath, RemappedWorldPath);
 	
 		if (bIsInstanced)
 		{
@@ -485,13 +488,16 @@ void UWorldPartition::Initialize(UWorld* InWorld, const FTransform& InTransform)
 		}
 	}
 
-	// Make sure to preload only AWorldDataLayers actor first (ShouldActorBeLoadedByEditorCells requires it)
-	for (UActorDescContainer::TIterator<> ActorDescIterator(this); ActorDescIterator; ++ActorDescIterator)
+	if (bIsEditor)
 	{
-		if (ActorDescIterator->GetActorClass()->IsChildOf<AWorldDataLayers>())
+		// Make sure to preload only AWorldDataLayers actor first (ShouldActorBeLoadedByEditorCells requires it)
+		for (UActorDescContainer::TIterator<> ActorDescIterator(this); ActorDescIterator; ++ActorDescIterator)
 		{
-			WorldDataLayersActor = FWorldPartitionReference(this, ActorDescIterator->GetGuid());
-			break;
+			if (ActorDescIterator->GetActorClass()->IsChildOf<AWorldDataLayers>())
+			{
+				WorldDataLayersActor = FWorldPartitionReference(this, ActorDescIterator->GetGuid());
+				break;
+			}
 		}
 	}
 
@@ -541,7 +547,7 @@ void UWorldPartition::Initialize(UWorld* InWorld, const FTransform& InTransform)
 		}
 
 		// Apply remapping of Persistent Level's SoftObjectPaths
-		FWorldPartitionLevelHelper::RemapLevelSoftObjectPaths(World->PersistentLevel, this);
+		FWorldPartitionLevelHelper::RemapLevelSoftObjectPaths(OuterWorld->PersistentLevel, this);
 	}
 #endif
 
@@ -551,16 +557,6 @@ PRAGMA_DISABLE_DEPRECATION_WARNINGS
 PRAGMA_ENABLE_DEPRECATION_WARNINGS
 }
 
-void UWorldPartition::RegisterStreamingSourceProvider(IWorldPartitionStreamingSourceProvider* StreamingSource)
-{
-	StreamingSourceProviders.Add(StreamingSource);
-}
-
-bool UWorldPartition::UnregisterStreamingSourceProvider(IWorldPartitionStreamingSourceProvider* StreamingSource)
-{
-	return !!StreamingSourceProviders.Remove(StreamingSource);
-}
-
 void UWorldPartition::Uninitialize()
 {
 	if (IsInitialized())
@@ -629,6 +625,12 @@ bool UWorldPartition::IsInitialized() const
 	return InitState == EWorldPartitionInitState::Initialized;
 }
 
+bool UWorldPartition::IsMainWorldPartition() const
+{
+	check(World);
+	return World == GetTypedOuter<UWorld>();
+}
+
 void UWorldPartition::OnPostBugItGoCalled(const FVector& Loc, const FRotator& Rot)
 {
 #if WITH_EDITOR
@@ -645,27 +647,30 @@ void UWorldPartition::RegisterDelegates()
 {
 	check(World); 
 
-#if WITH_EDITOR
-	if (GEditor && !IsTemplate() && !World->IsGameWorld())
+	if (IsMainWorldPartition())
 	{
-		FEditorDelegates::PreBeginPIE.AddUObject(this, &UWorldPartition::OnPreBeginPIE);
-		FEditorDelegates::PrePIEEnded.AddUObject(this, &UWorldPartition::OnPrePIEEnded);
-		FEditorDelegates::CancelPIE.AddUObject(this, &UWorldPartition::OnCancelPIE);
-		FGameDelegates::Get().GetEndPlayMapDelegate().AddUObject(this, &UWorldPartition::OnEndPlay);
+#if WITH_EDITOR
+		if (GEditor && !IsTemplate() && !World->IsGameWorld())
+		{
+			FEditorDelegates::PreBeginPIE.AddUObject(this, &UWorldPartition::OnPreBeginPIE);
+			FEditorDelegates::PrePIEEnded.AddUObject(this, &UWorldPartition::OnPrePIEEnded);
+			FEditorDelegates::CancelPIE.AddUObject(this, &UWorldPartition::OnCancelPIE);
+			FGameDelegates::Get().GetEndPlayMapDelegate().AddUObject(this, &UWorldPartition::OnEndPlay);
 
-		FCoreUObjectDelegates::PostReachabilityAnalysis.AddUObject(this, &UWorldPartition::OnGCPostReachabilityAnalysis);
+			FCoreUObjectDelegates::PostReachabilityAnalysis.AddUObject(this, &UWorldPartition::OnGCPostReachabilityAnalysis);
 
-		GEditor->OnPostBugItGoCalled().AddUObject(this, &UWorldPartition::OnPostBugItGoCalled);
-	}
+			GEditor->OnPostBugItGoCalled().AddUObject(this, &UWorldPartition::OnPostBugItGoCalled);
+		}
 #endif
 
-	if (World->IsGameWorld())
-	{
-		World->OnWorldMatchStarting.AddUObject(this, &UWorldPartition::OnWorldMatchStarting);
+		if (World->IsGameWorld())
+		{
+			World->OnWorldMatchStarting.AddUObject(this, &UWorldPartition::OnWorldMatchStarting);
 
 #if !UE_BUILD_SHIPPING
-		FCoreDelegates::OnGetOnScreenMessages.AddUObject(this, &UWorldPartition::GetOnScreenMessages);
+			FCoreDelegates::OnGetOnScreenMessages.AddUObject(this, &UWorldPartition::GetOnScreenMessages);
 #endif
+		}
 	}
 }
 
@@ -673,30 +678,33 @@ void UWorldPartition::UnregisterDelegates()
 {
 	check(World);
 
-#if WITH_EDITOR
-	if (GEditor && !IsTemplate() && !World->IsGameWorld())
+	if (IsMainWorldPartition())
 	{
-		FEditorDelegates::PreBeginPIE.RemoveAll(this);
-		FEditorDelegates::PrePIEEnded.RemoveAll(this);
-		FEditorDelegates::CancelPIE.RemoveAll(this);
-		FGameDelegates::Get().GetEndPlayMapDelegate().RemoveAll(this);
-
-		if (!IsEngineExitRequested())
+#if WITH_EDITOR
+		if (GEditor && !IsTemplate() && !World->IsGameWorld())
 		{
-			FCoreUObjectDelegates::PostReachabilityAnalysis.RemoveAll(this);
-		}
+			FEditorDelegates::PreBeginPIE.RemoveAll(this);
+			FEditorDelegates::PrePIEEnded.RemoveAll(this);
+			FEditorDelegates::CancelPIE.RemoveAll(this);
+			FGameDelegates::Get().GetEndPlayMapDelegate().RemoveAll(this);
 
-		GEditor->OnPostBugItGoCalled().RemoveAll(this);
-	}
+			if (!IsEngineExitRequested())
+			{
+				FCoreUObjectDelegates::PostReachabilityAnalysis.RemoveAll(this);
+			}
+
+			GEditor->OnPostBugItGoCalled().RemoveAll(this);
+		}
 #endif
 
-	if (World->IsGameWorld())
-	{
-		World->OnWorldMatchStarting.RemoveAll(this);
+		if (World->IsGameWorld())
+		{
+			World->OnWorldMatchStarting.RemoveAll(this);
 
 #if !UE_BUILD_SHIPPING
-		FCoreDelegates::OnGetOnScreenMessages.RemoveAll(this);
+			FCoreDelegates::OnGetOnScreenMessages.RemoveAll(this);
 #endif
+		}
 	}
 }
 
@@ -1135,6 +1143,7 @@ void UWorldPartition::OnActorDescRegistered(const FWorldPartitionActorDesc& Acto
 	check(Actor);
 	ApplyActorTransform(Actor, InstanceTransform);
 	Actor->GetLevel()->AddLoadedActor(Actor);
+
 }
 
 void UWorldPartition::OnActorDescUnregistered(const FWorldPartitionActorDesc& ActorDesc) 
@@ -1352,7 +1361,13 @@ bool UWorldPartition::CanDrawRuntimeHash() const
 	return GetWorld()->IsGameWorld() || UWorldPartition::IsSimulating();
 }
 
-void UWorldPartition::DrawRuntimeHash2D(class UCanvas* Canvas, const FVector2D& PartitionCanvasSize, FVector2D& Offset)
+FVector2D UWorldPartition::GetDrawRuntimeHash2DDesiredFootprint(const FVector2D& CanvasSize)
+{
+	check(CanDrawRuntimeHash());
+	return StreamingPolicy->GetDrawRuntimeHash2DDesiredFootprint(CanvasSize);
+}
+
+void UWorldPartition::DrawRuntimeHash2D(class UCanvas* Canvas, const FVector2D& PartitionCanvasSize, const FVector2D& Offset)
 {
 	check(CanDrawRuntimeHash());
 	StreamingPolicy->DrawRuntimeHash2D(Canvas, PartitionCanvasSize, Offset);
diff --git a/Engine/Source/Runtime/Engine/Private/WorldPartition/WorldPartitionLevelStreamingPolicy.cpp b/Engine/Source/Runtime/Engine/Private/WorldPartition/WorldPartitionLevelStreamingPolicy.cpp
--- a/Engine/Source/Runtime/Engine/Private/WorldPartition/WorldPartitionLevelStreamingPolicy.cpp
+++ b/Engine/Source/Runtime/Engine/Private/WorldPartition/WorldPartitionLevelStreamingPolicy.cpp
@@ -124,21 +124,21 @@ void UWorldPartitionLevelStreamingPolicy::RemapSoftObjectPath(FSoftObjectPath& O
 	{
 		if (!SrcObjectPath.GetSubPathString().IsEmpty())
 		{
-			UWorld* World = WorldPartition->GetWorld();
-			const FString PackagePath = UWorldPartitionLevelStreamingPolicy::GetCellPackagePath(*CellName, World);
+			UWorld* OuterWorld = WorldPartition->GetTypedOuter<UWorld>();
+			const FString PackagePath = UWorldPartitionLevelStreamingPolicy::GetCellPackagePath(*CellName, OuterWorld);
 			FString PrefixPath;
 			if (IsRunningCookCommandlet())
 			{
 				//@todo_ow: Temporary workaround. This information should be provided by the COTFS
-				const UPackage* Package = GetOuterUWorldPartition()->GetWorld()->GetPackage();
+				const UPackage* Package = OuterWorld->GetPackage();
 				PrefixPath = FString::Printf(TEXT("%s/%s/_Generated_"), *FPackageName::GetLongPackagePath(Package->GetPathName()), *FPackageName::GetShortName(Package->GetName()));
 			}
 
 			// Use the WorldPartition world name here instead of using the world name from the path to support converting level instance paths to main world paths.
-			ObjectPath.SetAssetPathName(FName(*FString::Printf(TEXT("%s%s.%s"), *PrefixPath, *PackagePath, *World->GetName())));
+			ObjectPath.SetAssetPathName(FName(*FString::Printf(TEXT("%s%s.%s"), *PrefixPath, *PackagePath, *OuterWorld->GetName())));
 			ObjectPath.SetSubPathString(SrcObjectPath.GetSubPathString());
 			// Put back PIE prefix
-			if (World->IsPlayInEditor() && (PIEInstanceID != INDEX_NONE))
+			if (OuterWorld->IsPlayInEditor() && (PIEInstanceID != INDEX_NONE))
 			{
 				ObjectPath.FixupForPIE(PIEInstanceID);
 			}
diff --git a/Engine/Source/Runtime/Engine/Private/WorldPartition/WorldPartitionRuntimeLevelStreamingCell.cpp b/Engine/Source/Runtime/Engine/Private/WorldPartition/WorldPartitionRuntimeLevelStreamingCell.cpp
--- a/Engine/Source/Runtime/Engine/Private/WorldPartition/WorldPartitionRuntimeLevelStreamingCell.cpp
+++ b/Engine/Source/Runtime/Engine/Private/WorldPartition/WorldPartitionRuntimeLevelStreamingCell.cpp
@@ -323,7 +323,15 @@ ULevel* UWorldPartitionRuntimeLevelStreamingCell::GetLevel() const
 
 bool UWorldPartitionRuntimeLevelStreamingCell::CanUnload() const
 {
-	return LevelStreaming ? LevelStreaming->GetWorld()->GetSubsystem<UHLODSubsystem>()->RequestUnloading(this) : true;
+	if (LevelStreaming)
+	{
+		if (UHLODSubsystem* HLODSubsystem = LevelStreaming->GetWorld()->GetSubsystem<UHLODSubsystem>())
+		{
+			return HLODSubsystem->RequestUnloading(this);
+		}
+	}
+
+	return true;
 }
 
 void UWorldPartitionRuntimeLevelStreamingCell::Unload() const
diff --git a/Engine/Source/Runtime/Engine/Private/WorldPartition/WorldPartitionRuntimeSpatialHash.cpp b/Engine/Source/Runtime/Engine/Private/WorldPartition/WorldPartitionRuntimeSpatialHash.cpp
--- a/Engine/Source/Runtime/Engine/Private/WorldPartition/WorldPartitionRuntimeSpatialHash.cpp
+++ b/Engine/Source/Runtime/Engine/Private/WorldPartition/WorldPartitionRuntimeSpatialHash.cpp
@@ -1230,7 +1230,12 @@ EWorldPartitionStreamingPerformance UWorldPartitionRuntimeSpatialHash::GetStream
 	return EWorldPartitionStreamingPerformance::Good;
 }
 
-void UWorldPartitionRuntimeSpatialHash::Draw2D(UCanvas* Canvas, const TArray<FWorldPartitionStreamingSource>& Sources, const FVector2D& PartitionCanvasSize, FVector2D& Offset) const
+FVector2D UWorldPartitionRuntimeSpatialHash::GetDraw2DDesiredFootprint(const FVector2D& CanvasSize) const
+{
+	return FVector2D(CanvasSize.X * GetFilteredStreamingGrids().Num(), CanvasSize.Y);
+}
+
+void UWorldPartitionRuntimeSpatialHash::Draw2D(UCanvas* Canvas, const TArray<FWorldPartitionStreamingSource>& Sources, const FVector2D& PartitionCanvasSize, const FVector2D& Offset) const
 {
 	TRACE_CPUPROFILER_EVENT_SCOPE(UWorldPartitionRuntimeSpatialHash::Draw2D);
 
@@ -1277,7 +1282,6 @@ void UWorldPartitionRuntimeSpatialHash::Draw2D(UCanvas* Canvas, const TArray<FWo
 		auto WorldToScreen = [&](const FVector2D& WorldPos) { return (WorldToScreenScale * (WorldPos - GridReferenceWorldPos)) + GridScreenOffset; };
 
 		StreamingGrid->Draw2D(Canvas, World, Sources, Region, GridScreenBounds, WorldToScreen);
-		Offset = GridScreenBounds.Max;
 
 		// Draw WorldPartition name
 		FVector2D GridInfoPos = GridScreenOffset - GridScreenHalfExtent;
diff --git a/Engine/Source/Runtime/Engine/Private/WorldPartition/WorldPartitionStreamingPolicy.cpp b/Engine/Source/Runtime/Engine/Private/WorldPartition/WorldPartitionStreamingPolicy.cpp
--- a/Engine/Source/Runtime/Engine/Private/WorldPartition/WorldPartitionStreamingPolicy.cpp
+++ b/Engine/Source/Runtime/Engine/Private/WorldPartition/WorldPartitionStreamingPolicy.cpp
@@ -7,6 +7,7 @@
 #include "WorldPartition/WorldPartitionStreamingPolicy.h"
 #include "WorldPartition/WorldPartitionRuntimeCell.h"
 #include "WorldPartition/WorldPartitionStreamingSource.h"
+#include "WorldPartition/WorldPartitionSubsystem.h"
 #include "WorldPartition/WorldPartition.h"
 #include "WorldPartition/WorldPartitionReplay.h"
 #include "WorldPartition/DataLayer/DataLayerSubsystem.h"
@@ -139,7 +140,9 @@ void UWorldPartitionStreamingPolicy::UpdateStreamingSources()
 			}
 		}
 	
-		for (IWorldPartitionStreamingSourceProvider* StreamingSourceProvider : WorldPartition->StreamingSourceProviders)
+		UWorldPartitionSubsystem* WorldPartitionSubsystem = WorldPartition->GetWorld()->GetSubsystem<UWorldPartitionSubsystem>();
+		check(WorldPartitionSubsystem);
+		for (IWorldPartitionStreamingSourceProvider* StreamingSourceProvider : WorldPartitionSubsystem->GetStreamingSourceProviders())
 		{
 			FWorldPartitionStreamingSource StreamingSource;
 			// Default Streaming Source provider's priority to be less than those based on player controllers
@@ -651,7 +654,12 @@ bool UWorldPartitionStreamingPolicy::IsStreamingCompleted(EWorldPartitionRuntime
 	return true;
 }
 
-void UWorldPartitionStreamingPolicy::DrawRuntimeHash2D(class UCanvas* Canvas, const FVector2D& PartitionCanvasSize, FVector2D& Offset)
+FVector2D UWorldPartitionStreamingPolicy::GetDrawRuntimeHash2DDesiredFootprint(const FVector2D& CanvasSize)
+{
+	return WorldPartition->RuntimeHash->GetDraw2DDesiredFootprint(CanvasSize);
+}
+
+void UWorldPartitionStreamingPolicy::DrawRuntimeHash2D(class UCanvas* Canvas, const FVector2D& PartitionCanvasSize, const FVector2D& Offset)
 {
 	if (StreamingSources.Num() > 0)
 	{
diff --git a/Engine/Source/Runtime/Engine/Private/WorldPartition/WorldPartitionSubsystem.cpp b/Engine/Source/Runtime/Engine/Private/WorldPartition/WorldPartitionSubsystem.cpp
--- a/Engine/Source/Runtime/Engine/Private/WorldPartition/WorldPartitionSubsystem.cpp
+++ b/Engine/Source/Runtime/Engine/Private/WorldPartition/WorldPartitionSubsystem.cpp
@@ -149,56 +149,76 @@ void UWorldPartitionSubsystem::Deinitialize()
 
 void UWorldPartitionSubsystem::OnWorldPartitionInitialized(UWorldPartition* InWorldPartition)
 {
-	check(InWorldPartition == GetWorldPartition());
-	if (InWorldPartition->CanDrawRuntimeHash() && (GetWorld()->GetNetMode() != NM_DedicatedServer))
+	if (RegisteredWorldPartitions.IsEmpty())
 	{
-		DrawHandle = UDebugDrawService::Register(TEXT("Game"), FDebugDrawDelegate::CreateUObject(this, &UWorldPartitionSubsystem::Draw));
-	}
+		if (InWorldPartition->CanDrawRuntimeHash() && (GetWorld()->GetNetMode() != NM_DedicatedServer))
+		{
+			DrawHandle = UDebugDrawService::Register(TEXT("Game"), FDebugDrawDelegate::CreateUObject(this, &UWorldPartitionSubsystem::Draw));
+		}
 
-	// Enforce some GC settings when using World Partition
-	if (GetWorld()->IsGameWorld())
-	{
-		LevelStreamingContinuouslyIncrementalGCWhileLevelsPendingPurge = GLevelStreamingContinuouslyIncrementalGCWhileLevelsPendingPurge;
-		LevelStreamingForceGCAfterLevelStreamedOut = GLevelStreamingForceGCAfterLevelStreamedOut;
+		// Enforce some GC settings when using World Partition
+		if (GetWorld()->IsGameWorld())
+		{
+			LevelStreamingContinuouslyIncrementalGCWhileLevelsPendingPurge = GLevelStreamingContinuouslyIncrementalGCWhileLevelsPendingPurge;
+			LevelStreamingForceGCAfterLevelStreamedOut = GLevelStreamingForceGCAfterLevelStreamedOut;
 
-		GLevelStreamingContinuouslyIncrementalGCWhileLevelsPendingPurge = GLevelStreamingContinuouslyIncrementalGCWhileLevelsPendingPurgeForWP;
-		GLevelStreamingForceGCAfterLevelStreamedOut = 0;
+			GLevelStreamingContinuouslyIncrementalGCWhileLevelsPendingPurge = GLevelStreamingContinuouslyIncrementalGCWhileLevelsPendingPurgeForWP;
+			GLevelStreamingForceGCAfterLevelStreamedOut = 0;
+		}
 	}
+
+	check(!RegisteredWorldPartitions.Contains(InWorldPartition));
+	RegisteredWorldPartitions.Add(InWorldPartition);
 }
 
 void UWorldPartitionSubsystem::OnWorldPartitionUninitialized(UWorldPartition* InWorldPartition)
 {
-	check(InWorldPartition == GetWorldPartition());
-	if (GetWorld()->IsGameWorld())
-	{
-		GLevelStreamingContinuouslyIncrementalGCWhileLevelsPendingPurge = LevelStreamingContinuouslyIncrementalGCWhileLevelsPendingPurge;
-		GLevelStreamingForceGCAfterLevelStreamedOut = LevelStreamingForceGCAfterLevelStreamedOut;
-	}
+	check(RegisteredWorldPartitions.Contains(InWorldPartition));
+	RegisteredWorldPartitions.Remove(InWorldPartition);
 
-	if (DrawHandle.IsValid())
+	if (RegisteredWorldPartitions.IsEmpty())
 	{
-		UDebugDrawService::Unregister(DrawHandle);
-		DrawHandle.Reset();
+		if (GetWorld()->IsGameWorld())
+		{
+			GLevelStreamingContinuouslyIncrementalGCWhileLevelsPendingPurge = LevelStreamingContinuouslyIncrementalGCWhileLevelsPendingPurge;
+			GLevelStreamingForceGCAfterLevelStreamedOut = LevelStreamingForceGCAfterLevelStreamedOut;
+		}
+
+		if (DrawHandle.IsValid())
+		{
+			UDebugDrawService::Unregister(DrawHandle);
+			DrawHandle.Reset();
+		}
 	}
 }
 
+void UWorldPartitionSubsystem::RegisterStreamingSourceProvider(IWorldPartitionStreamingSourceProvider* StreamingSource)
+{
+	StreamingSourceProviders.Add(StreamingSource);
+}
+
+bool UWorldPartitionSubsystem::UnregisterStreamingSourceProvider(IWorldPartitionStreamingSourceProvider* StreamingSource)
+{
+	return !!StreamingSourceProviders.Remove(StreamingSource);
+}
+
 void UWorldPartitionSubsystem::Tick(float DeltaSeconds)
 {
 	Super::Tick(DeltaSeconds);
 
-	if (UWorldPartition* Partition = GetWorldPartition())
+	for (UWorldPartition* RegisteredWorldPartition : RegisteredWorldPartitions)
 	{
-		Partition->Tick(DeltaSeconds);
+		RegisteredWorldPartition->Tick(DeltaSeconds);
 
-		if (GDrawRuntimeHash3D && Partition->CanDrawRuntimeHash())
+		if (GDrawRuntimeHash3D && RegisteredWorldPartition->CanDrawRuntimeHash())
 		{
-			Partition->DrawRuntimeHash3D();
+			RegisteredWorldPartition->DrawRuntimeHash3D();
 		}
 
 #if WITH_EDITOR
 		if (!GetWorld()->IsGameWorld())
 		{
-			Partition->DrawRuntimeHashPreview();
+			RegisteredWorldPartition->DrawRuntimeHashPreview();
 		}
 #endif
 	}
@@ -216,9 +236,12 @@ TStatId UWorldPartitionSubsystem::GetStatId() const
 
 bool UWorldPartitionSubsystem::IsStreamingCompleted(EWorldPartitionRuntimeCellState QueryState, const TArray<FWorldPartitionStreamingQuerySource>& QuerySources, bool bExactState) const
 {
-	if (const UWorldPartition* Partition = GetWorldPartition())
+	for (UWorldPartition* RegisteredWorldPartition : RegisteredWorldPartitions)
 	{
-		return Partition->IsStreamingCompleted(QueryState, QuerySources, bExactState);
+		if (!RegisteredWorldPartition->IsStreamingCompleted(QueryState, QuerySources, bExactState))
+		{
+			return false;
+		}
 	}
 
 	return true;
@@ -228,7 +251,7 @@ void UWorldPartitionSubsystem::DumpStreamingSources(FOutputDevice& OutputDevice)
 {
 	if (const UWorldPartition* WorldPartition = GetWorldPartition())
 	{
-		const TArray<FWorldPartitionStreamingSource>* StreamingSources = WorldPartition ? &WorldPartition->GetStreamingSources() : nullptr;
+		const TArray<FWorldPartitionStreamingSource>* StreamingSources = &WorldPartition->GetStreamingSources();
 		if (StreamingSources && (StreamingSources->Num() > 0))
 		{
 			OutputDevice.Logf(TEXT("Streaming Sources:"));
@@ -250,9 +273,9 @@ void UWorldPartitionSubsystem::UpdateStreamingState()
 	}
 #endif
 
-	if (UWorldPartition* Partition = GetWorldPartition())
+	for (UWorldPartition* RegisteredWorldPartition : RegisteredWorldPartitions)
 	{
-		Partition->UpdateStreamingState();
+		RegisteredWorldPartition->UpdateStreamingState();
 	}
 }
 
@@ -282,11 +305,26 @@ void UWorldPartitionSubsystem::Draw(UCanvas* Canvas, class APlayerController* PC
 		const FVector2D CanvasMinimumSize(100.f, 100.f);
 		const FVector2D CanvasMaxScreenSize = FVector2D::Max(MaxScreenRatio*FVector2D(Canvas->ClipX, Canvas->ClipY) - CanvasBottomRightPadding - CurrentOffset, CanvasMinimumSize);
 
-		if (UWorldPartition* Partition = GetWorldPartition())
+		FVector2D TotalFootprint(ForceInitToZero);
+		for (UWorldPartition* RegisteredWorldPartition : RegisteredWorldPartitions)
+		{
+			FVector2D Footprint = RegisteredWorldPartition->GetDrawRuntimeHash2DDesiredFootprint(CanvasMaxScreenSize);
+			TotalFootprint.X += Footprint.X;
+		}
+
+		if (TotalFootprint.X > 0.f)
 		{
-			FVector2D PartitionCanvasSize = FVector2D(CanvasMaxScreenSize.X, CanvasMaxScreenSize.Y);
-			Partition->DrawRuntimeHash2D(Canvas, PartitionCanvasSize, CurrentOffset);
+			for (UWorldPartition* RegisteredWorldPartition : RegisteredWorldPartitions)
+			{
+				FVector2D Footprint = RegisteredWorldPartition->GetDrawRuntimeHash2DDesiredFootprint(CanvasMaxScreenSize);
+				float FootprintRatio = Footprint.X / TotalFootprint.X;
+				FVector2D PartitionCanvasSize = FVector2D(CanvasMaxScreenSize.X * FootprintRatio, CanvasMaxScreenSize.Y);
+				RegisteredWorldPartition->DrawRuntimeHash2D(Canvas, PartitionCanvasSize, CurrentOffset);
+				CurrentOffset.X += PartitionCanvasSize.X;
+			}
+			
 			CurrentOffset.X = CanvasBottomRightPadding.X;
+			CurrentOffset.Y += CanvasMaxScreenSize.Y;
 		}
 	}
 	
diff --git a/Engine/Source/Runtime/Engine/Public/WorldPartition/HLOD/HLODSubsystem.h b/Engine/Source/Runtime/Engine/Public/WorldPartition/HLOD/HLODSubsystem.h
--- a/Engine/Source/Runtime/Engine/Public/WorldPartition/HLOD/HLODSubsystem.h
+++ b/Engine/Source/Runtime/Engine/Public/WorldPartition/HLOD/HLODSubsystem.h
@@ -71,9 +71,14 @@ private:
 		}
 	};
 
-	TMap<FName, FCellData> CellsData;
-	TSet<FCellData*> CellsToWarmup;
-
+	struct FWorldPartitionHLODRuntimeData
+	{
+		TMap<FName, FCellData> CellsData;
+		TSet<FCellData*> CellsToWarmup;
+	};
+	
+	TMap<TObjectPtr<UWorldPartition>, FWorldPartitionHLODRuntimeData> WorldPartitionsHLODRuntimeData;
+		
 	void OnWorldPartitionInitialized(UWorldPartition* InWorldPartition);
 	void OnWorldPartitionUninitialized(UWorldPartition* InWorldPartition);
 	void OnBeginRenderViews(const FSceneViewFamily& InViewFamily);
diff --git a/Engine/Source/Runtime/Engine/Public/WorldPartition/WorldPartition.h b/Engine/Source/Runtime/Engine/Public/WorldPartition/WorldPartition.h
--- a/Engine/Source/Runtime/Engine/Public/WorldPartition/WorldPartition.h
+++ b/Engine/Source/Runtime/Engine/Public/WorldPartition/WorldPartition.h
@@ -177,6 +177,7 @@ public:
 	bool IsInitialized() const;
 	virtual void Uninitialize() override;
 
+	bool IsMainWorldPartition() const;
 	const FTransform& GetInstanceTransform() const { return InstanceTransform; }
 
 	void Tick(float DeltaSeconds);
@@ -186,12 +187,10 @@ public:
 
 	const TArray<FWorldPartitionStreamingSource>& GetStreamingSources() const;
 
-	void RegisterStreamingSourceProvider(IWorldPartitionStreamingSourceProvider* StreamingSource);
-	bool UnregisterStreamingSourceProvider(IWorldPartitionStreamingSourceProvider* StreamingSource);
-
 	// Debugging Methods
 	bool CanDrawRuntimeHash() const;
-	void DrawRuntimeHash2D(UCanvas* Canvas, const FVector2D& PartitionCanvasSize, FVector2D& Offset);
+	FVector2D GetDrawRuntimeHash2DDesiredFootprint(const FVector2D& CanvasSize);
+	void DrawRuntimeHash2D(UCanvas* Canvas, const FVector2D& PartitionCanvasSize, const FVector2D& Offset);
 	void DrawRuntimeHash3D();
 	void DrawRuntimeCellsDetails(UCanvas* Canvas, FVector2D& Offset);
 	void DrawStreamingStatusLegend(UCanvas* Canvas, FVector2D& Offset);
@@ -244,8 +243,6 @@ private:
 	UPROPERTY()
 	mutable TObjectPtr<UWorldPartitionStreamingPolicy> StreamingPolicy;
 
-	TArray<IWorldPartitionStreamingSourceProvider*> StreamingSourceProviders;
-
 #if WITH_EDITORONLY_DATA
 	FLinkerInstancingContext InstancingContext;
 	TUniquePtr<FSoftObjectPathFixupArchive> InstancingSoftObjectPathFixupArchive;
diff --git a/Engine/Source/Runtime/Engine/Public/WorldPartition/WorldPartitionRuntimeHash.h b/Engine/Source/Runtime/Engine/Public/WorldPartition/WorldPartitionRuntimeHash.h
--- a/Engine/Source/Runtime/Engine/Public/WorldPartition/WorldPartitionRuntimeHash.h
+++ b/Engine/Source/Runtime/Engine/Public/WorldPartition/WorldPartitionRuntimeHash.h
@@ -63,7 +63,14 @@ class ENGINE_API UWorldPartitionRuntimeHash : public UObject
 	virtual void SortStreamingCellsByImportance(const TSet<const UWorldPartitionRuntimeCell*>& InCells, const TArray<FWorldPartitionStreamingSource>& InSources, TArray<const UWorldPartitionRuntimeCell*, TInlineAllocator<256>>& OutSortedCells) const;
 	EWorldPartitionStreamingPerformance GetStreamingPerformance(const TSet<const UWorldPartitionRuntimeCell*>& CellToActivate) const;
 
-	virtual void Draw2D(class UCanvas* Canvas, const TArray<FWorldPartitionStreamingSource>& Sources, const FVector2D& PartitionCanvasSize, FVector2D& Offset) const {}
+	/* Returns desired footprint that Draw2D should take relative to given Canvas size (the value can exceed the given size).
+	 * UWorldPartitionSubSystem will re-adapt the size relative to all others UWorldPartitionRuntimeHash and provide the correct size to Draw2D.
+	 *
+	 * Return Draw2D's desired footprint.
+	 */
+	virtual FVector2D GetDraw2DDesiredFootprint(const FVector2D& CanvasSize) const { return FVector2D::ZeroVector; }
+
+	virtual void Draw2D(class UCanvas* Canvas, const TArray<FWorldPartitionStreamingSource>& Sources, const FVector2D& PartitionCanvasSize, const FVector2D& Offset) const {}
 	virtual void Draw3D(const TArray<FWorldPartitionStreamingSource>& Sources) const {}
 	virtual bool ContainsRuntimeHash(const FString& Name) const { return false; }
 
diff --git a/Engine/Source/Runtime/Engine/Public/WorldPartition/WorldPartitionRuntimeSpatialHash.h b/Engine/Source/Runtime/Engine/Public/WorldPartition/WorldPartitionRuntimeSpatialHash.h
--- a/Engine/Source/Runtime/Engine/Public/WorldPartition/WorldPartitionRuntimeSpatialHash.h
+++ b/Engine/Source/Runtime/Engine/Public/WorldPartition/WorldPartitionRuntimeSpatialHash.h
@@ -280,7 +280,8 @@ private:
 
 	mutable TMap<FName, const FSpatialHashStreamingGrid*> NameToGridMapping;
 	
-	virtual void Draw2D(class UCanvas* Canvas, const TArray<FWorldPartitionStreamingSource>& Sources, const FVector2D& PartitionCanvasSize, FVector2D& Offset) const override;
+	virtual FVector2D GetDraw2DDesiredFootprint(const FVector2D& CanvasSize) const override;
+	virtual void Draw2D(class UCanvas* Canvas, const TArray<FWorldPartitionStreamingSource>& Sources, const FVector2D& PartitionCanvasSize, const FVector2D& Offset) const override;
 	virtual void Draw3D(const TArray<FWorldPartitionStreamingSource>& Sources) const override;
 	virtual bool ContainsRuntimeHash(const FString& Name) const override;
 
diff --git a/Engine/Source/Runtime/Engine/Public/WorldPartition/WorldPartitionStreamingPolicy.h b/Engine/Source/Runtime/Engine/Public/WorldPartition/WorldPartitionStreamingPolicy.h
--- a/Engine/Source/Runtime/Engine/Public/WorldPartition/WorldPartitionStreamingPolicy.h
+++ b/Engine/Source/Runtime/Engine/Public/WorldPartition/WorldPartitionStreamingPolicy.h
@@ -45,7 +45,8 @@ public:
 	virtual void UpdateStreamingState();
 	virtual void SetTargetStateForCells(EWorldPartitionRuntimeCellState TargetState, const TSet<const UWorldPartitionRuntimeCell*>& Cells);
 	virtual bool CanAddLoadedLevelToWorld(class ULevel* InLevel) const;
-	virtual void DrawRuntimeHash2D(class UCanvas* Canvas, const FVector2D& PartitionCanvasSize, FVector2D& Offset);
+	virtual FVector2D GetDrawRuntimeHash2DDesiredFootprint(const FVector2D& CanvasSize);
+	virtual void DrawRuntimeHash2D(class UCanvas* Canvas, const FVector2D& PartitionCanvasSize, const FVector2D& Offset);
 	virtual void DrawRuntimeHash3D();
 	virtual void DrawRuntimeCellsDetails(class UCanvas* Canvas, FVector2D& Offset) {}
 	virtual void DrawStreamingStatusLegend(class UCanvas* Canvas, FVector2D& Offset) {}
diff --git a/Engine/Source/Runtime/Engine/Public/WorldPartition/WorldPartitionSubsystem.h b/Engine/Source/Runtime/Engine/Public/WorldPartition/WorldPartitionSubsystem.h
--- a/Engine/Source/Runtime/Engine/Public/WorldPartition/WorldPartitionSubsystem.h
+++ b/Engine/Source/Runtime/Engine/Public/WorldPartition/WorldPartitionSubsystem.h
@@ -43,6 +43,10 @@ public:
 
 	void DumpStreamingSources(FOutputDevice& OutputDevice) const;
 
+	TArray<IWorldPartitionStreamingSourceProvider*> GetStreamingSourceProviders() const { return StreamingSourceProviders; }
+	void RegisterStreamingSourceProvider(IWorldPartitionStreamingSourceProvider* StreamingSource);
+	bool UnregisterStreamingSourceProvider(IWorldPartitionStreamingSourceProvider* StreamingSource);
+
 private:
 #if WITH_EDITOR
 	bool IsRunningConvertWorldPartitionCommandlet() const;
@@ -56,6 +60,10 @@ private:
 	void Draw(class UCanvas* Canvas, class APlayerController* PC);
 	friend class UWorldPartition;
 
+	TArray<TObjectPtr<UWorldPartition>> RegisteredWorldPartitions;
+
+	TArray<IWorldPartitionStreamingSourceProvider*> StreamingSourceProviders;
+
 	FDelegateHandle	DrawHandle;
 
 	// GC backup values
