Fixed failure to save a new level created with the Open World template after creating Folders.

diff --git a/Engine/Source/Runtime/Engine/Private/ActorFolder.cpp b/Engine/Source/Runtime/Engine/Private/ActorFolder.cpp
--- a/Engine/Source/Runtime/Engine/Private/ActorFolder.cpp
+++ b/Engine/Source/Runtime/Engine/Private/ActorFolder.cpp
@@ -26,10 +26,10 @@ UActorFolder* UActorFolder::Create(ULevel* InLevel, const FString& InFolderLabel
 	GloballyUniqueObjectPath += TEXT(".");
 	GloballyUniqueObjectPath += FolderShortName;
 
-	const bool bIsInstancedLevel = InLevel->IsInstancedLevel();
-	const bool bUseExternalObject = InLevel->IsUsingExternalObjects() && !bIsInstancedLevel;
+	const bool bIsTransientFolder = (InLevel->IsInstancedLevel() && !InLevel->IsPersistentLevel()) || InLevel->GetPackage()->HasAnyPackageFlags(PKG_PlayInEditor);
+	const bool bUseExternalObject = InLevel->IsUsingExternalObjects() && !bIsTransientFolder;
 	const bool bShouldDirtyLevel = !bUseExternalObject;
-	const EObjectFlags Flags = (bIsInstancedLevel ? RF_Transient : RF_NoFlags) | RF_Transactional;
+	const EObjectFlags Flags = (bIsTransientFolder ? RF_Transient : RF_NoFlags) | RF_Transactional;
 
 	UPackage* ExternalPackage = bUseExternalObject ? FExternalPackageHelper::CreateExternalPackage(InLevel, *GloballyUniqueObjectPath, GetExternalPackageFlags()) : nullptr;
 
