PR #8965: 00413401: Fix for ensure in FPreviewScene when using a GameMode that derives directly from AGameModeBase (Contributed by [REDACTED_PERSON])

diff --git a/Engine/Source/Runtime/Engine/Private/PreviewScene.cpp b/Engine/Source/Runtime/Engine/Private/PreviewScene.cpp
--- a/Engine/Source/Runtime/Engine/Private/PreviewScene.cpp
+++ b/Engine/Source/Runtime/Engine/Private/PreviewScene.cpp
@@ -64,7 +64,7 @@ FPreviewScene::FPreviewScene(FPreviewScene::ConstructionValues CVS)
 		{
 			PreviewWorld->SetGameMode(URL);
 
-			AGameMode* Mode = PreviewWorld->GetAuthGameMode<AGameMode>();
+			AGameModeBase* Mode = PreviewWorld->GetAuthGameMode<AGameModeBase>();
 			ensure(Mode);
 		}
 	}
