Material editor colors wires based on data type when using new HLSL generator

diff --git a/Engine/Source/Editor/GraphEditor/Private/MaterialGraphConnectionDrawingPolicy.cpp b/Engine/Source/Editor/GraphEditor/Private/MaterialGraphConnectionDrawingPolicy.cpp
--- a/Engine/Source/Editor/GraphEditor/Private/MaterialGraphConnectionDrawingPolicy.cpp
+++ b/Engine/Source/Editor/GraphEditor/Private/MaterialGraphConnectionDrawingPolicy.cpp
@@ -5,6 +5,7 @@
 #include "MaterialGraph/MaterialGraphSchema.h"
 #include "MaterialGraphConnectionDrawingPolicy.h"
 #include "MaterialGraphNode_Knot.h"
+#include "MaterialCachedHLSLTree.h"
 
 /////////////////////////////////////////////////////
 // FMaterialGraphConnectionDrawingPolicy
@@ -109,12 +110,20 @@ void FMaterialGraphConnectionDrawingPolicy::DetermineWiringStyle(UEdGraphPin* Ou
 	Params.AssociatedPin2 = InputPin;
 	Params.WireColor = MaterialGraphSchema->ActivePinColor;
 
+	FMaterialConnectionKey ConnectionKey;
+	bool bInactivePin = false;
+	bool bExecPin = false;
+
 	// Have to consider both pins as the input will be an 'output' when previewing a connection
 	if (OutputPin)
 	{
 		if (!MaterialGraph->IsInputActive(OutputPin))
 		{
-			Params.WireColor = MaterialGraphSchema->InactivePinColor;
+			bInactivePin = true;
+		}
+		if (OutputPin->PinType.PinCategory == UMaterialGraphSchema::PC_Exec)
+		{
+			bExecPin = true;
 		}
 
 		UEdGraphNode* OutputNode = OutputPin->GetOwningNode();
@@ -127,14 +136,23 @@ void FMaterialGraphConnectionDrawingPolicy::DetermineWiringStyle(UEdGraphPin* Ou
 		}
 		else if (!OutputNode->IsNodeEnabled() || OutputNode->IsDisplayAsDisabledForced() || OutputNode->IsNodeUnrelated())
 		{
-			Params.WireColor = MaterialGraphSchema->InactivePinColor;
+			bInactivePin = true;
+		}
+		else if(UMaterialGraphNode_Base* OutputMaterialNode = Cast<UMaterialGraphNode_Base>(OutputNode))
+		{
+			ConnectionKey.OutputObject = OutputMaterialNode->GetMaterialNodeOwner();
+			ConnectionKey.OutputIndex = OutputPin->SourceIndex;
 		}
 	}
 	if (InputPin)
 	{
 		if (!MaterialGraph->IsInputActive(InputPin))
 		{
-			Params.WireColor = MaterialGraphSchema->InactivePinColor;
+			bInactivePin = true;
+		}
+		if (InputPin->PinType.PinCategory == UMaterialGraphSchema::PC_Exec)
+		{
+			bExecPin = true;
 		}
 
 		UEdGraphNode* InputNode = InputPin->GetOwningNode();
@@ -147,7 +165,61 @@ void FMaterialGraphConnectionDrawingPolicy::DetermineWiringStyle(UEdGraphPin* Ou
 		}
 		else if (!InputNode->IsNodeEnabled() || InputNode->IsDisplayAsDisabledForced() || InputNode->IsNodeUnrelated())
 		{
-			Params.WireColor = MaterialGraphSchema->InactivePinColor;
+			bInactivePin = true;
+		}
+		else if (UMaterialGraphNode_Base* InputMaterialNode = Cast<UMaterialGraphNode_Base>(InputNode))
+		{
+			ConnectionKey.InputObject = InputMaterialNode->GetMaterialNodeOwner();
+			ConnectionKey.InputIndex = InputMaterialNode->GetInputIndexForPin(InputPin);
+		}
+	}
+
+	if (bInactivePin)
+	{
+		Params.WireColor = MaterialGraphSchema->InactivePinColor;
+	}
+	else if (bExecPin)
+	{
+		Params.WireColor = Settings->ExecutionPinTypeColor;
+		Params.WireThickness = Settings->DefaultExecutionWireThickness;
+	}
+	else
+	{
+		const UE::Shader::EValueType ConnectionType = MaterialGraph->GetConnectionType(ConnectionKey);
+		if (ConnectionType == UE::Shader::EValueType::Struct)
+		{
+			Params.WireColor = Settings->StructPinTypeColor;
+		}
+		else
+		{
+			const UE::Shader::FValueTypeDescription TypeDesc = UE::Shader::GetValueTypeDescription(ConnectionType);
+			if (UE::Shader::IsTextureType(TypeDesc.ComponentType))
+			{
+				Params.WireColor = Settings->ObjectPinTypeColor;
+			}
+			else if (TypeDesc.ComponentType == UE::Shader::EValueComponentType::Float)
+			{
+				if (TypeDesc.NumComponents == 1)
+				{
+					Params.WireColor = Settings->FloatPinTypeColor;
+				}
+				else
+				{
+					Params.WireColor = Settings->VectorPinTypeColor;
+				}
+			}
+			else if (TypeDesc.ComponentType == UE::Shader::EValueComponentType::Double)
+			{
+				Params.WireColor = Settings->DoublePinTypeColor;
+			}
+			else if (TypeDesc.ComponentType == UE::Shader::EValueComponentType::Bool)
+			{
+				Params.WireColor = Settings->BooleanPinTypeColor;
+			}
+			else if (TypeDesc.ComponentType == UE::Shader::EValueComponentType::Int)
+			{
+				Params.WireColor = Settings->IntPinTypeColor;
+			}
 		}
 	}
 
diff --git a/Engine/Source/Editor/UnrealEd/Classes/MaterialGraph/MaterialGraph.h b/Engine/Source/Editor/UnrealEd/Classes/MaterialGraph/MaterialGraph.h
--- a/Engine/Source/Editor/UnrealEd/Classes/MaterialGraph/MaterialGraph.h
+++ b/Engine/Source/Editor/UnrealEd/Classes/MaterialGraph/MaterialGraph.h
@@ -4,12 +4,20 @@
 
 #include "CoreMinimal.h"
 #include "UObject/ObjectMacros.h"
+#include "Templates/UniquePtr.h"
 #include "EdGraph/EdGraph.h"
 #include "Materials/Material.h"
 #include "MaterialGraph.generated.h"
 
 class UMaterialExpressionComment;
 class UMaterialExpressionComposite;
+struct FMaterialGraphCachedConnections;
+struct FMaterialConnectionKey;
+
+namespace UE::Shader
+{
+enum class EValueType : uint8;
+}
 
 DECLARE_DELEGATE_RetVal( bool, FRealtimeStateGetter );
 DECLARE_DELEGATE( FSetMaterialDirty );
@@ -135,7 +143,17 @@ class UNREALED_API UMaterialGraph : public UEdGraph
 	UPROPERTY()
 	FString	OriginalMaterialFullName;
 
+	//~ Begin UEdGraph interface
+	virtual void NotifyGraphChanged() override;
+protected:
+	virtual void NotifyGraphChanged(const FEdGraphEditAction& Action) override;
+	//~ End UEdGraph interface
+
 public:
+	UMaterialGraph();
+	UMaterialGraph(FVTableHelper& Helper);
+	virtual ~UMaterialGraph();
+
 	/**
 	 * Completely rebuild the graph from the material, removing all old nodes
 	 */
@@ -173,7 +191,7 @@ public:
 	void LinkGraphNodesFromMaterial();
 
 	/** Link the Material using the Graph node's connections */
-	void LinkMaterialExpressionsFromGraph() const;
+	void LinkMaterialExpressionsFromGraph();
 
 	/**
 	 * Check whether a material input should be marked as active
@@ -182,6 +200,9 @@ public:
 	 */
 	bool IsInputActive(class UEdGraphPin* GraphPin) const;
 
+	/** Returns the input index associated with the given property */
+	int32 GetInputIndexForProperty(EMaterialProperty Property) const;
+
 	/**
 	 * Get a list of nodes representing expressions that are not used in the Material
 	 *
@@ -189,6 +210,8 @@ public:
 	 */
 	void GetUnusedExpressions(TArray<class UEdGraphNode*>& UnusedNodes) const;
 
+	UE::Shader::EValueType GetConnectionType(const FMaterialConnectionKey& Key);
+
 private:
 	/**
 	 * Remove all Nodes from the graph
@@ -204,4 +227,5 @@ private:
 	 */
 	int32 GetValidOutputIndex(FExpressionInput* Input) const;
 
+	TUniquePtr<FMaterialGraphCachedConnections> CachedConnections;
 };
diff --git a/Engine/Source/Editor/UnrealEd/Classes/MaterialGraph/MaterialGraphNode.h b/Engine/Source/Editor/UnrealEd/Classes/MaterialGraph/MaterialGraphNode.h
--- a/Engine/Source/Editor/UnrealEd/Classes/MaterialGraph/MaterialGraphNode.h
+++ b/Engine/Source/Editor/UnrealEd/Classes/MaterialGraph/MaterialGraphNode.h
@@ -75,6 +75,7 @@ public:
 	//~ End UEdGraphNode Interface.
 
 	//~ Begin UMaterialGraphNode_Base Interface
+	virtual UObject* GetMaterialNodeOwner() const override { return MaterialExpression; }
 	virtual void CreateInputPins() override;
 	virtual void CreateOutputPins() override;
 	//~ End UMaterialGraphNode_Base Interface
diff --git a/Engine/Source/Editor/UnrealEd/Classes/MaterialGraph/MaterialGraphNode_Base.h b/Engine/Source/Editor/UnrealEd/Classes/MaterialGraph/MaterialGraphNode_Base.h
--- a/Engine/Source/Editor/UnrealEd/Classes/MaterialGraph/MaterialGraphNode_Base.h
+++ b/Engine/Source/Editor/UnrealEd/Classes/MaterialGraph/MaterialGraphNode_Base.h
@@ -21,6 +21,10 @@ class UMaterialGraphNode_Base : public UEdGraphNode
 	virtual void CreateOutputPins() {};
 	/** Is this the undeletable root node */
 	virtual bool IsRootNode() const {return false;}
+	/** Gets the object that owns this node, typically either a UMaterial, UMaterialFunction, UMaterialExpression */
+	virtual UObject* GetMaterialNodeOwner() const { return nullptr; }
+	/** Returns the InputIndex associated with the given pin, maps to FMaterialConnectionKey::InputIndex */
+	virtual int32 GetInputIndexForPin(const UEdGraphPin* Pin) const;
 	/** Get a single Input Pin via its index */
 	UNREALED_API class UEdGraphPin* GetInputPin(int32 InputIndex) const;
 	/** Get a single Output Pin via its index */
diff --git a/Engine/Source/Editor/UnrealEd/Classes/MaterialGraph/MaterialGraphNode_Root.h b/Engine/Source/Editor/UnrealEd/Classes/MaterialGraph/MaterialGraphNode_Root.h
--- a/Engine/Source/Editor/UnrealEd/Classes/MaterialGraph/MaterialGraphNode_Root.h
+++ b/Engine/Source/Editor/UnrealEd/Classes/MaterialGraph/MaterialGraphNode_Root.h
@@ -28,6 +28,8 @@ class UMaterialGraphNode_Root : public UMaterialGraphNode_Base
 	//~ End UEdGraphNode Interface.
 
 	//~ Begin UMaterialGraphNode_Base Interface
+	virtual UObject* GetMaterialNodeOwner() const override { return Material; }
+	virtual int32 GetInputIndexForPin(const UEdGraphPin* Pin) const override;
 	virtual uint32 GetPinMaterialType(const UEdGraphPin* Pin) const override;
 	virtual void CreateInputPins() override;
 	virtual bool IsRootNode() const override {return true;}
diff --git a/Engine/Source/Editor/UnrealEd/Private/MaterialGraph.cpp b/Engine/Source/Editor/UnrealEd/Private/MaterialGraph.cpp
--- a/Engine/Source/Editor/UnrealEd/Private/MaterialGraph.cpp
+++ b/Engine/Source/Editor/UnrealEd/Private/MaterialGraph.cpp
@@ -21,17 +21,36 @@
 #include "Materials/MaterialExpressionExecBegin.h"
 #include "Materials/MaterialExpressionExecEnd.h"
 
+#include "MaterialCachedHLSLTree.h"
+#include "HLSLTree/HLSLTreeEmit.h"
+#include "MaterialHLSLTree.h"
+
 #include "MaterialGraphNode_Knot.h"
 
 #include "Kismet2/BlueprintEditorUtils.h"
 
 #define LOCTEXT_NAMESPACE "MaterialGraph"
 
+struct FMaterialGraphCachedConnections
+{
+	TMap<FMaterialConnectionKey, UE::Shader::EValueType> ConnectionMap;
+};
+
 UMaterialGraph::UMaterialGraph(const FObjectInitializer& ObjectInitializer)
 	: Super(ObjectInitializer)
 {
 }
 
+UMaterialGraph::UMaterialGraph() = default;
+
+UMaterialGraph::UMaterialGraph(FVTableHelper& Helper)
+{
+}
+
+UMaterialGraph::~UMaterialGraph()
+{
+}
+
 void UMaterialGraph::RebuildGraph()
 {
 	check(Material);
@@ -78,6 +97,73 @@ static UMaterialGraphNode* InitExpressionNewNode(UMaterialGraph* Graph, UMateria
 	return NewNode;
 }
 
+namespace Private
+{
+void PrepareHLSLTree(UE::HLSLTree::FEmitContext& EmitContext,
+	const FMaterialCachedHLSLTree& CachedTree,
+	EShaderFrequency ShaderFrequency)
+{
+	using namespace UE::HLSLTree;
+	using namespace UE::Shader;
+
+	EmitContext.ShaderFrequency = ShaderFrequency;
+	EmitContext.bUseAnalyticDerivatives = true; // We want to consider expressions used for analytic derivatives
+	EmitContext.bMarkLiveValues = false;
+	FEmitScope* EmitResultScope = EmitContext.PrepareScope(CachedTree.GetResultScope());
+
+	FRequestedType RequestedAttributesType;
+	CachedTree.SetRequestedFields(ShaderFrequency, RequestedAttributesType);
+
+	const FPreparedType& ResultType = EmitContext.PrepareExpression(CachedTree.GetResultExpression(), *EmitResultScope, RequestedAttributesType);
+}
+
+FMaterialGraphCachedConnections* CreateConnections(const FMaterialCachedHLSLTree& CachedTree)
+{
+	using namespace UE::HLSLTree;
+
+	FNullErrorHandler NullErrorHandler;
+	FMemStackBase Allocator;
+	FEmitContext EmitContext(Allocator, FTargetParameters(), NullErrorHandler, CachedTree.GetTypeRegistry());
+
+	Material::FEmitData& EmitMaterialData = EmitContext.AcquireData<Material::FEmitData>();
+
+	PrepareHLSLTree(EmitContext, CachedTree, SF_Pixel);
+	PrepareHLSLTree(EmitContext, CachedTree, SF_Vertex);
+
+	FMaterialGraphCachedConnections* Connections = new FMaterialGraphCachedConnections();
+	for (const auto& It : CachedTree.GetConnections())
+	{
+		const UE::Shader::FType Type = EmitContext.GetType(It.Value);
+		if (!Type.IsVoid())
+		{
+			Connections->ConnectionMap.Add(It.Key, Type.ValueType);
+		}
+	}
+
+	return Connections;
+}
+
+} // namespace Private
+
+UE::Shader::EValueType UMaterialGraph::GetConnectionType(const FMaterialConnectionKey& Key)
+{
+	FMaterialGraphCachedConnections* LocalConnections = CachedConnections.Get();
+	if (!LocalConnections && Material->IsUsingNewHLSLGenerator())
+	{
+		LocalConnections = Private::CreateConnections(Material->GetCachedHLSLTree());
+		CachedConnections.Reset(LocalConnections);
+	}
+	if (LocalConnections)
+	{
+		const UE::Shader::EValueType* Result = LocalConnections->ConnectionMap.Find(Key);
+		if (Result)
+		{
+			return *Result;
+		}
+	}
+	return UE::Shader::EValueType::Void;
+}
+
 void UMaterialGraph::RebuildGraphInternal(const TMap<UMaterialExpression*, TArray<UMaterialExpression*>>& SubgraphExpressionMap, const TMap<UMaterialExpression*, TArray<UMaterialExpressionComment*>>& SubgraphCommentMap)
 {
 	Modify();
@@ -411,7 +497,7 @@ void UMaterialGraph::LinkGraphNodesFromMaterial()
 	NotifyGraphChanged();
 }
 
-void UMaterialGraph::LinkMaterialExpressionsFromGraph() const
+void UMaterialGraph::LinkMaterialExpressionsFromGraph()
 {
 	// Use GraphNodes to make Material Expression Connections
 	for (int32 NodeIndex = 0; NodeIndex < Nodes.Num(); ++NodeIndex)
@@ -584,10 +670,13 @@ void UMaterialGraph::LinkMaterialExpressionsFromGraph() const
 	}
 
 	// Also link subgraphs?
-	for (const UEdGraph* SubGraph : SubGraphs)
+	for (UEdGraph* SubGraph : SubGraphs)
 	{
 		CastChecked<UMaterialGraph>(SubGraph)->LinkMaterialExpressionsFromGraph();
 	}
+
+	// This is called when graph connections have been modified, so we need to reset any cached connection types
+	CachedConnections.Reset();
 }
 
 bool UMaterialGraph::IsInputActive(UEdGraphPin* GraphPin) const
@@ -606,6 +695,18 @@ bool UMaterialGraph::IsInputActive(UEdGraphPin* GraphPin) const
 	return true;
 }
 
+int32 UMaterialGraph::GetInputIndexForProperty(EMaterialProperty Property) const
+{
+	for (int32 Index = 0; Index < MaterialInputs.Num(); ++Index)
+	{
+		if (MaterialInputs[Index].GetProperty() == Property)
+		{
+			return Index;
+		}
+	}
+	return INDEX_NONE;
+}
+
 void UMaterialGraph::GetUnusedExpressions(TArray<UEdGraphNode*>& UnusedNodes) const
 {
 	UnusedNodes.Empty();
@@ -705,6 +806,18 @@ void UMaterialGraph::GetUnusedExpressions(TArray<UEdGraphNode*>& UnusedNodes) co
 	}
 }
 
+void UMaterialGraph::NotifyGraphChanged(const FEdGraphEditAction& Action)
+{
+	CachedConnections.Reset();
+	Super::NotifyGraphChanged(Action);
+}
+
+void UMaterialGraph::NotifyGraphChanged()
+{
+	CachedConnections.Reset();
+	Super::NotifyGraphChanged();
+}
+
 void UMaterialGraph::RemoveAllNodes()
 {
 	MaterialInputs.Empty();
diff --git a/Engine/Source/Editor/UnrealEd/Private/MaterialGraphNode_Base.cpp b/Engine/Source/Editor/UnrealEd/Private/MaterialGraphNode_Base.cpp
--- a/Engine/Source/Editor/UnrealEd/Private/MaterialGraphNode_Base.cpp
+++ b/Engine/Source/Editor/UnrealEd/Private/MaterialGraphNode_Base.cpp
@@ -17,6 +17,12 @@ UMaterialGraphNode_Base::UMaterialGraphNode_Base(const FObjectInitializer& Objec
 {
 }
 
+int32 UMaterialGraphNode_Base::GetInputIndexForPin(const UEdGraphPin* Pin) const
+{
+	// For most node types, the pin's SourceIndex will be the input index
+	return Pin ? Pin->SourceIndex : INDEX_NONE;
+}
+
 UEdGraphPin* UMaterialGraphNode_Base::GetInputPin(int32 InputIndex) const
 {
 	for (UEdGraphPin* Pin : Pins)
diff --git a/Engine/Source/Editor/UnrealEd/Private/MaterialGraphNode_Root.cpp b/Engine/Source/Editor/UnrealEd/Private/MaterialGraphNode_Root.cpp
--- a/Engine/Source/Editor/UnrealEd/Private/MaterialGraphNode_Root.cpp
+++ b/Engine/Source/Editor/UnrealEd/Private/MaterialGraphNode_Root.cpp
@@ -45,6 +45,21 @@ void UMaterialGraphNode_Root::PostPlacedNewNode()
 	}
 }
 
+int32 UMaterialGraphNode_Root::GetInputIndexForPin(const UEdGraphPin* Pin) const
+{
+	if (Pin)
+	{
+		// Pin->SourceIndex will be an index into the graph's MaterialInputs array
+		// the 'InputIndex' for root nodes is equal to the material property value
+		const UMaterialGraph* MaterialGraph = CastChecked<UMaterialGraph>(GetGraph());
+		if (MaterialGraph->MaterialInputs.IsValidIndex(Pin->SourceIndex))
+		{
+			return (int32)MaterialGraph->MaterialInputs[Pin->SourceIndex].GetProperty();
+		}
+	}
+	return INDEX_NONE;
+}
+
 uint32 UMaterialGraphNode_Root::GetPinMaterialType(const UEdGraphPin* Pin) const
 {
 	if (Pin->PinType.PinCategory == UMaterialGraphSchema::PC_Exec)
diff --git a/Engine/Source/Runtime/Engine/Private/HLSLTree/HLSLTreeCommon.cpp b/Engine/Source/Runtime/Engine/Private/HLSLTree/HLSLTreeCommon.cpp
--- a/Engine/Source/Runtime/Engine/Private/HLSLTree/HLSLTreeCommon.cpp
+++ b/Engine/Source/Runtime/Engine/Private/HLSLTree/HLSLTreeCommon.cpp
@@ -16,7 +16,7 @@ struct FPreshaderLoopScope
 	Shader::FPreshaderLabel BreakLabel;
 };
 
-FSwizzleParameters::FSwizzleParameters(int8 InR, int8 InG, int8 InB, int8 InA) : NumComponents(0)
+FSwizzleParameters::FSwizzleParameters(int8 InR, int8 InG, int8 InB, int8 InA) : NumComponents(0), bHasSwizzle(true)
 {
 	ComponentIndex[0] = InR;
 	ComponentIndex[1] = InG;
@@ -47,20 +47,6 @@ FSwizzleParameters::FSwizzleParameters(int8 InR, int8 InG, int8 InB, int8 InA) :
 	++NumComponents;
 }
 
-FRequestedType FSwizzleParameters::GetRequestedInputType(const FRequestedType& RequestedType) const
-{
-	FRequestedType RequestedInputType;
-	for (int32 Index = 0; Index < NumComponents; ++Index)
-	{
-		if (RequestedType.IsComponentRequested(Index))
-		{
-			const int32 SwizzledComponentIndex = ComponentIndex[Index];
-			RequestedInputType.SetComponentRequest(SwizzledComponentIndex);
-		}
-	}
-	return RequestedInputType;
-}
-
 FSwizzleParameters MakeSwizzleMask(bool bInR, bool bInG, bool bInB, bool bInA)
 {
 	int8 ComponentIndex[4] = { INDEX_NONE, INDEX_NONE, INDEX_NONE, INDEX_NONE };
@@ -594,25 +580,51 @@ void FExpressionDerivative::EmitValuePreshader(FEmitContext& Context, FEmitScope
 	OutResult.Preshader.WriteOpcode(Shader::EPreshaderOpcode::ConstantZero).Write(OutResult.Type);
 }
 
+namespace Private
+{
+FRequestedType GetRequestedSwizzleType(const FSwizzleParameters& Params, const FRequestedType& RequestedType)
+{
+	FRequestedType RequestedInputType;
+	for (int32 Index = 0; Index < Params.NumComponents; ++Index)
+	{
+		if (RequestedType.IsComponentRequested(Index))
+		{
+			const int32 SwizzledComponentIndex = Params.ComponentIndex[Index];
+			RequestedInputType.SetComponentRequest(SwizzledComponentIndex);
+		}
+	}
+	return RequestedInputType;
+}
+} // namespace Private
+
+const FExpression* FTree::NewSwizzle(const FSwizzleParameters& Params, const FExpression* Input)
+{
+	if (Params.bHasSwizzle)
+	{
+		return NewExpression<FExpressionSwizzle>(Params, Input);
+	}
+	return Input;
+}
+
 void FExpressionSwizzle::ComputeAnalyticDerivatives(FTree& Tree, FExpressionDerivatives& OutResult) const
 {
 	const FExpressionDerivatives InputDerivatives = Tree.GetAnalyticDerivatives(Input);
 	if (InputDerivatives.IsValid())
 	{
-		OutResult.ExpressionDdx = Tree.NewExpression<FExpressionSwizzle>(Parameters, InputDerivatives.ExpressionDdx);
-		OutResult.ExpressionDdy = Tree.NewExpression<FExpressionSwizzle>(Parameters, InputDerivatives.ExpressionDdy);
+		OutResult.ExpressionDdx = Tree.NewSwizzle(Parameters, InputDerivatives.ExpressionDdx);
+		OutResult.ExpressionDdy = Tree.NewSwizzle(Parameters, InputDerivatives.ExpressionDdy);
 	}
 }
 
 const FExpression* FExpressionSwizzle::ComputePreviousFrame(FTree& Tree, const FRequestedType& RequestedType) const
 {
-	const FRequestedType RequestedInputType = Parameters.GetRequestedInputType(RequestedType);
-	return Tree.NewExpression<FExpressionSwizzle>(Parameters, Tree.GetPreviousFrame(Input, RequestedInputType));
+	const FRequestedType RequestedInputType = Private::GetRequestedSwizzleType(Parameters, RequestedType);
+	return Tree.NewSwizzle(Parameters, Tree.GetPreviousFrame(Input, RequestedInputType));
 }
 
 bool FExpressionSwizzle::PrepareValue(FEmitContext& Context, FEmitScope& Scope, const FRequestedType& RequestedType, FPrepareValueResult& OutResult) const
 {
-	const FRequestedType RequestedInputType = Parameters.GetRequestedInputType(RequestedType);
+	const FRequestedType RequestedInputType = Private::GetRequestedSwizzleType(Parameters, RequestedType);
 
 	FPreparedType ResultType;
 	if (RequestedInputType.IsVoid())
@@ -653,7 +665,7 @@ void FExpressionSwizzle::EmitValueShader(FEmitContext& Context, FEmitScope& Scop
 	TCHAR LWCSwizzle[10] = TEXT("");
 	bool bHasSwizzleReorder = false;
 
-	const int32 NumComponents = FMath::Min(RequestedType.GetNumComponents(), Parameters.NumComponents);
+	const int32 NumComponents = FMath::Min<int32>(RequestedType.GetNumComponents(), Parameters.NumComponents);
 	FRequestedType RequestedInputType;
 	for (int32 ComponentIndex = 0; ComponentIndex < NumComponents; ++ComponentIndex)
 	{
@@ -712,7 +724,7 @@ void FExpressionSwizzle::EmitValueShader(FEmitContext& Context, FEmitScope& Scop
 
 void FExpressionSwizzle::EmitValuePreshader(FEmitContext& Context, FEmitScope& Scope, const FRequestedType& RequestedType, FEmitValuePreshaderResult& OutResult) const
 {
-	const FRequestedType RequestedInputType = Parameters.GetRequestedInputType(RequestedType);
+	const FRequestedType RequestedInputType = Private::GetRequestedSwizzleType(Parameters, RequestedType);
 	if (RequestedInputType.IsVoid())
 	{
 		Context.PreshaderStackPosition++;
diff --git a/Engine/Source/Runtime/Engine/Private/Materials/MaterialCachedData.cpp b/Engine/Source/Runtime/Engine/Private/Materials/MaterialCachedData.cpp
--- a/Engine/Source/Runtime/Engine/Private/Materials/MaterialCachedData.cpp
+++ b/Engine/Source/Runtime/Engine/Private/Materials/MaterialCachedData.cpp
@@ -32,7 +32,7 @@
 #include "Materials/MaterialFunctionInstance.h"
 #include "Materials/MaterialExpressionMaterialFunctionCall.h"
 #include "Materials/MaterialParameterCollection.h"
-#include "Materials/MaterialHLSLTree.h"
+#include "MaterialHLSLTree.h"
 #include "HLSLTree/HLSLTreeEmit.h"
 #include "HLSLTree/HLSLTree.h"
 #include "VT/RuntimeVirtualTexture.h"
@@ -457,14 +457,6 @@ void FMaterialCachedExpressionData::UpdateForExpressions(const FMaterialCachedEx
 
 namespace Private
 {
-class FNullErrorHandler : public UE::HLSLTree::FErrorHandlerInterface
-{
-public:
-	virtual void AddErrorInternal(UObject* InOwner, FStringView InError) override
-	{
-	}
-};
-
 void PrepareHLSLTree(UE::HLSLTree::FEmitContext& EmitContext,
 	const FMaterialCachedHLSLTree& CachedTree,
 	FMaterialCachedExpressionData& CachedData,
@@ -481,26 +473,13 @@ void PrepareHLSLTree(UE::HLSLTree::FEmitContext& EmitContext,
 	FRequestedType RequestedAttributesType;
 	CachedTree.SetRequestedFields(ShaderFrequency, RequestedAttributesType);
 
-	const TArray<FGuid>& OrderedVisibleAttributes = FMaterialAttributeDefinitionMap::GetOrderedVisibleAttributeList();
-	for (const FGuid& AttributeID : OrderedVisibleAttributes)
-	{
-		if (FMaterialAttributeDefinitionMap::GetShaderFrequency(AttributeID) == ShaderFrequency)
-		{
-			const FString& FieldName = FMaterialAttributeDefinitionMap::GetAttributeName(AttributeID);
-			const FStructField* Field = CachedTree.GetMaterialAttributesType()->FindFieldByName(*FieldName);
-			if (Field)
-			{
-				RequestedAttributesType.SetFieldRequested(Field);
-			}
-		}
-	}
-
 	const FPreparedType& ResultType = EmitContext.PrepareExpression(CachedTree.GetResultExpression(), *EmitResultScope, RequestedAttributesType);
 	if (!ResultType.IsVoid())
 	{
 		EmitContext.bMarkLiveValues = true;
 		EmitContext.PrepareExpression(CachedTree.GetResultExpression(), *EmitResultScope, RequestedAttributesType);
 
+		const TArray<FGuid>& OrderedVisibleAttributes = FMaterialAttributeDefinitionMap::GetOrderedVisibleAttributeList();
 		for (const FGuid& AttributeID : OrderedVisibleAttributes)
 		{
 			if (FMaterialAttributeDefinitionMap::GetShaderFrequency(AttributeID) == ShaderFrequency)
@@ -522,7 +501,7 @@ void FMaterialCachedExpressionData::UpdateForCachedHLSLTree(const FMaterialCache
 	using namespace UE::HLSLTree;
 
 	// We ignore errors here, errors will be captured when we actually emit HLSL from the tree
-	::Private::FNullErrorHandler NullErrorHandler;
+	FNullErrorHandler NullErrorHandler;
 
 	FMemStackBase Allocator;
 	FEmitContext EmitContext(Allocator, FTargetParameters(), NullErrorHandler, CachedTree.GetTypeRegistry());
diff --git a/Engine/Source/Runtime/Engine/Private/Materials/MaterialCachedHLSLTree.cpp b/Engine/Source/Runtime/Engine/Private/Materials/MaterialCachedHLSLTree.cpp
--- a/Engine/Source/Runtime/Engine/Private/Materials/MaterialCachedHLSLTree.cpp
+++ b/Engine/Source/Runtime/Engine/Private/Materials/MaterialCachedHLSLTree.cpp
@@ -155,6 +155,20 @@ void FMaterialCachedHLSLTree::SetRequestedFields(EShaderFrequency ShaderFrequenc
 			OutRequestedType.SetFieldRequested(CustomOutputField);
 		}
 	}
+
+	const TArray<FGuid>& OrderedVisibleAttributes = FMaterialAttributeDefinitionMap::GetOrderedVisibleAttributeList();
+	for (const FGuid& AttributeID : OrderedVisibleAttributes)
+	{
+		if (FMaterialAttributeDefinitionMap::GetShaderFrequency(AttributeID) == ShaderFrequency)
+		{
+			const FString& FieldName = FMaterialAttributeDefinitionMap::GetAttributeName(AttributeID);
+			const UE::Shader::FStructField* Field = MaterialAttributesType->FindFieldByName(*FieldName);
+			if (Field)
+			{
+				OutRequestedType.SetFieldRequested(Field);
+			}
+		}
+	}
 }
 
 void FMaterialCachedHLSLTree::EmitSharedCode(FStringBuilderBase& OutCode) const
diff --git a/Engine/Source/Runtime/Engine/Private/Materials/MaterialExpressionHLSL.cpp b/Engine/Source/Runtime/Engine/Private/Materials/MaterialExpressionHLSL.cpp
--- a/Engine/Source/Runtime/Engine/Private/Materials/MaterialExpressionHLSL.cpp
+++ b/Engine/Source/Runtime/Engine/Private/Materials/MaterialExpressionHLSL.cpp
@@ -136,7 +136,7 @@
 #include "Materials/MaterialExpressionForLoop.h"
 #include "Materials/MaterialExpressionWhileLoop.h"
 #include "Materials/MaterialFunctionInterface.h"
-#include "Materials/MaterialHLSLTree.h"
+#include "MaterialHLSLTree.h"
 #include "HLSLTree/HLSLTree.h"
 #include "HLSLTree/HLSLTreeCommon.h"
 #include "Engine/Engine.h"
diff --git a/Engine/Source/Runtime/Engine/Private/Materials/MaterialHLSLEmitter.cpp b/Engine/Source/Runtime/Engine/Private/Materials/MaterialHLSLEmitter.cpp
--- a/Engine/Source/Runtime/Engine/Private/Materials/MaterialHLSLEmitter.cpp
+++ b/Engine/Source/Runtime/Engine/Private/Materials/MaterialHLSLEmitter.cpp
@@ -11,7 +11,7 @@
 #include "Materials/MaterialExpressionFunctionOutput.h"
 #include "Materials/MaterialExpressionVolumetricAdvancedMaterialOutput.h"
 #include "Materials/Material.h"
-#include "Materials/MaterialHLSLTree.h"
+#include "MaterialHLSLTree.h"
 #include "MaterialCachedHLSLTree.h"
 #include "ShaderCore.h"
 #include "HLSLTree/HLSLTree.h"
@@ -808,18 +808,7 @@ bool MaterialEmitHLSL(const FMaterialCompileTargetParameters& InCompilerTarget,
 		// Prepare all fields *except* normal
 		FRequestedType RequestedPixelAttributesType;
 		CachedTree->SetRequestedFields(SF_Pixel, RequestedPixelAttributesType);
-		for (const FGuid& AttributeID : OrderedVisibleAttributes)
-		{
-			if (FMaterialAttributeDefinitionMap::GetShaderFrequency(AttributeID) == SF_Pixel)
-			{
-				const FString& FieldName = FMaterialAttributeDefinitionMap::GetAttributeName(AttributeID);
-				const FStructField* Field = CachedTree->GetMaterialAttributesType()->FindFieldByName(*FieldName);
-				if (Field && Field != NormalField)
-				{
-					RequestedPixelAttributesType.SetFieldRequested(Field);
-				}
-			}
-		}
+		RequestedPixelAttributesType.SetFieldRequested(NormalField, false);
 
 		const FPreparedType PixelResultType0 = EmitContext.PrepareExpression(CachedTree->GetResultExpression(), *EmitResultScope, RequestedPixelAttributesType);
 		if (PixelResultType0.IsVoid())
@@ -894,18 +883,6 @@ bool MaterialEmitHLSL(const FMaterialCompileTargetParameters& InCompilerTarget,
 	{
 		FRequestedType RequestedVertexAttributesType;
 		CachedTree->SetRequestedFields(SF_Vertex, RequestedVertexAttributesType);
-		for (const FGuid& AttributeID : OrderedVisibleAttributes)
-		{
-			if (FMaterialAttributeDefinitionMap::GetShaderFrequency(AttributeID) == SF_Vertex)
-			{
-				const FString& FieldName = FMaterialAttributeDefinitionMap::GetAttributeName(AttributeID);
-				const FStructField* Field = CachedTree->GetMaterialAttributesType()->FindFieldByName(*FieldName);
-				if (Field)
-				{
-					RequestedVertexAttributesType.SetFieldRequested(Field);
-				}
-			}
-		}
 		RequestedVertexAttributesType.SetFieldRequested(CachedTree->GetMaterialAttributesType()->FindFieldByName(TEXT("PrevWorldPositionOffset")));
 
 		EmitContext.ShaderFrequency = SF_Vertex;
diff --git a/Engine/Source/Runtime/Engine/Private/Materials/MaterialHLSLGenerator.cpp b/Engine/Source/Runtime/Engine/Private/Materials/MaterialHLSLGenerator.cpp
--- a/Engine/Source/Runtime/Engine/Private/Materials/MaterialHLSLGenerator.cpp
+++ b/Engine/Source/Runtime/Engine/Private/Materials/MaterialHLSLGenerator.cpp
@@ -11,7 +11,7 @@
 #include "Materials/MaterialExpressionCustomOutput.h"
 #include "Materials/MaterialExpressionVertexInterpolator.h"
 #include "Materials/MaterialExpressionExecBegin.h"
-#include "Materials/MaterialHLSLTree.h"
+#include "MaterialHLSLTree.h"
 #include "Materials/Material.h"
 #include "MaterialCachedHLSLTree.h"
 #include "ShaderCore.h"
@@ -53,6 +53,11 @@ const UE::Shader::FValue& FMaterialHLSLGenerator::GetMaterialAttributesDefaultVa
 	return CachedTree.GetMaterialAttributesDefaultValue();
 }
 
+UMaterialExpression* FMaterialHLSLGenerator::GetCurrentExpression() const
+{
+	return Cast<UMaterialExpression>(GetTree().GetCurrentOwner());
+}
+
 bool FMaterialHLSLGenerator::Generate()
 {
 	using namespace UE::HLSLTree;
@@ -135,7 +140,7 @@ static const UE::HLSLTree::FExpression* CompileMaterialInput(FMaterialHLSLGenera
 			else
 			{
 				check(InputDescription.Input);
-				Expression = InputDescription.Input->TryAcquireHLSLExpression(Generator, Scope);
+				Expression = InputDescription.Input->TryAcquireHLSLExpression(Generator, Scope, (int32)InputProperty);
 			}
 		}
 	}
@@ -156,9 +161,11 @@ bool FMaterialHLSLGenerator::GenerateResult(UE::HLSLTree::FScope& Scope)
 		// Result for function call
 		FFunction* HLSLFunction = FunctionEntry->HLSLFunction;
 		HLSLFunction->OutputExpressions.Reserve(FunctionEntry->FunctionOutputs.Num());
-		for (const UMaterialExpressionFunctionOutput* ExpressionOutput : FunctionEntry->FunctionOutputs)
+		for (int32 OutputIndex = 0; OutputIndex < FunctionEntry->FunctionOutputs.Num(); ++OutputIndex)
 		{
-			HLSLFunction->OutputExpressions.Add(ExpressionOutput->A.TryAcquireHLSLExpression(*this, Scope));
+			UMaterialExpressionFunctionOutput* ExpressionOutput = FunctionEntry->FunctionOutputs[OutputIndex];
+			FOwnerScope TreeOwnerScope(GetTree(), ExpressionOutput);
+			HLSLFunction->OutputExpressions.Add(ExpressionOutput->A.TryAcquireHLSLExpression(*this, Scope, 0));
 		}
 		FunctionEntry->bGeneratedResult = true;
 		bResult = true;
@@ -175,6 +182,8 @@ bool FMaterialHLSLGenerator::GenerateResult(UE::HLSLTree::FScope& Scope)
 		const FExpression* AttributesExpression = nullptr;
 		if (TargetMaterial)
 		{
+			FOwnerScope TreeOwnerScope(GetTree(), TargetMaterial);
+
 			const FStructField* PrevWPOField = CachedTree.GetMaterialAttributesType()->FindFieldByName(TEXT("PrevWorldPositionOffset"));
 			if (TargetMaterial->bUseMaterialAttributes)
 			{
@@ -182,7 +191,7 @@ bool FMaterialHLSLGenerator::GenerateResult(UE::HLSLTree::FScope& Scope)
 				if (TargetMaterial->GetExpressionInputDescription(MP_MaterialAttributes, InputDescription))
 				{
 					check(InputDescription.Type == UE::Shader::EValueType::Struct);
-					AttributesExpression = InputDescription.Input->TryAcquireHLSLExpression(*this, Scope);
+					AttributesExpression = InputDescription.Input->TryAcquireHLSLExpression(*this, Scope, (int32)MP_MaterialAttributes);
 
 					if (AttributesExpression)
 					{
@@ -240,7 +249,7 @@ bool FMaterialHLSLGenerator::GenerateResult(UE::HLSLTree::FScope& Scope)
 					const FStructField* CustomOutputField = CachedTree.GetMaterialAttributesType()->FindFieldByName(FieldName.ToString());
 					check(CustomOutputField);
 
-					const FExpression* CustomOutputExpression = AcquireExpression(Scope, CustomOutput, OutputIndex);
+					const FExpression* CustomOutputExpression = AcquireExpression(Scope, INDEX_NONE, CustomOutput, OutputIndex, FSwizzleParameters());
 					AttributesExpression = GetTree().NewExpression<FExpressionSetStructField>(CachedTree.GetMaterialAttributesType(), CustomOutputField, AttributesExpression, CustomOutputExpression);
 				}
 			}
@@ -253,7 +262,7 @@ bool FMaterialHLSLGenerator::GenerateResult(UE::HLSLTree::FScope& Scope)
 				// If we didn't hit the preview expression while generating the material normally, then generate it now
 				// Hardcoding output 0 as we don't have the UI to specify any other output
 				const int32 OutputIndex = 0;
-				PreviewExpressionResult = AcquireExpression(Scope, PreviewExpression, OutputIndex);
+				PreviewExpressionResult = AcquireExpression(Scope, INDEX_NONE, PreviewExpression, OutputIndex, FSwizzleParameters());
 			}
 			const FString& EmissiveColorName = FMaterialAttributeDefinitionMap::GetAttributeName(MP_EmissiveColor);
 			const FStructField* EmissiveColorField = CachedTree.GetMaterialAttributesType()->FindFieldByName(*EmissiveColorName);
@@ -321,7 +330,7 @@ const UE::HLSLTree::FExpression* FMaterialHLSLGenerator::NewExternalInput(UE::HL
 
 const UE::HLSLTree::FExpression* FMaterialHLSLGenerator::NewSwizzle(const UE::HLSLTree::FSwizzleParameters& Params, const UE::HLSLTree::FExpression* Input)
 {
-	return GetTree().NewExpression<UE::HLSLTree::FExpressionSwizzle>(Params, Input);
+	return GetTree().NewSwizzle(Params, Input);
 }
 
 const UE::Shader::FTextureValue* FMaterialHLSLGenerator::AcquireTextureValue(const UE::Shader::FTextureValue& InValue)
@@ -365,19 +374,43 @@ bool FMaterialHLSLGenerator::InternalError(FStringView ErrorMessage)
 	return false;
 }
 
-const UE::HLSLTree::FExpression* FMaterialHLSLGenerator::AcquireExpression(UE::HLSLTree::FScope& Scope, UMaterialExpression* MaterialExpression, int32 OutputIndex)
+int32 FMaterialHLSLGenerator::FindInputIndex(const FExpressionInput* Input) const
+{
+	UMaterialExpression* OwnerMaterialExpression = GetCurrentExpression();
+	int32 Index = INDEX_NONE;
+	if (OwnerMaterialExpression)
+	{
+		const TArray<FExpressionInput*> Inputs = OwnerMaterialExpression->GetInputs();
+		Index = Inputs.Find(const_cast<FExpressionInput*>(Input));
+	}
+	return Index;
+}
+
+const UE::HLSLTree::FExpression* FMaterialHLSLGenerator::AcquireExpression(UE::HLSLTree::FScope& Scope,
+	int32 InputIndex,
+	UMaterialExpression* MaterialExpression,
+	int32 OutputIndex,
+	const UE::HLSLTree::FSwizzleParameters& Swizzle)
 {
 	using namespace UE::HLSLTree;
+	UObject* InputOwner = GetTree().GetCurrentOwner();
 	FOwnerScope TreeOwnerScope(GetTree(), MaterialExpression);
 
 	const FExpression* Expression = nullptr;
 	if (MaterialExpression->GenerateHLSLExpression(*this, Scope, OutputIndex, Expression))
 	{
+		Expression = GetTree().NewSwizzle(Swizzle, Expression);
 		if (MaterialExpression == PreviewExpression &&
 			!PreviewExpressionResult)
 		{
 			PreviewExpressionResult = Expression;
 		}
+		if (InputOwner && InputIndex != INDEX_NONE)
+		{
+			const FMaterialConnectionKey Key{ InputOwner, MaterialExpression, InputIndex, OutputIndex };
+			CachedTree.ConnectionMap.Add(Key, Expression);
+		}
+
 		return Expression;
 	}
 	
@@ -392,20 +425,20 @@ const UE::HLSLTree::FExpression* FMaterialHLSLGenerator::AcquireFunctionInputExp
 	using namespace UE::HLSLTree;
 	const FFunctionCallEntry* FunctionEntry = FunctionCallStack.Last();
 	const FExpression* InputExpression = nullptr;
+	int32 InputIndex = INDEX_NONE;
 	if (FunctionEntry->MaterialFunction)
 	{
-		bool bFoundInput = false;
 		for (int32 Index = 0; Index < FunctionEntry->FunctionInputs.Num(); ++Index)
 		{
 			if (FunctionEntry->FunctionInputs[Index] == MaterialExpression)
 			{
-				bFoundInput = true;
+				InputIndex = Index;;
 				InputExpression = FunctionEntry->ConnectedInputs[Index];
 				break;
 			}
 		}
 
-		if (!bFoundInput)
+		if (InputIndex == INDEX_NONE)
 		{
 			// Finding a connected input is always expected if we're in a function call
 			Error(TEXT("Invalid function input"));
@@ -416,7 +449,7 @@ const UE::HLSLTree::FExpression* FMaterialHLSLGenerator::AcquireFunctionInputExp
 	if (!InputExpression && (MaterialExpression->bUsePreviewValueAsDefault || !FunctionEntry->MaterialFunction))
 	{
 		// Either we're previewing the material function, or the input isn't connected and we're using preview as default value
-		InputExpression = MaterialExpression->Preview.TryAcquireHLSLExpression(*this, Scope);
+		InputExpression = MaterialExpression->Preview.TryAcquireHLSLExpression(*this, Scope, InputIndex);
 		if (!InputExpression)
 		{
 			const FVector4f PreviewValue(MaterialExpression->PreviewValue);
diff --git a/Engine/Source/Runtime/Engine/Private/Materials/MaterialHLSLTree.cpp b/Engine/Source/Runtime/Engine/Private/Materials/MaterialHLSLTree.cpp
--- a/Engine/Source/Runtime/Engine/Private/Materials/MaterialHLSLTree.cpp
+++ b/Engine/Source/Runtime/Engine/Private/Materials/MaterialHLSLTree.cpp
@@ -1,7 +1,7 @@
 // Copyright Epic Games, Inc. All Rights Reserved.
 #if WITH_EDITOR
 
-#include "Materials/MaterialHLSLTree.h"
+#include "MaterialHLSLTree.h"
 #include "HLSLTree/HLSLTreeCommon.h"
 #include "HLSLTree/HLSLTreeEmit.h"
 #include "Misc/StringBuilder.h"
diff --git a/Engine/Source/Runtime/Engine/Private/Materials/MaterialHLSLTree.h b/Engine/Source/Runtime/Engine/Private/Materials/MaterialHLSLTree.h
deleted file mode 100644
diff --git a/Engine/Source/Runtime/Engine/Private/Materials/MaterialShared.cpp b/Engine/Source/Runtime/Engine/Private/Materials/MaterialShared.cpp
--- a/Engine/Source/Runtime/Engine/Private/Materials/MaterialShared.cpp
+++ b/Engine/Source/Runtime/Engine/Private/Materials/MaterialShared.cpp
@@ -278,32 +278,51 @@ int32 FExpressionInput::Compile(class FMaterialCompiler* Compiler)
 		return INDEX_NONE;
 }
 
-const UE::HLSLTree::FExpression* FExpressionInput::TryAcquireHLSLExpression(FMaterialHLSLGenerator& Generator, UE::HLSLTree::FScope& Scope) const
+const UE::HLSLTree::FExpression* FExpressionInput::TryAcquireHLSLExpression(FMaterialHLSLGenerator& Generator, UE::HLSLTree::FScope& Scope, int32 InputIndex) const
 {
-	const UE::HLSLTree::FExpression* Result = nullptr;
+	using namespace UE::HLSLTree;
+	const FExpression* Result = nullptr;
 	if (Expression)
 	{
 		Expression->ValidateState();
-		Result = Generator.AcquireExpression(Scope, Expression, OutputIndex);
-		if (Mask && Result)
-		{
-			const UE::HLSLTree::FSwizzleParameters SwizzleParams = UE::HLSLTree::MakeSwizzleMask(!!MaskR, !!MaskG, !!MaskB, !!MaskA);
-			Result = Generator.NewSwizzle(SwizzleParams, Result);
-		}
+		const FSwizzleParameters SwizzleParams = Mask ? MakeSwizzleMask(!!MaskR, !!MaskG, !!MaskB, !!MaskA) : FSwizzleParameters();
+		Result = Generator.AcquireExpression(Scope, InputIndex, Expression, OutputIndex, SwizzleParams);
 	}
-
 	return Result;
 }
 
-const UE::HLSLTree::FExpression* FExpressionInput::AcquireHLSLExpression(FMaterialHLSLGenerator& Generator, UE::HLSLTree::FScope& Scope) const
+const UE::HLSLTree::FExpression* FExpressionInput::TryAcquireHLSLExpression(FMaterialHLSLGenerator& Generator, UE::HLSLTree::FScope& Scope) const
+{
+	return TryAcquireHLSLExpression(Generator, Scope, Generator.FindInputIndex(this));
+}
+
+const UE::HLSLTree::FExpression* FExpressionInput::AcquireHLSLExpression(FMaterialHLSLGenerator& Generator, UE::HLSLTree::FScope& Scope, int32 InputIndex) const
 {
 	const FExpressionInput TracedInput = GetTracedInput();
 	if (!TracedInput.Expression)
 	{
-		Generator.Error(TEXT("Missing input"));
+		UMaterialExpression* OwnerExpression = Generator.GetCurrentExpression();
+		FName LocalInputName;
+		if (OwnerExpression && InputIndex != INDEX_NONE)
+		{
+			LocalInputName = OwnerExpression->GetInputName(InputIndex);
+		}
+		if (!LocalInputName.IsNone())
+		{
+			Generator.Errorf(TEXT("Missing input '%s'"), *LocalInputName.ToString());
+		}
+		else
+		{
+			Generator.Error(TEXT("Missing input"));
+		}
 		return nullptr;
 	}
-	return TryAcquireHLSLExpression(Generator, Scope);
+	return TryAcquireHLSLExpression(Generator, Scope, InputIndex);
+}
+
+const UE::HLSLTree::FExpression* FExpressionInput::AcquireHLSLExpression(FMaterialHLSLGenerator& Generator, UE::HLSLTree::FScope& Scope) const
+{
+	return AcquireHLSLExpression(Generator, Scope, Generator.FindInputIndex(this));
 }
 
 const UE::HLSLTree::FExpression* FExpressionInput::AcquireHLSLExpressionOrConstant(FMaterialHLSLGenerator& Generator, UE::HLSLTree::FScope& Scope, const UE::Shader::FValue& ConstantValue) const
diff --git a/Engine/Source/Runtime/Engine/Public/HLSLTree/HLSLTree.h b/Engine/Source/Runtime/Engine/Public/HLSLTree/HLSLTree.h
--- a/Engine/Source/Runtime/Engine/Public/HLSLTree/HLSLTree.h
+++ b/Engine/Source/Runtime/Engine/Public/HLSLTree/HLSLTree.h
@@ -112,6 +112,13 @@ private:
 	int32 NumErrors = 0;
 };
 
+class FNullErrorHandler : public FErrorHandlerInterface
+{
+public:
+	virtual void AddErrorInternal(UObject* InOwner, FStringView InError) override
+	{
+	}
+};
 
 /** Root class of the HLSL AST */
 class FNode
@@ -521,6 +528,7 @@ public:
 
 	/** Shortcuts to create various common expression types */
 	const FExpression* NewConstant(const Shader::FValue& Value);
+	const FExpression* NewSwizzle(const FSwizzleParameters& Params, const FExpression* Input);
 	const FExpression* NewUnaryOp(EOperation Op, const FExpression* Input);
 	const FExpression* NewBinaryOp(EOperation Op, const FExpression* Lhs, const FExpression* Rhs);
 
diff --git a/Engine/Source/Runtime/Engine/Public/HLSLTree/HLSLTreeCommon.h b/Engine/Source/Runtime/Engine/Public/HLSLTree/HLSLTreeCommon.h
--- a/Engine/Source/Runtime/Engine/Public/HLSLTree/HLSLTreeCommon.h
+++ b/Engine/Source/Runtime/Engine/Public/HLSLTree/HLSLTreeCommon.h
@@ -162,17 +162,6 @@ public:
 	virtual void EmitValuePreshader(FEmitContext& Context, FEmitScope& Scope, const FRequestedType& RequestedType, FEmitValuePreshaderResult& OutResult) const override;
 };
 
-struct FSwizzleParameters
-{
-	FSwizzleParameters() : NumComponents(0) { ComponentIndex[0] = ComponentIndex[1] = ComponentIndex[2] = ComponentIndex[3] = INDEX_NONE; }
-	explicit FSwizzleParameters(int8 IndexR, int8 IndexG = INDEX_NONE, int8 IndexB = INDEX_NONE, int8 IndexA = INDEX_NONE);
-
-	FRequestedType GetRequestedInputType(const FRequestedType& RequestedType) const;
-	bool HasSwizzle() const;
-
-	int8 ComponentIndex[4];
-	int32 NumComponents;
-};
 FSwizzleParameters MakeSwizzleMask(bool bInR, bool bInG, bool bInB, bool bInA);
 
 class FExpressionSwizzle : public FExpression
diff --git a/Engine/Source/Runtime/Engine/Public/HLSLTree/HLSLTreeEmit.h b/Engine/Source/Runtime/Engine/Public/HLSLTree/HLSLTreeEmit.h
--- a/Engine/Source/Runtime/Engine/Public/HLSLTree/HLSLTreeEmit.h
+++ b/Engine/Source/Runtime/Engine/Public/HLSLTree/HLSLTreeEmit.h
@@ -265,7 +265,7 @@ struct FTargetParameters
 };
 
 /** Tracks shared state while emitting HLSL code */
-class FEmitContext
+class ENGINE_API FEmitContext
 {
 public:
 	explicit FEmitContext(FMemStackBase& InAllocator,
diff --git a/Engine/Source/Runtime/Engine/Public/HLSLTree/HLSLTreeTypes.h b/Engine/Source/Runtime/Engine/Public/HLSLTree/HLSLTreeTypes.h
--- a/Engine/Source/Runtime/Engine/Public/HLSLTree/HLSLTreeTypes.h
+++ b/Engine/Source/Runtime/Engine/Public/HLSLTree/HLSLTreeTypes.h
@@ -137,6 +137,16 @@ struct FOperationDescription
 
 FOperationDescription GetOperationDescription(EOperation Op);
 
+struct FSwizzleParameters
+{
+	FSwizzleParameters() : NumComponents(0), bHasSwizzle(false) { ComponentIndex[0] = ComponentIndex[1] = ComponentIndex[2] = ComponentIndex[3] = INDEX_NONE; }
+	explicit FSwizzleParameters(int8 IndexR, int8 IndexG = INDEX_NONE, int8 IndexB = INDEX_NONE, int8 IndexA = INDEX_NONE);
+
+	int8 ComponentIndex[4];
+	int8 NumComponents;
+	bool bHasSwizzle;
+};
+
 struct FCustomHLSLInput
 {
 	FCustomHLSLInput() = default;
diff --git a/Engine/Source/Runtime/Engine/Public/MaterialCachedHLSLTree.h b/Engine/Source/Runtime/Engine/Public/MaterialCachedHLSLTree.h
--- a/Engine/Source/Runtime/Engine/Public/MaterialCachedHLSLTree.h
+++ b/Engine/Source/Runtime/Engine/Public/MaterialCachedHLSLTree.h
@@ -14,6 +14,29 @@ class UMaterialExpression;
 class UMaterialExpressionCustomOutput;
 struct FMaterialLayersFunctions;
 
+struct FMaterialConnectionKey
+{
+	const UObject* InputObject = nullptr;
+	const UObject* OutputObject = nullptr;
+	int32 InputIndex = INDEX_NONE;
+	int32 OutputIndex = INDEX_NONE;
+};
+inline uint32 GetTypeHash(const FMaterialConnectionKey& Key)
+{
+	return HashCombine(HashCombine(HashCombine(GetTypeHash(Key.InputObject), GetTypeHash(Key.OutputObject)), GetTypeHash(Key.InputIndex)), GetTypeHash(Key.OutputIndex));
+}
+inline bool operator==(const FMaterialConnectionKey& Lhs, const FMaterialConnectionKey& Rhs)
+{
+	return Lhs.InputObject == Rhs.InputObject &&
+		Lhs.OutputObject == Rhs.OutputObject &&
+		Lhs.InputIndex == Rhs.InputIndex &&
+		Lhs.OutputIndex == Rhs.OutputIndex;
+}
+inline bool operator!=(const FMaterialConnectionKey& Lhs, const FMaterialConnectionKey& Rhs)
+{
+	return !operator==(Lhs, Rhs);
+}
+
 class FMaterialCachedHLSLTree
 {
 public:
@@ -38,10 +61,12 @@ public:
 	const UE::Shader::FStructType* GetMaterialAttributesType() const { return MaterialAttributesType; }
 	const UE::Shader::FValue& GetMaterialAttributesDefaultValue() const { return MaterialAttributesDefaultValue; }
 
-	void SetRequestedFields(EShaderFrequency ShaderFrequency, UE::HLSLTree::FRequestedType& OutRequestedType) const;
+	const TMap<FMaterialConnectionKey, const UE::HLSLTree::FExpression*>& GetConnections() const { return ConnectionMap; }
+
+	ENGINE_API void SetRequestedFields(EShaderFrequency ShaderFrequency, UE::HLSLTree::FRequestedType& OutRequestedType) const;
 	void EmitSharedCode(FStringBuilderBase& OutCode) const;
 
-	bool IsAttributeUsed(UE::HLSLTree::FEmitContext& Context,
+	ENGINE_API bool IsAttributeUsed(UE::HLSLTree::FEmitContext& Context,
 		UE::HLSLTree::FEmitScope& Scope,
 		const UE::HLSLTree::FPreparedType& ResultType,
 		EMaterialProperty Property) const;
@@ -54,6 +79,7 @@ private:
 	UE::HLSLTree::FScope* ResultScope = nullptr;
 
 	TArray<UMaterialExpressionCustomOutput*> MaterialCustomOutputs;
+	TMap<FMaterialConnectionKey, const UE::HLSLTree::FExpression*> ConnectionMap;
 	const UE::Shader::FStructType* MaterialAttributesType = nullptr;
 	UE::Shader::FValue MaterialAttributesDefaultValue;
 
diff --git a/Engine/Source/Runtime/Engine/Public/MaterialExpressionIO.h b/Engine/Source/Runtime/Engine/Public/MaterialExpressionIO.h
--- a/Engine/Source/Runtime/Engine/Public/MaterialExpressionIO.h
+++ b/Engine/Source/Runtime/Engine/Public/MaterialExpressionIO.h
@@ -77,6 +77,9 @@ struct FExpressionInput
 
 #if WITH_EDITOR
 	ENGINE_API int32 Compile(class FMaterialCompiler* Compiler);
+	ENGINE_API const UE::HLSLTree::FExpression* TryAcquireHLSLExpression(FMaterialHLSLGenerator& Generator, UE::HLSLTree::FScope& Scope, int32 InputIndex) const;
+	ENGINE_API const UE::HLSLTree::FExpression* AcquireHLSLExpression(FMaterialHLSLGenerator& Generator, UE::HLSLTree::FScope& Scope, int32 InputIndex) const;
+
 	ENGINE_API const UE::HLSLTree::FExpression* TryAcquireHLSLExpression(FMaterialHLSLGenerator& Generator, UE::HLSLTree::FScope& Scope) const;
 	ENGINE_API const UE::HLSLTree::FExpression* AcquireHLSLExpression(FMaterialHLSLGenerator& Generator, UE::HLSLTree::FScope& Scope) const;
 	ENGINE_API const UE::HLSLTree::FExpression* AcquireHLSLExpressionOrConstant(FMaterialHLSLGenerator& Generator, UE::HLSLTree::FScope& Scope, const UE::Shader::FValue& ConstantValue) const;
diff --git a/Engine/Source/Runtime/Engine/Public/MaterialHLSLGenerator.h b/Engine/Source/Runtime/Engine/Public/MaterialHLSLGenerator.h
--- a/Engine/Source/Runtime/Engine/Public/MaterialHLSLGenerator.h
+++ b/Engine/Source/Runtime/Engine/Public/MaterialHLSLGenerator.h
@@ -28,6 +28,7 @@ class UMaterialExpressionFunctionOutput;
 class UMaterialExpressionCustomOutput;
 class UMaterialExpressionVertexInterpolator;
 struct FFunctionExpressionInput;
+struct FExpressionInput;
 class ITargetPlatform;
 enum class EMaterialParameterType : uint8;
 
@@ -75,6 +76,8 @@ public:
 	const UE::Shader::FStructType* GetMaterialAttributesType() const;
 	const UE::Shader::FValue& GetMaterialAttributesDefaultValue() const;
 
+	UMaterialExpression* GetCurrentExpression() const;
+
 	template<typename StringType>
 	inline bool Error(const StringType& InError)
 	{
@@ -105,12 +108,19 @@ public:
 
 	const UE::Shader::FTextureValue* AcquireTextureValue(const UE::Shader::FTextureValue& Value);
 
+	int32 FindInputIndex(const FExpressionInput* Input) const;
+
 	/**
 	 * Returns the appropriate HLSLNode representing the given UMaterialExpression.
 	 * The node will be created if it doesn't exist. Otherwise, the tree will be updated to ensure the node is visible in the given scope
 	 * Note that a given UMaterialExpression may only support 1 of these node types, attempting to access an invalid node type will generate an error
 	 */
-	const UE::HLSLTree::FExpression* AcquireExpression(UE::HLSLTree::FScope& Scope, UMaterialExpression* MaterialExpression, int32 OutputIndex);
+	const UE::HLSLTree::FExpression* AcquireExpression(UE::HLSLTree::FScope& Scope,
+		int32 InputIndex,
+		UMaterialExpression* MaterialExpression,
+		int32 OutputIndex,
+		const UE::HLSLTree::FSwizzleParameters& Swizzle);
+
 	const UE::HLSLTree::FExpression* AcquireFunctionInputExpression(UE::HLSLTree::FScope& Scope, const UMaterialExpressionFunctionInput* MaterialExpression);
 
 	bool GenerateStatements(UE::HLSLTree::FScope& Scope, UMaterialExpression* MaterialExpression);
@@ -175,7 +185,7 @@ private:
 	};
 
 	using FFunctionInputArray = TArray<const UMaterialExpressionFunctionInput*, TInlineAllocator<4>>;
-	using FFunctionOutputArray = TArray<const UMaterialExpressionFunctionOutput*, TInlineAllocator<4>>;
+	using FFunctionOutputArray = TArray<UMaterialExpressionFunctionOutput*, TInlineAllocator<4>>;
 	using FConnectedInputArray = TArray<const UE::HLSLTree::FExpression*, TInlineAllocator<4>>;
 
 	struct FFunctionCallEntry
diff --git a/Engine/Source/Runtime/Engine/Public/MaterialHLSLTree.h b/Engine/Source/Runtime/Engine/Public/MaterialHLSLTree.h
new file mode 100644
--- /dev/null
+++ b/Engine/Source/Runtime/Engine/Public/MaterialHLSLTree.h
@@ -0,0 +1,404 @@
+// Copyright Epic Games, Inc. All Rights Reserved.
+#pragma once
+
+#if WITH_EDITOR
+
+#include "HLSLTree/HLSLTreeCommon.h"
+#include "MaterialTypes.h"
+#include "Materials/MaterialLayersFunctions.h"
+#include "RHIDefinitions.h"
+
+enum class EMaterialParameterType : uint8;
+struct FMaterialCachedExpressionData;
+struct FMaterialLayersFunctions;
+
+namespace UE::HLSLTree::Material
+{
+
+enum class EExternalInput : uint8
+{
+	None,
+
+	TexCoord0,
+	TexCoord1,
+	TexCoord2,
+	TexCoord3,
+	TexCoord4,
+	TexCoord5,
+	TexCoord6,
+	TexCoord7,
+
+	TexCoord0_Ddx,
+	TexCoord1_Ddx,
+	TexCoord2_Ddx,
+	TexCoord3_Ddx,
+	TexCoord4_Ddx,
+	TexCoord5_Ddx,
+	TexCoord6_Ddx,
+	TexCoord7_Ddx,
+
+	TexCoord0_Ddy,
+	TexCoord1_Ddy,
+	TexCoord2_Ddy,
+	TexCoord3_Ddy,
+	TexCoord4_Ddy,
+	TexCoord5_Ddy,
+	TexCoord6_Ddy,
+	TexCoord7_Ddy,
+
+	LightmapTexCoord,
+	LightmapTexCoord_Ddx,
+	LightmapTexCoord_Ddy,
+
+	TwoSidedSign,
+	VertexColor,
+	VertexColor_Ddx,
+	VertexColor_Ddy,
+
+	WorldPosition,
+	WorldPosition_NoOffsets,
+	TranslatedWorldPosition,
+	TranslatedWorldPosition_NoOffsets,
+	ActorWorldPosition,
+
+	PrevWorldPosition,
+	PrevWorldPosition_NoOffsets,
+	PrevTranslatedWorldPosition,
+	PrevTranslatedWorldPosition_NoOffsets,
+
+	WorldPosition_Ddx,
+	WorldPosition_Ddy,
+
+	WorldVertexNormal,
+	WorldVertexTangent,
+	WorldNormal,
+	WorldReflection,
+
+	PreSkinnedPosition,
+	PreSkinnedNormal,
+	PreSkinnedLocalBoundsMin,
+	PreSkinnedLocalBoundsMax,
+
+	ViewportUV,
+	PixelPosition,
+	ViewSize,
+	RcpViewSize,
+	FieldOfView,
+	TanHalfFieldOfView,
+	CotanHalfFieldOfView,
+	TemporalSampleCount,
+	TemporalSampleIndex,
+	TemporalSampleOffset,
+	PreExposure,
+	RcpPreExposure,
+	EyeAdaptation,
+	RuntimeVirtualTextureOutputLevel,
+	RuntimeVirtualTextureOutputDerivative,
+	RuntimeVirtualTextureMaxLevel,
+
+	CameraVector,
+	CameraWorldPosition,
+	ViewWorldPosition,
+	PreViewTranslation,
+	TangentToWorld,
+	LocalToWorld,
+	WorldToLocal,
+	TranslatedWorldToCameraView,
+	TranslatedWorldToView,
+	CameraViewToTranslatedWorld,
+	ViewToTranslatedWorld,
+	WorldToParticle,
+	WorldToInstance,
+	ParticleToWorld,
+	InstanceToWorld,
+
+	PrevFieldOfView,
+	PrevTanHalfFieldOfView,
+	PrevCotanHalfFieldOfView,
+
+	PrevCameraWorldPosition,
+	PrevViewWorldPosition,
+	PrevPreViewTranslation,
+	PrevLocalToWorld,
+	PrevWorldToLocal,
+	PrevTranslatedWorldToCameraView,
+	PrevTranslatedWorldToView,
+	PrevCameraViewToTranslatedWorld,
+	PrevViewToTranslatedWorld,
+
+	PixelDepth,
+	PixelDepth_Ddx,
+	PixelDepth_Ddy,
+
+	GameTime,
+	RealTime,
+	DeltaTime,
+
+	PrevGameTime,
+	PrevRealTime,
+
+	ParticleColor,
+	ParticleTranslatedWorldPosition,
+	ParticleRadius,
+
+	Num,
+};
+static constexpr int32 MaxNumTexCoords = 8;
+
+struct FExternalInputDescription
+{
+	FExternalInputDescription(const TCHAR* InName, Shader::EValueType InType, EExternalInput InDdx = EExternalInput::None, EExternalInput InDdy = EExternalInput::None, EExternalInput InPreviousFrame = EExternalInput::None)
+		: Name(InName), Type(InType), Ddx(InDdx), Ddy(InDdy), PreviousFrame(InPreviousFrame)
+	{}
+
+	const TCHAR* Name;
+	Shader::EValueType Type;
+	EExternalInput Ddx;
+	EExternalInput Ddy;
+	EExternalInput PreviousFrame;
+};
+
+FExternalInputDescription GetExternalInputDescription(EExternalInput Input);
+
+inline bool IsTexCoord(EExternalInput Type)
+{
+	return FMath::IsWithin((int32)Type, (int32)EExternalInput::TexCoord0, (int32)EExternalInput::TexCoord0 + MaxNumTexCoords);
+}
+inline bool IsTexCoord_Ddx(EExternalInput Type)
+{
+	return FMath::IsWithin((int32)Type, (int32)EExternalInput::TexCoord0_Ddx, (int32)EExternalInput::TexCoord0_Ddx + MaxNumTexCoords);
+}
+inline bool IsTexCoord_Ddy(EExternalInput Type)
+{
+	return FMath::IsWithin((int32)Type, (int32)EExternalInput::TexCoord0_Ddy, (int32)EExternalInput::TexCoord0_Ddy + MaxNumTexCoords);
+}
+inline EExternalInput MakeInputTexCoord(int32 Index)
+{
+	check(Index >= 0 && Index < MaxNumTexCoords);
+	return (EExternalInput)((int32)EExternalInput::TexCoord0 + Index);
+}
+
+class FExpressionExternalInput : public FExpression
+{
+public:
+	FExpressionExternalInput(EExternalInput InInputType) : InputType(InInputType) {}
+
+	EExternalInput InputType;
+
+	virtual void ComputeAnalyticDerivatives(FTree& Tree, FExpressionDerivatives& OutResult) const override;
+	virtual const FExpression* ComputePreviousFrame(FTree& Tree, const FRequestedType& RequestedType) const override;
+	virtual bool PrepareValue(FEmitContext& Context, FEmitScope& Scope, const FRequestedType& RequestedType, FPrepareValueResult& OutResult) const override;
+	virtual void EmitValueShader(FEmitContext& Context, FEmitScope& Scope, const FRequestedType& RequestedType, FEmitValueShaderResult& OutResult) const override;
+};
+
+class FExpressionShadingModel : public FExpression
+{
+public:
+	explicit FExpressionShadingModel(EMaterialShadingModel InShadingModel)
+		: ShadingModel(InShadingModel)
+	{}
+
+	TEnumAsByte<EMaterialShadingModel> ShadingModel;
+
+	virtual void ComputeAnalyticDerivatives(FTree& Tree, FExpressionDerivatives& OutResult) const override;
+	virtual bool PrepareValue(FEmitContext& Context, FEmitScope& Scope, const FRequestedType& RequestedType, FPrepareValueResult& OutResult) const override;
+	virtual void EmitValuePreshader(FEmitContext& Context, FEmitScope& Scope, const FRequestedType& RequestedType, FEmitValuePreshaderResult& OutResult) const override;
+};
+
+inline void AppendHash(FHasher& Hasher, const FMaterialParameterInfo& Value)
+{
+	AppendHash(Hasher, Value.Name);
+	AppendHash(Hasher, Value.Index);
+	AppendHash(Hasher, Value.Association);
+}
+
+inline void AppendHash(FHasher& Hasher, const FMaterialParameterValue& Value)
+{
+	switch (Value.Type)
+	{
+	case EMaterialParameterType::Scalar: AppendHash(Hasher, Value.Float[0]); break;
+	case EMaterialParameterType::Vector: AppendHash(Hasher, Value.Float); break;
+	case EMaterialParameterType::DoubleVector: AppendHash(Hasher, Value.Double); break;
+	case EMaterialParameterType::Texture: AppendHash(Hasher, Value.Texture); break;
+	case EMaterialParameterType::Font: AppendHash(Hasher, Value.Font); break;
+	case EMaterialParameterType::RuntimeVirtualTexture: AppendHash(Hasher, Value.RuntimeVirtualTexture); break;
+	case EMaterialParameterType::StaticSwitch: AppendHash(Hasher, Value.Bool[0]); break;
+	case EMaterialParameterType::StaticComponentMask: AppendHash(Hasher, Value.Bool); break;
+	default: checkNoEntry(); break;
+	}
+}
+
+inline void AppendHash(FHasher& Hasher, const FMaterialParameterMetadata& Meta)
+{
+	AppendHash(Hasher, Meta.Value);
+}
+
+class FExpressionParameter : public FExpression
+{
+public:
+	explicit FExpressionParameter(const FMaterialParameterInfo& InParameterInfo, const FMaterialParameterMetadata& InParameterMeta, const Shader::FValue& InDefaultValue)
+		: ParameterInfo(InParameterInfo), ParameterMeta(InParameterMeta), DefaultValue(InDefaultValue)
+	{
+	}
+
+	FMaterialParameterInfo ParameterInfo;
+	FMaterialParameterMetadata ParameterMeta;
+	Shader::FValue DefaultValue;
+
+	virtual void ComputeAnalyticDerivatives(FTree& Tree, FExpressionDerivatives& OutResult) const override;
+	virtual bool PrepareValue(FEmitContext& Context, FEmitScope& Scope, const FRequestedType& RequestedType, FPrepareValueResult& OutResult) const override;
+	virtual void EmitValueShader(FEmitContext& Context, FEmitScope& Scope, const FRequestedType& RequestedType, FEmitValueShaderResult& OutResult) const override;
+	virtual void EmitValuePreshader(FEmitContext& Context, FEmitScope& Scope, const FRequestedType& RequestedType, FEmitValuePreshaderResult& OutResult) const override;
+};
+
+class FExpressionTextureSize : public FExpression
+{
+public:
+	explicit FExpressionTextureSize(const FMaterialParameterInfo& InParameterInfo, UTexture* InTexture)
+		: ParameterInfo(InParameterInfo), Texture(InTexture)
+	{}
+
+	FMaterialParameterInfo ParameterInfo;
+	UTexture* Texture;
+
+	virtual bool PrepareValue(FEmitContext& Context, FEmitScope& Scope, const FRequestedType& RequestedType, FPrepareValueResult& OutResult) const override;
+	virtual void EmitValuePreshader(FEmitContext& Context, FEmitScope& Scope, const FRequestedType& RequestedType, FEmitValuePreshaderResult& OutResult) const override;
+};
+
+class FExpressionFunctionCall : public FExpressionForward
+{
+public:
+	FExpressionFunctionCall(const FExpression* InExpression, UMaterialFunctionInterface* InMaterialFunction)
+		: FExpressionForward(InExpression)
+		, MaterialFunction(InMaterialFunction)
+	{}
+
+	UMaterialFunctionInterface* MaterialFunction;
+
+	virtual bool PrepareValue(FEmitContext& Context, FEmitScope& Scope, const FRequestedType& RequestedType, FPrepareValueResult& OutResult) const override;
+};
+
+class FExpressionMaterialLayers : public FExpressionForward
+{
+public:
+	FExpressionMaterialLayers(const FExpression* InExpression, const FMaterialLayersFunctions* InMaterialLayers)
+		: FExpressionForward(InExpression)
+		, MaterialLayers(InMaterialLayers)
+	{}
+
+	const FMaterialLayersFunctions* MaterialLayers;
+
+	virtual bool PrepareValue(FEmitContext& Context, FEmitScope& Scope, const FRequestedType& RequestedType, FPrepareValueResult& OutResult) const override;
+};
+
+class FExpressionSceneTexture : public FExpression
+{
+public:
+	FExpressionSceneTexture(const FExpression* InTexCoordExpression, uint32 InSceneTextureId, bool bInFiltered)
+		: TexCoordExpression(InTexCoordExpression)
+		, SceneTextureId(InSceneTextureId)
+		, bFiltered(bInFiltered)
+	{}
+
+	const FExpression* TexCoordExpression;
+	uint32 SceneTextureId;
+	bool bFiltered;
+
+	virtual bool PrepareValue(FEmitContext& Context, FEmitScope& Scope, const FRequestedType& RequestedType, FPrepareValueResult& OutResult) const override;
+	virtual void EmitValueShader(FEmitContext& Context, FEmitScope& Scope, const FRequestedType& RequestedType, FEmitValueShaderResult& OutResult) const override;
+};
+
+struct FNoiseParameters
+{
+	FNoiseParameters() { FMemory::Memzero(*this); }
+
+	int32 Quality;
+	int32 Levels;
+	float Scale;
+	uint32 RepeatSize;
+	float OutputMin;
+	float OutputMax;
+	float LevelScale;
+	uint8 NoiseFunction;
+	bool bTiling;
+	bool bTurbulence;
+};
+
+class FExpressionNoise : public FExpression
+{
+public:
+	FExpressionNoise(const FNoiseParameters& InParams, const FExpression* InPositionExpression, const FExpression* InFilterWidthExpression)
+		: PositionExpression(InPositionExpression)
+		, FilterWidthExpression(InFilterWidthExpression)
+		, Parameters(InParams)
+	{}
+
+	const FExpression* PositionExpression;
+	const FExpression* FilterWidthExpression;
+	FNoiseParameters Parameters;
+
+	virtual bool PrepareValue(FEmitContext& Context, FEmitScope& Scope, const FRequestedType& RequestedType, FPrepareValueResult& OutResult) const override;
+	virtual void EmitValueShader(FEmitContext& Context, FEmitScope& Scope, const FRequestedType& RequestedType, FEmitValueShaderResult& OutResult) const override;
+};
+
+/**
+ * Mechanics of passing values from PS->VS are specific to material shaders, to VertexInterpolator is included in the Material module
+ * Possible that parts of this could be moved to the common HLSLTree module, if interpolators are needed by another system
+ */
+class FExpressionVertexInterpolator : public FExpression
+{
+public:
+	FExpressionVertexInterpolator(const FExpression* InVertexExpression) : VertexExpression(InVertexExpression) {}
+
+	const FExpression* VertexExpression;
+
+	virtual void ComputeAnalyticDerivatives(FTree& Tree, FExpressionDerivatives& OutResult) const override;
+	virtual const FExpression* ComputePreviousFrame(FTree& Tree, const FRequestedType& RequestedType) const override;
+	virtual bool PrepareValue(FEmitContext& Context, FEmitScope& Scope, const FRequestedType& RequestedType, FPrepareValueResult& OutResult) const override;
+	virtual void EmitValueShader(FEmitContext& Context, FEmitScope& Scope, const FRequestedType& RequestedType, FEmitValueShaderResult& OutResult) const override;
+	virtual void EmitValuePreshader(FEmitContext& Context, FEmitScope& Scope, const FRequestedType& RequestedType, FEmitValuePreshaderResult& OutResult) const override;
+};
+
+struct FVertexInterpolator
+{
+	FVertexInterpolator() = default;
+	FVertexInterpolator(const FExpression* InExpression) : Expression(InExpression) {}
+
+	const FExpression* Expression = nullptr;
+	FRequestedType RequestedType;
+	FPreparedType PreparedType;
+};
+
+class FEmitData
+{
+public:
+	FEmitData()
+	{
+		for (int32 Index = 0; Index < SF_NumFrequencies; ++Index)
+		{
+			ExternalInputMask[Index].Init(false, (int32)EExternalInput::Num);
+		}
+	}
+
+	const FStaticParameterSet* StaticParameters = nullptr;
+	FMaterialCachedExpressionData* CachedExpressionData = nullptr;
+	TMap<Shader::FValue, uint32> DefaultUniformValues;
+	TArray<FVertexInterpolator, TInlineAllocator<8>> VertexInterpolators;
+	TBitArray<> ExternalInputMask[SF_NumFrequencies];
+	FMaterialShadingModelField ShadingModelsFromCompilation;
+	int32 NumInterpolatorComponents = 0;
+
+	bool IsExternalInputUsed(EShaderFrequency Frequency, EExternalInput Input) const { return ExternalInputMask[Frequency][(int32)Input]; }
+	bool IsExternalInputUsed(EExternalInput Input) const { return IsExternalInputUsed(SF_Vertex, Input) || IsExternalInputUsed(SF_Pixel, Input); }
+
+	int32 FindInterpolatorIndex(const FExpression* Expression) const;
+	void AddInterpolator(const FExpression* Expression, const FRequestedType& RequestedType, const FPreparedType& PreparedType);
+
+	void PrepareInterpolators(FEmitContext& Context, FEmitScope& Scope);
+	void EmitInterpolatorStatements(FEmitContext& Context, FEmitScope& Scope) const;
+	void EmitInterpolatorShader(FEmitContext& Context, FStringBuilderBase& OutCode);
+};
+
+} // namespace UE::HLSLTree::Material
+
+#endif // WITH_EDITOR
diff --git a/Engine/Source/Runtime/Engine/Public/Shader/ShaderTypes.h b/Engine/Source/Runtime/Engine/Public/Shader/ShaderTypes.h
--- a/Engine/Source/Runtime/Engine/Public/Shader/ShaderTypes.h
+++ b/Engine/Source/Runtime/Engine/Public/Shader/ShaderTypes.h
@@ -81,14 +81,14 @@ struct FValueComponentTypeDescription
 	FComponentBounds Bounds;
 };
 
-FValueComponentTypeDescription GetValueComponentTypeDescription(EValueComponentType Type);
+ENGINE_API FValueComponentTypeDescription GetValueComponentTypeDescription(EValueComponentType Type);
 inline const TCHAR* GetComponentTypeName(EValueComponentType Type) { return GetValueComponentTypeDescription(Type).Name; }
 inline uint32 GetComponentTypeSizeInBytes(EValueComponentType Type) { return GetValueComponentTypeDescription(Type).SizeInBytes; }
 inline bool IsNumericType(EValueComponentType Type) { return Type != EValueComponentType::Void && (uint8)Type < (uint8)EValueComponentType::FirstTexture; }
 inline bool IsTextureType(EValueComponentType Type) { return (uint8)Type >= (uint8)EValueComponentType::FirstTexture; }
 inline bool IsComponentTypeWithinBounds(EValueComponentType Type, FComponentBounds Bounds) { return IsWithinBounds(GetValueComponentTypeDescription(Type).Bounds, Bounds); }
 
-EValueComponentType CombineComponentTypes(EValueComponentType Lhs, EValueComponentType Rhs);
+ENGINE_API EValueComponentType CombineComponentTypes(EValueComponentType Lhs, EValueComponentType Rhs);
 
 inline EValueComponentType MakeNonLWCType(EValueComponentType Type) { return Type == EValueComponentType::Double ? EValueComponentType::Float : Type; }
 inline bool IsLWCType(EValueComponentType Type) { return Type == EValueComponentType::Double; }
@@ -145,15 +145,15 @@ struct FValueTypeDescription
 	int8 NumComponents;
 };
 
-FValueTypeDescription GetValueTypeDescription(EValueType Type);
+ENGINE_API FValueTypeDescription GetValueTypeDescription(EValueType Type);
 inline bool IsLWCType(EValueType Type) { return IsLWCType(GetValueTypeDescription(Type).ComponentType); }
-EValueType MakeValueType(EValueComponentType ComponentType, int32 NumComponents);
-EValueType MakeValueType(EValueType BaseType, int32 NumComponents);
-EValueType MakeValueTypeWithRequestedNumComponents(EValueType BaseType, int8 RequestedNumComponents);
-EValueType MakeNonLWCType(EValueType Type);
-EValueType MakeDerivativeType(EValueType Type);
-EValueType MakeArithmeticResultType(EValueType Lhs, EValueType Rhs, FString& OutErrorMessage);
-EValueType MakeComparisonResultType(EValueType Lhs, EValueType Rhs, FString& OutErrorMessage);
+ENGINE_API EValueType MakeValueType(EValueComponentType ComponentType, int32 NumComponents);
+ENGINE_API EValueType MakeValueType(EValueType BaseType, int32 NumComponents);
+ENGINE_API EValueType MakeValueTypeWithRequestedNumComponents(EValueType BaseType, int8 RequestedNumComponents);
+ENGINE_API EValueType MakeNonLWCType(EValueType Type);
+ENGINE_API EValueType MakeDerivativeType(EValueType Type);
+ENGINE_API EValueType MakeArithmeticResultType(EValueType Lhs, EValueType Rhs, FString& OutErrorMessage);
+ENGINE_API EValueType MakeComparisonResultType(EValueType Lhs, EValueType Rhs, FString& OutErrorMessage);
 
 inline bool IsNumericType(EValueType Type) { return IsNumericType(GetValueTypeDescription(Type).ComponentType); }
 inline bool IsTextureType(EValueType Type) { return IsTextureType(GetValueTypeDescription(Type).ComponentType); }
