World Partition: release the levels used by level instances when finalizing the level.

diff --git a/Engine/Source/Runtime/Engine/Private/WorldPartition/WorldPartitionLevelStreamingDynamic.cpp b/Engine/Source/Runtime/Engine/Private/WorldPartition/WorldPartitionLevelStreamingDynamic.cpp
--- a/Engine/Source/Runtime/Engine/Private/WorldPartition/WorldPartitionLevelStreamingDynamic.cpp
+++ b/Engine/Source/Runtime/Engine/Private/WorldPartition/WorldPartitionLevelStreamingDynamic.cpp
@@ -424,6 +424,8 @@ void UWorldPartitionLevelStreamingDynamic::FinalizeRuntimeLevel()
 
 	// Mark this package as fully loaded with regards to external objects
 	RuntimeLevel->GetPackage()->SetDynamicPIEPackagePending(false);
+
+	PackageCache.UnloadPackages();
 }
 
 /**
