Set an engine default for net.PartialBunchReliableThreshold to improve network performance of large bunches when there's packet loss.

diff --git a/Engine/Source/Runtime/Engine/Private/DataChannel.cpp b/Engine/Source/Runtime/Engine/Private/DataChannel.cpp
--- a/Engine/Source/Runtime/Engine/Private/DataChannel.cpp
+++ b/Engine/Source/Runtime/Engine/Private/DataChannel.cpp
@@ -70,7 +70,7 @@ static TAutoConsoleVariable<int32> CVarNetInstantReplayProcessQueuedBunchesMilli
 	8,
 	TEXT("Time threshold for processing queued bunches during instant replays. If it takes longer than this in a single frame, wait until the next frame to continue processing queued bunches. For unlimited time, set to 0."));
 
-int32 GCVarNetPartialBunchReliableThreshold = 0;
+int32 GCVarNetPartialBunchReliableThreshold = 8;
 FAutoConsoleVariableRef CVarNetPartialBunchReliableThreshold(
 	TEXT("net.PartialBunchReliableThreshold"),
 	GCVarNetPartialBunchReliableThreshold,
