Fix issue with USceneCaptureComponent::UpdateShowFlags() emptying the custom show flags bit array.
* FEngineShowFlags::GetSingleFlag() is now more resilient, will now test if custom show flag index is globally registered before asserting

diff --git a/Engine/Source/Runtime/Engine/Private/ShowFlags.cpp b/Engine/Source/Runtime/Engine/Private/ShowFlags.cpp
--- a/Engine/Source/Runtime/Engine/Private/ShowFlags.cpp
+++ b/Engine/Source/Runtime/Engine/Private/ShowFlags.cpp
@@ -141,9 +141,17 @@ bool FEngineShowFlags::GetSingleFlag(uint32 Index) const
 	#define SHOWFLAG_ALWAYS_ACCESSIBLE(a,...) case SF_##a: return a != 0;
 	#include "ShowFlagsValues.inl"
 	default:
-		if (Index >= SF_FirstCustom && (Index - SF_FirstCustom) < (uint32)CustomShowFlags.Num())
+		if (Index >= SF_FirstCustom)
 		{
-			return CustomShowFlags[Index - SF_FirstCustom];
+			ECustomShowFlag CustomShowFlag = ECustomShowFlag(Index - SF_FirstCustom);
+			if (CustomShowFlags.IsValidIndex((uint32)CustomShowFlag))
+			{
+				return CustomShowFlags[(uint32)CustomShowFlag];
+			}
+			else if (IsRegisteredCustomShowFlag(CustomShowFlag))
+			{
+				return false;
+			}
 		}
 		{
 			checkNoEntry();
@@ -961,6 +969,12 @@ FEngineShowFlags::ECustomShowFlag FEngineShowFlags::FindCustomShowFlagByName(con
 	return ECustomShowFlag::None;
 }
 
+bool FEngineShowFlags::IsRegisteredCustomShowFlag(ECustomShowFlag Index)
+{
+	FRWScopeLock Lock(CustomShowFlagsInternal::GetLock(), SLT_ReadOnly);
+	return CustomShowFlagsInternal::GetRegisteredCustomShowFlags().IsValidIndex((uint32)Index);
+}
+
 FString FEngineShowFlags::GetCustomShowFlagName(ECustomShowFlag Index)
 {
 	FRWScopeLock Lock(CustomShowFlagsInternal::GetLock(), SLT_ReadOnly);
diff --git a/Engine/Source/Runtime/Engine/Public/ShowFlags.h b/Engine/Source/Runtime/Engine/Public/ShowFlags.h
--- a/Engine/Source/Runtime/Engine/Public/ShowFlags.h
+++ b/Engine/Source/Runtime/Engine/Public/ShowFlags.h
@@ -539,6 +539,7 @@ private:
 	
 	// Lookup functions for custom show flags matching engine ones
 	static ENGINE_API bool FindCustomShowFlagDisplayName(const FString& InName, FText& OutText);
+	static bool IsRegisteredCustomShowFlag(ECustomShowFlag Index);
 	static FString GetCustomShowFlagName(ECustomShowFlag Index);
 	static FText GetCustomShowFlagDisplayName(ECustomShowFlag Index);
 	static ENGINE_API EShowFlagGroup GetCustomShowFlagGroup(ECustomShowFlag Index);
