Fix static analysis warnings.

diff --git a/Engine/Source/Runtime/Engine/Private/NetConnection.cpp b/Engine/Source/Runtime/Engine/Private/NetConnection.cpp
--- a/Engine/Source/Runtime/Engine/Private/NetConnection.cpp
+++ b/Engine/Source/Runtime/Engine/Private/NetConnection.cpp
@@ -3269,7 +3269,8 @@ void UNetConnection::ReceivedPacket( FBitReader& Reader, bool bIsReinjectedPacke
 
 				// Dispatch the raw, unsequenced bunch to the channel.
 				bool bLocalSkipAck = false;
-				Channel->ReceivedRawBunch( Bunch, bLocalSkipAck ); //warning: May destroy channel.
+				//warning: May destroy channel.
+				Channel->ReceivedRawBunch( Bunch, bLocalSkipAck ); //-V595
 				if ( bLocalSkipAck )
 				{
 					bSkipAck = true;
diff --git a/Engine/Source/Runtime/Engine/Private/RepLayout.cpp b/Engine/Source/Runtime/Engine/Private/RepLayout.cpp
--- a/Engine/Source/Runtime/Engine/Private/RepLayout.cpp
+++ b/Engine/Source/Runtime/Engine/Private/RepLayout.cpp
@@ -6495,7 +6495,7 @@ void FRepLayout::SerializeProperties_r(
 		}
 
 #if !(UE_BUILD_SHIPPING || UE_BUILD_TEST)
-		if (GDoReplicationContextString > 0)
+		if (GDoReplicationContextString > 0 && Map)
 		{
 			Map->SetDebugContextString(FString::Printf(TEXT("%s - %s"), *Owner->GetPathName(), *Cmd.Property->GetPathName()));
 		}
@@ -6571,7 +6571,7 @@ void FRepLayout::SerializeProperties_r(
 		}
 
 #if !(UE_BUILD_SHIPPING || UE_BUILD_TEST)
-		if (GDoReplicationContextString > 0)
+		if (GDoReplicationContextString > 0 && Map)
 		{
 			Map->ClearDebugContextString();
 		}
