Chaos : fix collision shape with negative scale

diff --git a/Engine/Source/Runtime/Engine/Private/PhysicsEngine/BodySetup.cpp b/Engine/Source/Runtime/Engine/Private/PhysicsEngine/BodySetup.cpp
--- a/Engine/Source/Runtime/Engine/Private/PhysicsEngine/BodySetup.cpp
+++ b/Engine/Source/Runtime/Engine/Private/PhysicsEngine/BodySetup.cpp
@@ -2164,8 +2164,8 @@ FKBoxElem FKBoxElem::GetFinalScaled(const FVector& Scale3D, const FTransform& Re
 	ScaledBox.Y *= Scale3DAbs.Y;
 	ScaledBox.Z *= Scale3DAbs.Z;
 
-	FTransform BoxTransform = GetTransform() * RelativeTM;
-	BoxTransform.ScaleTranslation(Scale3D);
+	FTransform ScaleTransform(FQuat::Identity, FVector::ZeroVector, Scale3D); 
+	FTransform BoxTransform = GetTransform() * ScaleTransform * RelativeTM;
 	ScaledBox.SetTransform(BoxTransform);
 
 	return ScaledBox;
@@ -2260,9 +2260,13 @@ FKSphylElem FKSphylElem::GetFinalScaled(const FVector& Scale3D, const FTransform
 	ScaledSphylElem.Radius = GetScaledRadius(Scale3DAbs);
 	ScaledSphylElem.Length = GetScaledCylinderLength(Scale3DAbs);
 
-	FVector LocalOrigin = RelativeTM.TransformPosition(Center) * Scale3D;
+	const FTransform ScaleTransform(FQuat::Identity, FVector::ZeroVector, Scale3D);
+	const FTransform RotationTransform(ScaledSphylElem.Rotation, FVector::ZeroVector, Scale3D);
+	const FTransform ScaledRotationTransform = RotationTransform * ScaleTransform;
+	
+	const FVector LocalOrigin = RelativeTM.TransformPosition(Center) * Scale3D;
 	ScaledSphylElem.Center = LocalOrigin;
-	ScaledSphylElem.Rotation = FRotator(RelativeTM.GetRotation() * FQuat(ScaledSphylElem.Rotation));
+	ScaledSphylElem.Rotation = FRotator(RelativeTM.GetRotation() * ScaledRotationTransform.GetRotation());
 
 	return ScaledSphylElem;
 }
@@ -2361,9 +2365,13 @@ FKTaperedCapsuleElem FKTaperedCapsuleElem::GetFinalScaled(const FVector& Scale3D
 	GetScaledRadii(Scale3DAbs, ScaledTaperedCapsuleElem.Radius0, ScaledTaperedCapsuleElem.Radius1);
 	ScaledTaperedCapsuleElem.Length = GetScaledCylinderLength(Scale3DAbs);
 
-	FVector LocalOrigin = RelativeTM.TransformPosition(Center) * Scale3D;
+	const FTransform ScaleTransform(FQuat::Identity, FVector::ZeroVector, Scale3D);
+	const FTransform RotationTransform(ScaledTaperedCapsuleElem.Rotation, FVector::ZeroVector, Scale3D);
+	const FTransform ScaledRotationTransform = RotationTransform * ScaleTransform;
+	
+	const FVector LocalOrigin = RelativeTM.TransformPosition(Center) * Scale3D;
 	ScaledTaperedCapsuleElem.Center = LocalOrigin;
-	ScaledTaperedCapsuleElem.Rotation = FRotator(RelativeTM.GetRotation() * FQuat(ScaledTaperedCapsuleElem.Rotation));
+	ScaledTaperedCapsuleElem.Rotation = FRotator(RelativeTM.GetRotation() * ScaledRotationTransform.GetRotation());
 
 	return ScaledTaperedCapsuleElem;
 }
diff --git a/Engine/Source/Runtime/Engine/Private/PhysicsEngine/PhysDrawing.cpp b/Engine/Source/Runtime/Engine/Private/PhysicsEngine/PhysDrawing.cpp
--- a/Engine/Source/Runtime/Engine/Private/PhysicsEngine/PhysDrawing.cpp
+++ b/Engine/Source/Runtime/Engine/Private/PhysicsEngine/PhysDrawing.cpp
@@ -964,8 +964,8 @@ void FKAggregateGeom::GetAggGeom(const FTransform& Transform, const FColor Color
 
 	for (int32 i = 0; i < BoxElems.Num(); i++)
 	{
-		FTransform ElemTM = BoxElems[i].GetTransform();
-		ElemTM.ScaleTranslation(Scale3D);
+		const FKBoxElem ScaledBox = BoxElems[i].GetFinalScaled(Scale3D, FTransform::Identity);
+		FTransform ElemTM = ScaledBox.GetTransform();
 		ElemTM *= ParentTM;
 
 		if (bDrawSolid)
@@ -976,8 +976,8 @@ void FKAggregateGeom::GetAggGeom(const FTransform& Transform, const FColor Color
 
 	for (int32 i = 0; i < SphylElems.Num(); i++)
 	{
-		FTransform ElemTM = SphylElems[i].GetTransform();
-		ElemTM.ScaleTranslation(Scale3D);
+		const FKSphylElem ScaledSphyl = SphylElems[i].GetFinalScaled(Scale3D, FTransform::Identity);  
+		FTransform ElemTM = ScaledSphyl.GetTransform();
 		ElemTM *= ParentTM;
 
 		if (bDrawSolid)
@@ -1067,8 +1067,8 @@ void FKAggregateGeom::GetAggGeom(const FTransform& Transform, const FColor Color
 
 	for (int32 i = 0; i < TaperedCapsuleElems.Num(); i++)
 	{
-		FTransform ElemTM = TaperedCapsuleElems[i].GetTransform();
-		ElemTM.ScaleTranslation(Scale3D);
+		const FKTaperedCapsuleElem ScaledTaperedCapsule = TaperedCapsuleElems[i].GetFinalScaled(Scale3D, FTransform::Identity);
+		FTransform ElemTM = ScaledTaperedCapsule.GetTransform();
 		ElemTM *= ParentTM;
 
 		if (bDrawSolid)
