Capture actor counts per level and provide an option to sort the list by count

diff --git a/Engine/Source/Runtime/Engine/Private/EngineUtils.cpp b/Engine/Source/Runtime/Engine/Private/EngineUtils.cpp
--- a/Engine/Source/Runtime/Engine/Private/EngineUtils.cpp
+++ b/Engine/Source/Runtime/Engine/Private/EngineUtils.cpp
@@ -370,7 +370,7 @@ bool EngineUtils::FindOrLoadAssetsByPath(const FString& Path, TArray<UObject*>&
 	return true;
 }
 
-TArray<FSubLevelStatus> GetSubLevelsStatus( UWorld* World )
+TArray<FSubLevelStatus> GetSubLevelsStatus( UWorld* World, bool SortByActorCount )
 {
 	TArray<FSubLevelStatus> Result;
 	FWorldContext &Context = GEngine->GetWorldContextFromWorldChecked(World);
@@ -383,11 +383,33 @@ TArray<FSubLevelStatus> GetSubLevelsStatus( UWorld* World )
 		LevelStatus.PackageName = World->GetOutermost()->GetFName();
 		LevelStatus.StreamingStatus = LEVEL_Visible;
 		LevelStatus.LODIndex = INDEX_NONE;
+		LevelStatus.ActorCount = World->GetActorCount();
 		Result.Add(LevelStatus);
 	}
-	
+
+	auto SortFunc = [](ULevelStreaming* LevelA, ULevelStreaming* LevelB)
+	{
+		if (!LevelA->GetLoadedLevel())
+		{
+			return false;
+		}
+		else if (!LevelB->GetLoadedLevel())
+		{
+			return true;
+		}
+
+		return LevelA->GetLoadedLevel()->Actors.Num() > LevelB->GetLoadedLevel()->Actors.Num();
+	};
+
+	TArray<ULevelStreaming*> SortedStreamingLevels = World->GetStreamingLevels();
+
+	if (SortByActorCount)
+	{
+		Algo::Sort(SortedStreamingLevels, SortFunc);
+	}
+
 	// Iterate over the world info's level streaming objects to find and see whether levels are loaded, visible or neither.
-	for (ULevelStreaming* LevelStreaming : World->GetStreamingLevels())
+	for (const ULevelStreaming* LevelStreaming : SortedStreamingLevels)
 	{
 		if( LevelStreaming 
 			&&  !LevelStreaming->GetWorldAsset().IsNull()
@@ -397,6 +419,33 @@ TArray<FSubLevelStatus> GetSubLevelsStatus( UWorld* World )
 			LevelStatus.PackageName = LevelStreaming->GetWorldAssetPackageFName();
 			LevelStatus.LODIndex = LevelStreaming->GetLevelLODIndex();
 			LevelStatus.StreamingStatus = LevelStreaming->GetLevelStreamingStatus();
+
+			if (LevelStreaming->GetLoadedLevel())
+			{
+				LevelStatus.ActorCount = LevelStreaming->GetLoadedLevel()->Actors.Num();
+
+				for (const AActor* Actor : LevelStreaming->GetLoadedLevel()->Actors)
+				{
+					if (Actor && !Actor->HasAnyFlags(RF_ArchetypeObject | RF_ClassDefaultObject))
+					{
+						const UClass* ParentNativeClass = GetParentNativeClass(Actor->GetClass());
+						FName NativeClassName = ParentNativeClass ? ParentNativeClass->GetFName() : NAME_None;
+
+						FName ActorClassName = Actor->GetClass()->GetFName();
+						FSubLevelActorDetails& ActorDetails = LevelStatus.ActorMapToCount.FindOrAdd(ActorClassName);
+						ActorDetails.Count++;
+						ActorDetails.NativeClassName = NativeClassName;
+					}
+				}
+			}
+
+			if (SortByActorCount)
+			{
+				LevelStatus.ActorMapToCount.ValueSort([](const FSubLevelActorDetails& A, const FSubLevelActorDetails& B) {
+					return A.Count > B.Count;
+				});
+			}
+
 			Result.Add(LevelStatus);
 		}
 	}
@@ -411,6 +460,7 @@ TArray<FSubLevelStatus> GetSubLevelsStatus( UWorld* World )
 		LevelStatus.PackageName = LevelName;
 		LevelStatus.StreamingStatus = LEVEL_Preloading;
 		LevelStatus.LODIndex = INDEX_NONE;
+		LevelStatus.ActorCount = 0;
 		Result.Add(LevelStatus);
 	}
 
diff --git a/Engine/Source/Runtime/Engine/Public/EngineUtils.h b/Engine/Source/Runtime/Engine/Public/EngineUtils.h
--- a/Engine/Source/Runtime/Engine/Public/EngineUtils.h
+++ b/Engine/Source/Runtime/Engine/Public/EngineUtils.h
@@ -707,6 +707,19 @@ private:
  */
 ENGINE_API void DrawStatsHUD( UWorld* InWorld, FViewport* Viewport, FCanvas* Canvas, UCanvas* CanvasObject, TArray<struct FDebugDisplayProperty>& DebugProperties, const FVector& ViewLocation, const FRotator& ViewRotation );
 
+/** SubLevel Actor breakdown information **/
+struct FSubLevelActorDetails
+{
+	FSubLevelActorDetails() :
+		Count(0)
+		, NativeClassName(NAME_None)
+	{
+	}
+
+	int32 Count;
+	FName NativeClassName;
+};
+
 /** SubLevel status information */
 struct FSubLevelStatus
 {
@@ -715,6 +728,8 @@ struct FSubLevelStatus
 	int32				LODIndex;
 	bool				bInConsiderList;
 	bool				bPlayerInside;
+	int32				ActorCount;
+	TMap<FName, FSubLevelActorDetails>	ActorMapToCount;
 };
 
 /**
@@ -722,7 +737,7 @@ struct FSubLevelStatus
  *	@param InWorld		World to gather sublevels stats from
  *	@return				sublevels status (streaming state, LOD index, where player is)
  */
-TArray<FSubLevelStatus> GetSubLevelsStatus( UWorld* InWorld );
+ENGINE_API TArray<FSubLevelStatus> GetSubLevelsStatus( UWorld* InWorld, bool SortByActorCount = false );
 
 #if !UE_BUILD_SHIPPING
 
