Added fix for replication breaking after a period of time.

diff --git a/Engine/Source/Runtime/Engine/Private/RepLayout.cpp b/Engine/Source/Runtime/Engine/Private/RepLayout.cpp
--- a/Engine/Source/Runtime/Engine/Private/RepLayout.cpp
+++ b/Engine/Source/Runtime/Engine/Private/RepLayout.cpp
@@ -1822,6 +1822,14 @@ ERepLayoutResult FRepLayout::CompareProperties(
 		TArray<uint16> SecondChangelistCopy = MoveTemp(RepChangelistState->ChangeHistory[SecondHistoryIndex].Changed);
 
 		MergeChangeList(Data, FirstChangelistRef, SecondChangelistCopy, RepChangelistState->ChangeHistory[SecondHistoryIndex].Changed);
+
+		// Remove any tiling in the history markers to keep them from wrapping over time
+		const uint16 NewHistoryCount = RepChangelistState->HistoryEnd - RepChangelistState->HistoryStart;
+
+		check(NewHistoryCount < FRepChangelistState::MAX_CHANGE_HISTORY);
+
+		RepChangelistState->HistoryStart = RepChangelistState->HistoryStart % FRepChangelistState::MAX_CHANGE_HISTORY;
+		RepChangelistState->HistoryEnd = RepChangelistState->HistoryStart + NewHistoryCount;
 	}
 
 	return Result;
