Initialized AudioDeviceManager member pointer in GEngine.

diff --git a/Engine/Source/Runtime/Engine/Classes/Engine/Engine.h b/Engine/Source/Runtime/Engine/Classes/Engine/Engine.h
--- a/Engine/Source/Runtime/Engine/Classes/Engine/Engine.h
+++ b/Engine/Source/Runtime/Engine/Classes/Engine/Engine.h
@@ -1808,7 +1808,7 @@ private:
 protected:
 
 	/** The audio device manager */
-	FAudioDeviceManager* AudioDeviceManager;
+	FAudioDeviceManager* AudioDeviceManager = nullptr;
 
 	/** Audio device handle to the main audio device. */
 	FAudioDeviceHandle MainAudioDeviceHandle;
