Skel mesh streaming bounds optimization - only notify the streamer when the bounds changes by more than some ratio since the last update. This is controllable via cvar and defaults to 0.1 (10%)

diff --git a/Engine/Source/Runtime/Engine/Classes/Components/SkinnedMeshComponent.h b/Engine/Source/Runtime/Engine/Classes/Components/SkinnedMeshComponent.h
--- a/Engine/Source/Runtime/Engine/Classes/Components/SkinnedMeshComponent.h
+++ b/Engine/Source/Runtime/Engine/Classes/Components/SkinnedMeshComponent.h
@@ -279,6 +279,9 @@ protected:
 	TArray<FTransform> PreviousComponentSpaceTransformsArray;
 
 protected:
+	/** The bounds radius at the point we last notified the streamer of a bounds radius change */
+	double LastStreamerUpdateBoundsRadius;
+
 	/** The index for the ComponentSpaceTransforms buffer we can currently write to */
 	int32 CurrentEditableComponentTransforms;
 
diff --git a/Engine/Source/Runtime/Engine/Private/Components/SkinnedMeshComponent.cpp b/Engine/Source/Runtime/Engine/Private/Components/SkinnedMeshComponent.cpp
--- a/Engine/Source/Runtime/Engine/Private/Components/SkinnedMeshComponent.cpp
+++ b/Engine/Source/Runtime/Engine/Private/Components/SkinnedMeshComponent.cpp
@@ -58,7 +58,7 @@ static TAutoConsoleVariable<int32> CVarAnimVisualizeLODs(
 	0,
 	TEXT("Visualize SkelMesh LODs"));
 
-int32 GUpdateBoundsNotifyStreaming = 0;
+int32 GUpdateBoundsNotifyStreaming = 1;
 FAutoConsoleVariableRef CVarUpdateBoundsNotifyStreaming(
 	TEXT("r.SkinnedMesh.UpdateBoundsNotifyStreaming"),
 	GUpdateBoundsNotifyStreaming,
@@ -66,6 +66,14 @@ FAutoConsoleVariableRef CVarUpdateBoundsNotifyStreaming(
 	ECVF_Default
 );
 
+float GUpdateBoundsNotifyStreamingRadiusChangeRatio = 0.1f;
+FAutoConsoleVariableRef CVarUpdateBoundsNotifyStreamingRadiusChangeRatio(
+	TEXT("r.SkinnedMesh.UpdateBoundsNotifyStreamingRadiusChangeRatio"),
+	GUpdateBoundsNotifyStreamingRadiusChangeRatio,
+	TEXT("Update the streaming manager when the radius changes by more than this ratio since the last update. Only relevant with r.SkinnedMesh.UpdateBoundsNotifyStreaming=1"),
+	ECVF_Default
+);
+
 namespace FAnimUpdateRateManager
 {
 	static float TargetFrameTimeForUpdateRate = 1.f / 30.f; //Target frame rate for lookahead URO
@@ -360,6 +368,7 @@ USkinnedMeshComponent::USkinnedMeshComponent(const FObjectInitializer& ObjectIni
 	CapsuleIndirectShadowMinVisibility = .1f;
 
 	bDoubleBufferedComponentSpaceTransforms = true;
+	LastStreamerUpdateBoundsRadius = -1.0;
 	CurrentEditableComponentTransforms = 0;
 	CurrentReadComponentTransforms = 1;
 	bNeedToFlipSpaceBaseBuffers = false;
@@ -1186,7 +1195,31 @@ void USkinnedMeshComponent::UpdateBounds()
 		// Avoid updating the streamer for small changes in size
 		if (!FMath::IsNearlyEqual(OldRadius, Bounds.SphereRadius, 0.1f) && GetForcedLOD() == 0)
 		{
-			IStreamingManager::Get().NotifyPrimitiveUpdated_Concurrent(this);
+			bool bNotify = true;
+
+		    // If a threshold is specified, don't notify the streamer if the radius has changed by less than the threshold since the last update
+		    if (GUpdateBoundsNotifyStreamingRadiusChangeRatio > 0.0f)
+		    {
+			    if (LastStreamerUpdateBoundsRadius >= 0.0 && LastStreamerUpdateBoundsRadius != Bounds.SphereRadius)
+			    {
+				    // Compute a change ratio based on the change in radius relative to the min of the two radii
+				    auto MinRadius = FMath::Min(Bounds.SphereRadius, LastStreamerUpdateBoundsRadius);
+				    auto MaxRadius = FMath::Max(Bounds.SphereRadius, LastStreamerUpdateBoundsRadius);
+				    if (MinRadius > 0.0)
+				    {
+					    float ChangeRatio = float(MaxRadius / MinRadius)-1.0f;
+					    if (ChangeRatio < GUpdateBoundsNotifyStreamingRadiusChangeRatio)
+					    {
+						    bNotify = false;
+					    }
+				    }
+			    }
+		    }
+		    if (bNotify)
+		    {
+			    IStreamingManager::Get().NotifyPrimitiveUpdated_Concurrent(this);
+			    LastStreamerUpdateBoundsRadius = Bounds.SphereRadius;
+		    }
 		}
 	}
 	else
