Add method to FDeviceProfileManager that allows for lookups to loadable Device Profiles.
Add Lookup and load of device profiles in ULyraSettingsLocal.  Fixes device profiles on cooked platforms.

diff --git a/Engine/Source/Runtime/Engine/Classes/DeviceProfiles/DeviceProfileManager.h b/Engine/Source/Runtime/Engine/Classes/DeviceProfiles/DeviceProfileManager.h
--- a/Engine/Source/Runtime/Engine/Classes/DeviceProfiles/DeviceProfileManager.h
+++ b/Engine/Source/Runtime/Engine/Classes/DeviceProfiles/DeviceProfileManager.h
@@ -74,6 +74,16 @@ public:
 	 */
 	UDeviceProfile* CreateProfile(const FString& ProfileName, const FString& ProfileType, const FString& ParentName=TEXT(""), const TCHAR* ConfigPlatform=nullptr);
 
+	/**
+	 * Tests to see if a named device profile is available to call CreateProfile with.
+	 *
+	 * @param ProfileName - The profile name.
+	 * @param ProfileToCopy - The profile to copy name.
+	 *
+	 * @return the created profile.
+	 */
+	bool HasLoadableProfileName(const FString& ProfileName, FName OptionalPlatformName = FName());
+
 	/**
 	 * Delete a profile.
 	 *
diff --git a/Engine/Source/Runtime/Engine/Private/DeviceProfiles/DeviceProfileManager.cpp b/Engine/Source/Runtime/Engine/Private/DeviceProfiles/DeviceProfileManager.cpp
--- a/Engine/Source/Runtime/Engine/Private/DeviceProfiles/DeviceProfileManager.cpp
+++ b/Engine/Source/Runtime/Engine/Private/DeviceProfiles/DeviceProfileManager.cpp
@@ -753,6 +753,39 @@ UDeviceProfile* UDeviceProfileManager::CreateProfile(const FString& ProfileName,
 	return DeviceProfile;
 }
 
+bool UDeviceProfileManager::HasLoadableProfileName(const FString& ProfileName, FName OptionalPlatformName)
+{
+	UDeviceProfile* DeviceProfile = FindObject<UDeviceProfile>(GetTransientPackage(), *ProfileName);
+	if (DeviceProfile != nullptr)
+	{
+		return true;
+	}
+
+	FConfigCacheIni* ConfigSystem = GConfig;
+
+	if (OptionalPlatformName != NAME_None)
+	{
+#if ALLOW_OTHER_PLATFORM_CONFIG
+		ConfigSystem = FConfigCacheIni::ForPlatform(OptionalPlatformName);
+#else
+
+		checkf(OptionalPlatformName == FName(FPlatformProperties::IniPlatformName()), 
+			"UDeviceProfileManager::HasLoadableProfileName - This platform cannot load configurations for other platforms.");
+#endif
+	}
+
+	// use ConfigPlatform ini hierarchy to look in for the parent profile
+	// @todo config: we could likely cache local ini files to speed this up,
+	// along with the ones we load in LoadConfig
+	// NOTE: This happens at runtime, so maybe only do this if !RequiresCookedData()?
+	FConfigFile* PlatformConfigFile;
+	FConfigFile LocalConfigFile;
+
+	PlatformConfigFile = GConfig->Find(GDeviceProfilesIni);
+	const FString SectionName = FString::Printf(TEXT("%s %s"), *ProfileName, *UDeviceProfile::StaticClass()->GetName());
+	return PlatformConfigFile->Contains(SectionName);
+}
+
 
 void UDeviceProfileManager::DeleteProfile( UDeviceProfile* Profile )
 {
