Fix Linux/Android build
#rnx

diff --git a/Engine/Source/Runtime/Engine/Private/DeviceProfiles/DeviceProfileManager.cpp b/Engine/Source/Runtime/Engine/Private/DeviceProfiles/DeviceProfileManager.cpp
--- a/Engine/Source/Runtime/Engine/Private/DeviceProfiles/DeviceProfileManager.cpp
+++ b/Engine/Source/Runtime/Engine/Private/DeviceProfiles/DeviceProfileManager.cpp
@@ -770,7 +770,7 @@ bool UDeviceProfileManager::HasLoadableProfileName(const FString& ProfileName, F
 #else
 
 		checkf(OptionalPlatformName == FName(FPlatformProperties::IniPlatformName()), 
-			"UDeviceProfileManager::HasLoadableProfileName - This platform cannot load configurations for other platforms.");
+			TEXT("UDeviceProfileManager::HasLoadableProfileName - This platform cannot load configurations for other platforms."));
 #endif
 	}
 
