Move UCSModifiedProperties in to sparse storage as only a small percentage of all components stores data in it

diff --git a/Engine/Source/Runtime/Core/Public/UObject/FortniteReleaseBranchCustomObjectVersion.h b/Engine/Source/Runtime/Core/Public/UObject/FortniteReleaseBranchCustomObjectVersion.h
--- a/Engine/Source/Runtime/Core/Public/UObject/FortniteReleaseBranchCustomObjectVersion.h
+++ b/Engine/Source/Runtime/Core/Public/UObject/FortniteReleaseBranchCustomObjectVersion.h
@@ -21,6 +21,9 @@ struct CORE_API FFortniteReleaseBranchCustomObjectVersion
 		// Chaos::TKinematicTarget no longer stores a full transform, only position/rotation.
 		ChaosKinematicTargetRemoveScale,
 
+		// Move UCSModifiedProperties out of ActorComponent and in to sparse storage
+		ActorComponentUCSModifiedPropertiesSparseStorage,
+
 		// -----<new versions can be added above this line>-------------------------------------------------
 		VersionPlusOne,
 		LatestVersion = VersionPlusOne - 1
diff --git a/Engine/Source/Runtime/Engine/Classes/Components/ActorComponent.h b/Engine/Source/Runtime/Engine/Classes/Components/ActorComponent.h
--- a/Engine/Source/Runtime/Engine/Classes/Components/ActorComponent.h
+++ b/Engine/Source/Runtime/Engine/Classes/Components/ActorComponent.h
@@ -341,6 +341,9 @@ public:
 	/** Removes specified properties from the list of UCS-modified properties */
 	void RemoveUCSModifiedProperties(const TArray<FProperty*>& Properties);
 
+	/** Clears the component's UCS modified properties */
+	void ClearUCSModifiedProperties();
+
 	/** True if this component can be modified when it was inherited from a parent actor class */
 	bool IsEditableWhenInherited() const;
 
@@ -402,9 +405,14 @@ public:
 	UPROPERTY(BlueprintAssignable, Category = "Components|Activation")
 	FActorComponentDeactivateSignature OnComponentDeactivated;
 
+#if WITH_EDITORONLY_DATA
 private:
 	UPROPERTY()
-	TArray<FSimpleMemberReference> UCSModifiedProperties;
+	TArray<FSimpleMemberReference> UCSModifiedProperties_DEPRECATED;
+#endif
+
+	static FRWLock AllUCSModifiedPropertiesLock;
+	static TMap<UActorComponent*, TArray<FSimpleMemberReference>> AllUCSModifiedProperties;
 
 public:
 	/**
@@ -950,6 +958,7 @@ public:
 	virtual bool Rename( const TCHAR* NewName=NULL, UObject* NewOuter=NULL, ERenameFlags Flags=REN_None ) override;
 	virtual void PostRename(UObject* OldOuter, const FName OldName) override;
 	virtual void Serialize(FArchive& Ar) override;
+	static void AddReferencedObjects(UObject* InThis, FReferenceCollector& Collector);
 #if WITH_EDITOR
 	virtual bool Modify( bool bAlwaysMarkDirty = true ) override;
 	virtual bool CanEditChange(const FProperty* InProperty) const override;
@@ -1087,7 +1096,6 @@ private:
 
 	friend struct FMarkComponentEndOfFrameUpdateState;
 	friend struct FSetUCSSerializationIndex;
-	friend struct FActorComponentInstanceData;
 	friend class FActorComponentDetails;
 	friend class FComponentReregisterContextBase;
 	friend class FComponentRecreateRenderStateContext;
diff --git a/Engine/Source/Runtime/Engine/Private/ComponentInstanceDataCache.cpp b/Engine/Source/Runtime/Engine/Private/ComponentInstanceDataCache.cpp
--- a/Engine/Source/Runtime/Engine/Private/ComponentInstanceDataCache.cpp
+++ b/Engine/Source/Runtime/Engine/Private/ComponentInstanceDataCache.cpp
@@ -386,7 +386,7 @@ void FActorComponentInstanceData::ApplyToComponent(UActorComponent* Component, c
 		else if (CacheApplyPhase == ECacheApplyPhase::NonConstructionScript)
 		{
 			// When this case is used, we want to apply all properties, even UCS modified ones
-			Component->UCSModifiedProperties.Reset();
+			Component->ClearUCSModifiedProperties();
 		}
 
 		for (const FActorComponentDuplicatedObjectData& DuplicatedObjectData : DuplicatedObjects)
diff --git a/Engine/Source/Runtime/Engine/Private/Components/ActorComponent.cpp b/Engine/Source/Runtime/Engine/Private/Components/ActorComponent.cpp
--- a/Engine/Source/Runtime/Engine/Private/Components/ActorComponent.cpp
+++ b/Engine/Source/Runtime/Engine/Private/Components/ActorComponent.cpp
@@ -34,6 +34,7 @@
 #include "HAL/LowLevelMemTracker.h"
 #include "Net/Core/PushModel/PushModel.h"
 #include "UObject/FrameworkObjectVersion.h"
+#include "UObject/FortniteReleaseBranchCustomObjectVersion.h"
 #include "Async/ParallelFor.h"
 
 #if WITH_EDITOR
@@ -300,6 +301,8 @@ FActorComponentGlobalDestroyPhysicsSignature UActorComponent::GlobalDestroyPhysi
 UActorComponent::FOnMarkRenderStateDirty UActorComponent::MarkRenderStateDirtyEvent;
 
 const FString UActorComponent::ComponentTemplateNameSuffix(TEXT("_GEN_VARIABLE"));
+TMap<UActorComponent*, TArray<FSimpleMemberReference>> UActorComponent::AllUCSModifiedProperties;
+FRWLock UActorComponent::AllUCSModifiedPropertiesLock;
 
 UActorComponent::UActorComponent(const FObjectInitializer& ObjectInitializer /*= FObjectInitializer::Get()*/)
 	: Super(ObjectInitializer)
@@ -432,13 +435,23 @@ void UActorComponent::PostLoad()
 		{
 			DetermineUCSModifiedProperties();
 		}
+#if WITH_EDITORONLY_DATA
+		else if (GetLinkerCustomVersion(FFortniteReleaseBranchCustomObjectVersion::GUID) < FFortniteReleaseBranchCustomObjectVersion::ActorComponentUCSModifiedPropertiesSparseStorage)
+		{
+			if (UCSModifiedProperties_DEPRECATED.Num())
+			{
+				FRWScopeLock Lock(AllUCSModifiedPropertiesLock, SLT_Write);
+				AllUCSModifiedProperties.Add(this, MoveTemp(UCSModifiedProperties_DEPRECATED));
+			}
+		}
+#endif
 	}
 	else
 	{
+#if WITH_EDITORONLY_DATA
 		// For a brief period of time we were inadvertently storing these for all components, need to clear it out
-		UCSModifiedProperties.Empty();
+		UCSModifiedProperties_DEPRECATED.Empty();
 
-#if WITH_EDITORONLY_DATA
 		if (CreationMethod == EComponentCreationMethod::UserConstructionScript)
 		{
 			if (GetLinkerCustomVersion(FFrameworkObjectVersion::GUID) < FFrameworkObjectVersion::StoringUCSSerializationIndex)
@@ -722,6 +735,8 @@ void UActorComponent::BeginDestroy()
 	UEngineElementsLibrary::DestroyEditorComponentElement(this);
 #endif	// WITH_EDITOR
 
+	ClearUCSModifiedProperties();
+
 	Super::BeginDestroy();
 }
 
@@ -2188,10 +2203,10 @@ bool UActorComponent::IsEditableWhenInherited() const
 
 void UActorComponent::DetermineUCSModifiedProperties()
 {
-	UCSModifiedProperties.Empty();
-
 	if (CreationMethod == EComponentCreationMethod::SimpleConstructionScript)
 	{
+		TArray<FSimpleMemberReference> UCSModifiedProperties;
+
 		class FComponentPropertySkipper : public FArchive
 		{
 		public:
@@ -2238,24 +2253,60 @@ void UActorComponent::DetermineUCSModifiedProperties()
 				}
 			}
 		}
+
+		FRWScopeLock Lock(AllUCSModifiedPropertiesLock, SLT_Write);
+		if (UCSModifiedProperties.Num() > 0)
+		{
+			AllUCSModifiedProperties.Add(this, MoveTemp(UCSModifiedProperties));
+		}
+		else
+		{
+			AllUCSModifiedProperties.Remove(this);
+		}
 	}
 }
 
 void UActorComponent::GetUCSModifiedProperties(TSet<const FProperty*>& ModifiedProperties) const
 {
-	for (const FSimpleMemberReference& MemberReference : UCSModifiedProperties)
+	FRWScopeLock Lock(AllUCSModifiedPropertiesLock, SLT_ReadOnly);
+	if (TArray<FSimpleMemberReference>* UCSModifiedProperties = AllUCSModifiedProperties.Find(this))
 	{
-		ModifiedProperties.Add(FMemberReference::ResolveSimpleMemberReference<FProperty>(MemberReference));
+		for (const FSimpleMemberReference& MemberReference : *UCSModifiedProperties)
+		{
+			ModifiedProperties.Add(FMemberReference::ResolveSimpleMemberReference<FProperty>(MemberReference));
+		}
 	}
 }
 
 void UActorComponent::RemoveUCSModifiedProperties(const TArray<FProperty*>& Properties)
 {
-	for (FProperty* Property : Properties)
+	FRWScopeLock Lock(AllUCSModifiedPropertiesLock, SLT_Write);
+	if (TArray<FSimpleMemberReference>* UCSModifiedProperties = AllUCSModifiedProperties.Find(this))
+	{
+		for (FProperty* Property : Properties)
+		{
+			FSimpleMemberReference MemberReference;
+			FMemberReference::FillSimpleMemberReference<FProperty>(Property, MemberReference);
+			UCSModifiedProperties->RemoveSwap(MemberReference);
+		}
+	}
+}
+
+void UActorComponent::ClearUCSModifiedProperties()
+{
+	FRWScopeLock Lock(AllUCSModifiedPropertiesLock, SLT_Write);
+	AllUCSModifiedProperties.Remove(this);
+}
+
+void UActorComponent::AddReferencedObjects(UObject* InThis, FReferenceCollector& Collector)
+{
+	FRWScopeLock Lock(AllUCSModifiedPropertiesLock, SLT_ReadOnly);
+	if (TArray<FSimpleMemberReference>* UCSModifiedProperties = AllUCSModifiedProperties.Find(CastChecked<UActorComponent>(InThis)))
 	{
-		FSimpleMemberReference MemberReference;
-		FMemberReference::FillSimpleMemberReference<FProperty>(Property, MemberReference);
-		UCSModifiedProperties.RemoveSwap(MemberReference);
+		for (FSimpleMemberReference& MemberReference : *UCSModifiedProperties)
+		{
+			Collector.AddReferencedObject(MemberReference.MemberParent);
+		}
 	}
 }
 
@@ -2291,11 +2342,44 @@ void UActorComponent::Serialize(FArchive& Ar)
 	Super::Serialize(Ar);
 
 	Ar.UsingCustomVersion(FFrameworkObjectVersion::GUID);
+	Ar.UsingCustomVersion(FFortniteReleaseBranchCustomObjectVersion::GUID);
 
 	if (Ar.IsLoading() && (Ar.HasAnyPortFlags(PPF_DuplicateForPIE)||!Ar.HasAnyPortFlags(PPF_Duplicate)) && !IsTemplate())
 	{
 		bHasBeenCreated = true;
 	}
+
+	if (Ar.CustomVer(FFortniteReleaseBranchCustomObjectVersion::GUID) >= FFortniteReleaseBranchCustomObjectVersion::ActorComponentUCSModifiedPropertiesSparseStorage)
+	{
+		if (Ar.IsLoading())
+		{
+			TArray<FSimpleMemberReference> UCSModifiedProperties;
+			Ar << UCSModifiedProperties;
+
+			FRWScopeLock Lock(AllUCSModifiedPropertiesLock, SLT_Write);
+			if (UCSModifiedProperties.Num() > 0)
+			{
+				AllUCSModifiedProperties.Add(this, MoveTemp(UCSModifiedProperties));
+			}
+			else
+			{
+				AllUCSModifiedProperties.Remove(this);
+			}
+		}
+		else
+		{
+			FRWScopeLock Lock(AllUCSModifiedPropertiesLock, SLT_ReadOnly);
+			if (TArray<FSimpleMemberReference>* UCSModifiedProperties = AllUCSModifiedProperties.Find(this))
+			{
+				Ar << *UCSModifiedProperties;
+			}
+			else
+			{
+				TArray<FSimpleMemberReference> EmptyUCSModifiedProperties;
+				Ar << EmptyUCSModifiedProperties;
+			}
+		}
+	}
 }
 
 void UActorComponent::PostApplyToComponent()
