Reduce size of UMaterialInstanceDynamic by 8 bytes to drop into smaller FMallocBinned2 pool.

diff --git a/Engine/Source/Runtime/Engine/Classes/Materials/MaterialInstance.h b/Engine/Source/Runtime/Engine/Classes/Materials/MaterialInstance.h
--- a/Engine/Source/Runtime/Engine/Classes/Materials/MaterialInstance.h
+++ b/Engine/Source/Runtime/Engine/Classes/Materials/MaterialInstance.h
@@ -525,9 +525,14 @@ class UMaterialInstance : public UMaterialInterface
 	uint8 bCastDynamicShadowAsMasked : 1;
 	uint8 bOutputTranslucentVelocity : 1;
 	uint8 bIsShadingModelFromMaterialExpression : 1;
+protected:
+	uint8 bLoadedCachedData : 1;
+public:
 
 	TEnumAsByte<EBlendMode> BlendMode;
 
+	FMaterialShadingModelField ShadingModels;
+
 	//Cached copies of the base property overrides or the value from the parent to avoid traversing the parent chain for each access.
 	float OpacityMaskClipValue;
 
@@ -584,8 +589,6 @@ class UMaterialInstance : public UMaterialInterface
 	uint32 OffsetToFirstResource;
 #endif
 
-	FMaterialShadingModelField ShadingModels;
-
 #if WITH_EDITOR
 	/** Flag to detect cycles in the material instance graph, this is only used at content creation time where the hierarchy can be changed. */
 	bool ReentrantFlag[2];
@@ -622,7 +625,6 @@ private:
 #endif
 
 protected:
-	bool bLoadedCachedData;
 	TUniquePtr<FMaterialInstanceCachedData> CachedData;
 
 private:
