remove temp comments
behavior nop

diff --git a/Engine/Source/Runtime/Engine/Private/SkeletalMesh.cpp b/Engine/Source/Runtime/Engine/Private/SkeletalMesh.cpp
--- a/Engine/Source/Runtime/Engine/Private/SkeletalMesh.cpp
+++ b/Engine/Source/Runtime/Engine/Private/SkeletalMesh.cpp
@@ -1713,7 +1713,6 @@ void USkeletalMesh::Serialize( FArchive& Ar )
 	Ar.UsingCustomVersion(FRenderingObjectVersion::GUID);
 	Ar.UsingCustomVersion(FFortniteMainBranchObjectVersion::GUID);
 	Ar.UsingCustomVersion(FNiagaraObjectVersion::GUID);
-	//@@!! NOTE: FUE5MainStreamObjectVersion was not here but does seem to be in the Archive anyway
 	Ar.UsingCustomVersion(FUE5MainStreamObjectVersion::GUID);
 
 	FStripDataFlags StripFlags( Ar );
@@ -1960,17 +1959,6 @@ void USkeletalMesh::Serialize( FArchive& Ar )
 			ThisLODInfo.BuildSettings.bUseBackwardsCompatibleF16TruncUVs = true;
 		}
 	}
-	
-	//@@!!  NOTE : will remove
-	// FUE5MainStreamObjectVersion is present in UE5 saves of SkeletalMesh
-	//	I'm not sure who's putting it there
-	//  but it does seem like I can use it even without an explicit Using here
-	/*
-	if (Ar.IsLoading() && Ar.CustomVer(FUE5MainStreamObjectVersion::GUID) >= FUE5MainStreamObjectVersion::DirLightsAreAtmosphereLightsByDefault)
-	{
-		UE_LOG(LogSkeletalMesh, Display, TEXT("Skeletal mesh %s has FUE5MainStreamObjectVersion"), *GetFullName() );
-	}
-	*/
 }
 
 void USkeletalMesh::DeclareCustomVersions(FArchive& Ar)
diff --git a/Engine/Source/Runtime/Engine/Private/StaticMesh.cpp b/Engine/Source/Runtime/Engine/Private/StaticMesh.cpp
--- a/Engine/Source/Runtime/Engine/Private/StaticMesh.cpp
+++ b/Engine/Source/Runtime/Engine/Private/StaticMesh.cpp
@@ -5134,7 +5134,6 @@ void UStaticMesh::Serialize(FArchive& Ar)
 			FStaticMeshSourceModel& SrcModel = GetSourceModel(i);
 			SrcModel.SerializeBulkData(Ar, this);
 
-			//@@!!
 			// Automatically detect assets saved before CL 16135278 which changed F16 to RTNE
 			//	set them to bUseBackwardsCompatibleF16TruncUVs	
 			if ( Ar.IsLoading() && Ar.CustomVer(FUE5MainStreamObjectVersion::GUID) < FUE5MainStreamObjectVersion::DirLightsAreAtmosphereLightsByDefault)
