VTDataBuilder make sure BuildTexture debug path names are unique
output nop
these are the names used for the uncompressed processing path so this is not the issue we saw in oodle with the same name being used for multiples slices

diff --git a/Engine/Source/Runtime/Engine/Private/VT/VirtualTextureDataBuilder.cpp b/Engine/Source/Runtime/Engine/Private/VT/VirtualTextureDataBuilder.cpp
--- a/Engine/Source/Runtime/Engine/Private/VT/VirtualTextureDataBuilder.cpp
+++ b/Engine/Source/Runtime/Engine/Private/VT/VirtualTextureDataBuilder.cpp
@@ -682,6 +682,7 @@ void FVirtualTextureDataBuilder::BuildTiles(const TArray<FVTSourceTileEntry>& Ti
 			TArray<FCompressedImage2D> CompressedMip;
 			TArray<FImage> EmptyList;
 			uint32 NumMipsInTail, ExtData;
+			// this is the Build for Tiles to do the encode to GPU formats, with no processing
 			if (!ensure(Compressor->BuildTexture(TileImages, EmptyList, TBSettings, DebugTilePathName, CompressedMip, NumMipsInTail, ExtData)))
 			{
 				bCompressionError = true;
@@ -983,14 +984,18 @@ void FVirtualTextureDataBuilder::BuildSourcePixels(const FTextureSourceData& Sou
 				--BlockData.MipBias;
 				LocalBlockSizeScale *= 2;
 			}
+			
+			// give each tile a unique DebugTexturePathName for DebugDump option :
+			FString CurDebugTexturePathName = FString::Printf(TEXT("%s_L%d_B%d"), *DebugTexturePathName, LayerIndex, BlockIndex);
 
 			// Use the texture compressor module to do all the hard work
+			// this is the Build to Uncompressed to apply processing to create the source for the tiles
 			TArray<FCompressedImage2D> CompressedMips;
 			bool bBuildTextureResult = false;
 			if (LocalBlockSizeScale == 1)
 			{
 				uint32 NumMipsInTail, ExtData;
-				bBuildTextureResult = Compressor->BuildTexture(SourceMips, *CompositeSourceMips, TBSettings, DebugTexturePathName, CompressedMips, NumMipsInTail, ExtData);
+				bBuildTextureResult = Compressor->BuildTexture(SourceMips, *CompositeSourceMips, TBSettings, CurDebugTexturePathName, CompressedMips, NumMipsInTail, ExtData);
 			}
 			else
 			{
@@ -1013,7 +1018,7 @@ void FVirtualTextureDataBuilder::BuildSourcePixels(const FTextureSourceData& Sou
 				}
 
 				uint32 NumMipsInTail, ExtData;
-				bBuildTextureResult = Compressor->BuildTexture(ScaledSourceMips, ScaledCompositeMips, TBSettings, DebugTexturePathName, CompressedMips, NumMipsInTail, ExtData);
+				bBuildTextureResult = Compressor->BuildTexture(ScaledSourceMips, ScaledCompositeMips, TBSettings, CurDebugTexturePathName, CompressedMips, NumMipsInTail, ExtData);
 			}
 
 			check(bBuildTextureResult);
@@ -1143,11 +1148,15 @@ void FVirtualTextureDataBuilder::BuildSourcePixels(const FTextureSourceData& Sou
 				TBSettings.MipGenSettings = TMGS_SimpleAverage;
 			}
 
+			// give each tile a unique DebugTexturePathName for DebugDump option :
+			FString CurDebugTexturePathName = FString::Printf(TEXT("%s_L%d_MT"), *DebugTexturePathName, LayerIndex);
+
 			// Use the texture compressor module to do all the hard work
 			// TODO - composite images?
 			TArray<FCompressedImage2D> CompressedMips;
 			uint32 NumMipsInTail, ExtData;
-			if (!Compressor->BuildTexture(MiptailInputImages, EmptyImageArray, TBSettings, DebugTexturePathName, CompressedMips, NumMipsInTail, ExtData))
+			// this is a Build to uncompressed, to apply processing
+			if (!Compressor->BuildTexture(MiptailInputImages, EmptyImageArray, TBSettings, CurDebugTexturePathName, CompressedMips, NumMipsInTail, ExtData))
 			{
 				check(false);
 			}
