Invalidate the SSS Profile when there is a change.

diff --git a/Engine/Source/Runtime/Engine/Private/Rendering/SubsurfaceProfile.cpp b/Engine/Source/Runtime/Engine/Private/Rendering/SubsurfaceProfile.cpp
--- a/Engine/Source/Runtime/Engine/Private/Rendering/SubsurfaceProfile.cpp
+++ b/Engine/Source/Runtime/Engine/Private/Rendering/SubsurfaceProfile.cpp
@@ -693,6 +693,7 @@ void USubsurfaceProfile::PostEditChangeProperty(struct FPropertyChangedEvent& Pr
 {
 	const FSubsurfaceProfileStruct SettingsLocal = this->Settings;
 	USubsurfaceProfile* Profile = this;
+	GetRendererModule().InvalidatePathTracedOutput();
 	ENQUEUE_RENDER_COMMAND(UpdateSubsurfaceProfile)(
 		[SettingsLocal, Profile](FRHICommandListImmediate& RHICmdList)
 		{
