UBlueprintAsyncActionBase:
- Made WorldContextObject parameter const in RegisterWithGameInstance.

UEditorUtilityBlueprintAsyncActionBase:
- Made WorldContextObject parameter const in overridden RegisterWithGameInstance.

diff --git a/Engine/Source/Editor/Blutility/Private/EditorUtilityLibrary.cpp b/Engine/Source/Editor/Blutility/Private/EditorUtilityLibrary.cpp
--- a/Engine/Source/Editor/Blutility/Private/EditorUtilityLibrary.cpp
+++ b/Engine/Source/Editor/Blutility/Private/EditorUtilityLibrary.cpp
@@ -20,7 +20,7 @@ UEditorUtilityBlueprintAsyncActionBase::UEditorUtilityBlueprintAsyncActionBase(c
 {
 }
 
-void UEditorUtilityBlueprintAsyncActionBase::RegisterWithGameInstance(UObject* WorldContextObject)
+void UEditorUtilityBlueprintAsyncActionBase::RegisterWithGameInstance(const UObject* WorldContextObject)
 {
 	UEditorUtilitySubsystem* EditorUtilitySubsystem = GEditor->GetEditorSubsystem<UEditorUtilitySubsystem>();
 	EditorUtilitySubsystem->RegisterReferencedObject(this);
diff --git a/Engine/Source/Editor/Blutility/Public/EditorUtilityLibrary.h b/Engine/Source/Editor/Blutility/Public/EditorUtilityLibrary.h
--- a/Engine/Source/Editor/Blutility/Public/EditorUtilityLibrary.h
+++ b/Engine/Source/Editor/Blutility/Public/EditorUtilityLibrary.h
@@ -20,7 +20,7 @@ class BLUTILITY_API UEditorUtilityBlueprintAsyncActionBase : public UBlueprintAs
 	GENERATED_UCLASS_BODY()
 
 public:
-	virtual void RegisterWithGameInstance(UObject* WorldContextObject) override;
+	virtual void RegisterWithGameInstance(const UObject* WorldContextObject) override;
 	virtual void SetReadyToDestroy() override;
 };
 
diff --git a/Engine/Source/Runtime/Engine/Classes/Kismet/BlueprintAsyncActionBase.h b/Engine/Source/Runtime/Engine/Classes/Kismet/BlueprintAsyncActionBase.h
--- a/Engine/Source/Runtime/Engine/Classes/Kismet/BlueprintAsyncActionBase.h
+++ b/Engine/Source/Runtime/Engine/Classes/Kismet/BlueprintAsyncActionBase.h
@@ -31,7 +31,7 @@ public:
 	 * Call to globally register this object with a game instance, it will not be destroyed until SetReadyToDestroy is called
 	 * This allows having an action stay alive until SetReadyToDestroy is manually called, allowing it to be used inside loops or if the calling BP goes away
 	 */
-	virtual void RegisterWithGameInstance(UObject* WorldContextObject);
+	virtual void RegisterWithGameInstance(const UObject* WorldContextObject);
 
 	/** Call when the action is completely done, this makes the action free to delete, and will unregister it with the game instance */
 	virtual void SetReadyToDestroy();
diff --git a/Engine/Source/Runtime/Engine/Private/Kismet/BlueprintAsyncActionBase.cpp b/Engine/Source/Runtime/Engine/Private/Kismet/BlueprintAsyncActionBase.cpp
--- a/Engine/Source/Runtime/Engine/Private/Kismet/BlueprintAsyncActionBase.cpp
+++ b/Engine/Source/Runtime/Engine/Private/Kismet/BlueprintAsyncActionBase.cpp
@@ -20,7 +20,7 @@ void UBlueprintAsyncActionBase::Activate()
 {
 }
 
-void UBlueprintAsyncActionBase::RegisterWithGameInstance(UObject* WorldContextObject)
+void UBlueprintAsyncActionBase::RegisterWithGameInstance(const UObject* WorldContextObject)
 {
 	UWorld* FoundWorld = GEngine->GetWorldFromContextObject(WorldContextObject, EGetWorldErrorMode::LogAndReturnNull);
 
