Fix the Total Shaders compiled counter and stat to account for the job cache ddc.

- We were potentially over counting actual shaders compiled if we found a job already in the cache.
- Increase the counter when we finish a job that was not cached.
- Use JobsCompleted on the FShaderCompilerStats as a proxy for the number of total shaders compiled instead of counting the number copmiled in the ShaderCompilerStats (Because these are wrong and don't account for the job cache).
- Add a trace cpu profiler scope to FinishCompilation.

diff --git a/Engine/Source/Runtime/Engine/Private/ShaderCompiler/ShaderCompiler.cpp b/Engine/Source/Runtime/Engine/Private/ShaderCompiler/ShaderCompiler.cpp
--- a/Engine/Source/Runtime/Engine/Private/ShaderCompiler/ShaderCompiler.cpp
+++ b/Engine/Source/Runtime/Engine/Private/ShaderCompiler/ShaderCompiler.cpp
@@ -635,11 +635,13 @@ void FShaderCompileJobCollection::HandlePrintStats()
 	GShaderCompilingManager->PrintStats(true);
 }
 
+TRACE_DECLARE_INT_COUNTER(Shaders_Compiled, TEXT("Shaders/Compiled"));
 void FShaderCompileJobCollection::ProcessFinishedJob(FShaderCommonCompileJob* FinishedJob, bool bWasCached)
 {
 	if (!bWasCached)
 	{
 		GShaderCompilerStats->RegisterFinishedJob(*FinishedJob);
+		TRACE_COUNTER_ADD(Shaders_Compiled, 1);
 	}
 
 	// TODO: have a pending shader map critical section? not clear at this point if we can be accessing the results on another thread at the same time
@@ -3049,29 +3051,8 @@ void FShaderCompilerStats::WriteStatSummary()
 
 uint32 FShaderCompilerStats::GetTotalShadersCompiled()
 {
-	TRACE_CPUPROFILER_EVENT_SCOPE(FShaderCompilerStats::GetTotalShadersCompiled);
-
-	uint32 TotalCompiled = 0;
-
-	{
-		FScopeLock Lock(&CompileStatsLock);
-		const TSparseArray<ShaderCompilerStats>& PlatformStats = GetShaderCompilerStats();
-		for (int32 Platform = 0; Platform < PlatformStats.GetMaxIndex(); ++Platform)
-		{
-			if (PlatformStats.IsValidIndex(Platform))
-			{
-				const ShaderCompilerStats& Stats = PlatformStats[Platform];
-				for (const auto& Pair : Stats)
-				{
-					const FShaderCompilerStats::FShaderStats& SingleStats = Pair.Value;
-
-					TotalCompiled += SingleStats.Compiled;
-				}
-			}
-		}
-	}
-
-	return TotalCompiled;
+	FScopeLock Lock(&CompileStatsLock);
+	return JobsCompleted;
 }
 
 void FShaderCompilerStats::RegisterLocalWorkerIdleTime(double IdleTime)
@@ -3246,11 +3227,8 @@ void FShaderCompilerStats::RegisterCookedShaders(uint32 NumCooked, float Compile
 	}
 }
 
-TRACE_DECLARE_INT_COUNTER(Shaders_Compiled, TEXT("Shaders/Compiled"));
 void FShaderCompilerStats::RegisterCompiledShaders(uint32 NumCompiled, EShaderPlatform Platform, const FString MaterialPath, FString PermutationString)
 {
-	TRACE_COUNTER_ADD(Shaders_Compiled, NumCompiled);
-
 	FScopeLock Lock(&CompileStatsLock);
 	if (!CompileStats.IsValidIndex(Platform))
 	{
@@ -4939,6 +4917,8 @@ void FShaderCompilingManager::CancelCompilation(const TCHAR* MaterialName, const
 
 void FShaderCompilingManager::FinishCompilation(const TCHAR* MaterialName, const TArray<int32>& ShaderMapIdsToFinishCompiling)
 {
+	TRACE_CPUPROFILER_EVENT_SCOPE(FShaderCompilingManager::FinishCompilation);
+
 	// nothing to do
 	if (!AllowShaderCompiling())
 	{
