# Match pathtracer and rasterizer/raytracing SSS

## Major benefit
 * SSS rasterizer matches as close as possible to pathtracer without transmission. Ultimate goal is the user can use rasterizer to look and dev subsurface scattering.
 * SSProfile will use mfp. Next is to unify all scattering parameters to mfp.
 * Smooth upgrade for rasterizer/raytracing real-time renderer. Path tracer appearance will match rasterizer as previously the parameters are used differently for path tracer and rasterizer.
 * Path tracer supports SSS below 1mm.

## The break down of the CL

 0. Use mfp instead of dmfp to select the channel to scatter (CL19255292).
 1. Match burley high quality rasterizer with pathtracer in ideal SSS condition (CL19227570).
 2. Match the low quality burley, separable filter(CL19244046).
 3. Upgrade dmfp to mfp in UI. (This one includings all previous CLs).

## Issue
There are two issues I am hunting down

0. Need to scale the scatter distance by 2 for separable filtering so that it matches MC sampling in high quality Burley.
1. Need to scale down the mfp by a magic number 0.6 to match path tracer with high quality burley. However it has no impact for new SS profiles. The mfp is as it is.

diff --git a/Engine/Shaders/Private/BurleyNormalizedSSSCommon.ush b/Engine/Shaders/Private/BurleyNormalizedSSSCommon.ush
--- a/Engine/Shaders/Private/BurleyNormalizedSSSCommon.ush
+++ b/Engine/Shaders/Private/BurleyNormalizedSSSCommon.ush
@@ -94,6 +94,21 @@ float3 GetSearchLightDiffuseScalingFactor3D(float3 SurfaceAlbedo)
 	return 3.5 + 100 * Value * Value * Value * Value;
 }
 
+float3 GetMFPFromDMFPApprox(float3 SurfaceAlbedo, float3 TargetSurfaceAlbedo, float3 DMFP)
+{
+	// TODO: find the source of this magic number to convert dmfp to mfp to match path tracer and rasterizer.
+	// DMFP/SearchLightDiffuseScalingFactor goes to the normalized distance(d=l/s), then multiply the PerpendicularScalingFactor to get the mfp
+
+	const float Dmfp2MfpMagicNumber = 0.6f;
+	float3 DMFPToMFPCoeff = Dmfp2MfpMagicNumber * GetPerpendicularScalingFactor3D(TargetSurfaceAlbedo) / GetSearchLightDiffuseScalingFactor3D(SurfaceAlbedo);
+	return DMFPToMFPCoeff * DMFP;
+	
+	// Mathmatically matching based on diffusion coefficient instead of burley's approximation. However, it leads to incorrect result as 
+	// we use burley's approximation (IOR=1.0) for screenspace diffuse scattering.
+	//float3 Alpha = 1 - exp(-11.43 * SurfaceAlbedo + 15.38 * SurfaceAlbedo * SurfaceAlbedo - 13.91 * SurfaceAlbedo * SurfaceAlbedo * SurfaceAlbedo);
+	//return DMFP * sqrt(3 * (1 - Alpha) / (2 - Alpha));
+}
+
 float GetScalingFactor(float A)
 {
 #if LIGHT_PERPENDICULAR
diff --git a/Engine/Shaders/Private/PathTracing/Material/PathTracingMaterialSampling.ush b/Engine/Shaders/Private/PathTracing/Material/PathTracingMaterialSampling.ush
--- a/Engine/Shaders/Private/PathTracing/Material/PathTracingMaterialSampling.ush
+++ b/Engine/Shaders/Private/PathTracing/Material/PathTracingMaterialSampling.ush
@@ -402,13 +402,13 @@ FPathTracingPayload UnpackPathTracingPayloadMaterial(FPackedPathTracingPayload P
 				// Burley parameterization
 				float3 SurfaceAlbedo = View.SSProfilesTexture.Load(int3(BSSS_SURFACEALBEDO_OFFSET, ProfileId, 0)).xyz;
 				float3 DiffuseMeanFreePath = DecodeDiffuseMeanFreePath(View.SSProfilesTexture.Load(int3(BSSS_DMFP_OFFSET, ProfileId, 0))).xyz;
-				float WorldUnitScale = DecodeWorldUnitScale(View.SSProfilesTexture.Load(int3(SSSS_TINT_SCALE_OFFSET, ProfileId, 0)).a) * BURLEY_MM_2_CM;
+				float WorldUnitScale = DecodeWorldUnitScale(View.SSProfilesTexture.Load(int3(SSSS_TINT_SCALE_OFFSET, ProfileId, 0)).a) * BURLEY_CM_2_MM;
 
 				// Opacity acts as a per-pixel radius multiplier
 				// NOTE: this seems backwards? Opacity=0 acts like default-lit while Opacity=1 acts like SSS?
-				float3 SSSRadius = HitPayload.Opacity * WorldUnitScale * DiffuseMeanFreePath * GetDiffuseSurfaceScalingFactor3D(SurfaceAlbedo);
+				float3 SSSRadius =GetMFPFromDMFPApprox(SurfaceAlbedo, HitPayload.DiffuseColor, HitPayload.Opacity * WorldUnitScale * DiffuseMeanFreePath);
 
-				HitPayload.SetSubsurfaceRadius(SSSRadius);
+				HitPayload.SetSubsurfaceRadius(SSSRadius * BURLEY_MM_2_CM);
 
 				// burley mode doesn't have any diffuse component
 				HitPayload.SetSubsurfaceColor(HitPayload.DiffuseColor);
diff --git a/Engine/Shaders/Private/PathTracing/PathTracingCore.ush b/Engine/Shaders/Private/PathTracing/PathTracingCore.ush
--- a/Engine/Shaders/Private/PathTracing/PathTracingCore.ush
+++ b/Engine/Shaders/Private/PathTracing/PathTracingCore.ush
@@ -515,7 +515,10 @@ bool ProcessSubsurfaceRandomWalk(inout FPathTracingPayload Payload, inout float3
 		Payload.SetSubsurfaceColor(0.0);
 		return true;
 	}
-	SSSRadius = max(SSSRadius, 0.1); // 1mm minimum to make sure random walk can move forward and to keep SigmaT finite
+	
+	// Support SSS for mfp smaller than 1mm.
+	SSSRadius = max(SSSRadius, 0.0009); // 0.009mm minimum to make sure random walk can move forward and to keep SigmaT finite
+
 
 	// decide randomly to evaluate the SSS part of the material, or keep shading the current part
 	// This also computes the energy conservation to take out the portion of lighting reflected away by the specular layer
diff --git a/Engine/Shaders/Private/SubsurfaceBurleyNormalized.ush b/Engine/Shaders/Private/SubsurfaceBurleyNormalized.ush
--- a/Engine/Shaders/Private/SubsurfaceBurleyNormalized.ush
+++ b/Engine/Shaders/Private/SubsurfaceBurleyNormalized.ush
@@ -823,6 +823,9 @@ float2 CalculateBurleyScale(float WorldUnitScale, float DepthAtCenter)
 
 	float SSSScaleX = SubsurfaceParams.x;
 	BurleyScale *= SSSScaleX / DepthAtCenter;
+	
+	// cast from cm to mm for depth, and remove the effect of SUBSURFACE_KERNEL_SIZE. 
+	BurleyScale *= SUBSURFACE_KERNEL_SIZE / BURLEY_CM_2_MM;
 
 	// account for Screen Percentage/Dyanmic Resolution Scaling
 	BurleyScale *= (SubsurfaceInput0_ViewportSize.x * SubsurfaceInput0_ExtentInverse.x);
@@ -861,8 +864,8 @@ float3 CalculateCenterSampleWeight(float Depth, FBurleyParameter BurleyParameter
 
 	float DiffuseMeanFreePath = GetDiffuseMeanFreePathForSampling(BurleyParameter.DiffuseMeanFreePath);
 
-	// To calculate the surface free path from albedo, use diffuse surface scaling.
-	float3 D = DiffuseMeanFreePath / GetDiffuseSurfaceScalingFactor3D(BurleyParameter.SurfaceAlbedo.xyz);
+	// To calculate the surface free path from albedo, use the default scaling.
+	float3 D = DiffuseMeanFreePath / GetScalingFactor3D(BurleyParameter.SurfaceAlbedo.xyz);
 
 	float3 CenterSampleWeight;
 
diff --git a/Engine/Source/Editor/UnrealEd/Private/Factories/EditorFactories.cpp b/Engine/Source/Editor/UnrealEd/Private/Factories/EditorFactories.cpp
--- a/Engine/Source/Editor/UnrealEd/Private/Factories/EditorFactories.cpp
+++ b/Engine/Source/Editor/UnrealEd/Private/Factories/EditorFactories.cpp
@@ -7691,6 +7691,12 @@ UObject* USubsurfaceProfileFactory::FactoryCreateNew(UClass* InClass, UObject* I
 	// the profiles loaded from files to be converted into burley if they
 	// are separable.
 	Object->Settings.bEnableBurley = true;
+
+	// Enable smooth upgrading from DMFP to MFP. This flag change allows
+	// all newly created subsurface profile in MFP, while allowing old profiles
+	// loaded from files to be automatically converted to MFP.
+	Object->Settings.bEnableMeanFreePath = true;
+
 	return Object;
 }
 
diff --git a/Engine/Source/Runtime/Engine/Classes/Engine/SubsurfaceProfile.h b/Engine/Source/Runtime/Engine/Classes/Engine/SubsurfaceProfile.h
--- a/Engine/Source/Runtime/Engine/Classes/Engine/SubsurfaceProfile.h
+++ b/Engine/Source/Runtime/Engine/Classes/Engine/SubsurfaceProfile.h
@@ -43,6 +43,12 @@ struct FSubsurfaceProfileStruct
 	UPROPERTY(Category = "Burley Normalized", EditAnywhere, BlueprintReadOnly, meta = (editcondition = "false", EditConditionHides))
 	bool  bEnableBurley;
 
+	/**
+	 * Switch to use Mean Free Path, otherwise use diffuse mean free path.
+	 */
+	UPROPERTY(Category = "Burley Normalized", EditAnywhere, BlueprintReadOnly, meta = (editcondition = "false", EditConditionHides))
+	bool  bEnableMeanFreePath;
+
 	/**
 	* Specifies the how much of the diffuse light gets into the material,
 	* can be seen as a per-channel mix factor between the original image,
@@ -124,6 +130,8 @@ struct FSubsurfaceProfileStruct
 		WorldUnitScale = 0.1f;
 		TransmissionTintColor = FLinearColor(1.0f, 1.0f, 1.0f);
 		Tint = FLinearColor(1.0f, 1.0f, 1.0f);
+
+		bEnableMeanFreePath = false;
 	}
 
 	void Invalidate()
diff --git a/Engine/Source/Runtime/Engine/Private/Rendering/BurleyNormalizedSSS.cpp b/Engine/Source/Runtime/Engine/Private/Rendering/BurleyNormalizedSSS.cpp
--- a/Engine/Source/Runtime/Engine/Private/Rendering/BurleyNormalizedSSS.cpp
+++ b/Engine/Source/Runtime/Engine/Private/Rendering/BurleyNormalizedSSS.cpp
@@ -111,6 +111,16 @@ FVector GetSearchLightDiffuseScalingFactor(FLinearColor SurfaceAlbedo)
 	);
 }
 
+FLinearColor GetMeanFreePathFromDiffuseMeanFreePath(FLinearColor SurfaceAlbedo, FLinearColor DiffuseMeanFreePath)
+{
+	return DiffuseMeanFreePath * FLinearColor(GetPerpendicularScalingFactor(SurfaceAlbedo) / GetSearchLightDiffuseScalingFactor(SurfaceAlbedo));
+}
+
+FLinearColor GetDiffuseMeanFreePathFromMeanFreePath(FLinearColor SurfaceAlbedo, FLinearColor MeanFreePath)
+{
+	return MeanFreePath * FLinearColor(GetSearchLightDiffuseScalingFactor(SurfaceAlbedo) / GetPerpendicularScalingFactor(SurfaceAlbedo));
+}
+
 void ComputeMirroredBSSSKernel(FLinearColor* TargetBuffer, uint32 TargetBufferSize,
 	FLinearColor SurfaceAlbedo, FLinearColor DiffuseMeanFreePath, float ScatterRadius)
 {
@@ -158,8 +168,8 @@ void ComputeMirroredBSSSKernel(FLinearColor* TargetBuffer, uint32 TargetBufferSi
 			kernel[i].B = t.Z;
 		}
 
-		// We still need to do a small tweak to get the radius to visually match. Multiplying by 4.0 seems to fix it.
-		const float StepScale = 4.0f;
+		// We still need to do a small tweak to get the radius to visually match. Multiplying by 2.0 seems to fix it.
+		const float StepScale = 2.0f;
 		for (int32 i = 0; i < nTotalSamples; i++)
 		{
 			kernel[i].A *= StepScale;
diff --git a/Engine/Source/Runtime/Engine/Private/Rendering/BurleyNormalizedSSS.h b/Engine/Source/Runtime/Engine/Private/Rendering/BurleyNormalizedSSS.h
--- a/Engine/Source/Runtime/Engine/Private/Rendering/BurleyNormalizedSSS.h
+++ b/Engine/Source/Runtime/Engine/Private/Rendering/BurleyNormalizedSSS.h
@@ -27,4 +27,7 @@ void ComputeMirroredBSSSKernel(FLinearColor* TargetBuffer, uint32 TargetBufferSi
 //@param FalloffColor from Separable SSS
 //@param SurfaceAlbedo mapped from FallofColor 
 //@param DiffuseMeanFreePath mapped from FallofColor
-void MapFallOffColor2SurfaceAlbedoAndDiffuseMeanFreePath(float FalloffColor, float& SurfaceAlbedo, float& DiffuseMeanFreePath);
\ No newline at end of file
+void MapFallOffColor2SurfaceAlbedoAndDiffuseMeanFreePath(float FalloffColor, float& SurfaceAlbedo, float& DiffuseMeanFreePath);
+
+FLinearColor GetMeanFreePathFromDiffuseMeanFreePath(FLinearColor SurfaceAlbedo, FLinearColor DiffuseMeanFreePath);
+FLinearColor GetDiffuseMeanFreePathFromMeanFreePath(FLinearColor SurfaceAlbedo, FLinearColor MeanFreePath);
diff --git a/Engine/Source/Runtime/Engine/Private/Rendering/SubsurfaceProfile.cpp b/Engine/Source/Runtime/Engine/Private/Rendering/SubsurfaceProfile.cpp
--- a/Engine/Source/Runtime/Engine/Private/Rendering/SubsurfaceProfile.cpp
+++ b/Engine/Source/Runtime/Engine/Private/Rendering/SubsurfaceProfile.cpp
@@ -17,6 +17,14 @@ static TAutoConsoleVariable<int32> CVarSSProfilesPreIntegratedTextureResolution(
 	ECVF_RenderThreadSafe
 );
 
+static TAutoConsoleVariable<int32> CVarSSProfilesSamplingChannelSelection(
+	TEXT("r.SSProfilesSamplingChannelSelection"),
+	1,
+	TEXT("0. Select the sampling channel based on max DMFP.\n")
+	TEXT("1. based on max MFP."),
+	ECVF_RenderThreadSafe
+);
+
 #if !(UE_BUILD_SHIPPING || UE_BUILD_TEST)
 static TAutoConsoleVariable<int32> CVarSSProfilesPreIntegratedTextureForceUpdate(
 	TEXT("r.SSProfilesPreIntegratedTextureForceUpdate"),
@@ -68,6 +76,37 @@ void UpgradeSeparableToBurley(FSubsurfaceProfileStruct& Settings)
 	Settings.bEnableBurley = true;
 }
 
+// Match the magic number in BurleyNormalizedSSSCommon.ush
+const float Dmfp2MfpMagicNumber = 0.6f;
+
+void UpgradeDiffuseMeanFreePathToMeanFreePath(FSubsurfaceProfileStruct& Settings)
+{
+	// 1. Update dmfp to mean free path
+	const float CmToMm = 10.0f;
+	
+	FLinearColor OldDmfp = Settings.MeanFreePathColor * Settings.MeanFreePathDistance;
+	
+	FLinearColor Mfp = GetMeanFreePathFromDiffuseMeanFreePath(Settings.SurfaceAlbedo, OldDmfp);
+	Mfp *= Dmfp2MfpMagicNumber;
+
+	Settings.MeanFreePathDistance = FMath::Max3(Mfp.R, Mfp.G, Mfp.B);
+	Settings.MeanFreePathColor = Mfp / Settings.MeanFreePathDistance;
+
+	// Support mfp < 0.1f
+	if (Settings.MeanFreePathDistance < 0.1f)
+	{
+		Settings.MeanFreePathColor = Settings.MeanFreePathColor * (Settings.MeanFreePathDistance / 0.1f);
+		Settings.MeanFreePathDistance = 0.1f;
+	}
+
+	Settings.bEnableMeanFreePath = true;
+
+	//2. Fix scaling
+	// Previously, the scaling is scaled up by 1/(SUBSURFACE_KERNEL_SIZE / BURLEY_CM_2_MM). To maintain the same
+	// visual appearance, apply that scale up to world unit scale
+	Settings.WorldUnitScale /= (SUBSURFACE_KERNEL_SIZE / CmToMm);
+}
+
 FSubsurfaceProfileTexture::FSubsurfaceProfileTexture()
 {
 	check(IsInGameThread());
@@ -77,6 +116,12 @@ FSubsurfaceProfileTexture::FSubsurfaceProfileTexture()
 	//The default burley in slot 0 behaves the same as Separable previously
 	UpgradeSeparableToBurley(DefaultSkin);
 
+	//Upgrade DMFP in slot 0 to MFP
+	if (!DefaultSkin.bEnableMeanFreePath)
+	{
+		UpgradeDiffuseMeanFreePathToMeanFreePath(DefaultSkin);
+	}
+
 	// add element 0, it is used as default profile
 	SubsurfaceProfileEntries.Add(FSubsurfaceProfileEntry(DefaultSkin, 0));
 }
@@ -325,14 +370,18 @@ FLinearColor DecodeDiffuseMeanFreePath(FLinearColor EncodedDiffuseMeanFreePath)
 
 void SetupSurfaceAlbedoAndDiffuseMeanFreePath(FLinearColor& SurfaceAlbedo, FLinearColor& Dmfp)
 {
+	int32 SamplingSelectionMethod = FMath::Clamp(CVarSSProfilesSamplingChannelSelection.GetValueOnAnyThread(), 0, 1);
+	FLinearColor Distance = SamplingSelectionMethod == 0 ?
+		Dmfp															// 0: by max diffuse mean free path
+		: GetMeanFreePathFromDiffuseMeanFreePath(SurfaceAlbedo, Dmfp);	// 1: by max mean free path
 	//Store the value that corresponds to the largest Dmfp (diffuse mean free path) channel to A channel.
 	//This is an optimization to shift finding the max correspondence workload
 	//to CPU.
-	const float MaxDmfpComp = FMath::Max3(Dmfp.R, Dmfp.G, Dmfp.B);
-	const uint32 IndexOfMaxDmfp = (Dmfp.R == MaxDmfpComp) ? 0 : ((Dmfp.G == MaxDmfpComp) ? 1 : 2);
+	const float MaxComp = FMath::Max3(Distance.R, Distance.G, Distance.B);
+	const uint32 IndexOfMaxComp = (Distance.R == MaxComp) ? 0 : ((Distance.G == MaxComp) ? 1 : 2);
 
-	SurfaceAlbedo.A = SurfaceAlbedo.Component(IndexOfMaxDmfp);
-	Dmfp.A = MaxDmfpComp;
+	SurfaceAlbedo.A = SurfaceAlbedo.Component(IndexOfMaxComp);
+	Dmfp.A = Dmfp.Component(IndexOfMaxComp);
 }
 
 float Sqrt2(float X)
@@ -378,6 +427,7 @@ void FSubsurfaceProfileTexture::CreateTexture(FRHICommandListImmediate& RHICmdLi
 	// 0.0001f turned out to be too small to fix the issue (for a small KernelSize)
 	const float Bias = 0.009f;
 	const float CmToMm = 10.f;
+	const float MmToCm = 0.1f;
 	const float FloatScaleInitial = 0x10000;
 	const float FloatScale = GetNextSmallerPositiveFloat(FloatScaleInitial);
 	check((int32)GetNextSmallerPositiveFloat(FloatScaleInitial) == (FloatScaleInitial - 1));
@@ -399,7 +449,18 @@ void FSubsurfaceProfileTexture::CreateTexture(FRHICommandListImmediate& RHICmdLi
 		
 		const float UnitToCm = Data.WorldUnitScale;
 
-		FLinearColor DifffuseMeanFreePathInMm = (Data.MeanFreePathColor*Data.MeanFreePathDistance) * CmToMm; // convert cm to mm.
+		FLinearColor DifffuseMeanFreePathInMm;
+
+		if (Data.bEnableMeanFreePath)
+		{
+			DifffuseMeanFreePathInMm = GetDiffuseMeanFreePathFromMeanFreePath(Data.SurfaceAlbedo, Data.MeanFreePathColor * Data.MeanFreePathDistance) * CmToMm / Dmfp2MfpMagicNumber;
+		}
+		else
+		{
+			DifffuseMeanFreePathInMm = (Data.MeanFreePathColor * Data.MeanFreePathDistance) * CmToMm;
+			UE_LOG(LogSubsurfaceProfile, Warning, TEXT("DMFP has already been upgraded to MFP. Should not reach here."));
+		}
+
 		SetupSurfaceAlbedoAndDiffuseMeanFreePath(Data.SurfaceAlbedo, DifffuseMeanFreePathInMm);
 		TextureRow[BSSS_SURFACEALBEDO_OFFSET] = Data.SurfaceAlbedo;
 		TextureRow[BSSS_DMFP_OFFSET] = EncodeDiffuseMeanFreePath(DifffuseMeanFreePathInMm);
@@ -425,20 +486,7 @@ void FSubsurfaceProfileTexture::CreateTexture(FRHICommandListImmediate& RHICmdLi
 
 		if (Data.bEnableBurley)
 		{
-
-			// When we need performance, we fallback Burley to Separable. In order to achieve that, we estimate Separable parameters
-			// from Burley. This fallback has two two advantages:
-			// 1. Burley Fallback is more expressive than Separable. It can use Burley's scattering profile with hdr dmfp. The original
-			//    separable can only use a fixed dmfp around 2.229 based on the fitting. Because of this, the fallback Separable can be
-			//    used to express any materials. Since we have Burley parameters set, the transmittance profile can also go physically
-			//    more expressive. We can have better transmittance.
-			// 2. It runs faster than Burley. Burley has already been optimized. However, under extreme and rare conditions, it can go slow. 
-			//    If fps is critical even under extreme conditions, we can use Burley Fallback. Under normal setting, it does not have too
-			//    much visual difference. And it looks more appealing than the original one.
-
-			// Estimate the scatter radius based on the mean free path distance and the path color. 0.1f is the minimal value shown in the GUI.
-			// 2.229f is divided because of fitting relationship between falloff color and dmfp that has a scale of 2.229.
-			Data.ScatterRadius = FMath::Max(Data.MeanFreePathDistance*Data.MeanFreePathColor.GetMax()/2.229f, 0.1f);
+			Data.ScatterRadius = FMath::Max(DifffuseMeanFreePathInMm.GetMax()*MmToCm, 0.1f);
 
 			ComputeMirroredBSSSKernel(&TextureRow[SSSS_KERNEL0_OFFSET], SSSS_KERNEL0_SIZE, Data.SurfaceAlbedo,
 				DifffuseMeanFreePathInMm, Data.ScatterRadius);
@@ -447,11 +495,7 @@ void FSubsurfaceProfileTexture::CreateTexture(FRHICommandListImmediate& RHICmdLi
 			ComputeMirroredBSSSKernel(&TextureRow[SSSS_KERNEL2_OFFSET], SSSS_KERNEL2_SIZE, Data.SurfaceAlbedo,
 				DifffuseMeanFreePathInMm, Data.ScatterRadius);
 
-			// Then, scale up by world unit scale and the fitting parameters to affect screen space sampling location.
-			// For high irradiance, the lose of energy due to insufficient sampling count needs to be compensated to
-			// make Burley fallback looks the same to Burley. Set 1.0f when we have enough samples.
-			const float SamplingCountCompensation = 0.707f;
-			Data.ScatterRadius *= (Data.WorldUnitScale * 10.0f)*2.229f* SamplingCountCompensation;
+			Data.ScatterRadius *= (Data.WorldUnitScale * 10.0f);
 		}
 		else
 		{
@@ -721,4 +765,9 @@ void USubsurfaceProfile::PostLoad()
 			UpgradeSeparableToBurley(Settings);
 		}
 	}
+
+	if (!Settings.bEnableMeanFreePath)
+	{
+		UpgradeDiffuseMeanFreePathToMeanFreePath(Settings);
+	}
 }
