PoseSearch: Fixed modular features ensure in FAnimNode_SequencePlayerBase::GetEffectiveStartPosition

diff --git a/Engine/Source/Runtime/Engine/Private/Animation/AnimPoseSearchProvider.cpp b/Engine/Source/Runtime/Engine/Private/Animation/AnimPoseSearchProvider.cpp
--- a/Engine/Source/Runtime/Engine/Private/Animation/AnimPoseSearchProvider.cpp
+++ b/Engine/Source/Runtime/Engine/Private/Animation/AnimPoseSearchProvider.cpp
@@ -10,15 +10,19 @@ const FName IPoseSearchProvider::ModularFeatureName(TEXT("AnimPoseSearch"));
 
 IPoseSearchProvider* IPoseSearchProvider::Get()
 {
-	if (IsAvailable())
+	IModularFeatures::FScopedLockModularFeatureList ScopedLockModularFeatureList;
+
+	if (IModularFeatures::Get().IsModularFeatureAvailable(ModularFeatureName))
 	{
 		return &IModularFeatures::Get().GetModularFeature<IPoseSearchProvider>(ModularFeatureName);
 	}
+
 	return nullptr;
 }
 
 bool IPoseSearchProvider::IsAvailable()
 {
+	IModularFeatures::FScopedLockModularFeatureList ScopedLockModularFeatureList;
 	return IModularFeatures::Get().IsModularFeatureAvailable(ModularFeatureName);
 }
 
