Change thin layer to act as an Strata operator.

This simplify Strata Slab UI, and make it easier to thin-film coat many BSDFs at the same time.

diff --git a/Engine/Shaders/Private/DebugProbes.usf b/Engine/Shaders/Private/DebugProbes.usf
--- a/Engine/Shaders/Private/DebugProbes.usf
+++ b/Engine/Shaders/Private/DebugProbes.usf
@@ -19,6 +19,7 @@ uint MaxBytesPerPixel;
 void UpdateAllBSDFsOperatorCoverageTransmittance(inout FStrataPixelHeader StrataPixelHeader, inout FStrataTree StrataTree, bool bRoughDiffuseEnabled, FStrataAddressing NullStrataAddressing, float3 V, float3 L) {}
 void UpdateAllOperatorsCoverageTransmittance(inout FStrataTree StrataTree) {}
 void UpdateAllBSDFWithBottomUpOperatorVisit(inout FStrataTree StrataTree) {}
+void PreUpdateAllBSDFWithBottomUpOperatorVisit(inout FStrataPixelHeader StrataPixelHeader, inout FStrataTree StrataTree, FStrataAddressing NullStrataAddressing, float3 V) {}
 
 #endif // STRATA_ENABLED
 
@@ -83,6 +84,15 @@ void MainPS(
 	, out float  OutDepth		: SV_DEPTH
 )
 {
+	OutColor = 0.f;
+#if PERMUTATION_PASS == RENDER_BASEPASS
+	OutGBufferA = 0;
+	OutGBufferB = 0;
+	OutGBufferC = 0;
+	OutGBufferD = 0;
+	OutGBufferE = 0;
+#endif
+
 	ResolvedView = ResolveView();
 	const uint2 PixelCoord = uint2(SvPosition.xy);
 
@@ -288,7 +298,6 @@ void MainPS(
 			/*float SSSDMFPScale					*/0.0f,
 			/*float3 Emissive						*/0.0f,
 			/*float Haziness						*/0.0f,
-			/*float ThinFilmThickness				*/0.0f,
 			/*float FuzzAmount						*/0.0f,
 			/*float3 FuzzColor						*/0.0f,
 			/*float Thickness						*/0.0f,
diff --git a/Engine/Shaders/Private/Strata/Strata.ush b/Engine/Shaders/Private/Strata/Strata.ush
--- a/Engine/Shaders/Private/Strata/Strata.ush
+++ b/Engine/Shaders/Private/Strata/Strata.ush
@@ -736,6 +736,11 @@ struct FStrataOperator
 	float				 ThicknessCm;					// The thickness of the layer
 	float				 VerticalTop_ThicknessCm;		// Only used by vertical layering operators. It represent the thickness that is used to drive the roughness propagation from top layers to sub layers.
 
+	// Thin film
+	float ThinFilmThickness;							// Thickness of thin-film coating layer. In micro-meters (1um == 1000nm).  Merge with regular ThicknessCm?
+	float ThinFilmIOR;									// IOR of the thin-film coating layer
+	bool  bThinFilmInline;								// If true the thin-film reflectance is evaluated during the base-pass and stored into F0/F90 to speed up lighting computation
+
 	// Topology information
 	bool				 bIsTop;
 
@@ -819,7 +824,6 @@ FStrataData GetStrataSlabBSDF(
 	float SSSProfileID, float3 SSSDMFP, float SSSDMFPScale,
 	float3 Emissive, 
 	float Haziness, 
-	float ThinFilmThickness,
 	float FuzzAmount, float3 FuzzColor,
 	float Thickness,
 	uint SharedLocalBasisIndex, inout uint SharedLocalBasisTypes)
@@ -881,7 +885,7 @@ FStrataData GetStrataSlabBSDF(
 	}
 
 	const uint bHasHaziness = Haziness > 0.0f ? 1 : 0;
-	const uint bHasThinFilm = ThinFilmThickness > 0.0f ? 1 : 0;
+	const uint bHasThinFilm = 0;
 	const uint bHasAnisotropy = Anisotropy != 0 ? 1 : 0;
 	const uint bHasFuzz = FuzzAmount > 0.0f ? 1 : 0;
 	if (bHasAnisotropy)
@@ -933,7 +937,7 @@ FStrataData GetStrataSlabBSDF(
 	SLAB_ROUGHNESS			(StrataData.InlinedBSDF)		= Roughness;
 	SLAB_ANISOTROPY			(StrataData.InlinedBSDF)		= Anisotropy;
 	SLAB_HAZINESS			(StrataData.InlinedBSDF)		= Haziness;
-	SLAB_THINFILM_THICKNESS	(StrataData.InlinedBSDF)		= ThinFilmThickness;
+	SLAB_THINFILM_THICKNESS	(StrataData.InlinedBSDF)		= 0.f;
 
 	if (bHasSSSProfile)
 	{
@@ -1307,7 +1311,30 @@ FStrataData StrataWeight(FStrataData A, float Weight, inout FStrataTree StrataTr
 	StrataTree.Operators[OperatorIndex].Weight					= Weight;
 	StrataTree.Operators[OperatorIndex].LeftIndex				= A.OperatorIndex;
 
-	StrataTree.Operators[A.OperatorIndex].ParentIndex				= OperatorIndex;
+	StrataTree.Operators[A.OperatorIndex].ParentIndex			= OperatorIndex;
+
+	return StrataData;
+}
+
+FStrataData StrataThinFilm(FStrataData A, float Thickness, float IOR, bool bInlineEvaluation, inout FStrataTree StrataTree, int OperatorIndex, uint MaxDistanceFromLeaves)
+{
+	const float SafeThickness = saturate(Thickness);
+	const float SafeIOR = clamp(IOR, 1.01f, 3.0f);
+
+	FStrataData StrataData = GetInitialisedStrataData();
+
+	StrataTree.OperatorCount++;
+	StrataData.OperatorIndex = OperatorIndex;
+
+	StrataTree.Operators[OperatorIndex] = GetInitialisedStrataOperator();
+	StrataTree.Operators[OperatorIndex].MaxDistanceFromLeaves = MaxDistanceFromLeaves;
+	StrataTree.Operators[OperatorIndex].Type = STRATA_OPERATOR_THINFILM;
+	StrataTree.Operators[OperatorIndex].ThinFilmThickness = SafeThickness;
+	StrataTree.Operators[OperatorIndex].ThinFilmIOR = SafeIOR;
+	StrataTree.Operators[OperatorIndex].bThinFilmInline = bInlineEvaluation;
+	StrataTree.Operators[OperatorIndex].LeftIndex = A.OperatorIndex;
+
+	StrataTree.Operators[A.OperatorIndex].ParentIndex = OperatorIndex;
 
 	return StrataData;
 }
@@ -1917,7 +1944,7 @@ float3 StrataGetBSDFDiffuseColor(in FStrataBSDF BSDF)
 	return 0.f;
 }
 
-float3 StrataGetBSDFSpecularColor(in FStrataBSDF BSDF)
+float3 StrataGetBSDFSpecularF0(in FStrataBSDF BSDF)
 {
 	switch (BSDF_GETTYPE(BSDF))
 	{
@@ -1927,6 +1954,7 @@ float3 StrataGetBSDFSpecularColor(in FStrataBSDF BSDF)
 	}
 	return 0.f;
 }
+float3 StrataGetBSDFSpecularColor(in FStrataBSDF BSDF) { return StrataGetBSDFSpecularF0(BSDF);  }
 
 float3 StrataGetBSDFSpecularF90(in FStrataBSDF BSDF)
 {
@@ -1939,6 +1967,26 @@ float3 StrataGetBSDFSpecularF90(in FStrataBSDF BSDF)
 	return 0.f;
 }
 
+void StrataSetBSDFSpecularF0(in FStrataBSDF BSDF, float3 F0)
+{
+	switch (BSDF_GETTYPE(BSDF))
+	{
+	case STRATA_BSDF_TYPE_SLAB: 			SLAB_F0(BSDF) = F0;
+	case STRATA_BSDF_TYPE_HAIR: 			HAIR_SPECULAR(BSDF) = F0ToDielectricSpecular(max3(F0.x, F0.y, F0.z)); break;
+	case STRATA_BSDF_TYPE_SINGLELAYERWATER: SLW_SPECULAR(BSDF) = F0ToDielectricSpecular(max3(F0.x, F0.y, F0.z)); break;
+	}
+}
+
+void StrataSetBSDFSpecularF90(in FStrataBSDF BSDF, float3 F90)
+{
+	switch (BSDF_GETTYPE(BSDF))
+	{
+	case STRATA_BSDF_TYPE_SLAB: 			SLAB_F90(BSDF) = F90;
+	case STRATA_BSDF_TYPE_HAIR: 			break;
+	case STRATA_BSDF_TYPE_SINGLELAYERWATER: break;
+	}
+}
+
 float3 StrataGetBSDFSubSurfaceColor(in FStrataBSDF BSDF)
 {
 	switch (BSDF_GETTYPE(BSDF))
@@ -2466,7 +2514,6 @@ void StrataConvertFromDBuffer(
 		0.f,				// SSSDMFPScale
 		0.f,				// Emissive
 		0.f,				// Haziness
-		0.f,				// ThinFilmThickness
 		0.f,				// FuzzAmount
 		0.f,				// FuzzColor
 		DefaultThickness,
@@ -2489,6 +2536,7 @@ void StrataConvertFromDBuffer(
 void UpdateAllBSDFsOperatorCoverageTransmittance(inout FStrataPixelHeader StrataPixelHeader, inout FStrataTree StrataTree, bool bRoughDiffuseEnabled, FStrataAddressing NullStrataAddressing, float3 V, float3 L);
 void UpdateAllOperatorsCoverageTransmittance(inout FStrataTree StrataTree);
 void UpdateAllBSDFWithBottomUpOperatorVisit(inout FStrataTree StrataTree);
+void PreUpdateAllBSDFWithBottomUpOperatorVisit(inout FStrataPixelHeader StrataPixelHeader, inout FStrataTree StrataTree, FStrataAddressing NullStrataAddressing, float3 V);
 #else
 void LegacyUpdateBSDFsOperators(inout FStrataPixelHeader StrataPixelHeader, inout FStrataTree StrataTree, bool bRoughDiffuseEnabled, FStrataAddressing NullStrataAddressing, float3 V, float3 L);
 #endif
@@ -2503,6 +2551,11 @@ void StrataUpdateTree(
 	bool bRoughDiffuseEnabled)
 {
 #if STRATA_USES_CONVERSION_FROM_LEGACY==0
+	//
+	// Pre-Update the BSDF by traversing up the strata tree until the root node. This allows to modify BSDF's parameters based on operators (i.e., thin-film coating)
+	//
+	PreUpdateAllBSDFWithBottomUpOperatorVisit(StrataPixelHeader, StrataTree, StrataAddressing, V);
+
 	//
 	// Update the coverage/transmittance of each leaves (==BSDFs) of the strata tree.
 	//
diff --git a/Engine/Shaders/Private/Strata/StrataLegacyConversion.ush b/Engine/Shaders/Private/Strata/StrataLegacyConversion.ush
--- a/Engine/Shaders/Private/Strata/StrataLegacyConversion.ush
+++ b/Engine/Shaders/Private/Strata/StrataLegacyConversion.ush
@@ -90,7 +90,6 @@ FStrataData CreateLegacySlab(
 			In.SSSDMFPScale,	// SSSDMFPScale
 			In.Emissive,		// Emissive
 			0.f,				// Haziness
-			0.f,				// ThinFilmThickness
 			In.FuzzAmount,		// FuzzAmount
 			In.FuzzColor,		// FuzzColor
 			In.Thickness,
diff --git a/Engine/Shaders/Private/Strata/StrataTree.ush b/Engine/Shaders/Private/Strata/StrataTree.ush
--- a/Engine/Shaders/Private/Strata/StrataTree.ush
+++ b/Engine/Shaders/Private/Strata/StrataTree.ush
@@ -263,6 +263,18 @@ void UpdateSingleOperatorCoverageTransmittance(inout FStrataTree StrataTree, int
 		break;
 	}
 
+	case STRATA_OPERATOR_THINFILM:
+	{
+		// Thin-film does not affect these property
+		Op.Coverage = OpA.Coverage;
+		Op.Transmittance = OpA.Transmittance;
+		Op.TopDownRefractionLobeStat = OpA.TopDownRefractionLobeStat;
+		Op.TopDownRefractionWorldNormal = OpA.TopDownRefractionWorldNormal;
+		Op.BottomUpRefractionLobeStat = OpA.BottomUpRefractionLobeStat;
+		Op.ThicknessCm = OpA.ThicknessCm;
+		break;
+	}
+
 	}
 
 #undef Op
@@ -276,7 +288,7 @@ void UpdateSingleOperatorCoverageTransmittance(inout FStrataTree StrataTree, int
 
 #define Op StrataTree.Operators[OpIndex]
 
-
+// Post-update visit operators
 void UpdateAllBSDFWithBottomUpOperatorVisit_Weight(
 	inout FStrataTree StrataTree,
 	inout FStrataBSDF CurrentBSDF,
@@ -353,6 +365,83 @@ void UpdateAllBSDFWithBottomUpOperatorVisit_Vertical(
 #endif // STRATA_INLINE_SHADING
 }
 
+// Pre-update visit operators
+void PreUpdateAllBSDFWithBottomUpOperatorVisit_Vertical(
+	FStrataPixelHeader StrataPixelHeader,
+	inout FStrataTree StrataTree,
+	inout FStrataBSDF CurrentBSDF,
+	FStrataAddressing NullStrataAddressing,
+	inout bool bCanReceiveThinFilm,
+	float3 V,
+	int OpIndex,
+	int PreviousIsInputA)
+{
+#if STRATA_INLINE_SHADING
+
+	// If the BSDF is coated with a vertical layer, it can no longer receive thin film-coating
+	if (PreviousIsInputA <= 0)
+	{
+		bCanReceiveThinFilm = false;
+	}
+
+#endif // STRATA_INLINE_SHADING
+}
+
+void PreUpdateAllBSDFWithBottomUpOperatorVisit_ThinFilm(
+	FStrataPixelHeader StrataPixelHeader,
+	inout FStrataTree StrataTree,
+	inout FStrataBSDF CurrentBSDF,
+	FStrataAddressing NullStrataAddressing,
+	inout bool bCanReceiveThinFilm,
+	float3 V,
+	int OpIndex,
+	int PreviousIsInputA)
+{
+#if STRATA_INLINE_SHADING
+
+	if (bCanReceiveThinFilm && BSDF_GETTYPE(CurrentBSDF) == STRATA_BSDF_TYPE_SLAB)
+	{
+		if (Op.bThinFilmInline)
+		{
+			const float3 F0  = SLAB_F0(CurrentBSDF);
+			const float3 F90 = SLAB_F90(CurrentBSDF);
+
+			const float3x3 TangentBasis = StrataGetBSDFSharedBasis_InlineShading(StrataPixelHeader, CurrentBSDF, NullStrataAddressing);
+			const float3 N = TangentBasis[2];
+			const float NoV = saturate(dot(N, V));
+			const float NoL = NoV;
+			const float VoH = NoV;
+			float3 RThinFilm = F_ThinFilm(NoV, NoL, VoH, F0, F90, Op.ThinFilmIOR, Op.ThinFilmThickness);
+
+			// Compute a F0 which match evaluation Thin-Film reflectance for the current view angle using Schlick's Fresnel
+			const float FReference = min(0.99f, F_Schlick(0, 1.0f, NoV).x);
+			RThinFilm = max(RThinFilm, FReference);
+
+			float3 F0_Thin = (RThinFilm - F90 * FReference) / (1 - FReference);
+
+			// Anything less than 2% is physically impossible and is instead considered to be shadowing
+			// Ensure the generated F0 doesn't cause micro-oclusion
+			const float MicroOcclusionThreshold = 0.02f;
+			F0_Thin = max(MicroOcclusionThreshold, F0_Thin);
+
+			// If micro-occlusion was set, re-apply it.
+			const float3 MicroOcclusion = saturate(F0 * 50.f);
+			F0_Thin *= MicroOcclusion;
+
+			SLAB_F0(CurrentBSDF)  = F0_Thin;
+			SLAB_F90(CurrentBSDF) = max(F0_Thin, F90);
+		}
+		else
+		{
+			const uint bHasThinFilm = Op.ThinFilmThickness > 0 ? 1u : 0u;
+			BSDF_SETHASTHINFILM(CurrentBSDF, bHasThinFilm);
+			SLAB_THINFILM_THICKNESS(CurrentBSDF) += Op.ThinFilmThickness;
+		}
+	}
+
+#endif // STRATA_INLINE_SHADING
+}
+
 #undef Op
 
 
diff --git a/Engine/Shaders/Shared/StrataDefinitions.h b/Engine/Shaders/Shared/StrataDefinitions.h
--- a/Engine/Shaders/Shared/StrataDefinitions.h
+++ b/Engine/Shaders/Shared/StrataDefinitions.h
@@ -51,3 +51,4 @@
 #define STRATA_OPERATOR_ADD				3
 #define STRATA_OPERATOR_BSDF			4
 #define STRATA_OPERATOR_BSDF_LEGACY		5
+#define STRATA_OPERATOR_THINFILM		6
diff --git a/Engine/Source/Runtime/Engine/Classes/Materials/MaterialExpressionStrata.h b/Engine/Source/Runtime/Engine/Classes/Materials/MaterialExpressionStrata.h
--- a/Engine/Source/Runtime/Engine/Classes/Materials/MaterialExpressionStrata.h
+++ b/Engine/Source/Runtime/Engine/Classes/Materials/MaterialExpressionStrata.h
@@ -278,12 +278,6 @@ class UMaterialExpressionStrataSlabBSDF : public UMaterialExpressionStrataBSDF
 	UPROPERTY()
 	FExpressionInput Haziness;
 
-	/**
-	 * Thin film controls the thin film layer coating the current slab. 0 means disabled and 1 means a coating layer of 10 micrometer. (type = float, unitless, default = 0)
-	 */
-	UPROPERTY()
-	FExpressionInput ThinFilmThickness;
-
 	/**
 	 * The slab thickness. (type = float, centimeters, default = 0.01 centimeter = 0.1 millimeter)
 	 */
@@ -325,7 +319,6 @@ class UMaterialExpressionStrataSlabBSDF : public UMaterialExpressionStrataBSDF
 	virtual const TArray<FExpressionInput*> GetInputs() override;
 
 	bool HasEdgeColor() const;
-	bool HasThinFilm() const;
 	bool HasFuzz() const;
 	bool HasHaziness() const;
 	bool HasSSS() const;
@@ -773,6 +766,44 @@ class UMaterialExpressionStrataWeight : public UMaterialExpressionStrataBSDF
 	//~ End UMaterialExpression Interface
 };
 
+UCLASS(MinimalAPI, collapsecategories, hidecategories = Object, DisplayName = "Strata Thin-Film")
+class UMaterialExpressionStrataThinFilm : public UMaterialExpressionStrataBSDF
+{
+	GENERATED_UCLASS_BODY()
+
+	/**
+	 * Strata material
+	 */
+	UPROPERTY()
+	FExpressionInput A;
+	
+	/**
+	 * Thin film controls the thin film layer coating the current slab. 0 means disabled and 1 means a coating layer of 10 micrometer. (type = float, unitless, default = 0)
+	 */
+	UPROPERTY()
+	FExpressionInput Thickness;
+
+	/**
+	 * Thin film IOR
+	 */
+	UPROPERTY()
+	FExpressionInput IOR;
+
+	UPROPERTY(EditAnywhere, Category = Mode)
+	uint32 bUseParameterBlending : 1;
+
+	//~ Begin UMaterialExpression Interface
+#if WITH_EDITOR
+	virtual int32 Compile(class FMaterialCompiler* Compiler, int32 OutputIndex) override;
+	virtual void GetCaption(TArray<FString>& OutCaptions) const override;
+	virtual uint32 GetOutputType(int32 OutputIndex) override;
+	virtual uint32 GetInputType(int32 InputIndex) override;
+	virtual bool IsResultStrataMaterial(int32 OutputIndex) override;
+	virtual void GatherStrataMaterialInfo(FStrataMaterialInfo& StrataMaterialInfo, int32 OutputIndex) override;
+	virtual FStrataOperator* StrataGenerateMaterialTopologyTree(class FMaterialCompiler* Compiler, class UMaterialExpression* Parent, int32 OutputIndex) override;
+#endif
+	//~ End UMaterialExpression Interface
+};
 
 
 ///////////////////////////////////////////////////////////////////////////////
diff --git a/Engine/Source/Runtime/Engine/Private/Materials/HLSLMaterialTranslator.cpp b/Engine/Source/Runtime/Engine/Private/Materials/HLSLMaterialTranslator.cpp
--- a/Engine/Source/Runtime/Engine/Private/Materials/HLSLMaterialTranslator.cpp
+++ b/Engine/Source/Runtime/Engine/Private/Materials/HLSLMaterialTranslator.cpp
@@ -174,6 +174,10 @@ static inline const TCHAR* GetStrataOperatorStr(int32 OperatorType)
 	{
 		return TEXT("BSDFLEGACY");
 	}
+	case STRATA_OPERATOR_THINFILM:
+	{
+		return TEXT("THINFILM");
+	}
 	}
 	return TEXT("UNKNOWN   ");
 };
@@ -1391,6 +1395,77 @@ bool FHLSLMaterialTranslator::Translate()
 			{
 				// Now implement the functions needed to process the material topology
 
+				// Pre-Update the slab BSDF with operators (like thin film coating, which can alter F0/F90)
+				{
+					// Update the coverage/transmittance of each node in the graph
+					check(StrataMaterialRootOperator);
+					int32 RootMaximumDistanceToLeaves = StrataMaterialRootOperator->MaxDistanceFromLeaves;
+
+					ResourcesString += "void PreUpdateAllBSDFWithBottomUpOperatorVisit(inout FStrataPixelHeader StrataPixelHeader, inout FStrataTree StrataTree, FStrataAddressing NullStrataAddressing, float3 V)\n";
+					ResourcesString += "{\n";
+					for (uint32 BSDFIndex = 0; BSDFIndex < StrataMaterialBSDFCount; ++BSDFIndex)
+					{
+						ResourcesString += "\t{\n";
+						ResourcesString += "\t\tbool bCanReceiveThinFilm = true;\n";
+						for (auto& It : StrataMaterialExpressionRegisteredOperators)
+						{
+							if (!It.IsDiscarded() && It.BSDFIndex == BSDFIndex)
+							{
+								// Walk up the graph to the root node and apply weight factors
+								std::function<void(const FStrataOperator&, int32)> WalkOperatorsUp = [&](const FStrataOperator& CurrentOperator, int32 PreviousOperatorIndex) -> void
+								{
+									switch (CurrentOperator.OperatorType)
+									{
+									case STRATA_OPERATOR_WEIGHT:
+									{
+										break; // NOP
+									}
+									case STRATA_OPERATOR_HORIZONTAL:
+									{
+										break; // NOP
+									}
+									case STRATA_OPERATOR_VERTICAL:
+									{
+										ResourcesString += FString::Printf(TEXT("\t PreUpdateAllBSDFWithBottomUpOperatorVisit_Vertical(StrataPixelHeader, StrataTree, StrataTree.BSDFs[%d], NullStrataAddressing, bCanReceiveThinFilm, V, %d /*Op index*/, %d /*PreviousIsInputA*/);\n"), BSDFIndex, CurrentOperator.Index, CurrentOperator.LeftIndex == PreviousOperatorIndex ? 1 : 0);
+										break;
+									}
+									case STRATA_OPERATOR_ADD:
+									{
+										break; // NOP
+									}
+									case STRATA_OPERATOR_THINFILM:
+									{
+										ResourcesString += FString::Printf(TEXT("\t PreUpdateAllBSDFWithBottomUpOperatorVisit_ThinFilm(StrataPixelHeader, StrataTree, StrataTree.BSDFs[%d], NullStrataAddressing, bCanReceiveThinFilm, V,  %d /*Op index*/, %d /*PreviousIsInputA*/);\n"), BSDFIndex, CurrentOperator.Index, CurrentOperator.LeftIndex == PreviousOperatorIndex ? 1 : 0);
+										break;
+									}
+									default:
+									case STRATA_OPERATOR_BSDF:
+									{
+										check(false);
+									}
+									}
+
+									if (CurrentOperator.ParentIndex != INDEX_NONE)
+									{
+										WalkOperatorsUp(StrataMaterialExpressionRegisteredOperators[CurrentOperator.ParentIndex], CurrentOperator.Index);
+									}
+								};
+
+								const int32 BSDFOperatorIndex = It.Index;
+								const FStrataOperator& BSDFOperator = StrataMaterialExpressionRegisteredOperators[BSDFOperatorIndex];
+
+								// Start visiting node up from the BSDF leaf only if it has a parent.
+								if (BSDFOperator.ParentIndex != INDEX_NONE)
+								{
+									WalkOperatorsUp(StrataMaterialExpressionRegisteredOperators[BSDFOperator.ParentIndex], BSDFOperator.Index);
+								}
+							}
+						}
+						ResourcesString += "\t}\n";
+					}
+					ResourcesString += "}\n";
+				}
+			
 				// Update the coverage/transmittance of each leaves (==BSDFs) of the strata tree.
 				{
 					ResourcesString += "void UpdateAllBSDFsOperatorCoverageTransmittance(inout FStrataPixelHeader StrataPixelHeader, inout FStrataTree StrataTree, bool bRoughDiffuseEnabled, FStrataAddressing NullStrataAddressing, float3 V, float3 L)\n";
@@ -1439,7 +1514,7 @@ bool FHLSLMaterialTranslator::Translate()
 							if (!It.IsDiscarded() && It.BSDFIndex == BSDFIndex)
 							{
 								// Walk up the graph to the root node and apply weight factors
-								std::function<void(FStrataOperator&, int32)> WalkOperatorsUp = [&](FStrataOperator& CurrentOperator, int32 PreviousOperatorIndex) -> void
+								std::function<void(const FStrataOperator&, int32)> WalkOperatorsUp = [&](const FStrataOperator& CurrentOperator, int32 PreviousOperatorIndex) -> void
 								{
 									switch (CurrentOperator.OperatorType)
 									{
@@ -1453,18 +1528,19 @@ bool FHLSLMaterialTranslator::Translate()
 										ResourcesString += FString::Printf(TEXT("\t UpdateAllBSDFWithBottomUpOperatorVisit_Horizontal(StrataTree, StrataTree.BSDFs[%d], %d /*Op index*/, %d /*PreviousIsInputA*/);\n"), BSDFIndex, CurrentOperator.Index, CurrentOperator.LeftIndex == PreviousOperatorIndex ? 1 : 0);
 										break;
 									}
-
 									case STRATA_OPERATOR_VERTICAL:
 									{
 										ResourcesString += FString::Printf(TEXT("\t UpdateAllBSDFWithBottomUpOperatorVisit_Vertical(StrataTree, StrataTree.BSDFs[%d], %d /*Op index*/, %d /*PreviousIsInputA*/);\n"), BSDFIndex, CurrentOperator.Index, CurrentOperator.LeftIndex == PreviousOperatorIndex ? 1 : 0);
 										break;
 									}
-
 									case STRATA_OPERATOR_ADD:
 									{
 										break; // NOP
 									}
-
+									case STRATA_OPERATOR_THINFILM:
+									{
+										break; // NOP
+									}
 									default:
 									case STRATA_OPERATOR_BSDF:
 									{
@@ -1479,7 +1555,7 @@ bool FHLSLMaterialTranslator::Translate()
 								};
 
 								const int32 BSDFOperatorIndex = It.Index;
-								FStrataOperator& BSDFOperator = StrataMaterialExpressionRegisteredOperators[BSDFOperatorIndex];
+								const FStrataOperator& BSDFOperator = StrataMaterialExpressionRegisteredOperators[BSDFOperatorIndex];
 
 								// Start visiting node up from the BSDF leaf only if it has a parent.
 								if (BSDFOperator.ParentIndex != INDEX_NONE)
@@ -9630,7 +9706,7 @@ FStrataOperator& FHLSLMaterialTranslator::StrataCompilationRegisterOperator(int3
 	NewOperator.LeftIndex   = INDEX_NONE;
 	NewOperator.RightIndex   = INDEX_NONE;
 
-	NewOperator.BSDFIndex = INDEX_NONE;	// Allocated later to be ableto account for inline
+	NewOperator.BSDFIndex = INDEX_NONE;	// Allocated later to be albedo account for inline
 
 	NewOperator.MaxDistanceFromLeaves = 0;
 	NewOperator.bIsBottom = false;
@@ -9697,6 +9773,7 @@ bool FHLSLMaterialTranslator::StrataGenerateDerivedMaterialOperatorData()
 
 		// Operators with a single child
 		case STRATA_OPERATOR_WEIGHT:
+		case STRATA_OPERATOR_THINFILM:
 		{
 			bMustHaveLeftChild = true;
 		}
@@ -9748,6 +9825,7 @@ bool FHLSLMaterialTranslator::StrataGenerateDerivedMaterialOperatorData()
 				break;
 			}
 			case STRATA_OPERATOR_WEIGHT:
+			case STRATA_OPERATOR_THINFILM:
 			{
 				WalkOperators(StrataMaterialExpressionRegisteredOperators[CurrentOperator.LeftIndex], bUseParameterBlending);
 				break;
@@ -9812,6 +9890,7 @@ bool FHLSLMaterialTranslator::StrataGenerateDerivedMaterialOperatorData()
 
 		// Operators with a single child
 		case STRATA_OPERATOR_WEIGHT:
+		case STRATA_OPERATOR_THINFILM:
 		{
 			check(It.LeftIndex != INDEX_NONE);
 		}
@@ -9828,6 +9907,7 @@ bool FHLSLMaterialTranslator::StrataGenerateDerivedMaterialOperatorData()
 			switch (CurrentOperator.OperatorType)
 			{
 			case STRATA_OPERATOR_WEIGHT:
+			case STRATA_OPERATOR_THINFILM:
 			{
 				CurrentOperator.MaxDistanceFromLeaves = StrataMaterialExpressionRegisteredOperators[CurrentOperator.LeftIndex].MaxDistanceFromLeaves + 1;
 				break;
@@ -9898,6 +9978,7 @@ bool FHLSLMaterialTranslator::StrataGenerateDerivedMaterialOperatorData()
 				break;
 			}
 			case STRATA_OPERATOR_WEIGHT:
+			case STRATA_OPERATOR_THINFILM:
 			{
 				WalkOperators(StrataMaterialExpressionRegisteredOperators[CurrentOperator.LeftIndex]);
 				break;
@@ -10059,7 +10140,6 @@ int32 FHLSLMaterialTranslator::StrataSlabBSDF(
 	int32 SSSProfileId, int32 SSSDMFP, int32 SSSDMFPScale, 
 	int32 EmissiveColor,	
 	int32 Haziness, 
-	int32 ThinFilmThickness, 
 	int32 FuzzAmount, int32 FuzzColor,
 	int32 Thickness, 
 	int32 Normal, int32 Tangent, const FString& SharedLocalBasisIndexMacro,
@@ -10074,7 +10154,7 @@ int32 FHLSLMaterialTranslator::StrataSlabBSDF(
 		check(PromoteToOperator->Index != INDEX_NONE);
 		check(PromoteToOperator->BSDFIndex != INDEX_NONE);
 		return AddCodeChunk(
-			MCT_Strata, TEXT("PromoteParameterBlendedBSDFToOperator(GetStrataSlabBSDF(Parameters.StrataPixelFootprint, %s, %s, %s, %s, %s, %s, %s, %s, %s, %s, %s, %s, %s, %s, %s, %s, %s, %s, %s, %s, Parameters.SharedLocalBases.Types) /* Normal = %s ; Tangent = %s */, Parameters.StrataTree, %u, %u, %u, %u)"),
+			MCT_Strata, TEXT("PromoteParameterBlendedBSDFToOperator(GetStrataSlabBSDF(Parameters.StrataPixelFootprint, %s, %s, %s, %s, %s, %s, %s, %s, %s, %s, %s, %s, %s, %s, %s, %s, %s, %s, %s, Parameters.SharedLocalBases.Types) /* Normal = %s ; Tangent = %s */, Parameters.StrataTree, %u, %u, %u, %u)"),
 			*GetParameterCode(UseMetalness),
 			*GetParameterCode(BaseColor), *GetParameterCode(EdgeColor), *GetParameterCode(Specular), *GetParameterCode(Metallic),
 			*GetParameterCode(DiffuseAlbedo), *GetParameterCode(F0), *GetParameterCode(F90),
@@ -10084,7 +10164,6 @@ int32 FHLSLMaterialTranslator::StrataSlabBSDF(
 			*GetParameterCode(SSSDMFPScale),
 			*GetParameterCode(EmissiveColor),
 			*GetParameterCode(Haziness),
-			*GetParameterCode(ThinFilmThickness),
 			*GetParameterCode(FuzzAmount), *GetParameterCode(FuzzColor),
 			*GetParameterCode(Thickness),
 			*SharedLocalBasisIndexMacro,
@@ -10097,7 +10176,7 @@ int32 FHLSLMaterialTranslator::StrataSlabBSDF(
 	}
 
 	return AddCodeChunk(
-		MCT_Strata, TEXT("GetStrataSlabBSDF(Parameters.StrataPixelFootprint, %s, %s, %s, %s, %s, %s, %s, %s, %s, %s, %s, %s, %s, %s, %s, %s, %s, %s, %s, %s, Parameters.SharedLocalBases.Types) /* Normal = %s ; Tangent = %s */"),
+		MCT_Strata, TEXT("GetStrataSlabBSDF(Parameters.StrataPixelFootprint, %s, %s, %s, %s, %s, %s, %s, %s, %s, %s, %s, %s, %s, %s, %s, %s, %s, %s, %s, Parameters.SharedLocalBases.Types) /* Normal = %s ; Tangent = %s */"),
 		* GetParameterCode(UseMetalness),
 		*GetParameterCode(BaseColor),		*GetParameterCode(EdgeColor),	*GetParameterCode(Specular), *GetParameterCode(Metallic),
 		*GetParameterCode(DiffuseAlbedo),	*GetParameterCode(F0),			*GetParameterCode(F90),
@@ -10107,7 +10186,6 @@ int32 FHLSLMaterialTranslator::StrataSlabBSDF(
 		*GetParameterCode(SSSDMFPScale),
 		*GetParameterCode(EmissiveColor),
 		*GetParameterCode(Haziness),
-		*GetParameterCode(ThinFilmThickness),
 		*GetParameterCode(FuzzAmount),		*GetParameterCode(FuzzColor),
 		*GetParameterCode(Thickness),
 		*SharedLocalBasisIndexMacro,
@@ -10409,6 +10487,23 @@ int32 FHLSLMaterialTranslator::StrataWeightParameterBlending(int32 A, int32 Weig
 	);
 }
 
+int32 FHLSLMaterialTranslator::StrataThinFilm(int32 A, int32 Thickness, int32 IOR, bool bInlinedEvaluation, int OperatorIndex, uint32 MaxDistanceFromLeaves)
+{
+	if (A == INDEX_NONE || Thickness == INDEX_NONE || IOR == INDEX_NONE)
+	{
+		return INDEX_NONE;
+	}
+	return AddCodeChunk(
+		MCT_Strata, TEXT("StrataThinFilm(%s, %s, %s, %d, Parameters.StrataTree, %u, %u)"),
+		*GetParameterCode(A),
+		*GetParameterCode(Thickness),
+		*GetParameterCode(IOR),
+		bInlinedEvaluation ? 1u : 0u,
+		OperatorIndex,
+		MaxDistanceFromLeaves
+	);
+
+}
 int32 FHLSLMaterialTranslator::StrataTransmittanceToMFP(int32 TransmittanceColor, int32 DesiredThickness, int32 OutputIndex)
 {
 	if (OutputIndex == INDEX_NONE)
diff --git a/Engine/Source/Runtime/Engine/Private/Materials/HLSLMaterialTranslator.h b/Engine/Source/Runtime/Engine/Private/Materials/HLSLMaterialTranslator.h
--- a/Engine/Source/Runtime/Engine/Private/Materials/HLSLMaterialTranslator.h
+++ b/Engine/Source/Runtime/Engine/Private/Materials/HLSLMaterialTranslator.h
@@ -1018,7 +1018,6 @@ protected:
 		int32 SSSProfileId, int32 SSSDMFP, int32 SSSDMFPScale,
 		int32 EmissiveColor, 
 		int32 Haziness, 
-		int32 ThinFilmThickness, 
 		int32 FuzzAmount, int32 FuzzColor, 
 		int32 Thickness,
 		int32 Normal, int32 Tangent, const FString& SharedLocalBasisIndexMacro,
@@ -1051,6 +1050,7 @@ protected:
 	virtual int32 StrataAddParameterBlending(int32 A, int32 B, int32 AMixWeight, const FString& SharedLocalBasisIndexMacro, FStrataOperator* PromoteToOperator) override;
 	virtual int32 StrataWeight(int32 A, int32 Weight, int OperatorIndex, uint32 MaxDistanceFromLeaves) override;
 	virtual int32 StrataWeightParameterBlending(int32 A, int32 Weight, FStrataOperator* PromoteToOperator) override;
+	virtual int32 StrataThinFilm(int32 A, int32 Thickness, int32 IOR, bool bInlinedEvaluation, int OperatorIndex, uint32 MaxDistanceFromLeaves) override;
 	virtual int32 StrataTransmittanceToMFP(int32 TransmittanceColor, int32 DesiredThickness, int32 OutputIndex) override;
 
 	virtual FStrataOperator& StrataCompilationRegisterOperator(int32 OperatorType, UMaterialExpression* Expression, UMaterialExpression* Parent, bool bUseParameterBlending = false) override;
diff --git a/Engine/Source/Runtime/Engine/Private/Materials/MaterialExpressions.cpp b/Engine/Source/Runtime/Engine/Private/Materials/MaterialExpressions.cpp
--- a/Engine/Source/Runtime/Engine/Private/Materials/MaterialExpressions.cpp
+++ b/Engine/Source/Runtime/Engine/Private/Materials/MaterialExpressions.cpp
@@ -20549,7 +20549,6 @@ int32 UMaterialExpressionStrataLegacyConversion::Compile(class FMaterialCompiler
 				false /*bHasSSS*/, 
 				false /*bHasDMFPPluggedIn*/, 
 				false /*bHasEdgeColor*/, 
-				false /*bHasThinFilm*/, 
 				false /*bHasFuzz*/, 
 				false /*bHasHaziness*/);
 		}
@@ -20568,7 +20567,6 @@ int32 UMaterialExpressionStrataLegacyConversion::Compile(class FMaterialCompiler
 				false /*bHasSSS*/,
 				false /*bHasDMFPPluggedIn*/,
 				false /*bHasEdgeColor*/,
-				false /*bHasThinFilm*/,
 				false /*bHasFuzz*/,
 				false /*bHasHaziness*/);
 		}
@@ -20579,7 +20577,6 @@ int32 UMaterialExpressionStrataLegacyConversion::Compile(class FMaterialCompiler
 				false /*bHasSSS*/,
 				false /*bHasDMFPPluggedIn*/,
 				false /*bHasEdgeColor*/,
-				false /*bHasThinFilm*/,
 				false /*bHasFuzz*/,
 				false /*bHasHaziness*/);
 		}
@@ -20796,7 +20793,6 @@ int32 UMaterialExpressionStrataSlabBSDF::Compile(class FMaterialCompiler* Compil
 	const FStrataRegisteredSharedLocalBasis NewRegisteredSharedLocalBasis = StrataCompilationInfoCreateSharedLocalBasis(Compiler, NormalCodeChunk, TangentCodeChunk);
 
 	const bool bHasEdgeColor = HasEdgeColor(); // This accounts for EdgeColor and also F90 when the non metalness worfklow is selected.
-	const bool bHasThinFilm = HasThinFilm();
 	const bool bHasFuzz = HasFuzz();
 	const bool bHasHaziness = HasHaziness();
 	const bool bHasDMFPPluggedIn = HasDMFPPluggedIn();
@@ -20838,7 +20834,6 @@ int32 UMaterialExpressionStrataSlabBSDF::Compile(class FMaterialCompiler* Compil
 		CompileWithDefaultFloat1(Compiler, SSSDMFPScale, 1.0f),
 		CompileWithDefaultFloat3(Compiler, EmissiveColor, 0.0f, 0.0f, 0.0f),
 		CompileWithDefaultFloat1(Compiler, Haziness, 0.0f),
-		CompileWithDefaultFloat1(Compiler, ThinFilmThickness, 0.0f),
 		CompileWithDefaultFloat1(Compiler, FuzzAmount, 0.0f),
 		CompileWithDefaultFloat3(Compiler, FuzzColor, 0.0f, 0.0f, 0.0f),
 		CompileWithDefaultFloat1(Compiler, Thickness, STRATA_LAYER_DEFAULT_THICKNESS_CM),
@@ -20846,7 +20841,7 @@ int32 UMaterialExpressionStrataSlabBSDF::Compile(class FMaterialCompiler* Compil
 		TangentCodeChunk,
 		Compiler->GetStrataSharedLocalBasisIndexMacro(NewRegisteredSharedLocalBasis),
 		!StrataOperator.bUseParameterBlending || (StrataOperator.bUseParameterBlending && StrataOperator.bRootOfParameterBlendingSubTree) ? &StrataOperator : nullptr);
-	StrataCompilationInfoCreateSingleBSDFMaterial(Compiler, OutputCodeChunk, NewRegisteredSharedLocalBasis, STRATA_BSDF_TYPE_SLAB, bHasSSS, bHasDMFPPluggedIn, bHasEdgeColor, bHasThinFilm, bHasFuzz, bHasHaziness);
+	StrataCompilationInfoCreateSingleBSDFMaterial(Compiler, OutputCodeChunk, NewRegisteredSharedLocalBasis, STRATA_BSDF_TYPE_SLAB, bHasSSS, bHasDMFPPluggedIn, bHasEdgeColor, bHasFuzz, bHasHaziness);
 
 	return OutputCodeChunk;
 }
@@ -20876,7 +20871,6 @@ const TArray<FExpressionInput*> UMaterialExpressionStrataSlabBSDF::GetInputs()
 	Result.Add(&SSSDMFPScale);
 	Result.Add(&EmissiveColor);
 	Result.Add(&Haziness);
-	Result.Add(&ThinFilmThickness);
 	Result.Add(&Thickness);
 	Result.Add(&FuzzAmount);
 	Result.Add(&FuzzColor);
@@ -21014,18 +21008,14 @@ uint32 UMaterialExpressionStrataSlabBSDF::GetInputType(int32 InputIndex)
 		return MCT_Float1; // Haziness
 	}
 	else if (InputIndex == (12 + SkipDisabledInputOffset))
-	{
-		return MCT_Float1; // ThinFilm Thickness
-	}
-	else if (InputIndex == (13 + SkipDisabledInputOffset))
 	{
 		return MCT_Float1; // Thickness
 	}
-	else if (InputIndex == (14 + SkipDisabledInputOffset))
+	else if (InputIndex == (13 + SkipDisabledInputOffset))
 	{
 		return MCT_Float1; // FuzzAmount
 	}
-	else if (InputIndex == (15 + SkipDisabledInputOffset))
+	else if (InputIndex == (14 + SkipDisabledInputOffset))
 	{
 		return MCT_Float3; // FuzzColor
 	}
@@ -21106,18 +21096,14 @@ FName UMaterialExpressionStrataSlabBSDF::GetInputName(int32 InputIndex) const
 		return TEXT("Haziness");
 	}
 	else if (InputIndex == (12 + SkipDisabledInputOffset))
-	{
-		return TEXT("ThinFilm Thickness");
-	}
-	else if (InputIndex == (13 + SkipDisabledInputOffset))
 	{
 		return TEXT("Thickness");
 	}
-	else if (InputIndex == (14 + SkipDisabledInputOffset))
+	else if (InputIndex == (13 + SkipDisabledInputOffset))
 	{
 		return TEXT("FuzzAmount");
 	}
-	else if (InputIndex == (15 + SkipDisabledInputOffset))
+	else if (InputIndex == (14 + SkipDisabledInputOffset))
 	{
 		return TEXT("FuzzColor");
 	}
@@ -21285,11 +21271,6 @@ bool UMaterialExpressionStrataSlabBSDF::HasEdgeColor() const
 	return bUseMetalness ? EdgeColor.IsConnected() : F90.IsConnected();
 }
 
-bool UMaterialExpressionStrataSlabBSDF::HasThinFilm() const
-{
-	return ThinFilmThickness.IsConnected();
-}
-
 bool UMaterialExpressionStrataSlabBSDF::HasFuzz() const
 {
 	return FuzzAmount.IsConnected();
@@ -22504,6 +22485,95 @@ FStrataOperator* UMaterialExpressionStrataWeight::StrataGenerateMaterialTopology
 
 
 
+UMaterialExpressionStrataThinFilm::UMaterialExpressionStrataThinFilm(const FObjectInitializer& ObjectInitializer)
+	: Super(ObjectInitializer)
+{
+	struct FConstructorStatics
+	{
+		FText NAME_Strata;
+		FConstructorStatics() : NAME_Strata(LOCTEXT("Strata Ops", "Strata Operators")) { }
+	};
+	static FConstructorStatics ConstructorStatics;
+#if WITH_EDITORONLY_DATA
+	MenuCategories.Add(ConstructorStatics.NAME_Strata);
+#endif
+}
+
+#if WITH_EDITOR
+int32 UMaterialExpressionStrataThinFilm::Compile(class FMaterialCompiler* Compiler, int32 OutputIndex)
+{
+	if (!A.GetTracedInput().Expression)
+	{
+		return Compiler->Errorf(TEXT("Missing A input"));
+	}
+
+	int32 ACodeChunk = A.Compile(Compiler);
+	int32 ThicknessCodeChunk = Thickness.GetTracedInput().Expression ? Thickness.Compile(Compiler) : Compiler->Constant(1.0f);
+	int32 IORCodeChunk = IOR.GetTracedInput().Expression ? IOR.Compile(Compiler) : Compiler->Constant(1.44f);
+
+	int32 OutputCodeChunk = INDEX_NONE;
+	FStrataOperator& StrataOperator = Compiler->StrataCompilationGetOperator(this);
+	{
+		OutputCodeChunk = Compiler->StrataThinFilm(ACodeChunk, ThicknessCodeChunk, IORCodeChunk, bUseParameterBlending, StrataOperator.Index, StrataOperator.MaxDistanceFromLeaves);
+
+		if (!Compiler->StrataCompilationInfoContainsCodeChunk(ACodeChunk))
+		{
+			return Compiler->Errorf(TEXT("Could not find ACodeChunk to coat with thin-film"));
+		}
+		FStrataMaterialCompilationInfo StrataInfo = StrataCompilationInfoThinFilm(Compiler, Compiler->GetStrataCompilationInfo(ACodeChunk));
+		Compiler->StrataCompilationInfoRegisterCodeChunk(OutputCodeChunk, StrataInfo);
+	}
+
+	return OutputCodeChunk;
+}
+
+void UMaterialExpressionStrataThinFilm::GetCaption(TArray<FString>& OutCaptions) const
+{
+	OutCaptions.Add(TEXT("Strata Thin-Film"));
+}
+
+uint32 UMaterialExpressionStrataThinFilm::GetOutputType(int32 OutputIndex)
+{
+	return MCT_Strata;
+}
+
+uint32 UMaterialExpressionStrataThinFilm::GetInputType(int32 InputIndex)
+{
+	if (InputIndex == 0) { return MCT_Strata; } // Material
+	if (InputIndex == 1) { return MCT_Float1; } // Thickness
+	if (InputIndex == 2) { return MCT_Float1; } // IOR
+	return MCT_Float1;
+}
+
+bool UMaterialExpressionStrataThinFilm::IsResultStrataMaterial(int32 OutputIndex)
+{
+	return true;
+}
+
+void UMaterialExpressionStrataThinFilm::GatherStrataMaterialInfo(FStrataMaterialInfo& StrataMaterialInfo, int32 OutputIndex)
+{
+	if (A.GetTracedInput().Expression)
+	{
+		A.GetTracedInput().Expression->GatherStrataMaterialInfo(StrataMaterialInfo, A.OutputIndex);
+	}
+}
+
+FStrataOperator* UMaterialExpressionStrataThinFilm::StrataGenerateMaterialTopologyTree(class FMaterialCompiler* Compiler, class UMaterialExpression* Parent, int32 OutputIndex)
+{
+	FStrataOperator& StrataOperator = Compiler->StrataCompilationRegisterOperator(STRATA_OPERATOR_THINFILM, this, Parent);
+
+	UMaterialExpression* ChildAExpression = A.GetTracedInput().Expression;
+	if (ChildAExpression)
+	{
+		StrataOperator.LeftIndex = ChildAExpression->StrataGenerateMaterialTopologyTree(Compiler, this, A.OutputIndex)->Index;
+	}
+
+	return &StrataOperator;
+}
+#endif // WITH_EDITOR
+
+
+
 UMaterialExpressionStrataTransmittanceToMFP::UMaterialExpressionStrataTransmittanceToMFP(const FObjectInitializer& ObjectInitializer)
 	: Super(ObjectInitializer)
 {
diff --git a/Engine/Source/Runtime/Engine/Private/Materials/StrataMaterial.cpp b/Engine/Source/Runtime/Engine/Private/Materials/StrataMaterial.cpp
--- a/Engine/Source/Runtime/Engine/Private/Materials/StrataMaterial.cpp
+++ b/Engine/Source/Runtime/Engine/Private/Materials/StrataMaterial.cpp
@@ -65,7 +65,7 @@ FStrataRegisteredSharedLocalBasis StrataCompilationInfoCreateSharedLocalBasis(FM
 
 void StrataCompilationInfoCreateSingleBSDFMaterial(FMaterialCompiler* Compiler, int32 CodeChunk,
 	const FStrataRegisteredSharedLocalBasis& RegisteredSharedLocalBasis,
-	uint8 BSDFType, bool bHasSSS, bool bHasDMFPPluggedIn, bool bHasEdgeColor, bool bHasThinFilm, bool bHasFuzz, bool bHasHaziness)
+	uint8 BSDFType, bool bHasSSS, bool bHasDMFPPluggedIn, bool bHasEdgeColor, bool bHasFuzz, bool bHasHaziness)
 {
 	FStrataMaterialCompilationInfo StrataInfo;
 	StrataInfo.LayerCount = 1;
@@ -75,7 +75,7 @@ void StrataCompilationInfoCreateSingleBSDFMaterial(FMaterialCompiler* Compiler,
 	StrataInfo.Layers[0].BSDFs[0].bHasSSS = bHasSSS;
 	StrataInfo.Layers[0].BSDFs[0].bHasDMFPPluggedIn = bHasDMFPPluggedIn;
 	StrataInfo.Layers[0].BSDFs[0].bHasEdgeColor = bHasEdgeColor;
-	StrataInfo.Layers[0].BSDFs[0].bHasThinFilm = bHasThinFilm;
+	StrataInfo.Layers[0].BSDFs[0].bHasThinFilm = false;
 	StrataInfo.Layers[0].BSDFs[0].bHasFuzz = bHasFuzz;
 	StrataInfo.Layers[0].BSDFs[0].bHasHaziness = bHasHaziness;
 	UpdateTotalBSDFCount(Compiler, StrataInfo);
@@ -93,7 +93,9 @@ void  StrataCompilationInfoCreateNullBSDF(FMaterialCompiler* Compiler, int32 Cod
 
 FStrataMaterialCompilationInfo StrataCompilationInfoWeight(FMaterialCompiler* Compiler, const FStrataMaterialCompilationInfo& A)
 {
+	// STRATA_TODO: This is not correct. FStrataMaterialCompilationInfo is meant to be removed and replaced with Strata tree / metadata
 	FStrataMaterialCompilationInfo StrataInfo = A;
+	StrataInfo.Layers[0].BSDFs[0].bHasThinFilm = true;
 	return StrataInfo;
 }
 
@@ -103,6 +105,11 @@ FStrataMaterialCompilationInfo StrataCompilationInfoWeightParamBlend(FMaterialCo
 	return StrataInfo;
 }
 
+FStrataMaterialCompilationInfo StrataCompilationInfoThinFilm(FMaterialCompiler* Compiler, const FStrataMaterialCompilationInfo& A)
+{
+	FStrataMaterialCompilationInfo StrataInfo = A;
+	return StrataInfo;
+}
 
 FStrataMaterialCompilationInfo StrataCompilationInfoAdd(FMaterialCompiler* Compiler, const FStrataMaterialCompilationInfo& A, const FStrataMaterialCompilationInfo& B)
 {
diff --git a/Engine/Source/Runtime/Engine/Private/Materials/StrataMaterial.h b/Engine/Source/Runtime/Engine/Private/Materials/StrataMaterial.h
--- a/Engine/Source/Runtime/Engine/Private/Materials/StrataMaterial.h
+++ b/Engine/Source/Runtime/Engine/Private/Materials/StrataMaterial.h
@@ -66,10 +66,11 @@ FStrataRegisteredSharedLocalBasis StrataCompilationInfoCreateSharedLocalBasis(FM
 void  StrataCompilationInfoCreateNullBSDF(FMaterialCompiler* Compiler, int32 CodeChunk);
 void  StrataCompilationInfoCreateSingleBSDFMaterial(FMaterialCompiler* Compiler, int32 CodeChunk, 
 	const FStrataRegisteredSharedLocalBasis& RegisteredSharedLocalBasis,
-	uint8 BSDFType, bool bHasSSS = false, bool bHasDMFPPluggedIn = false, bool bHasEdgeColor = false, bool bHasThinFilm = false, bool bHasFuzz = false, bool bHasHaziness = false);
+	uint8 BSDFType, bool bHasSSS = false, bool bHasDMFPPluggedIn = false, bool bHasEdgeColor = false, bool bHasFuzz = false, bool bHasHaziness = false);
 
 FStrataMaterialCompilationInfo StrataCompilationInfoWeight(FMaterialCompiler* Compiler, const FStrataMaterialCompilationInfo& A);
 FStrataMaterialCompilationInfo StrataCompilationInfoWeightParamBlend(FMaterialCompiler* Compiler, const FStrataMaterialCompilationInfo& A);
+FStrataMaterialCompilationInfo StrataCompilationInfoThinFilm(FMaterialCompiler* Compiler, const FStrataMaterialCompilationInfo& A);
 FStrataMaterialCompilationInfo StrataCompilationInfoAdd(FMaterialCompiler* Compiler, const FStrataMaterialCompilationInfo& A, const FStrataMaterialCompilationInfo& B);
 FStrataMaterialCompilationInfo StrataCompilationInfoAddParamBlend(FMaterialCompiler* Compiler, const FStrataMaterialCompilationInfo& A, const FStrataMaterialCompilationInfo& B, const FStrataRegisteredSharedLocalBasis& RegisteredSharedLocalBasis);
 FStrataMaterialCompilationInfo StrataCompilationInfoHorizontalMixing(FMaterialCompiler* Compiler, const FStrataMaterialCompilationInfo& Background, const FStrataMaterialCompilationInfo& Foreground);
diff --git a/Engine/Source/Runtime/Engine/Public/MaterialCompiler.h b/Engine/Source/Runtime/Engine/Public/MaterialCompiler.h
--- a/Engine/Source/Runtime/Engine/Public/MaterialCompiler.h
+++ b/Engine/Source/Runtime/Engine/Public/MaterialCompiler.h
@@ -495,7 +495,6 @@ public:
 		int32 SSSProfileId, int32 SSSDMFP, int32 SSSDMFPScale,
 		int32 EmissiveColor, 
 		int32 Haziness, 
-		int32 ThinFilmThickness, 
 		int32 FuzzAmount, int32 FuzzColor,
 		int32 Thickness,
 		int32 Normal, int32 Tangent, const FString& SharedLocalBasisIndexMacro,
@@ -527,6 +526,7 @@ public:
 	virtual int32 StrataAddParameterBlending(int32 A, int32 B, int32 AMixWeight, const FString& SharedLocalBasisIndexMacro, FStrataOperator* PromoteToOperator) = 0;
 	virtual int32 StrataWeight(int32 A, int32 Weight, int OperatorIndex, uint32 MaxDistanceFromLeaves) = 0;
 	virtual int32 StrataWeightParameterBlending(int32 A, int32 Weight, FStrataOperator* PromoteToOperator) = 0;
+	virtual int32 StrataThinFilm(int32 A, int32 Thickness, int32 IOR, bool bInlinedEvaluation, int OperatorIndex, uint32 MaxDistanceFromLeaves) = 0;
 	virtual int32 StrataTransmittanceToMFP(int32 TransmittanceColor, int32 DesiredThickness, int32 OutputIndex) = 0;
 
 	/**
@@ -1035,7 +1035,6 @@ public:
 		int32 SSSProfileId, int32 SSSDMFP, int32 SSSDMFPScale,
 		int32 EmissiveColor, 
 		int32 Haziness, 
-		int32 ThinFilmThickness, 
 		int32 FuzzAmount, int32 FuzzColor,
 		int32 Thickness,
 		int32 Normal, int32 Tangent, const FString& SharedLocalBasisIndexMacro,
@@ -1049,7 +1048,6 @@ public:
 			SSSProfileId, SSSDMFP, SSSDMFPScale,
 			EmissiveColor, 
 			Haziness, 
-			ThinFilmThickness,
 			FuzzAmount, FuzzColor,
 			Thickness,
 			Normal, Tangent, SharedLocalBasisIndexMacro,
@@ -1144,6 +1142,11 @@ public:
 		return Compiler->StrataWeight(A, Weight, OperatorIndex, MaxDistanceFromLeaves);
 	}
 
+	virtual int32 StrataThinFilm(int32 A, int32 Thickness, int32 IOR, bool bInlinedEvaluation, int OperatorIndex, uint32 MaxDistanceFromLeaves) override
+	{
+		return Compiler->StrataThinFilm(A, Thickness, IOR, bInlinedEvaluation, OperatorIndex, MaxDistanceFromLeaves);
+	}
+
 	virtual int32 StrataWeightParameterBlending(int32 A, int32 Weight, FStrataOperator* PromoteToOperator) override
 	{
 		return Compiler->StrataWeightParameterBlending(A, Weight, PromoteToOperator);
