Renderer: Reverted one files from CL 19371962 to fix crashes caused by RHICopySharedMips not being implemented on one of the platforms.

#rnx

diff --git a/Engine/Source/Runtime/Engine/Private/Streaming/Texture2DArrayStreaming.cpp b/Engine/Source/Runtime/Engine/Private/Streaming/Texture2DArrayStreaming.cpp
--- a/Engine/Source/Runtime/Engine/Private/Streaming/Texture2DArrayStreaming.cpp
+++ b/Engine/Source/Runtime/Engine/Private/Streaming/Texture2DArrayStreaming.cpp
@@ -7,6 +7,51 @@ Texture2DArrayStreaming.cpp: Helpers to stream in and out texture 2D array LODs.
 #include "Streaming/Texture2DArrayStreaming.h"
 #include "RenderUtils.h"
 
+//*****************************************************************************
+//***************************** Global Definitions ****************************
+//*****************************************************************************
+
+void RHICopySharedMips(FRHICommandList& RHICmdList, FRHITexture2DArray* DestTexture, FRHITexture2DArray* SrcTexture)
+{
+	// Transition to copy source and dest
+	{
+		FRHITransitionInfo TransitionsBefore[] = { FRHITransitionInfo(SrcTexture, ERHIAccess::SRVMask, ERHIAccess::CopySrc), FRHITransitionInfo(DestTexture, ERHIAccess::SRVMask, ERHIAccess::CopyDest) };
+		RHICmdList.Transition(MakeArrayView(TransitionsBefore, UE_ARRAY_COUNT(TransitionsBefore)));
+	}
+
+	// Copy 
+	{
+		FRHICopyTextureInfo CopyInfo;
+
+		auto SetCopyInfo = [&](FRHITexture2DArray* Texture2DArray)
+		{
+			CopyInfo.Size.X = Texture2DArray->GetSizeX();
+			CopyInfo.Size.Y = Texture2DArray->GetSizeY();
+			CopyInfo.NumSlices = Texture2DArray->GetSizeZ();
+			CopyInfo.NumMips = Texture2DArray->GetNumMips();
+		};
+
+		if (DestTexture->GetNumMips() < SrcTexture->GetNumMips())
+		{
+			SetCopyInfo(DestTexture);
+		}
+		else
+		{
+			SetCopyInfo(SrcTexture);
+		}
+
+		CopyInfo.SourceMipIndex = SrcTexture->GetNumMips() - CopyInfo.NumMips;
+		CopyInfo.DestMipIndex = DestTexture->GetNumMips() - CopyInfo.NumMips;
+		RHICmdList.CopyTexture(SrcTexture, DestTexture, CopyInfo);
+	}
+
+	// Transition to SRV
+	{
+		FRHITransitionInfo TransitionsAfter[] = { FRHITransitionInfo(SrcTexture, ERHIAccess::CopySrc, ERHIAccess::SRVMask), FRHITransitionInfo(DestTexture, ERHIAccess::CopyDest, ERHIAccess::SRVMask) };
+		RHICmdList.Transition(MakeArrayView(TransitionsAfter, UE_ARRAY_COUNT(TransitionsAfter)));
+	}
+}
+
 //*****************************************************************************
 //******************* FTexture2DArrayMipAllocator_Reallocate ******************
 //*****************************************************************************
@@ -71,7 +116,7 @@ bool FTexture2DArrayMipAllocator_Reallocate::FinalizeMips(const FTextureUpdateCo
 		ENQUEUE_RENDER_COMMAND(FCopySharedMipsForTexture2DArray)(
 			[&](FRHICommandListImmediate& RHICmdList)
 		{
-			RHICmdList.CopySharedMips(IntermediateTextureRHI.GetReference(), Context.Resource->GetTexture2DArrayRHI());
+			RHICopySharedMips(RHICmdList, IntermediateTextureRHI.GetReference(), Context.Resource->GetTexture2DArrayRHI());
 			bCopySharedMipsDone = true;
 		});
 		// Expected to execute immediately since ran on the renderthread.
