Add Stratga metalness to DiffuseAlbedo/F0 conversion node.

diff --git a/Engine/Shaders/Private/ShadingCommon.ush b/Engine/Shaders/Private/ShadingCommon.ush
--- a/Engine/Shaders/Private/ShadingCommon.ush
+++ b/Engine/Shaders/Private/ShadingCommon.ush
@@ -146,6 +146,11 @@ float3 ComputeF90(float3 F0, float3 EdgeColor, float Metallic)
 	return lerp(1.0, EdgeColor, Metallic.xxx);
 }
 
+float3 ComputeDiffuseAlbedo(float3 BaseColor, float Metallic)
+{
+	return BaseColor - BaseColor * Metallic;
+}
+
 float MakeRoughnessSafe(float Roughness, float MinRoughness=0.001f)
 {
 	return clamp(Roughness, MinRoughness, 1.0f);
diff --git a/Engine/Shaders/Private/Strata/Strata.ush b/Engine/Shaders/Private/Strata/Strata.ush
--- a/Engine/Shaders/Private/Strata/Strata.ush
+++ b/Engine/Shaders/Private/Strata/Strata.ush
@@ -944,7 +944,7 @@ FStrataData GetStrataSlabBSDF(
 		BSDF_SETHASF90(StrataData.InlinedBSDF, any(EdgeColor < 1.0f));
 
 		const float SafeMetallic = saturate(Metallic);
-		SLAB_DIFFUSEALBEDO(StrataData.InlinedBSDF)			= BaseColor - BaseColor * SafeMetallic;
+		SLAB_DIFFUSEALBEDO(StrataData.InlinedBSDF)			= ComputeDiffuseAlbedo(BaseColor, SafeMetallic);
 		SLAB_F0(StrataData.InlinedBSDF)						= ComputeF0(Specular, BaseColor, SafeMetallic);
 		SLAB_F90(StrataData.InlinedBSDF)					= ComputeF90(F0, EdgeColor, SafeMetallic);
 	}
diff --git a/Engine/Source/Runtime/Engine/Classes/Materials/MaterialExpressionStrata.h b/Engine/Source/Runtime/Engine/Classes/Materials/MaterialExpressionStrata.h
--- a/Engine/Source/Runtime/Engine/Classes/Materials/MaterialExpressionStrata.h
+++ b/Engine/Source/Runtime/Engine/Classes/Materials/MaterialExpressionStrata.h
@@ -839,3 +839,38 @@ class UMaterialExpressionStrataTransmittanceToMFP : public UMaterialExpressionSt
 	//~ End UMaterialExpression Interface
 };
 
+UCLASS(MinimalAPI, collapsecategories, hidecategories = Object, DisplayName = "Strata Metalness-To-DiffuseColorF0")
+class UMaterialExpressionStrataMetalnessToDiffuseAlbedoF0 : public UMaterialExpressionStrataBSDF
+{
+	GENERATED_UCLASS_BODY()
+
+	/**
+	 * Defines the overall color of the Material. (type = float3, unit = unitless, defaults to 0.18)
+	 */
+	UPROPERTY()
+	FExpressionInput BaseColor;
+
+	/**
+	 * Controls how \"metal-like\" your surface looks like. 0 means dielectric, 1 means conductor (type = float, unit = unitless, defaults to 0)
+	 */
+	UPROPERTY()
+	FExpressionInput Metallic;
+
+	/**
+	 * Used to scale the current amount of specularity on non-metallic surfaces and is a value between 0 and 1 (type = float, unit = unitless, defaults to plastic 0.5)
+	 */
+	UPROPERTY()
+	FExpressionInput Specular;
+
+	//~ Begin UMaterialExpression Interface
+#if WITH_EDITOR
+	virtual int32 Compile(class FMaterialCompiler* Compiler, int32 OutputIndex) override;
+	virtual void GetCaption(TArray<FString>& OutCaptions) const override;
+	virtual uint32 GetOutputType(int32 OutputIndex) override;
+	virtual uint32 GetInputType(int32 InputIndex) override;
+	virtual void GetConnectorToolTip(int32 InputIndex, int32 OutputIndex, TArray<FString>& OutToolTip) override;
+	virtual void GetExpressionToolTip(TArray<FString>& OutToolTip) override;
+#endif
+	//~ End UMaterialExpression Interface
+};
+
diff --git a/Engine/Source/Runtime/Engine/Private/Materials/HLSLMaterialTranslator.cpp b/Engine/Source/Runtime/Engine/Private/Materials/HLSLMaterialTranslator.cpp
--- a/Engine/Source/Runtime/Engine/Private/Materials/HLSLMaterialTranslator.cpp
+++ b/Engine/Source/Runtime/Engine/Private/Materials/HLSLMaterialTranslator.cpp
@@ -10515,6 +10515,7 @@ int32 FHLSLMaterialTranslator::StrataThinFilm(int32 A, int32 Thickness, int32 IO
 	);
 
 }
+
 int32 FHLSLMaterialTranslator::StrataTransmittanceToMFP(int32 TransmittanceColor, int32 DesiredThickness, int32 OutputIndex)
 {
 	if (OutputIndex == INDEX_NONE)
@@ -10542,6 +10543,32 @@ int32 FHLSLMaterialTranslator::StrataTransmittanceToMFP(int32 TransmittanceColor
 	return INDEX_NONE;
 }
 
+int32 FHLSLMaterialTranslator::StrataMetalnessToDiffuseAlbedoF0(int32 BaseColor, int32 Specular, int32 Metallic, int32 OutputIndex)
+{
+	if (OutputIndex == INDEX_NONE)
+	{
+		return INDEX_NONE;
+	}
+
+	switch (OutputIndex)
+	{
+	case 0:
+		return AddCodeChunk(MCT_Float3,
+			TEXT("ComputeDiffuseAlbedo(%s, saturate(%s))"),
+			*GetParameterCode(BaseColor),
+			*GetParameterCode(Metallic));
+		break;
+	case 1:
+		return AddCodeChunk(MCT_Float3,
+			TEXT("ComputeF0(%s, %s, saturate(%s))"),
+			*GetParameterCode(Specular),
+			*GetParameterCode(BaseColor),
+			*GetParameterCode(Metallic));
+		break;
+	}
+	return INDEX_NONE;
+}
+
 int32 FHLSLMaterialTranslator::MapARPassthroughCameraUV(int32 UV)
 {
 	if (UV == INDEX_NONE)
diff --git a/Engine/Source/Runtime/Engine/Private/Materials/HLSLMaterialTranslator.h b/Engine/Source/Runtime/Engine/Private/Materials/HLSLMaterialTranslator.h
--- a/Engine/Source/Runtime/Engine/Private/Materials/HLSLMaterialTranslator.h
+++ b/Engine/Source/Runtime/Engine/Private/Materials/HLSLMaterialTranslator.h
@@ -1053,7 +1053,8 @@ protected:
 	virtual int32 StrataWeightParameterBlending(int32 A, int32 Weight, FStrataOperator* PromoteToOperator) override;
 	virtual int32 StrataThinFilm(int32 A, int32 Thickness, int32 IOR, bool bInlinedEvaluation, int OperatorIndex, uint32 MaxDistanceFromLeaves) override;
 	virtual int32 StrataTransmittanceToMFP(int32 TransmittanceColor, int32 DesiredThickness, int32 OutputIndex) override;
-
+	virtual int32 StrataMetalnessToDiffuseAlbedoF0(int32 BaseColor, int32 Specular, int32 Metallic, int32 OutputIndex) override;
+	
 	virtual FStrataOperator& StrataCompilationRegisterOperator(int32 OperatorType, UMaterialExpression* Expression, UMaterialExpression* Parent, bool bUseParameterBlending = false) override;
 	virtual FStrataOperator& StrataCompilationGetOperator(UMaterialExpression* Expression) override;
 
diff --git a/Engine/Source/Runtime/Engine/Private/Materials/MaterialExpressions.cpp b/Engine/Source/Runtime/Engine/Private/Materials/MaterialExpressions.cpp
--- a/Engine/Source/Runtime/Engine/Private/Materials/MaterialExpressions.cpp
+++ b/Engine/Source/Runtime/Engine/Private/Materials/MaterialExpressions.cpp
@@ -22648,6 +22648,81 @@ void UMaterialExpressionStrataTransmittanceToMFP::GetExpressionToolTip(TArray<FS
 }
 #endif // WITH_EDITOR
 
+UMaterialExpressionStrataMetalnessToDiffuseAlbedoF0::UMaterialExpressionStrataMetalnessToDiffuseAlbedoF0(const FObjectInitializer& ObjectInitializer)
+	: Super(ObjectInitializer)
+{
+	struct FConstructorStatics
+	{
+		FText NAME_Strata;
+		FConstructorStatics() : NAME_Strata(LOCTEXT("Strata Helpers", "Strata Helpers")) { }
+	};
+	static FConstructorStatics ConstructorStatics;
+#if WITH_EDITORONLY_DATA
+	MenuCategories.Add(ConstructorStatics.NAME_Strata);
+
+	bShowOutputNameOnPin = true;
+
+	Outputs.Reset();
+	Outputs.Add(FExpressionOutput(TEXT("DiffuseAlbedo")));
+	Outputs.Add(FExpressionOutput(TEXT("F0")));
+#endif
+}
+
+#if WITH_EDITOR
+int32 UMaterialExpressionStrataMetalnessToDiffuseAlbedoF0::Compile(class FMaterialCompiler* Compiler, int32 OutputIndex)
+{
+	return Compiler->StrataMetalnessToDiffuseAlbedoF0(
+		BaseColor.GetTracedInput().Expression ? BaseColor.Compile(Compiler) : Compiler->Constant(0.18f),
+		Specular.GetTracedInput().Expression ? Specular.Compile(Compiler) : Compiler->Constant(0.5f),
+		Metallic.GetTracedInput().Expression ? Metallic.Compile(Compiler) : Compiler->Constant(0.f),
+		OutputIndex);
+}
+
+void UMaterialExpressionStrataMetalnessToDiffuseAlbedoF0::GetCaption(TArray<FString>& OutCaptions) const
+{
+	OutCaptions.Add(TEXT("Strata Metalness-To-DiffuseAlbedo-F0"));
+}
+
+uint32 UMaterialExpressionStrataMetalnessToDiffuseAlbedoF0::GetOutputType(int32 OutputIndex)
+{
+	switch (OutputIndex)
+	{
+	case 0:
+		return MCT_Float3; // Diffuse Albedo
+		break;
+	case 1:
+		return MCT_Float3; // F0
+		break;
+	}
+
+	check(false);
+	return MCT_Float1;
+}
+
+uint32 UMaterialExpressionStrataMetalnessToDiffuseAlbedoF0::GetInputType(int32 InputIndex)
+{
+	if (InputIndex == 0) { return MCT_Float3; }
+	if (InputIndex == 1) { return MCT_Float1; }
+	if (InputIndex == 2) { return MCT_Float1; }
+	return MCT_Float1;
+}
+
+void UMaterialExpressionStrataMetalnessToDiffuseAlbedoF0::GetConnectorToolTip(int32 InputIndex, int32 OutputIndex, TArray<FString>& OutToolTip)
+{
+	switch (OutputIndex)
+	{
+		case 1: ConvertToMultilineToolTip(TEXT("Defines the overall color of the Material. (type = float3, unit = unitless, defaults to 0.18)"), 80, OutToolTip); break;
+		case 2: ConvertToMultilineToolTip(TEXT("Controls how \"metal-like\" your surface looks like. 0 means dielectric, 1 means conductor (type = float, unit = unitless, defaults to 0)"), 80, OutToolTip); break;
+		case 3: ConvertToMultilineToolTip(TEXT("Used to scale the current amount of specularity on non-metallic surfaces and is a value between 0 and 1 (type = float, unit = unitless, defaults to plastic 0.5)"), 80, OutToolTip); break;
+	}
+}
+
+void UMaterialExpressionStrataMetalnessToDiffuseAlbedoF0::GetExpressionToolTip(TArray<FString>& OutToolTip)
+{
+	ConvertToMultilineToolTip(TEXT("Convert a metalness parameterization (BaseColor/Specular/Metallic) into DiffuseAlbedo/F0 parameterization."), 80, OutToolTip);
+
+}
+#endif // WITH_EDITOR
 
 
 UMaterialExpressionExecBegin::UMaterialExpressionExecBegin(const FObjectInitializer& ObjectInitializer)
diff --git a/Engine/Source/Runtime/Engine/Public/MaterialCompiler.h b/Engine/Source/Runtime/Engine/Public/MaterialCompiler.h
--- a/Engine/Source/Runtime/Engine/Public/MaterialCompiler.h
+++ b/Engine/Source/Runtime/Engine/Public/MaterialCompiler.h
@@ -528,9 +528,10 @@ public:
 	virtual int32 StrataWeightParameterBlending(int32 A, int32 Weight, FStrataOperator* PromoteToOperator) = 0;
 	virtual int32 StrataThinFilm(int32 A, int32 Thickness, int32 IOR, bool bInlinedEvaluation, int OperatorIndex, uint32 MaxDistanceFromLeaves) = 0;
 	virtual int32 StrataTransmittanceToMFP(int32 TransmittanceColor, int32 DesiredThickness, int32 OutputIndex) = 0;
-
+	virtual int32 StrataMetalnessToDiffuseAlbedoF0(int32 BaseColor, int32 Specular, int32 Metallic, int32 OutputIndex) = 0;
+	
 	/**
-	 * Register an operator of the tree representation the starta material and its topology.
+	 * Register an operator of the tree representation the Strata material and its topology.
 	 */
 	virtual FStrataOperator& StrataCompilationRegisterOperator(int32 OperatorType, UMaterialExpression* Expression, UMaterialExpression* Parent, bool bUseParameterBlending = false) = 0;
 	/**
@@ -1156,6 +1157,11 @@ public:
 	{
 		return Compiler->StrataTransmittanceToMFP(TransmittanceColor, DesiredThickness, OutputIndex);
 	}
+	
+	virtual int32 StrataMetalnessToDiffuseAlbedoF0(int32 BaseColor, int32 Specular, int32 Metallic, int32 OutputIndex) override
+	{
+		return Compiler->StrataMetalnessToDiffuseAlbedoF0(BaseColor, Specular, Metallic, OutputIndex);
+	}
 
 	virtual void StrataCompilationInfoRegisterCodeChunk(int32 CodeChunk, FStrataMaterialCompilationInfo& StrataMaterialCompilationInfo) override
 	{
