Fix Strata masked material.

diff --git a/Engine/Source/Runtime/Engine/Private/Materials/Material.cpp b/Engine/Source/Runtime/Engine/Private/Materials/Material.cpp
--- a/Engine/Source/Runtime/Engine/Private/Materials/Material.cpp
+++ b/Engine/Source/Runtime/Engine/Private/Materials/Material.cpp
@@ -3852,8 +3852,11 @@ void UMaterial::PostEditChangePropertyInternal(FPropertyChangedEvent& PropertyCh
 		}
 	}
 
+	//If we can be sure this material would be the same opaque as it is masked then allow it to be assumed opaque.
+	bCanMaskedBeAssumedOpaque = !OpacityMask.Expression && !(OpacityMask.UseConstant && OpacityMask.Constant < 0.999f) && !bUseMaterialAttributes;
+
 	// If the strata blending mode is changed, make sure we update the legacy blend mode too for the shaders to be filtered to passes correctly at runtime.
-	if (Engine_IsStrataEnabled() && PropertyChangedEvent.GetPropertyName() == GET_MEMBER_NAME_CHECKED(UMaterial, StrataBlendMode))
+	if (Engine_IsStrataEnabled())
 	{
 		switch (GetStrataBlendMode())
 		{
@@ -3874,10 +3877,6 @@ void UMaterial::PostEditChangePropertyInternal(FPropertyChangedEvent& PropertyCh
 		}
 	}
 
-
-	//If we can be sure this material would be the same opaque as it is masked then allow it to be assumed opaque.
-	bCanMaskedBeAssumedOpaque = !OpacityMask.Expression && !(OpacityMask.UseConstant && OpacityMask.Constant < 0.999f) && !bUseMaterialAttributes;
-
 	bool bRequiresCompilation = true;
 	if( PropertyThatChanged ) 
 	{
@@ -5545,6 +5544,7 @@ bool UMaterial::CastsRayTracedShadows() const
 static bool IsPropertyActive_Internal(EMaterialProperty InProperty,
 	EMaterialDomain Domain,
 	EBlendMode BlendMode,
+	EStrataBlendMode StrataBlendMode,
 	FMaterialShadingModelField ShadingModels,
 	ETranslucencyLightingMode TranslucencyLightingMode,
 	bool bBlendableOutputAlpha,
@@ -5680,7 +5680,7 @@ static bool IsPropertyActive_Internal(EMaterialProperty InProperty,
 		}
 		break;
 	case MP_OpacityMask:
-		Active = BlendMode == BLEND_Masked;
+		Active = (BlendMode == BLEND_Masked) || (bStrataEnabled && StrataBlendMode == SBM_Masked);
 		break;
 	case MP_BaseColor:
 	case MP_AmbientOcclusion:
@@ -5752,6 +5752,7 @@ bool UMaterial::IsPropertyActiveInEditor(EMaterialProperty InProperty) const
 	return IsPropertyActive_Internal(InProperty,
 		MaterialDomain,
 		BlendMode,
+		StrataBlendMode,
 		ShadingModels,
 		TranslucencyLightingMode,
 		BlendableOutputAlpha,
@@ -5766,6 +5767,7 @@ bool UMaterial::IsPropertyActiveInDerived(EMaterialProperty InProperty, const UM
 	return IsPropertyActive_Internal(InProperty,
 		MaterialDomain,
 		DerivedMaterial->GetBlendMode(),
+		DerivedMaterial->GetStrataBlendMode(),
 		DerivedMaterial->GetShadingModels(),
 		TranslucencyLightingMode,
 		BlendableOutputAlpha,
