Remove unnecessary material node inputs when Strata is enabled.

diff --git a/Engine/Source/Editor/UnrealEd/Classes/MaterialGraph/MaterialGraph.h b/Engine/Source/Editor/UnrealEd/Classes/MaterialGraph/MaterialGraph.h
--- a/Engine/Source/Editor/UnrealEd/Classes/MaterialGraph/MaterialGraph.h
+++ b/Engine/Source/Editor/UnrealEd/Classes/MaterialGraph/MaterialGraph.h
@@ -55,6 +55,11 @@ struct FMaterialInputInfo
 
 	bool IsVisiblePin(const UMaterial* Material, bool bIgnoreMaterialAttributes = false) const
 	{
+		if (!Material->IsPropertySupported(Property))
+		{
+			return false;
+		}
+
 		if(Material->bUseMaterialAttributes && !bIgnoreMaterialAttributes)
 		{
 			return Property == MP_MaterialAttributes;
diff --git a/Engine/Source/Runtime/Engine/Classes/Materials/Material.h b/Engine/Source/Runtime/Engine/Classes/Materials/Material.h
--- a/Engine/Source/Runtime/Engine/Classes/Materials/Material.h
+++ b/Engine/Source/Runtime/Engine/Classes/Materials/Material.h
@@ -1739,6 +1739,9 @@ public:
 	ENGINE_API bool HasVertexPositionOffsetConnected() const;
 	ENGINE_API bool HasPixelDepthOffsetConnected() const;
 
+	// Return true if the property is supported
+	ENGINE_API bool IsPropertySupported(EMaterialProperty InProperty) const;
+
 #if WITH_EDITOR
 	static void NotifyCompilationFinished(UMaterialInterface* Material);
 
diff --git a/Engine/Source/Runtime/Engine/Private/Materials/Material.cpp b/Engine/Source/Runtime/Engine/Private/Materials/Material.cpp
--- a/Engine/Source/Runtime/Engine/Private/Materials/Material.cpp
+++ b/Engine/Source/Runtime/Engine/Private/Materials/Material.cpp
@@ -5541,6 +5541,29 @@ bool UMaterial::CastsRayTracedShadows() const
 	return bCastRayTracedShadows;
 }
 
+// This is used to list the supported properties (i.e. when Strata is enabled/disabled)
+bool UMaterial::IsPropertySupported(EMaterialProperty InProperty) const
+{
+	bool bSupported = true;
+	if (Engine_IsStrataEnabled())
+	{
+		bSupported = false;
+		switch (InProperty)
+		{
+		case MP_Refraction:
+		case MP_Opacity:
+		case MP_OpacityMask:
+		case MP_AmbientOcclusion:
+		case MP_WorldPositionOffset:
+		case MP_PixelDepthOffset:
+		case MP_FrontMaterial:
+			bSupported = true;
+			break;
+		}
+	}
+	return bSupported;
+}
+
 static bool IsPropertyActive_Internal(EMaterialProperty InProperty,
 	EMaterialDomain Domain,
 	EBlendMode BlendMode,
@@ -5550,7 +5573,8 @@ static bool IsPropertyActive_Internal(EMaterialProperty InProperty,
 	bool bBlendableOutputAlpha,
 	bool bHasRefraction,
 	bool bUsesShadingModelFromMaterialExpression,
-	bool bIsTranslucencyWritingVelocity)
+	bool bIsTranslucencyWritingVelocity,
+	bool bIsSupported)
 {
 	const bool bStrataEnabled = Engine_IsStrataEnabled();
 
@@ -5663,78 +5687,112 @@ static bool IsPropertyActive_Internal(EMaterialProperty InProperty,
 	
 	bool Active = true;
 
-	switch (InProperty)
+	if (bStrataEnabled)
 	{
-	case MP_DiffuseColor:
-	case MP_SpecularColor:
 		Active = false;
-		break;
-	case MP_Refraction:
-		Active = (bIsTranslucentBlendMode && BlendMode != BLEND_AlphaHoldout && BlendMode != BLEND_Modulate) || ShadingModels.HasShadingModel(MSM_SingleLayerWater);
-		break;
-	case MP_Opacity:
-		Active = (bIsTranslucentBlendMode && BlendMode != BLEND_Modulate) || ShadingModels.HasShadingModel(MSM_SingleLayerWater);
-		if (IsSubsurfaceShadingModel(ShadingModels))
+		if (bIsSupported)
 		{
-			Active = true;
+			switch (InProperty)
+			{
+			case MP_Refraction:
+				Active = (bIsTranslucentBlendMode && BlendMode != BLEND_AlphaHoldout && BlendMode != BLEND_Modulate) || ShadingModels.HasShadingModel(MSM_SingleLayerWater);
+				break;
+			case MP_Opacity:
+				Active = (bIsTranslucentBlendMode && BlendMode != BLEND_Modulate) || ShadingModels.HasShadingModel(MSM_SingleLayerWater);
+				break;
+			case MP_OpacityMask:
+				Active = BlendMode == BLEND_Masked || StrataBlendMode == SBM_Masked;
+				break;
+			case MP_AmbientOcclusion:
+				Active = ShadingModels.IsLit();
+				break;
+			case MP_WorldPositionOffset:
+				Active = true;
+				break;
+			case MP_PixelDepthOffset:
+				Active = (!bIsTranslucentBlendMode) || (bIsTranslucencyWritingVelocity);
+				break;
+			case MP_FrontMaterial:
+				Active = true;
+				break;
+			}
 		}
-		break;
-	case MP_OpacityMask:
-		Active = (BlendMode == BLEND_Masked) || (bStrataEnabled && StrataBlendMode == SBM_Masked);
-		break;
-	case MP_BaseColor:
-	case MP_AmbientOcclusion:
-		Active = ShadingModels.IsLit();
-		break;
-	case MP_Specular:
-	case MP_Roughness:
-		Active = ShadingModels.IsLit() && (!bIsTranslucentBlendMode || !bIsVolumetricTranslucencyLightingMode);
-		break;
-	case MP_Anisotropy:
-		Active = ShadingModels.HasAnyShadingModel({ MSM_DefaultLit, MSM_ClearCoat }) && (!bIsTranslucentBlendMode || !bIsVolumetricTranslucencyLightingMode);
-		break;
-	case MP_Metallic:
-		// Subsurface models store opacity in place of Metallic in the GBuffer
-		Active = ShadingModels.IsLit() && (!bIsTranslucentBlendMode || !bIsVolumetricTranslucencyLightingMode);
-		break;
-	case MP_Normal:
-		Active = (ShadingModels.IsLit() && (!bIsTranslucentBlendMode || !bIsNonDirectionalTranslucencyLightingMode)) || bHasRefraction;
-		break;
-	case MP_Tangent:
-		Active = ShadingModels.HasAnyShadingModel({ MSM_DefaultLit, MSM_ClearCoat }) && (!bIsTranslucentBlendMode || !bIsVolumetricTranslucencyLightingMode);
-		break;
-	case MP_SubsurfaceColor:
-		Active = ShadingModels.HasAnyShadingModel({ MSM_Subsurface, MSM_PreintegratedSkin, MSM_TwoSidedFoliage, MSM_Cloth });
-		break;
-	case MP_CustomData0:
-		Active = ShadingModels.HasAnyShadingModel({ MSM_ClearCoat, MSM_Hair, MSM_Cloth, MSM_Eye, MSM_SubsurfaceProfile });
-		break;
-	case MP_CustomData1:
-		Active = ShadingModels.HasAnyShadingModel({ MSM_ClearCoat, MSM_Eye });
-		break;
-	case MP_EmissiveColor:
-		// Emissive is always active, even for light functions and post process materials, 
-		// but not for AlphaHoldout
-		Active = BlendMode != BLEND_AlphaHoldout;
-		break;
-	case MP_WorldPositionOffset:
-		Active = true;
-		break;
-	case MP_PixelDepthOffset:
-		Active = (!bIsTranslucentBlendMode) || (bIsTranslucencyWritingVelocity);
-		break;
-	case MP_ShadingModel:
-		Active = bUsesShadingModelFromMaterialExpression;
-		break;
-	case MP_FrontMaterial:
-		{
-			Active = bStrataEnabled;
+	}
+	else
+	{
+		switch (InProperty)
+		{
+		case MP_DiffuseColor:
+		case MP_SpecularColor:
+			Active = false;
+			break;
+		case MP_Refraction:
+			Active = (bIsTranslucentBlendMode && BlendMode != BLEND_AlphaHoldout && BlendMode != BLEND_Modulate) || ShadingModels.HasShadingModel(MSM_SingleLayerWater);
+			break;
+		case MP_Opacity:
+			Active = (bIsTranslucentBlendMode && BlendMode != BLEND_Modulate) || ShadingModels.HasShadingModel(MSM_SingleLayerWater);
+			if (IsSubsurfaceShadingModel(ShadingModels))
+			{
+				Active = true;
+			}
+			break;
+		case MP_OpacityMask:
+			Active = BlendMode == BLEND_Masked;
+			break;
+		case MP_BaseColor:
+		case MP_AmbientOcclusion:
+			Active = ShadingModels.IsLit();
+			break;
+		case MP_Specular:
+		case MP_Roughness:
+			Active = ShadingModels.IsLit() && (!bIsTranslucentBlendMode || !bIsVolumetricTranslucencyLightingMode);
+			break;
+		case MP_Anisotropy:
+			Active = ShadingModels.HasAnyShadingModel({ MSM_DefaultLit, MSM_ClearCoat }) && (!bIsTranslucentBlendMode || !bIsVolumetricTranslucencyLightingMode);
+			break;
+		case MP_Metallic:
+			// Subsurface models store opacity in place of Metallic in the GBuffer
+			Active = ShadingModels.IsLit() && (!bIsTranslucentBlendMode || !bIsVolumetricTranslucencyLightingMode);
+			break;
+		case MP_Normal:
+			Active = (ShadingModels.IsLit() && (!bIsTranslucentBlendMode || !bIsNonDirectionalTranslucencyLightingMode)) || bHasRefraction;
+			break;
+		case MP_Tangent:
+			Active = ShadingModels.HasAnyShadingModel({ MSM_DefaultLit, MSM_ClearCoat }) && (!bIsTranslucentBlendMode || !bIsVolumetricTranslucencyLightingMode);
+			break;
+		case MP_SubsurfaceColor:
+			Active = ShadingModels.HasAnyShadingModel({ MSM_Subsurface, MSM_PreintegratedSkin, MSM_TwoSidedFoliage, MSM_Cloth });
+			break;
+		case MP_CustomData0:
+			Active = ShadingModels.HasAnyShadingModel({ MSM_ClearCoat, MSM_Hair, MSM_Cloth, MSM_Eye, MSM_SubsurfaceProfile });
+			break;
+		case MP_CustomData1:
+			Active = ShadingModels.HasAnyShadingModel({ MSM_ClearCoat, MSM_Eye });
+			break;
+		case MP_EmissiveColor:
+			// Emissive is always active, even for light functions and post process materials, 
+			// but not for AlphaHoldout
+			Active = BlendMode != BLEND_AlphaHoldout;
+			break;
+		case MP_WorldPositionOffset:
+			Active = true;
+			break;
+		case MP_PixelDepthOffset:
+			Active = (!bIsTranslucentBlendMode) || (bIsTranslucencyWritingVelocity);
+			break;
+		case MP_ShadingModel:
+			Active = bUsesShadingModelFromMaterialExpression;
+			break;
+		case MP_FrontMaterial:
+			{
+				Active = bStrataEnabled;
+				break;
+			}
+		case MP_MaterialAttributes:
+		default:
+			Active = true;
 			break;
 		}
-	case MP_MaterialAttributes:
-	default:
-		Active = true;
-		break;
 	}
 	return Active;
 }
@@ -5758,7 +5816,8 @@ bool UMaterial::IsPropertyActiveInEditor(EMaterialProperty InProperty) const
 		BlendableOutputAlpha,
 		Refraction.IsConnected(),
 		IsShadingModelFromMaterialExpression(),
-		IsTranslucencyWritingVelocity());
+		IsTranslucencyWritingVelocity(),
+		IsPropertySupported(InProperty));
 }
 #endif // WITH_EDITOR
 
@@ -5773,7 +5832,8 @@ bool UMaterial::IsPropertyActiveInDerived(EMaterialProperty InProperty, const UM
 		BlendableOutputAlpha,
 		Refraction.IsConnected(),
 		DerivedMaterial->IsShadingModelFromMaterialExpression(),
-		IsTranslucencyWritingVelocity());
+		IsTranslucencyWritingVelocity(),
+		IsPropertySupported(InProperty));
 }
 
 #if WITH_EDITORONLY_DATA
