UDataLayer deprecated in favor of UDataLayerInstance & UDataLayerAsset

diff --git a/Engine/Source/Runtime/Engine/Private/WorldPartition/DataLayer/DataLayerSubsystem.cpp b/Engine/Source/Runtime/Engine/Private/WorldPartition/DataLayer/DataLayerSubsystem.cpp
--- a/Engine/Source/Runtime/Engine/Private/WorldPartition/DataLayer/DataLayerSubsystem.cpp
+++ b/Engine/Source/Runtime/Engine/Private/WorldPartition/DataLayer/DataLayerSubsystem.cpp
@@ -2,7 +2,7 @@
 
 #include "WorldPartition/DataLayer/DataLayerSubsystem.h"
 #include "WorldPartition/DataLayer/WorldDataLayers.h"
-#include "WorldPartition/DataLayer/DataLayer.h"
+#include "WorldPartition/DataLayer/DataLayerAsset.h"
 #include "WorldPartition/WorldPartitionDebugHelper.h"
 #include "WorldPartition/WorldPartitionSubsystem.h"
 #include "WorldPartition/WorldPartitionRuntimeCell.h"
@@ -67,82 +67,56 @@ const TSet<FName>& UDataLayerSubsystem::GetEffectiveLoadedDataLayerNames() const
 	return WorldDataLayers ? WorldDataLayers->GetEffectiveLoadedDataLayerNames() : EmptySet;
 }
 
-UDataLayer* UDataLayerSubsystem::GetDataLayer(const FActorDataLayer& InDataLayer) const
-{
-	return GetDataLayerFromName(InDataLayer.Name);
-}
-
-UDataLayer* UDataLayerSubsystem::GetDataLayerFromLabel(FName InDataLayerLabel) const
+UDataLayerInstance* UDataLayerSubsystem::GetDataLayerFromAsset(const UDataLayerAsset* InDataLayerAsset) const
 {
 	const AWorldDataLayers* WorldDataLayers = GetWorld()->GetWorldDataLayers();
-	return WorldDataLayers ? const_cast<UDataLayer*>(WorldDataLayers->GetDataLayerFromLabel(InDataLayerLabel)) : nullptr;
+	return WorldDataLayers ? const_cast<UDataLayerInstance*>(WorldDataLayers->GetDataLayerInstance(InDataLayerAsset)) : nullptr;
 }
 
-UDataLayer* UDataLayerSubsystem::GetDataLayerFromName(FName InDataLayerName) const
+void UDataLayerSubsystem::SetDataLayerInstanceRuntimeState(const UDataLayerAsset* InDataLayerAsset, EDataLayerRuntimeState InState, bool bInIsRecursive)
 {
-	const AWorldDataLayers* WorldDataLayers = GetWorld()->GetWorldDataLayers();
-	return WorldDataLayers ? const_cast<UDataLayer*>(WorldDataLayers->GetDataLayerFromName(InDataLayerName)) : nullptr;
+	SetDataLayerRuntimeState(GetDataLayerFromAsset(InDataLayerAsset), InState, bInIsRecursive);
 }
 
-void UDataLayerSubsystem::SetDataLayerRuntimeState(const UDataLayer* InDataLayer, EDataLayerRuntimeState InState, bool bInIsRecursive)
+EDataLayerRuntimeState UDataLayerSubsystem::GetDataLayerInstanceRuntimeState(const UDataLayerAsset* InDataLayerAsset) const
 {
-	if (InDataLayer)
-	{
-		if (AWorldDataLayers* WorldDataLayers = GetWorld()->GetWorldDataLayers())
-		{
-			WorldDataLayers->SetDataLayerRuntimeState(FActorDataLayer(InDataLayer->GetFName()), InState, bInIsRecursive);
-		}
-	}
-	else
-	{
-		UE_LOG(LogWorldPartition, Warning, TEXT("UDataLayerSubsystem::SetDataLayerRuntimeState called with null Data Layer"));
-	}
+	return GetDataLayerRuntimeState(GetDataLayerFromAsset(InDataLayerAsset));
 }
 
-void UDataLayerSubsystem::SetDataLayerRuntimeStateByName(const FName& InDataLayerName, EDataLayerRuntimeState InState, bool bInIsRecursive)
+EDataLayerRuntimeState UDataLayerSubsystem::GetDataLayerInstanceEffectiveRuntimeState(const UDataLayerAsset* InDataLayerAsset) const
 {
-	if (UDataLayer* DataLayer = GetDataLayerFromName(InDataLayerName))
-	{
-		SetDataLayerRuntimeState(DataLayer, InState, bInIsRecursive);
-	}
-	else
-	{
-		UE_LOG(LogWorldPartition, Warning, TEXT("UDataLayerSubsystem::SetDataLayerRuntimeStateByName unknown Data Layer: '%s'"), *InDataLayerName.ToString());
-	}
+	return GetDataLayerEffectiveRuntimeState(GetDataLayerFromAsset(InDataLayerAsset));
 }
 
-void UDataLayerSubsystem::SetDataLayerRuntimeState(const FActorDataLayer& InDataLayer, EDataLayerRuntimeState InState, bool bInIsRecursive)
+UDataLayerInstance* UDataLayerSubsystem::GetDataLayerInstance(const FName& InDataLayerInstanceName) const
 {
-	if (UDataLayer* DataLayer = GetDataLayerFromName(InDataLayer.Name))
-	{
-		SetDataLayerRuntimeState(DataLayer, InState, bInIsRecursive);
-	}
-	else
-	{
-		UE_LOG(LogWorldPartition, Warning, TEXT("UDataLayerSubsystem::SetDataLayerRuntimeState unknown Data Layer: '%s'"), *InDataLayer.Name.ToString());
-	}
+	const AWorldDataLayers* WorldDataLayers = GetWorld()->GetWorldDataLayers();
+	return WorldDataLayers ? const_cast<UDataLayerInstance*>(WorldDataLayers->GetDataLayerInstance(InDataLayerInstanceName)) : nullptr;
 }
 
-void UDataLayerSubsystem::SetDataLayerRuntimeStateByLabel(const FName& InDataLayerLabel, EDataLayerRuntimeState InState, bool bInIsRecursive)
+void UDataLayerSubsystem::SetDataLayerRuntimeState(const UDataLayerInstance* InDataLayerInstance, EDataLayerRuntimeState InState, bool bInIsRecursive)
 {
-	if (UDataLayer* DataLayer = GetDataLayerFromLabel(InDataLayerLabel))
+	if (InDataLayerInstance)
 	{
-		SetDataLayerRuntimeState(DataLayer, InState, bInIsRecursive);
+		if (AWorldDataLayers* WorldDataLayers = GetWorld()->GetWorldDataLayers())
+		{
+			WorldDataLayers->SetDataLayerRuntimeState(InDataLayerInstance, InState, bInIsRecursive);
+		}
 	}
 	else
 	{
-		UE_LOG(LogWorldPartition, Warning, TEXT("UDataLayerSubsystem::SetDataLayerRuntimeStateByLabel unknown Data Layer: '%s'"), *InDataLayerLabel.ToString());
+		UE_LOG(LogWorldPartition, Warning, TEXT("UDataLayerSubsystem::SetDataLayerRuntimeState called with null Data Layer"));
 	}
 }
 
-EDataLayerRuntimeState UDataLayerSubsystem::GetDataLayerRuntimeState(const UDataLayer* InDataLayer) const
+EDataLayerRuntimeState UDataLayerSubsystem::GetDataLayerRuntimeState(const UDataLayerInstance* InDataLayer) const
 {
 	if (!InDataLayer)
 	{
 		return EDataLayerRuntimeState::Unloaded;
 	}
 
-	return GetDataLayerRuntimeStateByName(InDataLayer->GetFName());
+	return GetDataLayerRuntimeStateByName(InDataLayer->GetDataLayerFName());
 }
 
 EDataLayerRuntimeState UDataLayerSubsystem::GetDataLayerRuntimeStateByName(const FName& InDataLayerName) const
@@ -156,24 +130,14 @@ EDataLayerRuntimeState UDataLayerSubsystem::GetDataLayerRuntimeStateByName(const
 	return WorldDataLayers->GetDataLayerRuntimeStateByName(InDataLayerName);
 }
 
-EDataLayerRuntimeState UDataLayerSubsystem::GetDataLayerRuntimeState(const FActorDataLayer& InDataLayer) const
-{
-	return GetDataLayerRuntimeStateByName(InDataLayer.Name);
-}
-
-EDataLayerRuntimeState UDataLayerSubsystem::GetDataLayerRuntimeStateByLabel(const FName& InDataLayerLabel) const
-{
-	return GetDataLayerRuntimeState(GetDataLayerFromLabel(InDataLayerLabel));
-}
-
-EDataLayerRuntimeState UDataLayerSubsystem::GetDataLayerEffectiveRuntimeState(const UDataLayer* InDataLayer) const
+EDataLayerRuntimeState UDataLayerSubsystem::GetDataLayerEffectiveRuntimeState(const UDataLayerInstance* InDataLayer) const
 {
 	if (!InDataLayer)
 	{
 		return EDataLayerRuntimeState::Unloaded;
 	}
 
-	return GetDataLayerEffectiveRuntimeStateByName(InDataLayer->GetFName());
+	return GetDataLayerEffectiveRuntimeStateByName(InDataLayer->GetDataLayerFName());
 }
 
 EDataLayerRuntimeState UDataLayerSubsystem::GetDataLayerEffectiveRuntimeStateByName(const FName& InDataLayerName) const
@@ -187,16 +151,6 @@ EDataLayerRuntimeState UDataLayerSubsystem::GetDataLayerEffectiveRuntimeStateByN
 	return WorldDataLayers->GetDataLayerEffectiveRuntimeStateByName(InDataLayerName);
 }
 
-EDataLayerRuntimeState UDataLayerSubsystem::GetDataLayerEffectiveRuntimeState(const FActorDataLayer& InDataLayer) const
-{
-	return GetDataLayerEffectiveRuntimeStateByName(InDataLayer.Name);
-}
-
-EDataLayerRuntimeState UDataLayerSubsystem::GetDataLayerEffectiveRuntimeStateByLabel(const FName& InDataLayerLabel) const
-{
-	return GetDataLayerEffectiveRuntimeState(GetDataLayerFromLabel(InDataLayerLabel));
-}
-
 bool UDataLayerSubsystem::IsAnyDataLayerInEffectiveRuntimeState(const TArray<FName>& InDataLayerNames, EDataLayerRuntimeState InState) const
 {
 	for (FName DataLayerName : InDataLayerNames)
@@ -229,22 +183,22 @@ void UDataLayerSubsystem::DrawDataLayersStatus(UCanvas* Canvas, FVector2D& Offse
 		{
 			FWorldPartitionDebugHelper::DrawText(Canvas, Title, GEngine->GetSmallFont(), HeaderColor, Pos, &MaxTextWidth);
 
-			TArray<UDataLayer*> DataLayers;
+			TArray<const UDataLayerInstance*> DataLayers;
 			DataLayers.Reserve(LayerNames.Num());
 			for (const FName& DataLayerName : LayerNames)
 			{
-				if (UDataLayer* DataLayer = GetDataLayerFromName(DataLayerName))
+				if (const UDataLayerInstance* DataLayer = GetDataLayerInstance(DataLayerName))
 				{
 					DataLayers.Add(DataLayer);
 				}
 			}
 
-			DataLayers.Sort([](const UDataLayer& A, const UDataLayer& B) { return A.GetDataLayerLabel().LexicalLess(B.GetDataLayerLabel()); });
+			DataLayers.Sort([](const UDataLayerInstance& A, const UDataLayerInstance& B) { return A.GetDataLayerFullName() < B.GetDataLayerFullName(); });
 
 			UFont* DataLayerFont = GEngine->GetSmallFont();
-			for (UDataLayer* DataLayer : DataLayers)
+			for (const UDataLayerInstance* DataLayer : DataLayers)
 			{
-				FString DataLayerString = DataLayer->GetDataLayerLabel().ToString();
+				FString DataLayerString = DataLayer->GetDataLayerShortName();
 
 				if (GDrawDataLayersLoadTime)
 				{
@@ -275,11 +229,11 @@ void UDataLayerSubsystem::DrawDataLayersStatus(UCanvas* Canvas, FVector2D& Offse
 	TSet<FName> UnloadedDataLayers;
 	if (AWorldDataLayers* WorldDataLayers = GetWorld()->GetWorldDataLayers())
 	{		
-		WorldDataLayers->ForEachDataLayer([&LoadedDataLayers, &ActiveDataLayers, &UnloadedDataLayers](UDataLayer* DataLayer)
+		WorldDataLayers->ForEachDataLayer([&LoadedDataLayers, &ActiveDataLayers, &UnloadedDataLayers](UDataLayerInstance* DataLayer)
 		{
 			if (DataLayer->IsRuntime())
 			{
-				const FName DataLayerName = DataLayer->GetFName();
+				const FName DataLayerName = DataLayer->GetDataLayerFName();
 				if (!LoadedDataLayers.Contains(DataLayerName) && !ActiveDataLayers.Contains(DataLayerName))
 				{
 					UnloadedDataLayers.Add(DataLayerName);
@@ -297,32 +251,32 @@ void UDataLayerSubsystem::DrawDataLayersStatus(UCanvas* Canvas, FVector2D& Offse
 	{
 		for (FName DataLayerName : UnloadedDataLayers)
 		{
-			if (UDataLayer* DataLayer = GetDataLayerFromName(DataLayerName))
+			if (const UDataLayerInstance* DataLayer = GetDataLayerInstance(DataLayerName))
 			{
 				ActiveDataLayersLoadTime.Remove(DataLayer);
 			}
 		}
 
-		TArray<UDataLayer*> LoadingDataLayers;
+		TArray<const UDataLayerInstance*> LoadingDataLayers;
 		LoadingDataLayers.Reserve(LoadedDataLayers.Num() + ActiveDataLayers.Num());
-		auto CopyLambda = [this](FName DataLayerName) { return GetDataLayerFromName(DataLayerName); };
+		auto CopyLambda = [this](FName DataLayerName) { return GetDataLayerInstance(DataLayerName); };
 		Algo::Transform(LoadedDataLayers, LoadingDataLayers, CopyLambda);
 		Algo::Transform(ActiveDataLayers, LoadingDataLayers, CopyLambda);
 
-		for (UDataLayer* DataLayer : LoadingDataLayers)
+		for (const UDataLayerInstance* DataLayer : LoadingDataLayers)
 		{
 			double* DataLayerLoadTime = ActiveDataLayersLoadTime.Find(DataLayer);
 
-			auto IsDataLayerReady = [WorldPartitionSubsystem](UDataLayer* DataLayer, EWorldPartitionRuntimeCellState TargetState)
+			auto IsDataLayerReady = [WorldPartitionSubsystem](const UDataLayerInstance* DataLayer, EWorldPartitionRuntimeCellState TargetState)
 			{
 				FWorldPartitionStreamingQuerySource QuerySource;
 				QuerySource.bDataLayersOnly = true;
 				QuerySource.bSpatialQuery = false;
-				QuerySource.DataLayers.Add(DataLayer->GetFName());
+				QuerySource.DataLayers.Add(DataLayer->GetDataLayerFName());
 				return WorldPartitionSubsystem->IsStreamingCompleted(TargetState, { QuerySource }, true);
 			};
 
-			const EWorldPartitionRuntimeCellState TargetState = ActiveDataLayers.Contains(DataLayer->GetFName()) ? EWorldPartitionRuntimeCellState::Activated : EWorldPartitionRuntimeCellState::Loaded;
+			const EWorldPartitionRuntimeCellState TargetState = ActiveDataLayers.Contains(DataLayer->GetDataLayerFName()) ? EWorldPartitionRuntimeCellState::Activated : EWorldPartitionRuntimeCellState::Loaded;
 
 			if (!DataLayerLoadTime)
 			{
@@ -347,9 +301,9 @@ void UDataLayerSubsystem::DrawDataLayersStatus(UCanvas* Canvas, FVector2D& Offse
 	}
 }
 
-TArray<UDataLayer*> UDataLayerSubsystem::ConvertArgsToDataLayers(UWorld* World, const TArray<FString>& InArgs)
+TArray<UDataLayerInstance*> UDataLayerSubsystem::ConvertArgsToDataLayers(UWorld* World, const TArray<FString>& InArgs)
 {
-	TSet<UDataLayer*> OutDataLayers;
+	TSet<UDataLayerInstance*> OutDataLayers;
 
 	const TCHAR* QuoteChar = TEXT("\"");
 	bool bQuoteStarted = false;
@@ -393,28 +347,33 @@ TArray<UDataLayer*> UDataLayerSubsystem::ConvertArgsToDataLayers(UWorld* World,
 
 	for (const FString& Arg : Args)
 	{
-		FName DataLayerLabel = FName(Arg);
+		FName DataLayerName = FName(Arg);
 		if (const AWorldDataLayers* WorldDataLayers = World->GetWorldDataLayers())
 		{
-			UDataLayer* DataLayer = const_cast<UDataLayer*>(WorldDataLayers->GetDataLayerFromLabel(DataLayerLabel));
-			if (!DataLayer)
+			bool bShortNameFound = false;
+			FString SanitizedDataLayerName = DataLayerName.ToString().Replace(TEXT(" "), TEXT(""));
+			WorldDataLayers->ForEachDataLayer([&SanitizedDataLayerName, &OutDataLayers, &bShortNameFound](UDataLayerInstance* It)
 			{
-				// Try to find Data Layer by ignoring whitespaces and case
-				FString DataLayerLabelToFind = DataLayerLabel.ToString().Replace(TEXT(" "), TEXT(""));
-				WorldDataLayers->ForEachDataLayer([&DataLayerLabelToFind, &DataLayer](UDataLayer* It)
+
+				if (It->GetDataLayerShortName().Find(SanitizedDataLayerName, ESearchCase::IgnoreCase) == 0)
 				{
-					if (It->GetDataLayerLabel().ToString().Replace(TEXT(" "), TEXT("")).Compare(DataLayerLabelToFind, ESearchCase::IgnoreCase) == 0)
+					if (bShortNameFound)
 					{
-						DataLayer = It;
+						UE_LOG(LogWorldPartition, Error, TEXT("Found 2 data layers with the ShortName %s when converting arguments. Consider using the data layers FullName or renaming one of the two."));
 						return false;
 					}
-					return true;
-				});
-			}
-			if (DataLayer)
-			{
-				OutDataLayers.Add(DataLayer);
-			}
+
+					OutDataLayers.Add(It);
+					bShortNameFound = true;
+				}
+				else if (It->GetDataLayerFullName().Find(SanitizedDataLayerName, ESearchCase::IgnoreCase) == 0)
+				{
+					OutDataLayers.Add(It);
+					return false;
+				}
+
+				return true;
+			});
 		}
 	}
 
@@ -431,7 +390,7 @@ void UDataLayerSubsystem::DumpDataLayers(FOutputDevice& OutputDevice) const
 
 FAutoConsoleCommand UDataLayerSubsystem::ToggleDataLayerActivation(
 	TEXT("wp.Runtime.ToggleDataLayerActivation"),
-	TEXT("Toggles DataLayers active state. Args [DataLayerLabels]"),
+	TEXT("Toggles DataLayers active state. Args [DataLayerNames]"),
 	FConsoleCommandWithArgsDelegate::CreateLambda([](const TArray<FString>& InArgs)
 	{
 		for (const FWorldContext& Context : GEngine->GetWorldContexts())
@@ -441,10 +400,10 @@ FAutoConsoleCommand UDataLayerSubsystem::ToggleDataLayerActivation(
 			{
 				if (UDataLayerSubsystem* DataLayerSubsystem = World->GetSubsystem<UDataLayerSubsystem>())
 				{
-					TArray<UDataLayer*> DataLayers = UDataLayerSubsystem::ConvertArgsToDataLayers(World, InArgs);
-					for (UDataLayer* DataLayer : DataLayers)
+					TArray<UDataLayerInstance*> DataLayerInstances = UDataLayerSubsystem::ConvertArgsToDataLayers(World, InArgs);
+					for (UDataLayerInstance* DataLayerInstance : DataLayerInstances)
 					{
-						DataLayerSubsystem->SetDataLayerRuntimeState(DataLayer, DataLayerSubsystem->GetDataLayerRuntimeState(DataLayer) == EDataLayerRuntimeState::Activated ? EDataLayerRuntimeState::Unloaded : EDataLayerRuntimeState::Activated);
+						DataLayerSubsystem->SetDataLayerRuntimeState(DataLayerInstance, DataLayerSubsystem->GetDataLayerRuntimeState(DataLayerInstance) == EDataLayerRuntimeState::Activated ? EDataLayerRuntimeState::Unloaded : EDataLayerRuntimeState::Activated);
 					}
 				}
 			}
@@ -454,7 +413,7 @@ FAutoConsoleCommand UDataLayerSubsystem::ToggleDataLayerActivation(
 
 FAutoConsoleCommand UDataLayerSubsystem::SetDataLayerRuntimeStateCommand(
 	TEXT("wp.Runtime.SetDataLayerRuntimeState"),
-	TEXT("Sets Runtime DataLayers state. Args [State = Unloaded, Loaded, Activated] [DataLayerLabels]"),
+	TEXT("Sets Runtime DataLayers state. Args [State = Unloaded, Loaded, Activated] [DataLayerNames]"),
 	FConsoleCommandWithArgsDelegate::CreateLambda([](const TArray<FString>& InArgs)
 	{
 		if (InArgs.Num() < 2)
@@ -480,10 +439,10 @@ FAutoConsoleCommand UDataLayerSubsystem::SetDataLayerRuntimeStateCommand(
 			{
 				if (UDataLayerSubsystem* DataLayerSubsystem = World->GetSubsystem<UDataLayerSubsystem>())
 				{
-					TArray<UDataLayer*> DataLayers = UDataLayerSubsystem::ConvertArgsToDataLayers(World, Args);
-					for (UDataLayer* DataLayer : DataLayers)
+					TArray<UDataLayerInstance*> DataLayerInstances = UDataLayerSubsystem::ConvertArgsToDataLayers(World, Args);
+					for (UDataLayerInstance* DataLayerInstance : DataLayerInstances)
 					{
-						DataLayerSubsystem->SetDataLayerRuntimeState(DataLayer, State);
+						DataLayerSubsystem->SetDataLayerRuntimeState(DataLayerInstance, State);
 					}
 				}
 			}
@@ -501,9 +460,9 @@ void UDataLayerSubsystem::GetDataLayerDebugColors(TMap<FName, FColor>& OutMappin
 		return;
 	}
 	
-	WorldDataLayers->ForEachDataLayer([&OutMapping](UDataLayer* DataLayer)
+	WorldDataLayers->ForEachDataLayer([&OutMapping](UDataLayerInstance* DataLayer)
 	{
-		OutMapping.Add(DataLayer->GetFName(), DataLayer->GetDebugColor());
+		OutMapping.Add(DataLayer->GetDataLayerFName(), DataLayer->GetDebugColor());
 		return true;
 	});
 }
@@ -513,13 +472,82 @@ uint32 UDataLayerSubsystem::GetDataLayerEditorContextHash() const
 {
 	if (const AWorldDataLayers* WorldDataLayers = GetWorld()->GetWorldDataLayers())
 	{
-		TArray<FName> DataLayers;
-		for (UDataLayer* DataLayer : WorldDataLayers->GetActorEditorContextDataLayers())
+		TArray<FName> DataLayerInstanceNames;
+		for (UDataLayerInstance* DataLayerInstance : WorldDataLayers->GetActorEditorContextDataLayers())
 		{
-			DataLayers.Add(DataLayer->GetFName());
+			DataLayerInstanceNames.Add(DataLayerInstance->GetDataLayerFName());
 		}
-		return FDataLayerEditorContext(GetWorld(), DataLayers).GetHash();
+		return FDataLayerEditorContext(GetWorld(), DataLayerInstanceNames).GetHash();
 	}
 	return FDataLayerEditorContext::EmptyHash;
 }
-#endif
\ No newline at end of file
+#endif
+
+//~ Begin Deprecated
+
+PRAGMA_DISABLE_DEPRECATION_WARNINGS
+
+UDataLayerInstance* UDataLayerSubsystem::GetDataLayer(const FActorDataLayer& InDataLayer) const
+{
+	return GetDataLayerFromName(InDataLayer.Name);
+}
+
+UDataLayerInstance* UDataLayerSubsystem::GetDataLayerFromLabel(FName InDataLayerLabel) const
+{
+	const AWorldDataLayers* WorldDataLayers = GetWorld()->GetWorldDataLayers();
+	return WorldDataLayers ? const_cast<UDataLayerInstance*>(WorldDataLayers->GetDataLayerFromLabel(InDataLayerLabel)) : nullptr;
+}
+
+EDataLayerRuntimeState UDataLayerSubsystem::GetDataLayerRuntimeStateByLabel(const FName& InDataLayerLabel) const
+{
+	return GetDataLayerRuntimeState(GetDataLayerFromLabel(InDataLayerLabel));
+}
+
+EDataLayerRuntimeState UDataLayerSubsystem::GetDataLayerEffectiveRuntimeStateByLabel(const FName& InDataLayerLabel) const
+{
+	return GetDataLayerEffectiveRuntimeState(GetDataLayerFromLabel(InDataLayerLabel));
+}
+
+void UDataLayerSubsystem::SetDataLayerRuntimeState(const FActorDataLayer& InDataLayer, EDataLayerRuntimeState InState, bool bInIsRecursive)
+{
+	if (UDataLayerInstance* DataLayerInstance = GetDataLayerFromName(InDataLayer.Name))
+	{
+		SetDataLayerRuntimeState(DataLayerInstance, InState, bInIsRecursive);
+	}
+	else
+	{
+		UE_LOG(LogWorldPartition, Warning, TEXT("UDataLayerSubsystem::SetDataLayerRuntimeState unknown Data Layer: '%s'"), *InDataLayer.Name.ToString());
+	}
+}
+
+void UDataLayerSubsystem::SetDataLayerRuntimeStateByLabel(const FName& InDataLayerLabel, EDataLayerRuntimeState InState, bool bInIsRecursive)
+{
+	if (UDataLayerInstance* DataLayerInstance = GetDataLayerFromLabel(InDataLayerLabel))
+	{
+		SetDataLayerRuntimeState(DataLayerInstance, InState, bInIsRecursive);
+	}
+	else
+	{
+		UE_LOG(LogWorldPartition, Warning, TEXT("UDataLayerSubsystem::SetDataLayerRuntimeStateByLabel unknown Data Layer: '%s'"), *InDataLayerLabel.ToString());
+	}
+}
+
+PRAGMA_ENABLE_DEPRECATION_WARNINGS
+
+UDataLayerInstance* UDataLayerSubsystem::GetDataLayerFromName(FName InDataLayerName) const
+{
+	const AWorldDataLayers* WorldDataLayers = GetWorld()->GetWorldDataLayers();
+	return WorldDataLayers ? const_cast<UDataLayerInstance*>(WorldDataLayers->GetDataLayerInstance(InDataLayerName)) : nullptr;
+}
+
+EDataLayerRuntimeState UDataLayerSubsystem::GetDataLayerRuntimeState(const FActorDataLayer& InDataLayer) const
+{
+	return GetDataLayerRuntimeStateByName(InDataLayer.Name);
+}
+
+EDataLayerRuntimeState UDataLayerSubsystem::GetDataLayerEffectiveRuntimeState(const FActorDataLayer& InDataLayer) const
+{
+	return GetDataLayerEffectiveRuntimeStateByName(InDataLayer.Name);
+}
+
+//~ End Deprecated
\ No newline at end of file
diff --git a/Engine/Source/Runtime/Engine/Private/WorldPartition/DataLayer/DataLayersID.cpp b/Engine/Source/Runtime/Engine/Private/WorldPartition/DataLayer/DataLayersID.cpp
--- a/Engine/Source/Runtime/Engine/Private/WorldPartition/DataLayer/DataLayersID.cpp
+++ b/Engine/Source/Runtime/Engine/Private/WorldPartition/DataLayer/DataLayersID.cpp
@@ -2,6 +2,8 @@
 
 #include "WorldPartition/DataLayer/DataLayersID.h"
 
+#include "WorldPartition/DataLayer/DataLayerInstance.h"
+
 #include "Algo/Transform.h"
 
 #if WITH_EDITOR
@@ -9,11 +11,11 @@ FDataLayersID::FDataLayersID()
 	: Hash(0)
 {}
 
-FDataLayersID::FDataLayersID(const TArray<const UDataLayer*>& InDataLayers)
+FDataLayersID::FDataLayersID(const TArray<const UDataLayerInstance*>& InDataLayerInstances)
 	: Hash(0)
 {
 	TArray<FName> DataLayers;
-	Algo::TransformIf(InDataLayers, DataLayers, [](const UDataLayer* Item) { return Item->IsRuntime(); }, [](const UDataLayer* Item) { return Item->GetFName(); });
+	Algo::TransformIf(InDataLayerInstances, DataLayers, [](const UDataLayerInstance* Item) { return Item->IsRuntime(); }, [](const UDataLayerInstance* Item) { return Item->GetDataLayerFName(); });
 
 	if (DataLayers.Num())
 	{
diff --git a/Engine/Source/Runtime/Engine/Private/WorldPartition/DataLayer/DeprecatedDataLayerInstance.cpp b/Engine/Source/Runtime/Engine/Private/WorldPartition/DataLayer/DeprecatedDataLayerInstance.cpp
new file mode 100644
--- /dev/null
+++ b/Engine/Source/Runtime/Engine/Private/WorldPartition/DataLayer/DeprecatedDataLayerInstance.cpp
@@ -0,0 +1,97 @@
+// Copyright Epic Games, Inc. All Rights Reserved.
+
+#include "WorldPartition/DataLayer/DeprecatedDataLayerInstance.h"
+
+#include "WorldPartition/DataLayer/WorldDataLayers.h"
+
+UDeprecatedDataLayerInstance::UDeprecatedDataLayerInstance(const FObjectInitializer& ObjectInitializer)
+	: Super(ObjectInitializer),
+	DataLayerType(EDataLayerType::Editor)
+{
+
+}
+
+PRAGMA_DISABLE_DEPRECATION_WARNINGS
+
+#if WITH_EDITOR
+
+FName UDeprecatedDataLayerInstance::MakeName()
+{
+	return FName(FString::Format(TEXT("DataLayer_{0}"), { FGuid::NewGuid().ToString() }));
+}
+
+void UDeprecatedDataLayerInstance::OnCreated()
+{
+	check(GetOuterAWorldDataLayers()->HasDeprecatedDataLayers());
+
+	Modify(/*bAlwaysMarkDirty*/false);
+
+	RelabelDataLayer(TEXT("DataLayer"));
+
+	DeprecatedDataLayerFName = TEXT("");
+
+	FRandomStream RandomStream(GetDataLayerFName());
+	const uint8 R = (uint8)(RandomStream.GetFraction() * 255.f);
+	const uint8 G = (uint8)(RandomStream.GetFraction() * 255.f);
+	const uint8 B = (uint8)(RandomStream.GetFraction() * 255.f);
+	DebugColor = FColor(R, G, B);
+}
+
+FName UDeprecatedDataLayerInstance::MakeName(const UDEPRECATED_DataLayer* DeprecatedDataLayer)
+{
+	return FName(FString::Format(TEXT("{0}_{1}"), { *DeprecatedDataLayer->GetDataLayerLabel().ToString(), FGuid::NewGuid().ToString() }));
+}
+
+void UDeprecatedDataLayerInstance::OnCreated(const UDEPRECATED_DataLayer* DeprecatedDataLayer)
+{
+	check(GetOuterAWorldDataLayers()->HasDeprecatedDataLayers());
+
+	Modify(/*bAlwaysMarkDirty*/false);
+
+	Label = DeprecatedDataLayer->GetDataLayerLabel();
+
+	DeprecatedDataLayerFName = DeprecatedDataLayer->GetFName();
+
+	DataLayerType = DeprecatedDataLayer->IsRuntime() ? EDataLayerType::Runtime : EDataLayerType::Editor;
+	DebugColor = DeprecatedDataLayer->GetDebugColor();
+
+	bIsVisible = DeprecatedDataLayer->bIsVisible;
+	bIsInitiallyVisible = DeprecatedDataLayer->bIsInitiallyVisible;
+	bIsInitiallyLoadedInEditor = DeprecatedDataLayer->bIsInitiallyLoadedInEditor;
+	bIsLocked = DeprecatedDataLayer->bIsLocked;
+	InitialRuntimeState = DeprecatedDataLayer->InitialRuntimeState;
+}
+
+bool UDeprecatedDataLayerInstance::AddActor(AActor* Actor) const
+{
+	return Actor->AddDataLayer(FActorDataLayer(GetDataLayerFName()));
+}
+
+bool UDeprecatedDataLayerInstance::RemoveActor(AActor* Actor) const
+{
+	return Actor->RemoveDataLayer(FActorDataLayer(GetDataLayerFName()));
+}
+
+bool UDeprecatedDataLayerInstance::ContainsActor(const AActor* Actor) const
+{
+	return Actor->ContainsDataLayer(FActorDataLayer(GetDataLayerFName()));
+}
+
+bool UDeprecatedDataLayerInstance::RelabelDataLayer(FName InDataLayerLabel)
+{
+	FName SanitizedLabel = DataLayerUtils::GetSanitizedDataLayerLabel(InDataLayerLabel);
+	FName UniqueNewDataLayerLabel = GetOuterAWorldDataLayers()->GenerateUniqueDataLayerLabel(SanitizedLabel);
+	if (Label != UniqueNewDataLayerLabel)
+	{
+		Modify();
+		check(!GetOuterAWorldDataLayers()->GetDataLayerFromLabel(UniqueNewDataLayerLabel));
+		Label = UniqueNewDataLayerLabel;
+		return true;
+	}
+
+	return false;
+}
+
+#endif // WITH_EDITOR
+
+PRAGMA_ENABLE_DEPRECATION_WARNINGS
\ No newline at end of file
diff --git a/Engine/Source/Runtime/Engine/Private/WorldPartition/DataLayer/WorldDataLayers.cpp b/Engine/Source/Runtime/Engine/Private/WorldPartition/DataLayer/WorldDataLayers.cpp
--- a/Engine/Source/Runtime/Engine/Private/WorldPartition/DataLayer/WorldDataLayers.cpp
+++ b/Engine/Source/Runtime/Engine/Private/WorldPartition/DataLayer/WorldDataLayers.cpp
@@ -5,8 +5,9 @@
 =============================================================================*/
 
 #include "WorldPartition/DataLayer/WorldDataLayers.h"
-#include "WorldPartition/DataLayer/DataLayer.h"
 #include "WorldPartition/DataLayer/DataLayerSubsystem.h"
+#include "WorldPartition/DataLayer/DataLayerInstanceWithAsset.h"
+#include "WorldPartition/DataLayer/DeprecatedDataLayerInstance.h"
 #include "EngineUtils.h"
 #include "Engine/CoreSettings.h"
 #include "Net/UnrealNetwork.h"
@@ -19,19 +20,19 @@
 
 int32 AWorldDataLayers::DataLayersStateEpoch = 0;
 
-FString JoinDataLayerLabelsFromNames(AWorldDataLayers* InWorldDataLayers, const TArray<FName>& InDataLayerNames)
+FString JoinDataLayerShortNamesFromInstanceNames(AWorldDataLayers* InWorldDataLayers, const TArray<FName>& InDataLayerInstanceNames)
 {
 	check(InWorldDataLayers);
-	TArray<FString> DataLayerLabels;
-	DataLayerLabels.Reserve(InDataLayerNames.Num());
-	for (const FName& DataLayerName : InDataLayerNames)
+	TArray<FString> DataLayerShortNames;
+	DataLayerShortNames.Reserve(InDataLayerInstanceNames.Num());
+	for (const FName& DataLayerInstanceName : InDataLayerInstanceNames)
 	{
-		if (const UDataLayer* DataLayer = InWorldDataLayers->GetDataLayerFromName(DataLayerName))
+		if (const UDataLayerInstance* DataLayerInstance = InWorldDataLayers->GetDataLayerInstance(DataLayerInstanceName))
 		{
-			DataLayerLabels.Add(DataLayer->GetDataLayerLabel().ToString());
+			DataLayerShortNames.Add(DataLayerInstance->GetDataLayerShortName());
 		}
 	}
-	return FString::Join(DataLayerLabels, TEXT(","));
+	return FString::Join(DataLayerShortNames, TEXT(","));
 }
 
 AWorldDataLayers::AWorldDataLayers(const FObjectInitializer& ObjectInitializer)
@@ -84,17 +85,17 @@ void AWorldDataLayers::InitializeDataLayerRuntimeStates()
 
 	if (GetWorld()->IsGameWorld())
 	{
-		ForEachDataLayer([this](class UDataLayer* DataLayer)
+		ForEachDataLayer([this](class UDataLayerInstance* DataLayer)
 		{
 			if (DataLayer && DataLayer->IsRuntime())
 			{
 				if (DataLayer->GetInitialRuntimeState() == EDataLayerRuntimeState::Activated)
 				{
-					ActiveDataLayerNames.Add(DataLayer->GetFName());
+					ActiveDataLayerNames.Add(DataLayer->GetDataLayerFName());
 				}
 				else if (DataLayer->GetInitialRuntimeState() == EDataLayerRuntimeState::Loaded)
 				{
-					LoadedDataLayerNames.Add(DataLayer->GetFName());
+					LoadedDataLayerNames.Add(DataLayer->GetDataLayerFName());
 				}
 			}
 			return true;
@@ -103,7 +104,7 @@ void AWorldDataLayers::InitializeDataLayerRuntimeStates()
 		RepActiveDataLayerNames = ActiveDataLayerNames.Array();
 		RepLoadedDataLayerNames = LoadedDataLayerNames.Array();
 
-		ForEachDataLayer([this](class UDataLayer* DataLayer)
+		ForEachDataLayer([this](class UDataLayerInstance* DataLayer)
 		{
 			if (DataLayer && DataLayer->IsRuntime())
 			{
@@ -115,7 +116,7 @@ void AWorldDataLayers::InitializeDataLayerRuntimeStates()
 		RepEffectiveActiveDataLayerNames = EffectiveActiveDataLayerNames.Array();
 		RepEffectiveLoadedDataLayerNames = EffectiveLoadedDataLayerNames.Array();
 
-		UE_LOG(LogWorldPartition, Log, TEXT("Initial Data Layer Effective States Activated(%s) Loaded(%s)"), *JoinDataLayerLabelsFromNames(this, RepEffectiveActiveDataLayerNames), *JoinDataLayerLabelsFromNames(this, RepEffectiveLoadedDataLayerNames));
+		UE_LOG(LogWorldPartition, Log, TEXT("Initial Data Layer Effective States Activated(%s) Loaded(%s)"), *JoinDataLayerShortNamesFromInstanceNames(this, RepEffectiveActiveDataLayerNames), *JoinDataLayerShortNamesFromInstanceNames(this, RepEffectiveLoadedDataLayerNames));
 	}
 }
 
@@ -132,47 +133,46 @@ void AWorldDataLayers::ResetDataLayerRuntimeStates()
 	RepEffectiveLoadedDataLayerNames.Reset();
 }
 
-void AWorldDataLayers::SetDataLayerRuntimeState(FActorDataLayer InDataLayer, EDataLayerRuntimeState InState, bool bInIsRecursive)
+void AWorldDataLayers::SetDataLayerRuntimeState(const UDataLayerInstance* InDataLayerInstance, EDataLayerRuntimeState InState, bool bInIsRecursive)
 {
 	if (ensure(GetLocalRole() == ROLE_Authority))
 	{
-		const UDataLayer* DataLayer = GetDataLayerFromName(InDataLayer.Name);
-		if (!DataLayer || !DataLayer->IsRuntime())
+		if (!InDataLayerInstance || !InDataLayerInstance->IsRuntime())
 		{
 			return;
 		}
 
-		EDataLayerRuntimeState CurrentState = GetDataLayerRuntimeStateByName(InDataLayer.Name);
+		EDataLayerRuntimeState CurrentState = GetDataLayerRuntimeStateByName(InDataLayerInstance->GetDataLayerFName());
 		if (CurrentState != InState)
 		{
-PRAGMA_DISABLE_DEPRECATION_WARNINGS
-			if (GetWorld()->IsGameWorld())
-			{
-				FName DataLayerLabelName = DataLayer->GetDataLayerLabel();
-				if (DataLayersFilterDelegate.IsBound())
+			PRAGMA_DISABLE_DEPRECATION_WARNINGS
+				if (GetWorld()->IsGameWorld())
 				{
-					if (!DataLayersFilterDelegate.Execute(DataLayerLabelName, CurrentState, InState))
+					FName DataLayerShortName(InDataLayerInstance->GetDataLayerShortName());
+					if (DataLayersFilterDelegate.IsBound())
 					{
-						UE_LOG(LogWorldPartition, Log, TEXT("Data Layer '%s' was filtered out: %s -> %s"),
-							*DataLayerLabelName.ToString(),
-							*StaticEnum<EDataLayerRuntimeState>()->GetDisplayNameTextByValue((int64)CurrentState).ToString(),
-							*StaticEnum<EDataLayerRuntimeState>()->GetDisplayNameTextByValue((int64)InState).ToString());
-						return;
+						if (!DataLayersFilterDelegate.Execute(DataLayerShortName, CurrentState, InState))
+						{
+							UE_LOG(LogWorldPartition, Log, TEXT("Data Layer '%s' was filtered out: %s -> %s"),
+								*DataLayerShortName.ToString(),
+								*StaticEnum<EDataLayerRuntimeState>()->GetDisplayNameTextByValue((int64)CurrentState).ToString(),
+								*StaticEnum<EDataLayerRuntimeState>()->GetDisplayNameTextByValue((int64)InState).ToString());
+							return;
+						}
 					}
 				}
-			}
-PRAGMA_ENABLE_DEPRECATION_WARNINGS
+			PRAGMA_ENABLE_DEPRECATION_WARNINGS
 
-			LoadedDataLayerNames.Remove(InDataLayer.Name);
-			ActiveDataLayerNames.Remove(InDataLayer.Name);
+			LoadedDataLayerNames.Remove(InDataLayerInstance->GetDataLayerFName());
+			ActiveDataLayerNames.Remove(InDataLayerInstance->GetDataLayerFName());
 
 			if (InState == EDataLayerRuntimeState::Loaded)
 			{
-				LoadedDataLayerNames.Add(InDataLayer.Name);
+				LoadedDataLayerNames.Add(InDataLayerInstance->GetDataLayerFName());
 			}
 			else if (InState == EDataLayerRuntimeState::Activated)
 			{
-				ActiveDataLayerNames.Add(InDataLayer.Name);
+				ActiveDataLayerNames.Add(InDataLayerInstance->GetDataLayerFName());
 			}
 			else if (InState == EDataLayerRuntimeState::Unloaded)
 			{
@@ -186,30 +186,29 @@ PRAGMA_ENABLE_DEPRECATION_WARNINGS
 			++DataLayersStateEpoch;
 
 #if !NO_LOGGING || CSV_PROFILER
-			const FString DataLayerLabel = DataLayer->GetDataLayerLabel().ToString();
+			const FString DataLayerShortName = InDataLayerInstance->GetDataLayerShortName();
 			UE_LOG(LogWorldPartition, Log, TEXT("Data Layer '%s' state changed: %s -> %s"),
-				*DataLayerLabel, 
+				*DataLayerShortName,
 				*StaticEnum<EDataLayerRuntimeState>()->GetDisplayNameTextByValue((int64)CurrentState).ToString(),
 				*StaticEnum<EDataLayerRuntimeState>()->GetDisplayNameTextByValue((int64)InState).ToString());
 
-			CSV_EVENT_GLOBAL(TEXT("DataLayer-%s-%s"), *DataLayerLabel, *StaticEnum<EDataLayerRuntimeState>()->GetDisplayNameTextByValue((int64)InState).ToString());
+			CSV_EVENT_GLOBAL(TEXT("DataLayer-%s-%s"), *DataLayerShortName, *StaticEnum<EDataLayerRuntimeState>()->GetDisplayNameTextByValue((int64)InState).ToString());
 #endif
 
-			ResolveEffectiveRuntimeState(DataLayer);
+			ResolveEffectiveRuntimeState(InDataLayerInstance);
 		}
-
 		if (bInIsRecursive)
 		{
-			DataLayer->ForEachChild([this, InState, bInIsRecursive](const UDataLayer* Child)
-			{
-				SetDataLayerRuntimeState(Child->GetFName(), InState, bInIsRecursive);
-				return true;
-			});
+			InDataLayerInstance->ForEachChild([this, InState, bInIsRecursive](const UDataLayerInstance* Child)
+				{
+					SetDataLayerRuntimeState(Child, InState, bInIsRecursive);
+					return true;
+				});
 		}
 	}
 }
 
-void AWorldDataLayers::OnDataLayerRuntimeStateChanged_Implementation(const UDataLayer* InDataLayer, EDataLayerRuntimeState InState)
+void AWorldDataLayers::OnDataLayerRuntimeStateChanged_Implementation(const UDataLayerInstance* InDataLayer, EDataLayerRuntimeState InState)
 {
 	UDataLayerSubsystem* DataLayerSubsystem = GetWorld()->GetSubsystem<UDataLayerSubsystem>();
 	DataLayerSubsystem->OnDataLayerRuntimeStateChanged.Broadcast(InDataLayer, InState);
@@ -271,19 +270,19 @@ EDataLayerRuntimeState AWorldDataLayers::GetDataLayerEffectiveRuntimeStateByName
 	return EDataLayerRuntimeState::Unloaded;
 }
 
-void AWorldDataLayers::ResolveEffectiveRuntimeState(const UDataLayer* InDataLayer, bool bInNotifyChange)
+void AWorldDataLayers::ResolveEffectiveRuntimeState(const UDataLayerInstance* InDataLayer, bool bInNotifyChange)
 {
 	check(InDataLayer);
-	const FName DataLayerName = InDataLayer->GetFName();
+	const FName DataLayerName = InDataLayer->GetDataLayerFName();
 	EDataLayerRuntimeState CurrentEffectiveRuntimeState = GetDataLayerEffectiveRuntimeStateByName(DataLayerName);
 	EDataLayerRuntimeState NewEffectiveRuntimeState = GetDataLayerRuntimeStateByName(DataLayerName);
-	const UDataLayer* Parent = InDataLayer->GetParent();
+	const UDataLayerInstance* Parent = InDataLayer->GetParent();
 	while (Parent && (NewEffectiveRuntimeState != EDataLayerRuntimeState::Unloaded))
 	{
 		if (Parent->IsRuntime())
 		{
 			// Apply min logic with parent DataLayers
-			NewEffectiveRuntimeState = (EDataLayerRuntimeState)FMath::Min((int32)NewEffectiveRuntimeState, (int32)GetDataLayerRuntimeStateByName(Parent->GetFName()));
+			NewEffectiveRuntimeState = (EDataLayerRuntimeState)FMath::Min((int32)NewEffectiveRuntimeState, (int32)GetDataLayerRuntimeStateByName(Parent->GetDataLayerFName()));
 		}
 		Parent = Parent->GetParent();
 	};
@@ -311,23 +310,23 @@ void AWorldDataLayers::ResolveEffectiveRuntimeState(const UDataLayer* InDataLaye
 		if (bInNotifyChange)
 		{
 			UE_LOG(LogWorldPartition, Log, TEXT("Data Layer '%s' effective state changed: %s -> %s"),
-				*InDataLayer->GetDataLayerLabel().ToString(),
+				*InDataLayer->GetDataLayerShortName(),
 				*StaticEnum<EDataLayerRuntimeState>()->GetDisplayNameTextByValue((int64)CurrentEffectiveRuntimeState).ToString(),
 				*StaticEnum<EDataLayerRuntimeState>()->GetDisplayNameTextByValue((int64)NewEffectiveRuntimeState).ToString());
 
 			OnDataLayerRuntimeStateChanged(InDataLayer, NewEffectiveRuntimeState);
 		}
 
-		for (const UDataLayer* Child : InDataLayer->GetChildren())
+		for (const UDataLayerInstance* Child : InDataLayer->GetChildren())
 		{
 			ResolveEffectiveRuntimeState(Child);
 		}
 	}
 }
 
-void AWorldDataLayers::DumpDataLayerRecursively(const UDataLayer* DataLayer, FString Prefix, FOutputDevice& OutputDevice) const
+void AWorldDataLayers::DumpDataLayerRecursively(const UDataLayerInstance* DataLayer, FString Prefix, FOutputDevice& OutputDevice) const
 {
-	auto GetDataLayerRuntimeStateString = [this](const UDataLayer* DataLayer)
+	auto GetDataLayerRuntimeStateString = [this](const UDataLayerInstance* DataLayer)
 	{
 		if (DataLayer->IsRuntime())
 		{
@@ -338,8 +337,8 @@ void AWorldDataLayers::DumpDataLayerRecursively(const UDataLayer* DataLayer, FSt
 			else
 			{
 				return FString::Printf(TEXT("(Effective State = %s | Target State = %s)"),
-					GetDataLayerRuntimeStateName(GetDataLayerEffectiveRuntimeStateByName(DataLayer->GetFName())),
-					GetDataLayerRuntimeStateName(GetDataLayerRuntimeStateByName(DataLayer->GetFName()))
+					GetDataLayerRuntimeStateName(GetDataLayerEffectiveRuntimeStateByName(DataLayer->GetDataLayerFName())),
+					GetDataLayerRuntimeStateName(GetDataLayerRuntimeStateByName(DataLayer->GetDataLayerFName()))
 				);
 			}
 		}
@@ -349,10 +348,10 @@ void AWorldDataLayers::DumpDataLayerRecursively(const UDataLayer* DataLayer, FSt
 	OutputDevice.Logf(TEXT(" %s%s%s %s"),
 		*Prefix,
 		(DataLayer->GetChildren().IsEmpty() && DataLayer->GetParent()) ? TEXT("") : TEXT("[+]"),
-		*DataLayer->GetDataLayerLabel().ToString(),
+		*DataLayer->GetDataLayerShortName(),
 		*GetDataLayerRuntimeStateString(DataLayer));
 
-	for (const UDataLayer* Child : DataLayer->GetChildren())
+	for (const UDataLayerInstance* Child : DataLayer->GetChildren())
 	{
 		DumpDataLayerRecursively(Child, Prefix + TEXT(" | "), OutputDevice);
 	}
@@ -367,16 +366,16 @@ void AWorldDataLayers::DumpDataLayers(FOutputDevice& OutputDevice) const
 
 	if (GetWorld()->IsGameWorld())
 	{
-		auto DumpDataLayersRuntimeState = [this, &OutputDevice](const TCHAR* InStateName, const TSet<FName>& InDataLayers)
+		auto DumpDataLayersRuntimeState = [this, &OutputDevice](const TCHAR* InStateName, const TSet<FName>& InDataLayerInstanceNames)
 		{
-			if (InDataLayers.Num())
+			if (InDataLayerInstanceNames.Num())
 			{
 				OutputDevice.Logf(TEXT(" - %s Data Layers:"), InStateName);
-				for (const FName& DataLayerName : InDataLayers)
+				for (const FName& DataLayerInstanceName : InDataLayerInstanceNames)
 				{
-					if (const UDataLayer* DataLayer = GetDataLayerFromName(DataLayerName))
+					if (const UDataLayerInstance* DataLayerInstance = GetDataLayerInstance(DataLayerInstanceName))
 					{
-						OutputDevice.Logf(TEXT("    - %s"), *DataLayer->GetDataLayerLabel().ToString());
+						OutputDevice.Logf(TEXT("    - %s"), *DataLayerInstance->GetDataLayerShortName());
 					}
 				}
 			}
@@ -395,11 +394,11 @@ void AWorldDataLayers::DumpDataLayers(FOutputDevice& OutputDevice) const
 	
 	OutputDevice.Logf(TEXT("----------------------------------------------------"));
 	OutputDevice.Logf(TEXT(" Data Layers Hierarchy"));
-	ForEachDataLayer([this, &OutputDevice](UDataLayer* DataLayer)
+	ForEachDataLayer([this, &OutputDevice](UDataLayerInstance* DataLayerInstance)
 	{
-		if (DataLayer && !DataLayer->GetParent())
+		if (DataLayerInstance && !DataLayerInstance->GetParent())
 		{
-			DumpDataLayerRecursively(DataLayer, TEXT(""), OutputDevice);
+			DumpDataLayerRecursively(DataLayerInstance, TEXT(""), OutputDevice);
 		}
 		return true;
 	});
@@ -407,57 +406,6 @@ void AWorldDataLayers::DumpDataLayers(FOutputDevice& OutputDevice) const
 }
 
 #if WITH_EDITOR
-void AWorldDataLayers::OverwriteDataLayerRuntimeStates(TArray<FActorDataLayer>* InActiveDataLayers, TArray<FActorDataLayer>* InLoadedDataLayers)
-{
-	if (GetLocalRole() == ROLE_Authority)
-	{
-		// This should get called before game starts. It doesn't send out events
-		check(!GetWorld()->bMatchStarted);
-		if (InActiveDataLayers)
-		{
-			ActiveDataLayerNames.Empty(InActiveDataLayers->Num());
-			for (const FActorDataLayer& ActorDataLayer : *InActiveDataLayers)
-			{
-				const UDataLayer* DataLayer = GetDataLayerFromName(ActorDataLayer);
-				if (DataLayer && DataLayer->IsRuntime())
-				{
-					ActiveDataLayerNames.Add(ActorDataLayer);
-				}
-			}
-			RepActiveDataLayerNames = ActiveDataLayerNames.Array();
-		}
-
-		if (InLoadedDataLayers)
-		{
-			LoadedDataLayerNames.Empty(InLoadedDataLayers->Num());
-			for (const FActorDataLayer& ActorDataLayer : *InLoadedDataLayers)
-			{
-				const UDataLayer* DataLayer = GetDataLayerFromName(ActorDataLayer);
-				if (DataLayer && DataLayer->IsRuntime())
-				{
-					LoadedDataLayerNames.Add(ActorDataLayer);
-				}
-			}
-			RepLoadedDataLayerNames = LoadedDataLayerNames.Array();
-		}
-
-		UE_LOG(LogWorldPartition, Log, TEXT("Overwrite Data Layer States Activated(%s) Loaded(%s)"), *JoinDataLayerLabelsFromNames(this, RepActiveDataLayerNames), *JoinDataLayerLabelsFromNames(this, RepLoadedDataLayerNames));
-
-		ForEachDataLayer([this](class UDataLayer* DataLayer)
-		{
-			if (DataLayer && DataLayer->IsRuntime())
-			{
-				ResolveEffectiveRuntimeState(DataLayer, /*bNotifyChange*/false);
-			}
-			return true;
-		});
-
-		RepEffectiveActiveDataLayerNames = EffectiveActiveDataLayerNames.Array();
-		RepEffectiveLoadedDataLayerNames = EffectiveLoadedDataLayerNames.Array();
-
-		UE_LOG(LogWorldPartition, Log, TEXT("Overwrite Data Layer Effective States Activated(%s) Loaded(%s)"), *JoinDataLayerLabelsFromNames(this, RepEffectiveActiveDataLayerNames), *JoinDataLayerLabelsFromNames(this, RepEffectiveLoadedDataLayerNames));
-	}
-}
 
 void AWorldDataLayers::GetUserLoadedInEditorStates(TArray<FName>& OutDataLayersLoadedInEditor, TArray<FName>& OutDataLayersNotLoadedInEditor) const
 {
@@ -466,30 +414,30 @@ void AWorldDataLayers::GetUserLoadedInEditorStates(TArray<FName>& OutDataLayersL
 
 	OutDataLayersNotLoadedInEditor.Empty();
 	OutDataLayersLoadedInEditor.Empty();
-	ForEachDataLayer([&OutDataLayersNotLoadedInEditor, &OutDataLayersLoadedInEditor, &SettingsDataLayersNotLoadedInEditor, &SettingsDataLayersLoadedInEditor](UDataLayer* DataLayer)
+	ForEachDataLayer([&OutDataLayersNotLoadedInEditor, &OutDataLayersLoadedInEditor, &SettingsDataLayersNotLoadedInEditor, &SettingsDataLayersLoadedInEditor](UDataLayerInstance* DataLayer)
 	{
 		if (DataLayer->IsLoadedInEditorChangedByUserOperation())
 		{
 			if (!DataLayer->IsLoadedInEditor() && DataLayer->IsInitiallyLoadedInEditor())
 			{
-				OutDataLayersNotLoadedInEditor.Add(DataLayer->GetFName());
+				OutDataLayersNotLoadedInEditor.Add(DataLayer->GetDataLayerFName());
 			}
 			else if (DataLayer->IsLoadedInEditor() && !DataLayer->IsInitiallyLoadedInEditor())
 			{
-				OutDataLayersLoadedInEditor.Add(DataLayer->GetFName());
+				OutDataLayersLoadedInEditor.Add(DataLayer->GetDataLayerFName());
 			}
 			
 			DataLayer->ClearLoadedInEditorChangedByUserOperation();
 		}
 		else
 		{
-			if (SettingsDataLayersNotLoadedInEditor.Contains(DataLayer->GetFName()))
+			if (SettingsDataLayersNotLoadedInEditor.Contains(DataLayer->GetDataLayerFName()))
 			{
-				OutDataLayersNotLoadedInEditor.Add(DataLayer->GetFName());
+				OutDataLayersNotLoadedInEditor.Add(DataLayer->GetDataLayerFName());
 			}
-			else if (SettingsDataLayersLoadedInEditor.Contains(DataLayer->GetFName()))
+			else if (SettingsDataLayersLoadedInEditor.Contains(DataLayer->GetDataLayerFName()))
 			{
-				OutDataLayersLoadedInEditor.Add(DataLayer->GetFName());
+				OutDataLayersLoadedInEditor.Add(DataLayer->GetDataLayerFName());
 			}
 		}
 		return true;
@@ -532,39 +480,27 @@ AWorldDataLayers* AWorldDataLayers::Create(UWorld* World)
 	return WorldDataLayers;
 }
 
-FName AWorldDataLayers::GenerateUniqueDataLayerLabel(const FName& InDataLayerLabel) const
-{
-	int32 DataLayerIndex = 0;
-	const FName DataLayerLabelSanitized = UDataLayer::GetSanitizedDataLayerLabel(InDataLayerLabel);
-	FName UniqueNewDataLayerLabel = DataLayerLabelSanitized;
-	while (GetDataLayerFromLabel(UniqueNewDataLayerLabel))
-	{
-		UniqueNewDataLayerLabel = FName(*FString::Printf(TEXT("%s%d"), *DataLayerLabelSanitized.ToString(), ++DataLayerIndex));
-	};
-	return UniqueNewDataLayerLabel;
-}
-
-TArray<FName> AWorldDataLayers::GetDataLayerNames(const TArray<FActorDataLayer>& InDataLayers) const
+TArray<FName> AWorldDataLayers::GetDataLayerInstanceNames(const TArray<TObjectPtr<const UDataLayerAsset>>& InDataLayersAssets) const
 {
 	TArray<FName> OutDataLayerNames;
-	OutDataLayerNames.Reserve(DataLayers.Num());
+	OutDataLayerNames.Reserve(InDataLayersAssets.Num());
 
-	for (const UDataLayer* DataLayer : GetDataLayerObjects(InDataLayers))
+	for (const UDataLayerInstance* DataLayerInstance : GetDataLayerInstances(InDataLayersAssets))
 	{
-		OutDataLayerNames.Add(DataLayer->GetFName());
+		OutDataLayerNames.Add(DataLayerInstance->GetDataLayerFName());
 	}
 
 	return OutDataLayerNames;
 }
 
-TArray<const UDataLayer*> AWorldDataLayers::GetDataLayerObjects(const TArray<FName>& InDataLayerNames) const
+TArray<const UDataLayerInstance*> AWorldDataLayers::GetDataLayerInstances(const TArray<FName>& InDataLayerInstanceNames) const
 {
-	TArray<const UDataLayer*> OutDataLayers;
-	OutDataLayers.Reserve(DataLayers.Num());
+	TArray<const UDataLayerInstance*> OutDataLayers;
+	OutDataLayers.Reserve(InDataLayerInstanceNames.Num());
 
-	for (const FName& DataLayerName : InDataLayerNames)
+	for (const FName& DataLayerInstanceName : InDataLayerInstanceNames)
 	{
-		if (const UDataLayer* DataLayerObject = GetDataLayerFromName(DataLayerName))
+		if (const UDataLayerInstance* DataLayerObject = GetDataLayerInstance(DataLayerInstanceName))
 		{
 			OutDataLayers.AddUnique(DataLayerObject);
 		}
@@ -573,14 +509,14 @@ TArray<const UDataLayer*> AWorldDataLayers::GetDataLayerObjects(const TArray<FNa
 	return OutDataLayers;
 }
 
-TArray<const UDataLayer*> AWorldDataLayers::GetDataLayerObjects(const TArray<FActorDataLayer>& InDataLayers) const
+TArray<const UDataLayerInstance*> AWorldDataLayers::GetDataLayerInstances(const TArray<TObjectPtr<const UDataLayerAsset>>& InDataLayersAssets) const
 {
-	TArray<const UDataLayer*> OutDataLayers;
-	OutDataLayers.Reserve(DataLayers.Num());
+	TArray<const UDataLayerInstance*> OutDataLayers;
+	OutDataLayers.Reserve(InDataLayersAssets.Num());
 
-	for (const FActorDataLayer& DataLayer : InDataLayers)
+	for (const TObjectPtr<const UDataLayerAsset>& DataLayerAsset : InDataLayersAssets)
 	{
-		if (const UDataLayer* DataLayerObject = GetDataLayerFromName(DataLayer.Name))
+		if (const UDataLayerInstance* DataLayerObject = GetDataLayerInstance(DataLayerAsset.Get()))
 		{
 			OutDataLayers.AddUnique(DataLayerObject);
 		}
@@ -589,44 +525,41 @@ TArray<const UDataLayer*> AWorldDataLayers::GetDataLayerObjects(const TArray<FAc
 	return OutDataLayers;
 }
 
-UDataLayer* AWorldDataLayers::CreateDataLayer(FName InName, EObjectFlags InObjectFlags)
-{
-	Modify();
-
-	// Make sure new DataLayer name (not label) is unique and never re-used so that actors still referencing on deleted DataLayer's don't get valid again.
-	const FName DataLayerUniqueName = *FString::Format(TEXT("DataLayer_{0}"), { FGuid::NewGuid().ToString()});
-	UDataLayer* NewDataLayer = NewObject<UDataLayer>(this, DataLayerUniqueName, RF_Transactional | InObjectFlags);
-	check(NewDataLayer != NULL);
-	FName DataLayerLabel = GenerateUniqueDataLayerLabel(InName);
-	NewDataLayer->SetDataLayerLabel(DataLayerLabel);
-	NewDataLayer->SetVisible(true);
-	WorldDataLayers.Add(NewDataLayer);
-	check(GetDataLayerFromName(NewDataLayer->GetFName()));
-	return NewDataLayer;
-}
-
-bool AWorldDataLayers::RemoveDataLayers(const TArray<UDataLayer*>& InDataLayers)
+bool AWorldDataLayers::RemoveDataLayers(const TArray<UDataLayerInstance*>& InDataLayerInstances)
 {
 	bool bIsModified = false;
-	for (UDataLayer* DataLayer : InDataLayers)
+
+	for (UDataLayerInstance* DataLayerInstance : InDataLayerInstances)
 	{
-		if (ContainsDataLayer(DataLayer))
+		if (ContainsDataLayer(DataLayerInstance))
 		{
 			Modify();
-			DataLayer->SetChildParent(DataLayer->GetParent());
-			WorldDataLayers.Remove(const_cast<UDataLayer*>(DataLayer));
+			DataLayerInstance->SetChildParent(DataLayerInstance->GetParent());
+			DataLayerInstances.Remove(DataLayerInstance);
 			bIsModified = true;
 		}
 	}
+
+	if (bIsModified)
+	{
+		UpdateContainsDeprecatedDataLayers();
+	}
+
 	return bIsModified;
 }
 
-bool AWorldDataLayers::RemoveDataLayer(UDataLayer* InDataLayer)
+bool AWorldDataLayers::RemoveDataLayer(const UDataLayerInstance* InDataLayerInstance)
 {
-	if (ContainsDataLayer(InDataLayer))
+	if (ContainsDataLayer(InDataLayerInstance))
 	{
 		Modify();
-		WorldDataLayers.Remove(const_cast<UDataLayer*>(InDataLayer));
+		DataLayerInstances.Remove(InDataLayerInstance);
+
+		if (InDataLayerInstance->IsA<UDeprecatedDataLayerInstance>())
+		{
+			UpdateContainsDeprecatedDataLayers();
+		}
+
 		return true;
 	}
 	return false;
@@ -641,12 +574,12 @@ void AWorldDataLayers::SetAllowRuntimeDataLayerEditing(bool bInAllowRuntimeDataL
 	}
 }
 
-bool AWorldDataLayers::IsInActorEditorContext(const UDataLayer* InDataLayer) const
+bool AWorldDataLayers::IsInActorEditorContext(const UDataLayerInstance* InDataLayerInstance) const
 {
-	for (const FName& DataLayerName : CurrentDataLayers.DataLayers)
+	for (const FName& DataLayerInstanceName : CurrentDataLayers.DataLayerInstanceNames)
 	{
-		const UDataLayer* DataLayer = GetDataLayerFromName(DataLayerName);
-		if (DataLayer && (DataLayer == InDataLayer) && !DataLayer->IsLocked())
+		const UDataLayerInstance* DataLayerInstance = GetDataLayerInstance(DataLayerInstanceName);
+		if (DataLayerInstance && (DataLayerInstance == InDataLayerInstance) && !DataLayerInstance->IsLocked())
 		{
 			return true;
 		}
@@ -654,25 +587,25 @@ bool AWorldDataLayers::IsInActorEditorContext(const UDataLayer* InDataLayer) con
 	return false;
 }
 
-bool AWorldDataLayers::AddToActorEditorContext(UDataLayer* InDataLayer)
+bool AWorldDataLayers::AddToActorEditorContext(UDataLayerInstance* InDataLayerInstance)
 {
-	check(WorldDataLayers.Contains(InDataLayer));
-	if (!CurrentDataLayers.DataLayers.Contains(InDataLayer->GetFName()))
+	check(DataLayerInstances.Contains(InDataLayerInstance));
+	if (!CurrentDataLayers.DataLayerInstanceNames.Contains(InDataLayerInstance->GetDataLayerFName()))
 	{
 		Modify(/*bDirty*/false);
-		CurrentDataLayers.DataLayers.Add(InDataLayer->GetFName());
+		CurrentDataLayers.DataLayerInstanceNames.Add(InDataLayerInstance->GetDataLayerFName());
 		return true;
 	}
 	return false;
 }
 
-bool AWorldDataLayers::RemoveFromActorEditorContext(UDataLayer* InDataLayer)
+bool AWorldDataLayers::RemoveFromActorEditorContext(UDataLayerInstance* InDataLayerInstance)
 {
-	check(WorldDataLayers.Contains(InDataLayer));
-	if (CurrentDataLayers.DataLayers.Contains(InDataLayer->GetFName()))
+	check(DataLayerInstances.Contains(InDataLayerInstance));
+	if (CurrentDataLayers.DataLayerInstanceNames.Contains(InDataLayerInstance->GetDataLayerFName()))
 	{
 		Modify(/*bDirty*/false);
-		CurrentDataLayers.DataLayers.Remove(InDataLayer->GetFName());
+		CurrentDataLayers.DataLayerInstanceNames.Remove(InDataLayerInstance->GetDataLayerFName());
 		return true;
 	}
 	return false;
@@ -692,15 +625,15 @@ void AWorldDataLayers::PopActorEditorContext()
 	CurrentDataLayers = CurrentDataLayersStack.Pop();
 }
 
-TArray<UDataLayer*> AWorldDataLayers::GetActorEditorContextDataLayers() const
+TArray<UDataLayerInstance*> AWorldDataLayers::GetActorEditorContextDataLayers() const
 {
-	TArray<UDataLayer*> Result;
-	for (const FName& DataLayerName : CurrentDataLayers.DataLayers)
+	TArray<UDataLayerInstance*> Result;
+	for (const FName& DataLayerInstanceName : CurrentDataLayers.DataLayerInstanceNames)
 	{
-		const UDataLayer* DataLayer = GetDataLayerFromName(DataLayerName);
-		if (DataLayer && !DataLayer->IsLocked())
+		const UDataLayerInstance* DataLayerInstance = GetDataLayerInstance(DataLayerInstanceName);
+		if (DataLayerInstance && !DataLayerInstance->IsLocked())
 		{
-			Result.Add(const_cast<UDataLayer*>(DataLayer));
+			Result.Add(const_cast<UDataLayerInstance*>(DataLayerInstance));
 		}
 	};
 	return Result;
@@ -708,53 +641,84 @@ TArray<UDataLayer*> AWorldDataLayers::GetActorEditorContextDataLayers() const
 
 #endif
 
-bool AWorldDataLayers::ContainsDataLayer(const UDataLayer* InDataLayer) const
+bool AWorldDataLayers::ContainsDataLayer(const UDataLayerInstance* InDataLayerInstance) const 
 {
-	return WorldDataLayers.Contains(InDataLayer);
+	return DataLayerInstances.Contains(InDataLayerInstance);
 }
 
-const UDataLayer* AWorldDataLayers::GetDataLayerFromName(const FName& InDataLayerName) const
+const UDataLayerInstance* AWorldDataLayers::GetDataLayerInstance(const FName& InDataLayerInstanceName) const
 {
 #if WITH_EDITOR	
-	for (UDataLayer* DataLayer : WorldDataLayers)
+	for (UDataLayerInstance* DataLayerInstance : DataLayerInstances)
 	{
-		if (DataLayer->GetFName() == InDataLayerName)
+		if (DataLayerInstance->GetDataLayerFName() == InDataLayerInstanceName)
 		{
-			return DataLayer;
+			return DataLayerInstance;
 		}
 	}
 #else
-	if (const UDataLayer* const* FoundDataLayer = NameToDataLayer.Find(InDataLayerName))
+	if (const UDataLayerInstance* const* FoundDataLayerInstance = InstanceNameToInstance.Find(InDataLayerInstanceName))
 	{
-		return *FoundDataLayer;
+		return *FoundDataLayerInstance;
 	}
 #endif
+
+#if DATALAYER_TO_INSTANCE_RUNTIME_CONVERSION_ENABLED
+	if (TWeakObjectPtr<UDataLayerInstance> const* FoundDataLayerInstance = DeprecatedDataLayerNameToDataLayerInstance.Find(InDataLayerInstanceName))
+	{
+		return FoundDataLayerInstance->Get();
+	}
+#endif // DATALAYER_TO_INSTANCE_RUNTIME_CONVERSION_ENABLED
+
 	return nullptr;
 }
 
-const UDataLayer* AWorldDataLayers::GetDataLayerFromLabel(const FName& InDataLayerLabel) const
+const UDataLayerInstance* AWorldDataLayers::GetDataLayerInstanceFromAssetName(const FName& InDataLayerAssetFullName) const
 {
-	const FName DataLayerLabelSanitized = UDataLayer::GetSanitizedDataLayerLabel(InDataLayerLabel);
 #if WITH_EDITOR	
-	for (const UDataLayer* DataLayer : WorldDataLayers)
+	for (UDataLayerInstance* DataLayerInstance : DataLayerInstances)
 	{
-		if (DataLayer->GetDataLayerLabel() == DataLayerLabelSanitized)
+		if (DataLayerInstance->GetDataLayerFullName().Compare(InDataLayerAssetFullName.ToString(), ESearchCase::IgnoreCase))
 		{
-			return DataLayer;
+			return DataLayerInstance;
+		}
+	}
+#else
+	if (const UDataLayerInstance* const* FoundDataLayerInstance = AssetNameToInstance.Find(InDataLayerAssetFullName.ToString()))
+	{
+		return *FoundDataLayerInstance;
+	}
+#endif
+	return nullptr;
+}
+
+const UDataLayerInstance* AWorldDataLayers::GetDataLayerInstance(const UDataLayerAsset* InDataLayerAsset) const
+{
+#if WITH_EDITOR	
+	for (const UDataLayerInstance* DataLayerInstance : DataLayerInstances)
+	{
+		static_assert(DATALAYER_TO_INSTANCE_RUNTIME_CONVERSION_ENABLED, "Remove unnecessary cast. All DataLayerInstance now have assets");
+		if (const UDataLayerInstanceWithAsset* DataLayerInstanceWithAsset = Cast<UDataLayerInstanceWithAsset>(DataLayerInstance))
+		{
+			if (DataLayerInstanceWithAsset->GetAsset() == InDataLayerAsset)
+			{
+				return DataLayerInstance;
+			}
 		}
 	}
 #else
-	if (const UDataLayer* const* FoundDataLayer = LabelToDataLayer.Find(DataLayerLabelSanitized))
+	if (const UDataLayerInstance* const* FoundDataLayerInstance = AssetNameToInstance.Find(InDataLayerAsset->GetFullName()))
 	{
-		return *FoundDataLayer;
+		return *FoundDataLayerInstance;
 	}
 #endif
+
 	return nullptr;
 }
 
-void AWorldDataLayers::ForEachDataLayer(TFunctionRef<bool(UDataLayer*)> Func)
+void AWorldDataLayers::ForEachDataLayer(TFunctionRef<bool(UDataLayerInstance*)> Func)
 {
-	for (UDataLayer* DataLayer : WorldDataLayers)
+	for (UDataLayerInstance* DataLayer : DataLayerInstances)
 	{
 		if (!Func(DataLayer))
 		{
@@ -763,9 +727,9 @@ void AWorldDataLayers::ForEachDataLayer(TFunctionRef<bool(UDataLayer*)> Func)
 	}
 }
 
-void AWorldDataLayers::ForEachDataLayer(TFunctionRef<bool(UDataLayer*)> Func) const
+void AWorldDataLayers::ForEachDataLayer(TFunctionRef<bool(UDataLayerInstance*)> Func) const
 {
-	for (UDataLayer* DataLayer : WorldDataLayers)
+	for (UDataLayerInstance* DataLayer : DataLayerInstances)
 	{
 		if (!Func(DataLayer))
 		{
@@ -783,17 +747,19 @@ void AWorldDataLayers::PostLoad()
 	GetTypedOuter<UWorld>()->SetWorldDataLayers(this);
 
 #if WITH_EDITOR
+	ConvertDataLayerToInstancces();
+
 	// Remove all Editor Data Layers when cooking or when in a game world
 	if (IsRunningCookCommandlet() || GetWorld()->IsGameWorld())
 	{
-		ForEachDataLayer([](UDataLayer* DataLayer)
+		ForEachDataLayer([](UDataLayerInstance* DataLayer)
 		{
 			DataLayer->ConditionalPostLoad();
 			return true;
 		});
 
-		TArray<UDataLayer*> EditorDataLayers;
-		ForEachDataLayer([&EditorDataLayers](UDataLayer* DataLayer)
+		TArray<UDataLayerInstance*> EditorDataLayers;
+		ForEachDataLayer([&EditorDataLayers](UDataLayerInstance* DataLayer)
 		{
 			if (DataLayer && !DataLayer->IsRuntime())
 			{
@@ -805,41 +771,209 @@ void AWorldDataLayers::PostLoad()
 	}
 
 	// Setup defaults before overriding with user settings
-	for (UDataLayer* DataLayer : WorldDataLayers)
+	for (UDataLayerInstance* DataLayer : DataLayerInstances)
 	{
 		DataLayer->SetIsLoadedInEditor(DataLayer->IsInitiallyLoadedInEditor(), /*bFromUserChange*/false);
 	}
 
 	// Initialize DataLayer's IsLoadedInEditor based on DataLayerEditorPerProjectUserSettings
 	TArray<FName> SettingsDataLayersNotLoadedInEditor = GetMutableDefault<UWorldPartitionEditorPerProjectUserSettings>()->GetWorldDataLayersNotLoadedInEditor(GetWorld());
-	for (const FName& DataLayerName : SettingsDataLayersNotLoadedInEditor)
+	for (const FName& DataLayerAssetName : SettingsDataLayersNotLoadedInEditor)
 	{
-		if (UDataLayer* DataLayer = const_cast<UDataLayer*>(GetDataLayerFromName(DataLayerName)))
+		if (UDataLayerInstance* DataLayerInstance = const_cast<UDataLayerInstance*>(GetDataLayerInstanceFromAssetName(DataLayerAssetName)))
 		{
-			DataLayer->SetIsLoadedInEditor(false, /*bFromUserChange*/false);
+			DataLayerInstance->SetIsLoadedInEditor(false, /*bFromUserChange*/false);
 		}
 	}
 
 	TArray<FName> SettingsDataLayersLoadedInEditor = GetMutableDefault<UWorldPartitionEditorPerProjectUserSettings>()->GetWorldDataLayersLoadedInEditor(GetWorld());
-	for (const FName& DataLayerName : SettingsDataLayersLoadedInEditor)
+	for (const FName& DataLayerAssetName : SettingsDataLayersLoadedInEditor)
 	{
-		if (UDataLayer* DataLayer = const_cast<UDataLayer*>(GetDataLayerFromName(DataLayerName)))
+		if (UDataLayerInstance* DataLayerInstance = const_cast<UDataLayerInstance*>(GetDataLayerInstanceFromAssetName(DataLayerAssetName)))
 		{
-			DataLayer->SetIsLoadedInEditor(true, /*bFromUserChange*/false);
+			DataLayerInstance->SetIsLoadedInEditor(true, /*bFromUserChange*/false);
 		}
 	}
 
 	bListedInSceneOutliner = true;
 #else
 	// Build acceleration tables
-	for (const UDataLayer* DataLayer : WorldDataLayers)
+	for (const UDataLayerInstance* DataLayerInstance : DataLayerInstances)
 	{
-		LabelToDataLayer.Add(DataLayer->GetDataLayerLabel(), DataLayer);
-		NameToDataLayer.Add(DataLayer->GetFName(), DataLayer);
+		InstanceNameToInstance.Add(DataLayerInstance->GetDataLayerFName(), DataLayerInstance);
+
+		static_assert(DATALAYER_TO_INSTANCE_RUNTIME_CONVERSION_ENABLED, "Remove unnecessary cast. All DataLayerInstance now have assets");
+		if (const UDataLayerInstanceWithAsset* DataLayerInstanceWithAsset = Cast<UDataLayerInstanceWithAsset>(DataLayerInstance))
+		{
+			AssetNameToInstance.Add(DataLayerInstanceWithAsset->GetAsset()->GetFullName(), DataLayerInstance);
+		}
 	}
 #endif
 
 	InitializeDataLayerRuntimeStates();
 }
 
+#if WITH_EDITOR
+
+void AWorldDataLayers::ConvertDataLayerToInstancces()
+{
+	static_assert(DATALAYER_TO_INSTANCE_RUNTIME_CONVERSION_ENABLED, "AWorldDataLayers::ConvertDataLayerToInstancces function is deprecated and needs to be deleted.");
+	
+	PRAGMA_DISABLE_DEPRECATION_WARNINGS
+	bHasDeprecatedDataLayers = !WorldDataLayers_DEPRECATED.IsEmpty();
+
+	for (UDEPRECATED_DataLayer* DeprecatedDataLayer : WorldDataLayers_DEPRECATED)
+	{
+		UDeprecatedDataLayerInstance* DataLayerInstance = CreateDataLayer<UDeprecatedDataLayerInstance>(DeprecatedDataLayer);
+		DeprecatedDataLayerNameToDataLayerInstance.Add(DeprecatedDataLayer->GetFName(), DataLayerInstance);
+	}
+
+	for (UDEPRECATED_DataLayer* DeprecatedDataLayer : WorldDataLayers_DEPRECATED)
+	{
+		if (DeprecatedDataLayer->GetParent() != nullptr)
+		{
+			UDataLayerInstance* ParentInstance = const_cast<UDataLayerInstance*>(GetDataLayerInstance(DeprecatedDataLayer->GetParent()->GetFName()));
+			UDataLayerInstance* ChildInstance = const_cast<UDataLayerInstance*>(GetDataLayerInstance(DeprecatedDataLayer->GetFName()));
+			if (!ChildInstance->SetParent(ParentInstance))
+			{
+				UE_LOG(LogWorldPartition, Error, TEXT("Failed to Convert DataLayer %s' hieararchy to DataLayerInstances. Run DataLayerToAsset Commandlet or fix the hierarchy manually."), *DeprecatedDataLayer->GetDataLayerLabel().ToString());
+			}
+		}
+	}
+
+	WorldDataLayers_DEPRECATED.Empty();
+
+	UpdateContainsDeprecatedDataLayers();
+
+	PRAGMA_ENABLE_DEPRECATION_WARNINGS
+}
+
+void AWorldDataLayers::UpdateContainsDeprecatedDataLayers()
+{
+	static_assert(DATALAYER_TO_INSTANCE_RUNTIME_CONVERSION_ENABLED, "AWorldDataLayers::UpdateContainsDeprecatedDataLayers function is deprecated and needs to be deleted.");
+
+	bHasDeprecatedDataLayers = !WorldDataLayers_DEPRECATED.IsEmpty();
+
+	if (!bHasDeprecatedDataLayers)
+	{
+		for (const UDataLayerInstance* DataLayerInstance : DataLayerInstances)
+		{
+			if (DataLayerInstance->IsA<UDeprecatedDataLayerInstance>())
+			{
+				bHasDeprecatedDataLayers = true;
+				break;
+			}
+		}
+	}
+}
+
+#endif
+
+//~ Begin Deprecated
+
+PRAGMA_DISABLE_DEPRECATION_WARNINGS
+
+#if WITH_EDITOR
+
+bool AWorldDataLayers::RemoveDataLayer(const UDEPRECATED_DataLayer* InDataLayer)
+{
+	if (ContainsDataLayer(InDataLayer))
+	{
+		Modify();
+		WorldDataLayers_DEPRECATED.Remove(InDataLayer);
+
+		UpdateContainsDeprecatedDataLayers();
+
+		return true;
+	}
+	return false;
+}
+
+FName AWorldDataLayers::GenerateUniqueDataLayerLabel(const FName& InDataLayerLabel) const
+{
+	int32 DataLayerIndex = 0;
+	const FName DataLayerLabelSanitized = DataLayerUtils::GetSanitizedDataLayerLabel(InDataLayerLabel);
+	FName UniqueNewDataLayerLabel = DataLayerLabelSanitized;
+	while (GetDataLayerFromLabel(UniqueNewDataLayerLabel))
+	{
+		UniqueNewDataLayerLabel = FName(*FString::Printf(TEXT("%s%d"), *DataLayerLabelSanitized.ToString(), ++DataLayerIndex));
+	};
+	return UniqueNewDataLayerLabel;
+}
+
+#endif
+
+bool AWorldDataLayers::ContainsDataLayer(const UDEPRECATED_DataLayer* InDataLayer) const
+{
+	return WorldDataLayers_DEPRECATED.Contains(InDataLayer);
+}
+
+const UDataLayerInstance* AWorldDataLayers::GetDataLayerFromLabel(const FName& InDataLayerLabel) const
+{
+	const FName DataLayerLabelSanitized = DataLayerUtils::GetSanitizedDataLayerLabel(InDataLayerLabel);
+
+	for (const UDataLayerInstance* DataLayer : DataLayerInstances)
+	{
+		if (FName(DataLayer->GetDataLayerShortName()) == DataLayerLabelSanitized)
+		{
+			return DataLayer;
+		}
+	}
+
+	return nullptr;
+}
+
+PRAGMA_ENABLE_DEPRECATION_WARNINGS
+
+const UDataLayerInstance* AWorldDataLayers::GetDataLayerInstance(const FActorDataLayer& InActorDataLayer) const
+{
+	return GetDataLayerInstance(InActorDataLayer.Name);
+}
+
+TArray<FName> AWorldDataLayers::GetDataLayerInstanceNames(const TArray<FActorDataLayer>& InActorDataLayers) const
+{
+#if DATALAYER_TO_INSTANCE_RUNTIME_CONVERSION_ENABLED
+	TArray<FName> OutDataLayerNames;
+	OutDataLayerNames.Reserve(InActorDataLayers.Num());
+
+	for (const FActorDataLayer& ActorDataLayer : InActorDataLayers)
+	{
+		const UDataLayerInstance* DataLayerInstance = GetDataLayerInstance(ActorDataLayer.Name);
+		if (DataLayerInstance != nullptr)
+		{
+			OutDataLayerNames.Add(DataLayerInstance->GetDataLayerFName());
+		}
+	}
+
+	return OutDataLayerNames;
+#else 
+	static_assert(0, "AWorldDataLayers::GetDataLayerInstanceNames function is deprecated and needs to be deleted.");
+	return TArray<FName>();
+#endif
+}
+
+TArray<const UDataLayerInstance*> AWorldDataLayers::GetDataLayerInstances(const TArray<FActorDataLayer>& InActorDataLayers) const
+{
+#if DATALAYER_TO_INSTANCE_RUNTIME_CONVERSION_ENABLED
+	TArray<const UDataLayerInstance*> OutDataLayerInstances;
+	OutDataLayerInstances.Reserve(InActorDataLayers.Num());
+
+	for (const FActorDataLayer& ActorDataLayer : InActorDataLayers)
+	{
+		const UDataLayerInstance* DataLayerInstance = GetDataLayerInstance(ActorDataLayer.Name);
+		if (DataLayerInstance != nullptr)
+		{
+			OutDataLayerInstances.Add(DataLayerInstance);
+		}
+	}
+
+	return OutDataLayerInstances;
+#else 
+	static_assert(0, "AWorldDataLayers::GetDataLayerInstances function is deprecated and needs to be deleted.");
+	return TArray<const UDataLayerInstance*>();
+#endif
+}
+
+//~ End Deprecated
+
 #undef LOCTEXT_NAMESPACE
\ No newline at end of file
diff --git a/Engine/Source/Runtime/Engine/Private/WorldPartition/ErrorHandling/WorldPartitionStreamingGenerationLogErrorHandler.cpp b/Engine/Source/Runtime/Engine/Private/WorldPartition/ErrorHandling/WorldPartitionStreamingGenerationLogErrorHandler.cpp
--- a/Engine/Source/Runtime/Engine/Private/WorldPartition/ErrorHandling/WorldPartitionStreamingGenerationLogErrorHandler.cpp
+++ b/Engine/Source/Runtime/Engine/Private/WorldPartition/ErrorHandling/WorldPartitionStreamingGenerationLogErrorHandler.cpp
@@ -3,6 +3,7 @@
 #if WITH_EDITOR
 #include "WorldPartition/ErrorHandling/WorldPartitionStreamingGenerationLogErrorHandler.h"
 #include "WorldPartition/WorldPartitionLog.h"
+#include "WorldPartition/DataLayer/DataLayerInstanceWithAsset.h"
 #include "GameFramework/Actor.h"
 
 void FStreamingGenerationLogErrorHandler::OnInvalidReference(const FWorldPartitionActorDescView& ActorDescView, const FGuid& ReferenceGuid)
@@ -37,4 +38,19 @@ void FStreamingGenerationLogErrorHandler::OnInvalidReferenceLevelScriptDataLayer
 {
 	UE_LOG(LogWorldPartition, Log, TEXT("Level Script Blueprint references streamed actor %s with a non empty set of data layers"), *ActorDescView.GetActorLabelOrName().ToString());
 }
+
+void FStreamingGenerationLogErrorHandler::OnInvalidReferenceDataLayerAsset(const UDataLayerInstanceWithAsset* DataLayerInstance)
+{
+	UE_LOG(LogWorldPartition, Log, TEXT("Data Layer %s does not have a Data Layer asset"), *DataLayerInstance->GetDataLayerFName().ToString());
+}
+
+void FStreamingGenerationLogErrorHandler::OnDataLayerHierarchyTypeMismatch(const UDataLayerInstance* DataLayerInstance, const UDataLayerInstance* Parent)
+{
+	UE_LOG(LogWorldPartition, Log, TEXT("Data Layer %s is of Type %s and its parent %s is of type %s"), *DataLayerInstance->GetDataLayerFullName(), *UEnum::GetValueAsString(DataLayerInstance->GetType()), *Parent->GetDataLayerFullName(), *UEnum::GetValueAsString(Parent->GetType()));
+}
+
+void FStreamingGenerationLogErrorHandler::OnDataLayerAssetConflict(const UDataLayerInstanceWithAsset* DataLayerInstance, const UDataLayerInstanceWithAsset* ConflictingDataLayerInstance)
+{
+	UE_LOG(LogWorldPartition, Log, TEXT("Data Layer Instance %s and Data Layer Instance %s are both referencing Data Layer Asset %s"), *DataLayerInstance->GetDataLayerFName().ToString(), *ConflictingDataLayerInstance->GetDataLayerFName().ToString(), *DataLayerInstance->GetAsset()->GetFullName());
+}
 #endif
diff --git a/Engine/Source/Runtime/Engine/Private/WorldPartition/ErrorHandling/WorldPartitionStreamingGenerationMapCheckErrorHandler.cpp b/Engine/Source/Runtime/Engine/Private/WorldPartition/ErrorHandling/WorldPartitionStreamingGenerationMapCheckErrorHandler.cpp
--- a/Engine/Source/Runtime/Engine/Private/WorldPartition/ErrorHandling/WorldPartitionStreamingGenerationMapCheckErrorHandler.cpp
+++ b/Engine/Source/Runtime/Engine/Private/WorldPartition/ErrorHandling/WorldPartitionStreamingGenerationMapCheckErrorHandler.cpp
@@ -7,6 +7,7 @@
 #include "Misc/UObjectToken.h"
 #include "Misc/MapErrors.h"
 #include "GameFramework/Actor.h"
+#include "WorldPartition/DataLayer/DataLayerInstanceWithAsset.h"
 
 #define LOCTEXT_NAMESPACE "WorldPartition"
 
diff --git a/Engine/Source/Runtime/Engine/Private/WorldPartition/ErrorHandling/WorldPartitionStreamingGenerationTokenizedMessageErrorHandler.cpp b/Engine/Source/Runtime/Engine/Private/WorldPartition/ErrorHandling/WorldPartitionStreamingGenerationTokenizedMessageErrorHandler.cpp
--- a/Engine/Source/Runtime/Engine/Private/WorldPartition/ErrorHandling/WorldPartitionStreamingGenerationTokenizedMessageErrorHandler.cpp
+++ b/Engine/Source/Runtime/Engine/Private/WorldPartition/ErrorHandling/WorldPartitionStreamingGenerationTokenizedMessageErrorHandler.cpp
@@ -6,6 +6,7 @@
 #include "GameFramework/Actor.h"
 #include "Logging/MessageLog.h"
 #include "Misc/UObjectToken.h"
+#include "WorldPartition/DataLayer/DataLayerInstanceWithAsset.h"
 
 #define LOCTEXT_NAMESPACE "WorldPartition"
 
@@ -89,6 +90,47 @@ void ITokenizedMessageErrorHandler::OnInvalidReferenceLevelScriptDataLayers(cons
 	HandleTokenizedMessage(MoveTemp(Message));
 }
 
+void ITokenizedMessageErrorHandler::OnInvalidReferenceDataLayerAsset(const UDataLayerInstanceWithAsset* DataLayerInstance)
+{
+	TSharedRef<FTokenizedMessage> Message = FTokenizedMessage::Create(EMessageSeverity::Error);
+	Message->AddToken(FTextToken::Create(LOCTEXT("TokenMessage_DataLayers_DataLayer", "Data layer")))
+		->AddToken(FUObjectToken::Create(DataLayerInstance))
+		->AddToken(FTextToken::Create(LOCTEXT("TokenMessage_DataLayers_NullAsset", "Does not have Data Layer Asset")));
+
+	AddAdditionalNameToken(Message, FName(TEXT("DataLayers_InvalidAsset_CheckForErrors")));
+
+	HandleTokenizedMessage(MoveTemp(Message));
+}
+
+void ITokenizedMessageErrorHandler::OnDataLayerHierarchyTypeMismatch(const UDataLayerInstance* DataLayerInstance, const UDataLayerInstance* Parent)
+{
+	TSharedRef<FTokenizedMessage> Message = FTokenizedMessage::Create(EMessageSeverity::Error);
+	Message->AddToken(FTextToken::Create(LOCTEXT("TokenMessage_DataLayers_DataLayer", "Data layer")))
+		->AddToken(FTextToken::Create(FText::FromString(DataLayerInstance->GetDataLayerShortName())))
+		->AddToken(FTextToken::Create(LOCTEXT("TokenMessage_DataLayers_IsRntime", "is Runtime but its parent data layer")))
+		->AddToken(FTextToken::Create(FText::FromString(Parent->GetDataLayerShortName())))
+		->AddToken(FTextToken::Create(LOCTEXT("TokenMessage_DataLayers_IsNot", "is not")));
+
+	AddAdditionalNameToken(Message, FName(TEXT("DataLayers_HierarchyTypeMisMatch_CheckForErrors")));
+
+	HandleTokenizedMessage(MoveTemp(Message));
+}
+
+void ITokenizedMessageErrorHandler::OnDataLayerAssetConflict(const UDataLayerInstanceWithAsset* DataLayerInstance, const UDataLayerInstanceWithAsset* ConflictingDataLayerInstance)
+{
+	TSharedRef<FTokenizedMessage> Message = FTokenizedMessage::Create(EMessageSeverity::Error);
+	Message->AddToken(FTextToken::Create(LOCTEXT("TokenMessage_DataLayers_DataLayerInstance", "Data layer Instance")))
+		->AddToken(FUObjectToken::Create(DataLayerInstance))
+		->AddToken(FTextToken::Create(LOCTEXT("TokenMessage_DataLayers_AndDataLayerInstance", "and Data Layer Instance")))
+		->AddToken(FUObjectToken::Create(ConflictingDataLayerInstance))
+		->AddToken(FTextToken::Create(LOCTEXT("TokenMessage_DataLayers_BothReferencing", "are both referencing Data Layer Asset")))
+		->AddToken(FUObjectToken::Create(DataLayerInstance->GetAsset()));
+
+	AddAdditionalNameToken(Message, FName(TEXT("DataLayers_AssetConflict_CheckForErrors")));
+
+	HandleTokenizedMessage(MoveTemp(Message));
+}
+
 #undef LOCTEXT_NAMESPACE
 
 #endif
diff --git a/Engine/Source/Runtime/Engine/Private/WorldPartition/HLOD/HLODActorDesc.cpp b/Engine/Source/Runtime/Engine/Private/WorldPartition/HLOD/HLODActorDesc.cpp
--- a/Engine/Source/Runtime/Engine/Private/WorldPartition/HLOD/HLODActorDesc.cpp
+++ b/Engine/Source/Runtime/Engine/Private/WorldPartition/HLOD/HLODActorDesc.cpp
@@ -28,7 +28,7 @@ void FHLODActorDesc::Init(const AActor* InActor)
 		uint64 GridIndexZ;
 		HLODActor->GetGridIndices(GridIndexX, GridIndexY, GridIndexZ);
 
-		FDataLayersID DataLayersID(HLODActor->GetDataLayerObjects());
+		FDataLayersID DataLayersID(HLODActor->GetDataLayerInstances());
 
 		CellHash = ComputeCellHash(SubActorsHLODLayer->GetName(), GridIndexX, GridIndexY, GridIndexZ, DataLayersID);
 	}
diff --git a/Engine/Source/Runtime/Engine/Private/WorldPartition/RuntimeSpatialHash/RuntimeSpatialHashGridHelper.cpp b/Engine/Source/Runtime/Engine/Private/WorldPartition/RuntimeSpatialHash/RuntimeSpatialHashGridHelper.cpp
--- a/Engine/Source/Runtime/Engine/Private/WorldPartition/RuntimeSpatialHash/RuntimeSpatialHashGridHelper.cpp
+++ b/Engine/Source/Runtime/Engine/Private/WorldPartition/RuntimeSpatialHash/RuntimeSpatialHashGridHelper.cpp
@@ -3,6 +3,7 @@
 #include "WorldPartition/RuntimeSpatialHash/RuntimeSpatialHashGridHelper.h"
 #include "Algo/Transform.h"
 #include "ProfilingDebugging/ScopedTimers.h"
+#include "WorldPartition/DataLayer/DataLayerInstance.h"
 
 bool GRuntimeSpatialHashUseAlignedGridLevels = true;
 static FAutoConsoleVariableRef CVarRuntimeSpatialHashUseAlignedGridLevels(
@@ -263,7 +264,7 @@ FSquare2DGridHelper GetPartitionedActors(const UWorldPartition* WorldPartition,
 				if (ClusterInstance->DataLayers.Num())
 				{
 					TArray<FString> DataLayers;
-					Algo::Transform(ClusterInstance->DataLayers, DataLayers, [](const UDataLayer* DataLayer) { return DataLayer->GetDataLayerLabel().ToString(); });
+					Algo::Transform(ClusterInstance->DataLayers, DataLayers, [](const UDataLayerInstance* DataLayer) { return DataLayer->GetDataLayerFullName(); });
 					UE_LOG(LogWorldPartition, Verbose, TEXT("   - DataLayers: %s"), *FString::Join(DataLayers, TEXT(", ")));
 				}
 
diff --git a/Engine/Source/Runtime/Engine/Private/WorldPartition/RuntimeSpatialHash/RuntimeSpatialHashGridHelper.h b/Engine/Source/Runtime/Engine/Private/WorldPartition/RuntimeSpatialHash/RuntimeSpatialHashGridHelper.h
--- a/Engine/Source/Runtime/Engine/Private/WorldPartition/RuntimeSpatialHash/RuntimeSpatialHashGridHelper.h
+++ b/Engine/Source/Runtime/Engine/Private/WorldPartition/RuntimeSpatialHash/RuntimeSpatialHashGridHelper.h
@@ -9,6 +9,7 @@
 #include "WorldPartition/WorldPartitionRuntimeSpatialHash.h"
 #include "WorldPartition/WorldPartitionActorCluster.h"
 #include "WorldPartition/DataLayer/DataLayersID.h"
+#include "WorldPartition/DataLayer/DataLayerInstance.h"
 
 extern ENGINE_API bool GRuntimeSpatialHashUseAlignedGridLevels;
 extern ENGINE_API bool GRuntimeSpatialHashSnapNonAlignedGridLevelsToLowerLevels;
@@ -261,21 +262,21 @@ struct FSquare2DGridHelper
 #if WITH_EDITOR
 		struct FGridCellDataChunk
 		{
-			FGridCellDataChunk(const TArray<const UDataLayer*>& InDataLayers)
+			FGridCellDataChunk(const TArray<const UDataLayerInstance*>& InDataLayers)
 			{
-				Algo::TransformIf(InDataLayers, DataLayers, [](const UDataLayer* DataLayer) { return DataLayer->IsRuntime(); }, [](const UDataLayer* DataLayer) { return DataLayer; });
+				Algo::TransformIf(InDataLayers, DataLayers, [](const UDataLayerInstance* DataLayer) { return DataLayer->IsRuntime(); }, [](const UDataLayerInstance* DataLayer) { return DataLayer; });
 				DataLayersID = FDataLayersID(DataLayers);
 			}
 
 			void AddActor(FActorInstance ActorInstance) { Actors.Add(MoveTemp(ActorInstance)); }
 			const TSet<FActorInstance>& GetActors() const { return Actors; }
 			bool HasDataLayers() const { return !DataLayers.IsEmpty(); }
-			const TArray<const UDataLayer*>& GetDataLayers() const { return DataLayers; }
+			const TArray<const UDataLayerInstance*>& GetDataLayers() const { return DataLayers; }
 			const FDataLayersID& GetDataLayersID() const { return DataLayersID; }
 
 		private:
 			TSet<FActorInstance> Actors;
-			TArray<const UDataLayer*> DataLayers;
+			TArray<const UDataLayerInstance*> DataLayers;
 			FDataLayersID DataLayersID;
 		};
 
@@ -285,7 +286,7 @@ struct FSquare2DGridHelper
 				: Coords(InCoords)
 			{}
 
-			void AddActors(const TSet<FGuid>& InActors, const FActorContainerInstance* ContainerInstance, const TArray<const UDataLayer*>& InDataLayers)
+			void AddActors(const TSet<FGuid>& InActors, const FActorContainerInstance* ContainerInstance, const TArray<const UDataLayerInstance*>& InDataLayers)
 			{
 				for (const FGuid& Actor : InActors)
 				{
diff --git a/Engine/Source/Runtime/Engine/Private/WorldPartition/WorldPartition.cpp b/Engine/Source/Runtime/Engine/Private/WorldPartition/WorldPartition.cpp
--- a/Engine/Source/Runtime/Engine/Private/WorldPartition/WorldPartition.cpp
+++ b/Engine/Source/Runtime/Engine/Private/WorldPartition/WorldPartition.cpp
@@ -41,6 +41,7 @@
 #include "WorldPartition/WorldPartitionEditorPerProjectUserSettings.h"
 #include "WorldPartition/WorldPartitionMiniMap.h"
 #include "WorldPartition/WorldPartitionMiniMapHelper.h"
+#include "WorldPartition/DataLayer/DataLayerInstance.h"
 #include "WorldPartition/DataLayer/WorldDataLayers.h"
 #include "WorldPartition/WorldPartitionActorDescViewProxy.h"
 #include "WorldPartition/HLOD/HLODLayer.h"
@@ -64,9 +65,9 @@ TMap<FName, FString> GetDataLayersDumpString(const UWorldPartition* WorldPartiti
 	TMap<FName, FString> DataLayersDumpString;
 	if (const AWorldDataLayers* WorldDataLayers = WorldPartition->GetWorld()->GetWorldDataLayers())
 	{
-		WorldDataLayers->ForEachDataLayer([&DataLayersDumpString](const UDataLayer* DataLayer)
+		WorldDataLayers->ForEachDataLayer([&DataLayersDumpString](const UDataLayerInstance* DataLayer)
 		{
-			DataLayersDumpString.FindOrAdd(DataLayer->GetFName()) = FString::Format(TEXT("{0}({1})"), { DataLayer->GetDataLayerLabel().ToString(), DataLayer->GetFName().ToString() });
+			DataLayersDumpString.FindOrAdd(DataLayer->GetDataLayerFName()) = FString::Format(TEXT("{0}({1})"), { DataLayer->GetDataLayerShortName(), DataLayer->GetDataLayerFName().ToString() });
 			return true;
 		});
 	}
@@ -107,7 +108,7 @@ FString GetActorDescDumpString(const FWorldPartitionActorDesc* ActorDesc, const
 		ActorDesc->GetBounds().GetExtent().X,
 		ActorDesc->GetBounds().GetExtent().Y,
 		ActorDesc->GetBounds().GetExtent().Z,
-		*GetDataLayerString(ActorDesc->GetDataLayers())
+		*GetDataLayerString(ActorDesc->GetDataLayerInstanceNames())
 	);
 }
 
@@ -999,10 +1000,10 @@ bool UWorldPartition::ShouldActorBeLoadedByEditorCells(const FWorldPartitionActo
 		if (IsRunningCookCommandlet())
 		{
 			// When running cook commandlet, dont allow loading of actors with dynamically loaded data layers
-			for (const FName& DataLayerName : ActorDescProxy.GetDataLayers())
+			for (const FName& DataLayerInstanceName : ActorDescProxy.GetDataLayers())
 			{
-				const UDataLayer* DataLayer = WorldDataLayers->GetDataLayerFromName(DataLayerName);
-				if (DataLayer && DataLayer->IsRuntime())
+				const UDataLayerInstance* DataLayerInstance = WorldDataLayers->GetDataLayerInstance(DataLayerInstanceName);
+				if (DataLayerInstance && DataLayerInstance->IsRuntime())
 				{
 					return false;
 				}
@@ -1011,11 +1012,11 @@ bool UWorldPartition::ShouldActorBeLoadedByEditorCells(const FWorldPartitionActo
 		else
 		{
 			uint32 NumValidLayers = 0;
-			for (const FName& DataLayerName : ActorDescProxy.GetDataLayers())
+			for (const FName& DataLayerInstanceName : ActorDescProxy.GetDataLayers())
 			{
-				if (const UDataLayer* DataLayer = WorldDataLayers->GetDataLayerFromName(DataLayerName))
+				if (const UDataLayerInstance* DataLayerInstance = WorldDataLayers->GetDataLayerInstance(DataLayerInstanceName))
 				{
-					if (DataLayer->IsEffectiveLoadedInEditor())
+					if (DataLayerInstance->IsEffectiveLoadedInEditor())
 					{
 						return true;
 					}
diff --git a/Engine/Source/Runtime/Engine/Private/WorldPartition/WorldPartitionActorCluster.cpp b/Engine/Source/Runtime/Engine/Private/WorldPartition/WorldPartitionActorCluster.cpp
--- a/Engine/Source/Runtime/Engine/Private/WorldPartition/WorldPartitionActorCluster.cpp
+++ b/Engine/Source/Runtime/Engine/Private/WorldPartition/WorldPartitionActorCluster.cpp
@@ -10,25 +10,26 @@
 #include "WorldPartition/ActorDescContainer.h"
 #include "WorldPartition/WorldPartitionActorDescView.h"
 #include "WorldPartition/DataLayer/WorldDataLayers.h"
+#include "WorldPartition/DataLayer/DataLayerInstance.h"
 
 template<class LayerNameContainer>
-TSet<const UDataLayer*> GetDataLayers(UWorld* InWorld, const LayerNameContainer& DataLayerNames)
+TSet<const UDataLayerInstance*> GetDataLayers(UWorld* InWorld, const LayerNameContainer& DataLayerInstanceNames)
 {
-	TSet<const UDataLayer*> DataLayers;
+	TSet<const UDataLayerInstance*> DataLayerInstances;
 	if (const AWorldDataLayers* WorldDataLayers = InWorld->GetWorldDataLayers())
 	{
-		for (const FName& DataLayerName : DataLayerNames)
+		for (const FName& DataLayerInstanceName : DataLayerInstanceNames)
 		{
-			if (const UDataLayer* DataLayer = WorldDataLayers->GetDataLayerFromName(DataLayerName))
+			if (const UDataLayerInstance* DataLayerInstance = WorldDataLayers->GetDataLayerInstance(DataLayerInstanceName))
 			{
-				if (DataLayer->IsRuntime())
+				if (DataLayerInstance->IsRuntime())
 				{
-					DataLayers.Add(DataLayer);
+					DataLayerInstances.Add(DataLayerInstance);
 				}
 			}
 		}
 	}
-	return DataLayers;
+	return DataLayerInstances;
 }
 
 FActorCluster::FActorCluster(UWorld* InWorld, const FWorldPartitionActorDescView& InActorDescView)
@@ -60,7 +61,7 @@ FActorClusterInstance::FActorClusterInstance(const FActorCluster* InCluster, con
 {
 	Bounds = Cluster->Bounds.TransformBy(ContainerInstance->Transform);
 	
-	TSet<const UDataLayer*> DataLayerSet;
+	TSet<const UDataLayerInstance*> DataLayerSet;
 	DataLayerSet.Reserve(Cluster->DataLayers.Num() + ContainerInstance->DataLayers.Num());
 	DataLayerSet.Append(Cluster->DataLayers);
 	DataLayerSet.Append(ContainerInstance->DataLayers);
diff --git a/Engine/Source/Runtime/Engine/Private/WorldPartition/WorldPartitionActorDesc.cpp b/Engine/Source/Runtime/Engine/Private/WorldPartition/WorldPartitionActorDesc.cpp
--- a/Engine/Source/Runtime/Engine/Private/WorldPartition/WorldPartitionActorDesc.cpp
+++ b/Engine/Source/Runtime/Engine/Private/WorldPartition/WorldPartitionActorDesc.cpp
@@ -49,7 +49,7 @@ void FWorldPartitionActorDesc::Init(const AActor* InActor)
 	bLevelBoundsRelevant = InActor->IsLevelBoundsRelevant();
 	bActorIsHLODRelevant = InActor->IsHLODRelevant();
 	HLODLayer = InActor->GetHLODLayer() ? FName(InActor->GetHLODLayer()->GetPathName()) : FName();
-	DataLayers = InActor->GetDataLayerNames();
+	DataLayers = InActor->GetDataLayerInstanceNames();
 	ActorPackage = InActor->GetPackage()->GetFName();
 	ActorPath = *InActor->GetPathName();
 	FolderPath = InActor->GetFolderPath();
diff --git a/Engine/Source/Runtime/Engine/Private/WorldPartition/WorldPartitionActorDescView.cpp b/Engine/Source/Runtime/Engine/Private/WorldPartition/WorldPartitionActorDescView.cpp
--- a/Engine/Source/Runtime/Engine/Private/WorldPartition/WorldPartitionActorDescView.cpp
+++ b/Engine/Source/Runtime/Engine/Private/WorldPartition/WorldPartitionActorDescView.cpp
@@ -75,7 +75,7 @@ FName FWorldPartitionActorDescView::GetHLODLayer() const
 const TArray<FName>& FWorldPartitionActorDescView::GetDataLayers() const
 {
 	static TArray<FName> EmptyDataLayers;
-	return bInvalidDataLayers ? EmptyDataLayers : ActorDesc->GetDataLayers();
+	return bInvalidDataLayers ? EmptyDataLayers : ActorDesc->GetDataLayerInstanceNames();
 }
 
 FName FWorldPartitionActorDescView::GetActorPackage() const
diff --git a/Engine/Source/Runtime/Engine/Private/WorldPartition/WorldPartitionDebugHelper.cpp b/Engine/Source/Runtime/Engine/Private/WorldPartition/WorldPartitionDebugHelper.cpp
--- a/Engine/Source/Runtime/Engine/Private/WorldPartition/WorldPartitionDebugHelper.cpp
+++ b/Engine/Source/Runtime/Engine/Private/WorldPartition/WorldPartitionDebugHelper.cpp
@@ -3,7 +3,7 @@
 #include "WorldPartition/WorldPartitionDebugHelper.h"
 #include "WorldPartition/WorldPartitionRuntimeHash.h"
 #include "WorldPartition/WorldPartition.h"
-#include "WorldPartition/DataLayer/DataLayer.h"
+#include "WorldPartition/DataLayer/DataLayerInstance.h"
 #include "WorldPartition/DataLayer/DataLayerSubsystem.h"
 #include "ProfilingDebugging/ProfilingHelpers.h"
 #include "Engine/Canvas.h"
@@ -57,10 +57,10 @@ FAutoConsoleCommand FWorldPartitionDebugHelper::DebugFilterByDataLayerCommand(
 			UWorld* World = Context.World();
 			if (World && World->IsGameWorld())
 			{
-				TArray<UDataLayer*> DataLayers = UDataLayerSubsystem::ConvertArgsToDataLayers(World, InArgs);
-				for (UDataLayer* DataLayer : DataLayers)
+				TArray<UDataLayerInstance*> DataLayers = UDataLayerSubsystem::ConvertArgsToDataLayers(World, InArgs);
+				for (UDataLayerInstance* DataLayer : DataLayers)
 				{
-					DebugDataLayerFilter.Add(DataLayer->GetFName());
+					DebugDataLayerFilter.Add(DataLayer->GetDataLayerFName());
 				}
 			}
 		}
diff --git a/Engine/Source/Runtime/Engine/Private/WorldPartition/WorldPartitionRuntimeCell.cpp b/Engine/Source/Runtime/Engine/Private/WorldPartition/WorldPartitionRuntimeCell.cpp
--- a/Engine/Source/Runtime/Engine/Private/WorldPartition/WorldPartitionRuntimeCell.cpp
+++ b/Engine/Source/Runtime/Engine/Private/WorldPartition/WorldPartitionRuntimeCell.cpp
@@ -23,14 +23,14 @@ UWorldPartitionRuntimeCell::UWorldPartitionRuntimeCell(const FObjectInitializer&
 {}
 
 #if WITH_EDITOR
-void UWorldPartitionRuntimeCell::SetDataLayers(const TArray<const UDataLayer*>& InDataLayers)
+void UWorldPartitionRuntimeCell::SetDataLayers(const TArray<const UDataLayerInstance*>& InDataLayerInstances)
 {
 	check(DataLayers.IsEmpty());
-	DataLayers.Reserve(InDataLayers.Num());
-	for (const UDataLayer* DataLayer : InDataLayers)
+	DataLayers.Reserve(InDataLayerInstances.Num());
+	for (const UDataLayerInstance* DataLayerInstance : InDataLayerInstances)
 	{
-		check(DataLayer->IsRuntime());
-		DataLayers.Add(DataLayer->GetFName());
+		check(DataLayerInstance->IsRuntime());
+		DataLayers.Add(DataLayerInstance->GetDataLayerFName());
 	}
 	DataLayers.Sort([](const FName& A, const FName& B) { return A.ToString() < B.ToString(); });
 	UpdateDebugName();
@@ -52,15 +52,15 @@ void UWorldPartitionRuntimeCell::UpdateDebugName()
 	int32 DataLayerCount = DataLayers.Num();
 
 	const AWorldDataLayers* WorldDataLayers = GetOuterUWorldPartition()->GetWorld()->GetWorldDataLayers();
-	TArray<const UDataLayer*> DataLayerObjects;
+	TArray<const UDataLayerInstance*> DataLayerObjects;
 	if (WorldDataLayers && (DataLayerCount > 0))
 	{
 		Builder += TEXT(" DL[");
 		for (int i = 0; i < DataLayerCount; ++i)
 		{
-			const UDataLayer* DataLayer = WorldDataLayers->GetDataLayerFromName(DataLayers[i]);
+			const UDataLayerInstance* DataLayer = WorldDataLayers->GetDataLayerInstance(DataLayers[i]);
 			DataLayerObjects.Add(DataLayer);
-			Builder += DataLayer->GetDataLayerLabel().ToString();
+			Builder += DataLayer->GetDataLayerShortName();
 			Builder += TEXT(",");
 		}
 		Builder += FString::Printf(TEXT("ID:%X]"), FDataLayersID(DataLayerObjects).GetHash());
diff --git a/Engine/Source/Runtime/Engine/Private/WorldPartition/WorldPartitionStreamingGeneration.cpp b/Engine/Source/Runtime/Engine/Private/WorldPartition/WorldPartitionStreamingGeneration.cpp
--- a/Engine/Source/Runtime/Engine/Private/WorldPartition/WorldPartitionStreamingGeneration.cpp
+++ b/Engine/Source/Runtime/Engine/Private/WorldPartition/WorldPartitionStreamingGeneration.cpp
@@ -9,6 +9,7 @@
 #include "WorldPartition/WorldPartitionRuntimeHash.h"
 #include "WorldPartition/WorldPartitionStreamingPolicy.h"
 #include "WorldPartition/WorldPartitionActorCluster.h"
+#include "WorldPartition/DataLayer/WorldDataLayers.h"
 #include "WorldPartition/ErrorHandling/WorldPartitionStreamingGenerationNullErrorHandler.h"
 #include "WorldPartition/ErrorHandling/WorldPartitionStreamingGenerationLogErrorHandler.h"
 #include "WorldPartition/ErrorHandling/WorldPartitionStreamingGenerationMapCheckErrorHandler.h"
@@ -388,6 +389,28 @@ class FWorldPartitionStreamingGenerator
 		}
 	}
 
+	void ValidateDataLayers()
+	{
+		for (auto ContainerIt = ContainerDescriptorsMap.CreateIterator(); ContainerIt; ++ContainerIt)
+		{
+			const FActorContainerID& ContainerID = ContainerIt.Key();
+			FContainerDescriptor& ContainerDescriptor = ContainerIt.Value();
+			if (ContainerID.IsMainContainer() && ContainerDescriptor.Container->GetWorld())
+			{
+				if (AWorldDataLayers* WorldDataLayers = ContainerDescriptor.Container->GetWorld()->GetWorldDataLayers())
+				{
+					WorldDataLayers->ForEachDataLayer([this](const UDataLayerInstance* DataLayerInstance) 
+					{
+						DataLayerInstance->Validate(ErrorHandler);
+						return true;
+					});
+				}
+
+				break;
+			}
+		}
+	}
+
 public:
 	FWorldPartitionStreamingGenerator(FActorDescList* InModifiedActorsDescList, IStreamingGenerationErrorHandler* InErrorHandler)
 	: ModifiedActorsDescList(InModifiedActorsDescList)
@@ -402,6 +425,9 @@ public:
 		// Preparation Phase :: Actor Descriptor Views Validation
 		ValidateActorDescriptorViews();
 
+		// Preparation Phase :: Data Layer Validation
+		ValidateDataLayers();
+
 		// Update container descriptors
 		UpdateContainerDescriptors();
 	}
diff --git a/Engine/Source/Runtime/Engine/Public/WorldPartition/DataLayer/DataLayer.h b/Engine/Source/Runtime/Engine/Public/WorldPartition/DataLayer/DataLayer.h
--- a/Engine/Source/Runtime/Engine/Public/WorldPartition/DataLayer/DataLayer.h
+++ b/Engine/Source/Runtime/Engine/Public/WorldPartition/DataLayer/DataLayer.h
@@ -8,6 +8,7 @@
 #include "UObject/Object.h"
 #include "ActorDataLayer.h"
 #include "Math/Color.h"
+#include "WorldPartition/DataLayer/DataLayerInstance.h"
 
 #include "DataLayer.generated.h"
 
@@ -19,26 +20,6 @@ enum class UE_DEPRECATED(5.0, "Use EDataLayerRuntimeState instead.") EDataLayerS
 	Activated
 };
 
-UENUM(BlueprintType)
-enum class EDataLayerRuntimeState : uint8
-{
-	Unloaded,
-	Loaded,		// Loaded but not in the world
-	Activated	// Loaded and in the world
-};
-
-const inline TCHAR* GetDataLayerRuntimeStateName(EDataLayerRuntimeState State)
-{
-	switch(State)
-	{
-	case EDataLayerRuntimeState::Unloaded: return TEXT("Unloaded");
-	case EDataLayerRuntimeState::Loaded: return TEXT("Loaded");
-	case EDataLayerRuntimeState::Activated: return TEXT("Activated");
-	default: check(0);
-	}
-	return TEXT("Invalid");
-}
-
 // Used for debugging
 bool inline GetDataLayerRuntimeStateFromName(const FString& InStateName, EDataLayerRuntimeState& OutState)
 {
@@ -62,17 +43,22 @@ bool inline GetDataLayerRuntimeStateFromName(const FString& InStateName, EDataLa
 
 static_assert(EDataLayerRuntimeState::Unloaded < EDataLayerRuntimeState::Loaded && EDataLayerRuntimeState::Loaded < EDataLayerRuntimeState::Activated, "Streaming Query code is dependent on this being true");
 
-UCLASS(Config = Engine, PerObjectConfig, Within = WorldDataLayers, BlueprintType, AutoCollapseCategories = ("Data Layer|Advanced"), AutoExpandCategories = ("Data Layer|Editor", "Data Layer|Advanced|Runtime"))
-class ENGINE_API UDataLayer : public UObject
+
+static_assert(DATALAYER_TO_INSTANCE_RUNTIME_CONVERSION_ENABLED, "UDEPRECATED_DataLayer class is deprecated and needs to be deleted.");
+class UE_DEPRECATED(5.1, "Use UDataLayerInstance & UDataLayerAsset to create DataLayers") UDEPRECATED_DataLayer;
+
+UCLASS(Config = Engine, PerObjectConfig, Within = WorldDataLayers, BlueprintType, Deprecated)
+class ENGINE_API UDEPRECATED_DataLayer : public UObject
 {
 	GENERATED_UCLASS_BODY()
 
+	friend class UDeprecatedDataLayerInstance;
+
 public:
 #if WITH_EDITOR
 	virtual bool CanEditChange(const FProperty* Property) const;
 	virtual void PostEditChangeProperty(FPropertyChangedEvent& PropertyChangedEvent) override;
 
-	void SetDataLayerLabel(FName DataLayerLabel);
 	void SetVisible(bool bIsVisible);
 	void SetIsInitiallyVisible(bool bIsInitiallyVisible);
 	void SetIsRuntime(bool bIsRuntime);
@@ -86,22 +72,18 @@ public:
 	bool IsLoadedInEditorChangedByUserOperation() const { return bIsLoadedInEditorChangedByUserOperation; }
 	void ClearLoadedInEditorChangedByUserOperation() { bIsLoadedInEditorChangedByUserOperation = false; }
 
-	bool CanParent(const UDataLayer* InParent) const;
-	void SetParent(UDataLayer* InParent);
-	void SetChildParent(UDataLayer* InParent);
-
-	bool IsInActorEditorContext() const;
-	bool AddToActorEditorContext();
-	bool RemoveFromActorEditorContext();
+	bool CanParent(const UDEPRECATED_DataLayer* InParent) const;
+	void SetParent(UDEPRECATED_DataLayer* InParent);
+	void SetChildParent(UDEPRECATED_DataLayer* InParent);
 
-	static FText GetDataLayerText(const UDataLayer* InDataLayer);
+	static FText GetDataLayerText(const UDEPRECATED_DataLayer* InDataLayer);
 	const TCHAR* GetDataLayerIconName() const;
 #endif
-	const TArray<TObjectPtr<UDataLayer>>& GetChildren() const { return Children; }
-	void ForEachChild(TFunctionRef<bool(const UDataLayer*)> Operation) const;
+	const TArray<TObjectPtr<UDEPRECATED_DataLayer>>& GetChildren() const { return Children_DEPRECATED; }
+	void ForEachChild(TFunctionRef<bool(const UDEPRECATED_DataLayer*)> Operation) const;
 
-	const UDataLayer* GetParent() const { return Parent; }
-	UDataLayer* GetParent() { return Parent; }
+	const UDEPRECATED_DataLayer* GetParent() const { return Parent_DEPRECATED; }
+	UDEPRECATED_DataLayer* GetParent() { return Parent_DEPRECATED; }
 
 	virtual void PostLoad() override;
 
@@ -129,9 +111,6 @@ public:
 	UFUNCTION(Category = "Data Layer|Runtime", BlueprintCallable)
 	EDataLayerRuntimeState GetInitialRuntimeState() const { return IsRuntime() ? InitialRuntimeState : EDataLayerRuntimeState::Unloaded; }
 
-	/** Returns a sanitized version of the provided Data Layer Label */
-	static FName GetSanitizedDataLayerLabel(FName DataLayerLabel);
-
 	//~ Begin Deprecated
 
 	UE_DEPRECATED(5.0, "Use IsRuntime() instead.")
@@ -151,9 +130,9 @@ public:
 	//~ End Deprecated
 private:
 
-	void AddChild(UDataLayer* DataLayer);
+	void AddChild(UDEPRECATED_DataLayer* DataLayer);
 #if WITH_EDITOR
-	void RemoveChild(UDataLayer* DataLayer);
+	void RemoveChild(UDEPRECATED_DataLayer* DataLayer);
 	void PropagateIsRuntime();
 #endif
 
@@ -201,8 +180,8 @@ private:
 	FColor DebugColor;
 
 	UPROPERTY()
-	TObjectPtr<UDataLayer> Parent;
+	TObjectPtr<UDEPRECATED_DataLayer> Parent_DEPRECATED;
 
 	UPROPERTY(Transient)
-	TArray<TObjectPtr<UDataLayer>> Children;
+	TArray<TObjectPtr<UDEPRECATED_DataLayer>> Children_DEPRECATED;
 };
\ No newline at end of file
diff --git a/Engine/Source/Runtime/Engine/Public/WorldPartition/DataLayer/DataLayerEditorContext.h b/Engine/Source/Runtime/Engine/Public/WorldPartition/DataLayer/DataLayerEditorContext.h
--- a/Engine/Source/Runtime/Engine/Public/WorldPartition/DataLayer/DataLayerEditorContext.h
+++ b/Engine/Source/Runtime/Engine/Public/WorldPartition/DataLayer/DataLayerEditorContext.h
@@ -6,7 +6,7 @@
 
 #if WITH_EDITOR
 
-struct FActorDataLayer;
+class UDataLayerInstance;
 class UWorld;
 
 /**
@@ -18,12 +18,12 @@ struct ENGINE_API FDataLayerEditorContext
 public:
 	static const uint32 EmptyHash = 0;
 	FDataLayerEditorContext() : Hash(FDataLayerEditorContext::EmptyHash) {}
-	FDataLayerEditorContext(UWorld* InWorld, const TArray<FName>& InDataLayers);
-	FORCEINLINE bool IsEmpty() const { return (Hash == FDataLayerEditorContext::EmptyHash) && DataLayers.IsEmpty(); }
+	FDataLayerEditorContext(UWorld* InWorld, const TArray<FName>& InDataLayerInstances);
+	FORCEINLINE bool IsEmpty() const { return (Hash == FDataLayerEditorContext::EmptyHash) && DataLayerInstances.IsEmpty(); }
 	FORCEINLINE uint32 GetHash() const { return Hash; }
-	FORCEINLINE const TArray<FName>& GetDataLayers() const { return DataLayers; }
+	FORCEINLINE const TArray<FName>& GetDataLayerInstanceNames() const { return DataLayerInstances; }
 private:
 	uint32 Hash;
-	TArray<FName> DataLayers;
+	TArray<FName> DataLayerInstances;
 };
 #endif
\ No newline at end of file
