UDataLayer deprecated in favor of UDataLayerInstance & UDataLayerAsset

diff --git a/Engine/Source/Runtime/Engine/Public/WorldPartition/DataLayer/DataLayerInstance.h b/Engine/Source/Runtime/Engine/Public/WorldPartition/DataLayer/DataLayerInstance.h
new file mode 100644
--- /dev/null
+++ b/Engine/Source/Runtime/Engine/Public/WorldPartition/DataLayer/DataLayerInstance.h
@@ -0,0 +1,168 @@
+// Copyright Epic Games, Inc. All Rights Reserved.
+
+#pragma once
+
+#include "CoreMinimal.h"
+#include "UObject/ObjectMacros.h"
+#include "UObject/UObjectGlobals.h"
+#include "UObject/Object.h"
+#include "WorldPartition/DataLayer/DataLayerAsset.h"
+
+#include "DataLayerInstance.generated.h"
+
+class IStreamingGenerationErrorHandler;
+
+UENUM(BlueprintType)
+enum class EDataLayerRuntimeState : uint8
+{
+	Unloaded,
+	Loaded,
+	Activated
+};
+
+const inline TCHAR* GetDataLayerRuntimeStateName(EDataLayerRuntimeState State)
+{
+	switch (State)
+	{
+	case EDataLayerRuntimeState::Unloaded: return TEXT("Unloaded");
+	case EDataLayerRuntimeState::Loaded: return TEXT("Loaded");
+	case EDataLayerRuntimeState::Activated: return TEXT("Activated");
+	default: check(0);
+	}
+	return TEXT("Invalid");
+}
+
+UCLASS(Config = Engine, PerObjectConfig, Within = WorldDataLayers, BlueprintType, AutoCollapseCategories = ("Data Layer|Advanced"), AutoExpandCategories = ("Data Layer|Editor", "Data Layer|Advanced|Runtime"))
+class ENGINE_API UDataLayerInstance : public UObject
+{
+	GENERATED_UCLASS_BODY()
+
+	friend class UDataLayerConversionInfo;
+
+public:
+	virtual void PostLoad() override;
+
+#if WITH_EDITOR
+	void SetVisible(bool bIsVisible);
+	void SetIsInitiallyVisible(bool bIsInitiallyVisible);
+	void SetIsLoadedInEditor(bool bIsLoadedInEditor, bool bFromUserChange);
+	void SetIsLocked(bool bInIsLocked) { bIsLocked = bInIsLocked; }
+
+	bool IsLocked() const;
+	bool IsInitiallyLoadedInEditor() const { return bIsInitiallyLoadedInEditor; }
+	bool IsLoadedInEditor() const { return bIsLoadedInEditor; }
+	bool IsEffectiveLoadedInEditor() const;
+	bool IsLoadedInEditorChangedByUserOperation() const { return bIsLoadedInEditorChangedByUserOperation; }
+	void ClearLoadedInEditorChangedByUserOperation() { bIsLoadedInEditorChangedByUserOperation = false; }
+
+	const TCHAR* GetDataLayerIconName() const { return DataLayerUtils::GetDataLayerIconName(GetType()); }
+
+	bool CanParent(const UDataLayerInstance* InParent) const;
+	bool IsDataLayerTypeValidToParent(EDataLayerType ParentDataLayerType) const;
+	bool SetParent(UDataLayerInstance* InParent);
+
+	void SetChildParent(UDataLayerInstance* InParent);
+
+	static FText GetDataLayerText(const UDataLayerInstance* InDataLayer);
+
+	virtual bool AddActor(AActor* Actor) const { return false; }
+	virtual bool RemoveActor(AActor* Actor) const { return false; }
+	virtual bool ContainsActor(const AActor* Actor) const { return false; }
+
+	bool IsInActorEditorContext() const;
+	bool AddToActorEditorContext();
+	bool RemoveFromActorEditorContext();
+
+	virtual bool Validate(IStreamingGenerationErrorHandler* ErrorHandler) const;
+
+#if DATALAYER_TO_INSTANCE_RUNTIME_CONVERSION_ENABLED
+	virtual bool SupportRelabeling() const { return false; }
+	virtual bool RelabelDataLayer(FName NewDataLayerLabel) { return false; }
+
+#endif // DATALAYER_TO_INSTANCE_RUNTIME_CONVERSION_ENABLED
+
+#endif // WITH_EDITOR
+
+	UFUNCTION(Category = "Data Layer", BlueprintCallable)
+	virtual EDataLayerType GetType() const { return EDataLayerType::Unknown; }
+
+	UFUNCTION(Category = "Data Layer|Editor", BlueprintCallable)
+	bool IsInitiallyVisible() const;
+
+	UFUNCTION(Category = "Data Layer|Editor", BlueprintCallable)
+	bool IsVisible() const;
+
+	UFUNCTION(Category = "Data Layer|Editor", BlueprintCallable)
+	bool IsEffectiveVisible() const;
+
+	UFUNCTION(Category = "Data Layer|Runtime", BlueprintCallable)
+	virtual bool IsRuntime() const { return false; }
+
+	UFUNCTION(Category = "Data Layer|Runtime", BlueprintCallable)
+	virtual FColor GetDebugColor() const { return FColor::Black; }
+
+	UFUNCTION(Category = "Data Layer|Runtime", BlueprintCallable)
+	EDataLayerRuntimeState GetInitialRuntimeState() const { return IsRuntime() ? InitialRuntimeState : EDataLayerRuntimeState::Unloaded; }
+
+	virtual FString GetDataLayerShortName() const { return TEXT("Invalid Data Layer"); }
+	virtual FString GetDataLayerFullName() const { return TEXT("Invalid Data Layer"); }
+
+	const UDataLayerInstance* GetParent() const { return Parent; }
+	UDataLayerInstance* GetParent() { return Parent; }
+
+	const TArray<TObjectPtr<UDataLayerInstance>>& GetChildren() const { return Children; }
+	void ForEachChild(TFunctionRef<bool(const UDataLayerInstance*)> Operation) const;
+
+#if DATALAYER_TO_INSTANCE_RUNTIME_CONVERSION_ENABLED
+	virtual FName GetDataLayerFName() const { return GetFName(); }
+
+private:
+	FName GetFName() const { return Super::GetFName(); }
+	FString GetName() const { return Super::GetName();  }
+public:
+#endif // DATALAYER_TO_INSTANCE_RUNTIME_CONVERSION_ENABLED
+
+private:
+	void AddChild(UDataLayerInstance* DataLayer);
+#if WITH_EDITOR
+	void RemoveChild(UDataLayerInstance* DataLayer);
+#endif
+
+protected:
+#if WITH_EDITORONLY_DATA
+	/** Whether actors associated with the DataLayer are visible in the viewport */
+	UPROPERTY(Transient)
+	uint32 bIsVisible : 1;
+
+	/** Whether actors associated with the Data Layer should be initially visible in the viewport when loading the map */
+	UPROPERTY(Category = "Data Layer|Editor", EditAnywhere)
+	uint32 bIsInitiallyVisible : 1;
+
+	/** Determines the default value of the data layer's loaded state in editor if it hasn't been changed in data layer outliner by the user */
+	UPROPERTY(Category = "Data Layer|Editor", EditAnywhere, meta = (DisplayName = "Is Initially Loaded"))
+	uint32 bIsInitiallyLoadedInEditor : 1;
+
+	/** Wheter the data layer is loaded in editor (user setting) */
+	UPROPERTY(Transient)
+	uint32 bIsLoadedInEditor : 1;
+
+	/** Whether this data layer editor visibility was changed by a user operation */
+	UPROPERTY(Transient)
+	uint32 bIsLoadedInEditorChangedByUserOperation : 1;
+
+	/** Whether this data layer is locked, which means the user can't change actors assignation, remove or rename it */
+	UPROPERTY()
+	uint32 bIsLocked : 1;
+#endif
+
+	UPROPERTY(Category = "Data Layer|Advanced|Runtime", EditAnywhere, meta = (EditConditionHides, EditCondition = "bIsRuntime"))
+	EDataLayerRuntimeState InitialRuntimeState;
+
+private:
+	UPROPERTY()
+	TObjectPtr<UDataLayerInstance> Parent;
+
+	UPROPERTY(Transient)
+	TArray<TObjectPtr<UDataLayerInstance>> Children;
+
+};
\ No newline at end of file
diff --git a/Engine/Source/Runtime/Engine/Public/WorldPartition/DataLayer/DataLayerInstanceWithAsset.h b/Engine/Source/Runtime/Engine/Public/WorldPartition/DataLayer/DataLayerInstanceWithAsset.h
new file mode 100644
--- /dev/null
+++ b/Engine/Source/Runtime/Engine/Public/WorldPartition/DataLayer/DataLayerInstanceWithAsset.h
@@ -0,0 +1,43 @@
+// Copyright Epic Games, Inc. All Rights Reserved.
+#pragma once
+
+#include "WorldPartition/DataLayer/DataLayerInstance.h"
+
+#include "DataLayerInstanceWithAsset.generated.h"
+
+class UDataLayerAsset;
+
+UCLASS(Config = Engine, PerObjectConfig, Within = WorldDataLayers, AutoCollapseCategories = ("Data Layer|Advanced"), AutoExpandCategories = ("Data Layer|Editor", "Data Layer|Advanced|Runtime"))
+class ENGINE_API UDataLayerInstanceWithAsset : public UDataLayerInstance
+{
+	GENERATED_UCLASS_BODY()
+
+		friend class UDataLayerConversionInfo;
+
+public:
+#if WITH_EDITOR
+	static FName MakeName(const UDataLayerAsset* DeprecatedDataLayer);
+	void OnCreated(const UDataLayerAsset* Asset);
+
+	virtual bool AddActor(AActor* Actor) const override;
+	virtual bool RemoveActor(AActor* Actor) const override;
+	virtual bool ContainsActor(const AActor* Actor) const override;
+
+	virtual bool Validate(IStreamingGenerationErrorHandler* ErrorHandler) const override;
+#endif
+
+	const UDataLayerAsset* GetAsset() const { return DataLayerAsset; }
+
+	virtual EDataLayerType GetType() const override { return DataLayerAsset != nullptr ? DataLayerAsset->GetType() : EDataLayerType::Unknown; }
+
+	virtual bool IsRuntime() const override { return DataLayerAsset != nullptr ? DataLayerAsset->IsRuntime() : false; }
+
+	virtual FColor GetDebugColor() const override { return DataLayerAsset != nullptr ? DataLayerAsset->GetDebugColor() : FColor::Black; }
+
+	virtual FString GetDataLayerShortName() const override { return DataLayerAsset != nullptr ? DataLayerAsset->GetName() : GetDataLayerFName().ToString(); }
+	virtual FString GetDataLayerFullName() const override { return DataLayerAsset != nullptr ? DataLayerAsset->GetFullName() : GetDataLayerFName().ToString(); }
+
+private:
+	UPROPERTY(Category = "Data Layer", EditAnywhere)
+	TObjectPtr<const UDataLayerAsset> DataLayerAsset;
+};
\ No newline at end of file
diff --git a/Engine/Source/Runtime/Engine/Public/WorldPartition/DataLayer/DataLayerSubsystem.h b/Engine/Source/Runtime/Engine/Public/WorldPartition/DataLayer/DataLayerSubsystem.h
--- a/Engine/Source/Runtime/Engine/Public/WorldPartition/DataLayer/DataLayerSubsystem.h
+++ b/Engine/Source/Runtime/Engine/Public/WorldPartition/DataLayer/DataLayerSubsystem.h
@@ -4,7 +4,7 @@
 
 #include "Subsystems/WorldSubsystem.h"
 #include "WorldPartition/DataLayer/ActorDataLayer.h"
-#include "WorldPartition/DataLayer/DataLayer.h"
+#include "WorldPartition/DataLayer/DataLayerInstance.h"
 #include "WorldPartition/DataLayer/DataLayerEditorContext.h"
 #include "DataLayerSubsystem.generated.h"
 
@@ -14,7 +14,7 @@
 
 class UCanvas;
 
-DECLARE_DYNAMIC_MULTICAST_DELEGATE_TwoParams(FOnDataLayerRuntimeStateChanged, const UDataLayer*, DataLayer, EDataLayerRuntimeState, State);
+DECLARE_DYNAMIC_MULTICAST_DELEGATE_TwoParams(FOnDataLayerRuntimeStateChanged, const UDataLayerInstance*, DataLayer, EDataLayerRuntimeState, State);
 
 UCLASS()
 class ENGINE_API UDataLayerSubsystem : public UWorldSubsystem
@@ -30,41 +30,19 @@ public:
 
 	//~ Begin Blueprint callable functions
 
-	/** Find a Data Layer by name. */
+	/** Find a Data Layer by its asset. */
 	UFUNCTION(BlueprintCallable, Category = DataLayers)
-	UDataLayer* GetDataLayer(const FActorDataLayer& InDataLayer) const;
+	UDataLayerInstance* GetDataLayerFromAsset(const UDataLayerAsset* InDataLayerAsset) const;
 
-	/** Find a Data Layer by label. */
 	UFUNCTION(BlueprintCallable, Category = DataLayers)
-	UDataLayer* GetDataLayerFromLabel(FName InDataLayerLabel) const;
+	EDataLayerRuntimeState GetDataLayerInstanceRuntimeState(const UDataLayerAsset* InDataLayerAsset) const;
 
-	/** Find a Data Layer by name. */
 	UFUNCTION(BlueprintCallable, Category = DataLayers)
-	UDataLayer* GetDataLayerFromName(FName InDataLayerName) const;
+	EDataLayerRuntimeState GetDataLayerInstanceEffectiveRuntimeState(const UDataLayerAsset* InDataLayerAsset) const;
 
 	/** Set the Data Layer state using its name. */
 	UFUNCTION(BlueprintCallable, Category = DataLayers, BlueprintAuthorityOnly)
-	void SetDataLayerRuntimeState(const FActorDataLayer& InDataLayer, EDataLayerRuntimeState InState, bool bInIsRecursive = false);
-
-	/** Set the Data Layer state using its label. */
-	UFUNCTION(BlueprintCallable, Category = DataLayers, BlueprintAuthorityOnly)
-	void SetDataLayerRuntimeStateByLabel(const FName& InDataLayerLabel, EDataLayerRuntimeState InState, bool bInIsRecursive = false);
-		
-	/** Get the Data Layer state using its name. */
-	UFUNCTION(BlueprintCallable, Category = DataLayers)
-	EDataLayerRuntimeState GetDataLayerRuntimeState(const FActorDataLayer& InDataLayer) const;
-
-	/** Get the Data Layer state using its label. */
-	UFUNCTION(BlueprintCallable, Category = DataLayers)
-	EDataLayerRuntimeState GetDataLayerRuntimeStateByLabel(const FName& InDataLayerLabel) const;
-
-	/** Get the Data Layer effective state using its name. */
-	UFUNCTION(BlueprintCallable, Category = DataLayers)
-	EDataLayerRuntimeState GetDataLayerEffectiveRuntimeState(const FActorDataLayer& InDataLayer) const;
-
-	/** Get the Data Layer effective state using its label. */
-	UFUNCTION(BlueprintCallable, Category = DataLayers)
-	EDataLayerRuntimeState GetDataLayerEffectiveRuntimeStateByLabel(const FName& InDataLayerLabel) const;
+	void SetDataLayerInstanceRuntimeState(const UDataLayerAsset* InDataLayerAsset, EDataLayerRuntimeState InState, bool bInIsRecursive = false);
 
 	/** Called when a Data Layer changes state. */
 	UPROPERTY(BlueprintAssignable)
@@ -72,11 +50,11 @@ public:
 
 	//~ End Blueprint callable functions
 
-	void SetDataLayerRuntimeState(const UDataLayer* InDataLayer, EDataLayerRuntimeState InState, bool bInIsRecursive = false);
-	void SetDataLayerRuntimeStateByName(const FName& InDataLayerName, EDataLayerRuntimeState InState, bool bInIsRecursive = false);
-	EDataLayerRuntimeState GetDataLayerRuntimeState(const UDataLayer* InDataLayer) const;
+	UDataLayerInstance* GetDataLayerInstance(const FName& InDataLayerInstanceName) const;
+	void SetDataLayerRuntimeState(const UDataLayerInstance* InDataLayerInstance, EDataLayerRuntimeState InState, bool bInIsRecursive = false);
+	EDataLayerRuntimeState GetDataLayerRuntimeState(const UDataLayerInstance* InDataLayer) const;
 	EDataLayerRuntimeState GetDataLayerRuntimeStateByName(const FName& InDataLayerName) const;
-	EDataLayerRuntimeState GetDataLayerEffectiveRuntimeState(const UDataLayer* InDataLayer) const;
+	EDataLayerRuntimeState GetDataLayerEffectiveRuntimeState(const UDataLayerInstance* InDataLayer) const;
 	EDataLayerRuntimeState GetDataLayerEffectiveRuntimeStateByName(const FName& InDataLayerName) const;
 	bool IsAnyDataLayerInEffectiveRuntimeState(const TArray<FName>& InDataLayerNames, EDataLayerRuntimeState InState) const;
 	const TSet<FName>& GetEffectiveActiveDataLayerNames() const;
@@ -84,7 +62,7 @@ public:
 
 	void GetDataLayerDebugColors(TMap<FName, FColor>& OutMapping) const;
 	void DrawDataLayersStatus(UCanvas* Canvas, FVector2D& Offset) const;
-	static TArray<UDataLayer*> ConvertArgsToDataLayers(UWorld* World, const TArray<FString>& InArgs);
+	static TArray<UDataLayerInstance*> ConvertArgsToDataLayers(UWorld* World, const TArray<FString>& InArgs);
 
 	void DumpDataLayers(FOutputDevice& OutputDevice) const;
 
@@ -100,12 +78,6 @@ public:
 	UFUNCTION(BlueprintCallable, Category = DataLayers, BlueprintAuthorityOnly, meta = (DeprecatedFunction, DeprecationMessage = "Use SetDataLayerRuntimeStateByLabel instead"))
 	void SetDataLayerStateByLabel(const FName& InDataLayerLabel, EDataLayerState InState) { SetDataLayerRuntimeStateByLabel(InDataLayerLabel, (EDataLayerRuntimeState)InState); }
 
-	UE_DEPRECATED(5.0, "Use SetDataLayerRuntimeState() instead.")
-	void SetDataLayerState(const UDataLayer* InDataLayer, EDataLayerState InState) { SetDataLayerRuntimeState(InDataLayer, (EDataLayerRuntimeState)InState); }
-
-	UE_DEPRECATED(5.0, "Use SetDataLayerRuntimeStateByName() instead.")
-	void SetDataLayerStateByName(const FName& InDataLayerName, EDataLayerState InState) { SetDataLayerRuntimeStateByName(InDataLayerName, (EDataLayerRuntimeState)InState); }
-
 	UE_DEPRECATED(5.0, "Use GetDataLayerRuntimeState() instead.")
 	UFUNCTION(BlueprintCallable, Category = DataLayers, meta = (DeprecatedFunction, DeprecationMessage = "Use GetDataLayerRuntimeState instead"))
 	EDataLayerState GetDataLayerState(const FActorDataLayer& InDataLayer) const { return (EDataLayerState)GetDataLayerRuntimeState(InDataLayer); }
@@ -115,7 +87,7 @@ public:
 	EDataLayerState GetDataLayerStateByLabel(const FName& InDataLayerLabel) const { return (EDataLayerState)GetDataLayerRuntimeStateByLabel(InDataLayerLabel); }
 
 	UE_DEPRECATED(5.0, "Use GetDataLayerRuntimeState() instead.")
-	EDataLayerState GetDataLayerState(const UDataLayer* InDataLayer) const { return (EDataLayerState)GetDataLayerRuntimeState(InDataLayer); }
+	EDataLayerState GetDataLayerState(const UDataLayerInstance* InDataLayer) const { return (EDataLayerState)GetDataLayerRuntimeState(InDataLayer); }
 
 	UE_DEPRECATED(5.0, "Use GetDataLayerRuntimeStateByName() instead.")
 	EDataLayerState GetDataLayerStateByName(const FName& InDataLayerName) const { return (EDataLayerState)GetDataLayerRuntimeStateByName(InDataLayerName); }
@@ -132,6 +104,43 @@ public:
 	UE_DEPRECATED(5.0, "GetLoadedDataLayerNames will be removed.")
 	UFUNCTION(BlueprintCallable, Category = DataLayers, meta = (DeprecatedFunction, DeprecationMessage = "GetLoadedDataLayerNames will be removed."))
 	const TSet<FName>& GetLoadedDataLayerNames() const { return GetEffectiveLoadedDataLayerNames(); }
+
+	UE_DEPRECATED(5.1, "Use GetDataLayerFromAsset() instead.")
+	UFUNCTION(BlueprintCallable, Category = DataLayers)
+	UDataLayerInstance* GetDataLayer(const FActorDataLayer& InDataLayer) const;
+
+	UE_DEPRECATED(5.1, "Use GetDataLayerFromAsset() instead.")
+	UFUNCTION(BlueprintCallable, Category = DataLayers)
+	UDataLayerInstance* GetDataLayerFromName(FName InDataLayerName) const;
+
+	UE_DEPRECATED(5.1, "Use GetDataLayerFromAsset() instead.")
+	UFUNCTION(BlueprintCallable, Category = DataLayers)
+	UDataLayerInstance* GetDataLayerFromLabel(FName InDataLayerLabel) const;
+
+	UE_DEPRECATED(5.1, "Use GetDataLayerInstanceRuntimeState() instead.")
+	UFUNCTION(BlueprintCallable, Category = DataLayers)
+	EDataLayerRuntimeState GetDataLayerRuntimeState(const FActorDataLayer& InDataLayer) const;
+
+	UE_DEPRECATED(5.1, "Use GetDataLayerInstanceRuntimeState() instead.")
+	UFUNCTION(BlueprintCallable, Category = DataLayers)
+	EDataLayerRuntimeState GetDataLayerRuntimeStateByLabel(const FName& InDataLayerLabel) const;
+
+	UE_DEPRECATED(5.1, "Use GetDataLayerInstanceEffectiveRuntimeState() instead.")
+	UFUNCTION(BlueprintCallable, Category = DataLayers)
+	EDataLayerRuntimeState GetDataLayerEffectiveRuntimeState(const FActorDataLayer& InDataLayer) const;
+
+	UE_DEPRECATED(5.1, "Use GetDataLayerInstanceEffectiveRuntimeState() instead.")
+	UFUNCTION(BlueprintCallable, Category = DataLayers)
+	EDataLayerRuntimeState GetDataLayerEffectiveRuntimeStateByLabel(const FName& InDataLayerLabel) const;
+
+	UE_DEPRECATED(5.1, "Use SetDataLayerRuntimeState() with UDataLayerAsset* overload instead.")
+	UFUNCTION(BlueprintCallable, Category = DataLayers, BlueprintAuthorityOnly)
+	void SetDataLayerRuntimeState(const FActorDataLayer& InDataLayer, EDataLayerRuntimeState InState, bool bInIsRecursive = false);
+
+	UE_DEPRECATED(5.1, "Use SetDataLayerInstanceRuntimeState() instead.")
+	UFUNCTION(BlueprintCallable, Category = DataLayers, BlueprintAuthorityOnly)
+	void SetDataLayerRuntimeStateByLabel(const FName& InDataLayerLabel, EDataLayerRuntimeState InState, bool bInIsRecursive = false);
+
 	//~ End Deprecated
 
 #if WITH_EDITOR
@@ -147,7 +156,7 @@ private:
 	static class FAutoConsoleCommand ToggleDataLayerActivation;
 
 	/** Data layers load time */
-	mutable TMap<UDataLayer*, double> ActiveDataLayersLoadTime;
+	mutable TMap<const UDataLayerInstance*, double> ActiveDataLayersLoadTime;
 
 	/** Console command used to set Runtime DataLayer state*/
 	static class FAutoConsoleCommand SetDataLayerRuntimeStateCommand;
diff --git a/Engine/Source/Runtime/Engine/Public/WorldPartition/DataLayer/DataLayerUtils.h b/Engine/Source/Runtime/Engine/Public/WorldPartition/DataLayer/DataLayerUtils.h
new file mode 100644
--- /dev/null
+++ b/Engine/Source/Runtime/Engine/Public/WorldPartition/DataLayer/DataLayerUtils.h
@@ -0,0 +1,42 @@
+// Copyright Epic Games, Inc. All Rights Reserved.
+
+#pragma once
+
+#include "CoreMinimal.h"
+#include "UObject/ObjectMacros.h"
+#include "UObject/UObjectGlobals.h"
+
+#include "DataLayerUtils.generated.h"
+
+#define DATALAYER_TO_INSTANCE_RUNTIME_CONVERSION_ENABLED 1
+
+UENUM(BlueprintType)
+enum class EDataLayerType : uint8
+{
+	Runtime,
+	Editor,
+	Unknown UMETA(Hidden),
+
+	Size UMETA(Hidden)
+};
+
+namespace DataLayerUtils
+{
+#if WITH_EDITOR
+	constexpr const TCHAR* IconNameByType[static_cast<int>(EDataLayerType::Size)] = { TEXT("DataLayer.Runtime") , TEXT("DataLayer.Editor"), TEXT("") };
+
+	constexpr const TCHAR* GetDataLayerIconName(EDataLayerType DataLayerType)
+	{
+		return IconNameByType[static_cast<uint32>(DataLayerType)];
+	}
+#endif
+
+#if DATALAYER_TO_INSTANCE_RUNTIME_CONVERSION_ENABLED
+	UE_DEPRECATED(5.1, "Label usage is deprecated.")
+	static FName GetSanitizedDataLayerLabel(FName InDataLayerLabel)
+	{
+		return FName(InDataLayerLabel.ToString().TrimStartAndEnd().Replace(TEXT("\""), TEXT("")));
+	}
+
+#endif
+};
\ No newline at end of file
diff --git a/Engine/Source/Runtime/Engine/Public/WorldPartition/DataLayer/DataLayersID.h b/Engine/Source/Runtime/Engine/Public/WorldPartition/DataLayer/DataLayersID.h
--- a/Engine/Source/Runtime/Engine/Public/WorldPartition/DataLayer/DataLayersID.h
+++ b/Engine/Source/Runtime/Engine/Public/WorldPartition/DataLayer/DataLayersID.h
@@ -2,14 +2,16 @@
 
 #pragma once
 
-#include "WorldPartition/DataLayer/DataLayer.h"
+#include "CoreMinimal.h"
+
+class UDataLayerInstance;
 
 #if WITH_EDITOR
 class ENGINE_API FDataLayersID
 {
 public:
 	FDataLayersID();
-	FDataLayersID(const TArray<const UDataLayer*>& InDataLayers);
+	FDataLayersID(const TArray<const UDataLayerInstance*>& DataLayerInstance);
 
 	bool operator==(const FDataLayersID& Other) const
 	{
diff --git a/Engine/Source/Runtime/Engine/Public/WorldPartition/DataLayer/DeprecatedDataLayerInstance.h b/Engine/Source/Runtime/Engine/Public/WorldPartition/DataLayer/DeprecatedDataLayerInstance.h
new file mode 100644
--- /dev/null
+++ b/Engine/Source/Runtime/Engine/Public/WorldPartition/DataLayer/DeprecatedDataLayerInstance.h
@@ -0,0 +1,73 @@
+// Copyright Epic Games, Inc. All Rights Reserved.
+#pragma once
+
+#include "WorldPartition/DataLayer/DataLayerInstance.h"
+
+#include "DeprecatedDataLayerInstance.generated.h"
+
+class UDEPRECATED_DataLayer;
+
+// Class used for Runtime Conversion of the Deprecated UDataLayer Class to UDataLayerInstance + UDataLayerAsset.
+// This class is not to be inherited. It is solely used by AWorldDatalayers to convert UDataLayers to UDataLayerInstances on Level Boot.
+// You will need to run the DataLayerToAsset CommandLet to convert the deprecated datalayers to UDataLayerAssets and UDataLayerInstanceWithAsset.
+UCLASS(Config = Engine, PerObjectConfig, Within = WorldDataLayers, AutoCollapseCategories = ("Data Layer|Advanced"), AutoExpandCategories = ("Data Layer|Editor", "Data Layer|Advanced|Runtime"))
+class ENGINE_API UDeprecatedDataLayerInstance final : public UDataLayerInstance
+{
+	static_assert(DATALAYER_TO_INSTANCE_RUNTIME_CONVERSION_ENABLED, "UDeprecatedDataLayerInstance is deprecated and needs to be deleted.");
+
+	GENERATED_UCLASS_BODY()
+
+	friend class AWorldDataLayers;
+	friend class UDataLayerConversionInfo;
+	friend class UDataLayerToAssetCommandletContext;
+	friend class UDataLayerToAssetCommandlet;
+
+public:
+#if WITH_EDITOR
+	static FName MakeName();
+	void OnCreated();
+
+	PRAGMA_DISABLE_DEPRECATION_WARNINGS
+	static FName MakeName(const UDEPRECATED_DataLayer* DeprecatedDataLayer);
+	void OnCreated(const UDEPRECATED_DataLayer* DeprecatedDataLayer);
+	PRAGMA_ENABLE_DEPRECATION_WARNINGS
+
+	void SetDataLayerLabel(FName InDataLayerLabel);
+#endif
+	
+	FName GetDataLayerLabel() const { return Label; }
+
+private:
+#if WITH_EDITOR
+	virtual bool AddActor(AActor* Actor) const override;
+	virtual bool RemoveActor(AActor* Actor) const override;
+	virtual bool ContainsActor(const AActor* Actor) const override;
+
+	virtual bool SupportRelabeling() const override { return true; }
+	virtual bool RelabelDataLayer(FName NewDataLayerLabel) override;
+#endif
+
+	virtual FName GetDataLayerFName() const override { return !DeprecatedDataLayerFName.IsNone() ? DeprecatedDataLayerFName : Super::GetDataLayerFName(); }
+
+	virtual EDataLayerType GetType() const override { return DataLayerType; }
+
+	virtual bool IsRuntime() const override { return DataLayerType == EDataLayerType::Runtime; }
+
+	virtual FColor GetDebugColor() const override { return DebugColor; }
+
+	virtual FString GetDataLayerShortName() const override { return Label.ToString(); }
+	virtual FString GetDataLayerFullName() const override { return TEXT("Deprecated_") + Label.ToString(); }
+
+private:
+	UPROPERTY()
+	FName Label;
+
+	UPROPERTY()
+	FName DeprecatedDataLayerFName;
+
+	UPROPERTY(Category = "Data Layer", EditAnywhere)
+	EDataLayerType DataLayerType;
+
+	UPROPERTY(Category = "Data Layer|Runtime", EditAnywhere, meta = (EditConditionHides, EditCondition = "DataLayerType == EDataLayerType::Runtime"))
+	FColor DebugColor;
+};
diff --git a/Engine/Source/Runtime/Engine/Public/WorldPartition/DataLayer/WorldDataLayers.h b/Engine/Source/Runtime/Engine/Public/WorldPartition/DataLayer/WorldDataLayers.h
--- a/Engine/Source/Runtime/Engine/Public/WorldPartition/DataLayer/WorldDataLayers.h
+++ b/Engine/Source/Runtime/Engine/Public/WorldPartition/DataLayer/WorldDataLayers.h
@@ -6,10 +6,15 @@
 #include "UObject/ObjectMacros.h"
 #include "GameFramework/Info.h"
 #include "WorldPartition/DataLayer/ActorDataLayer.h"
+#include "WorldPartition/DataLayer/DataLayerInstance.h"
 #include "WorldPartition/DataLayer/DataLayer.h"
+#include "Engine/World.h"
 #include "WorldDataLayers.generated.h"
 
-class UDataLayer;
+class UDEPRECATED_DataLayer;
+class UDataLayerInstance;
+class UDataLayerInstanceWithAsset;
+class UDataLayerAsset;
 
 USTRUCT()
 struct FActorPlacementDataLayers
@@ -17,11 +22,11 @@ struct FActorPlacementDataLayers
 	GENERATED_BODY()
 
 	UPROPERTY()
-	TArray<FName> DataLayers;
+	TArray<FName> DataLayerInstanceNames;
 
 	void Reset()
 	{
-		DataLayers.Reset();
+		DataLayerInstanceNames.Reset();
 	}
 };
 
@@ -33,6 +38,8 @@ class ENGINE_API AWorldDataLayers : public AInfo
 {
 	GENERATED_UCLASS_BODY()
 
+	friend class UDataLayerToAssetCommandlet;
+
 public:
 	virtual void PostLoad() override;
 	virtual void RewindForReplay() override;
@@ -45,48 +52,55 @@ public:
 	virtual bool IsUserManaged() const override { return false; }
 
 	static AWorldDataLayers* Create(UWorld* World);
-	UDataLayer* CreateDataLayer(FName InName = TEXT("DataLayer"), EObjectFlags InObjectFlags = RF_NoFlags);
-	bool RemoveDataLayer(UDataLayer* InDataLayer);
-	bool RemoveDataLayers(const TArray<UDataLayer*>& InDataLayers);
-	FName GenerateUniqueDataLayerLabel(const FName& InDataLayerLabel) const;
+
+	bool IsEmpty() const { return DataLayerInstances.IsEmpty(); }
+	bool HasDeprecatedDataLayers() const { return bHasDeprecatedDataLayers; }
+
+	template<class DataLayerInstanceType, typename ...CreationsArgs>
+	DataLayerInstanceType* CreateDataLayer(CreationsArgs... InCreationArgs);
+
+	bool RemoveDataLayer(const UDataLayerInstance* InDataLayer);
+	bool RemoveDataLayers(const TArray<UDataLayerInstance*>& InDataLayerInstances);
 	void SetAllowRuntimeDataLayerEditing(bool bInAllowRuntimeDataLayerEditing);
 	bool GetAllowRuntimeDataLayerEditing() const { return bAllowRuntimeDataLayerEditing; }
 
-	bool IsInActorEditorContext(const UDataLayer* InDataLayer) const;
-	bool AddToActorEditorContext(UDataLayer* InDataLayer);
-	bool RemoveFromActorEditorContext(UDataLayer* InDataLayer);
+	bool IsInActorEditorContext(const UDataLayerInstance* InDataLayerInstance) const;
+	bool AddToActorEditorContext(UDataLayerInstance* InDataLayerInstance);
+	bool RemoveFromActorEditorContext(UDataLayerInstance* InDataLayerInstance);
 	void PushActorEditorContext();
 	void PopActorEditorContext();
-	TArray<UDataLayer*> GetActorEditorContextDataLayers() const;
+	TArray<UDataLayerInstance*> GetActorEditorContextDataLayers() const;
 
 	//~ Begin Helper Functions
-	TArray<const UDataLayer*> GetDataLayerObjects(const TArray<FActorDataLayer>& DataLayers) const;
-	TArray<const UDataLayer*> GetDataLayerObjects(const TArray<FName>& InDataLayerNames) const;
-	TArray<FName> GetDataLayerNames(const TArray<FActorDataLayer>& DataLayers) const;
+	TArray<const UDataLayerInstance*> GetDataLayerInstances(const TArray<FName>& InDataLayerInstanceNames) const;
+	TArray<const UDataLayerInstance*> GetDataLayerInstances(const TArray<TObjectPtr<const UDataLayerAsset>>& InDataLayersAssets) const;
+	TArray<FName> GetDataLayerInstanceNames(const TArray<TObjectPtr<const UDataLayerAsset>>& InDataLayersAssets) const;
 	//~ End Helper Functions
 
 	// Allows overriding of DataLayers with PlayFromHere
-	void OverwriteDataLayerRuntimeStates(TArray<FActorDataLayer>* InActiveDataLayers = nullptr, TArray<FActorDataLayer>* InLoadedDataLayers = nullptr);
+	template<class T>
+	void OverwriteDataLayerRuntimeStates(const TArray<T>* InActiveDataLayers, const TArray<T>* InLoadedDataLayers );
 
 	// Returns the DataLayer user loaded editor states
 	void GetUserLoadedInEditorStates(TArray<FName>& OutDataLayersLoadedInEditor, TArray<FName>& OutDataLayersNotLoadedInEditor) const;
 #endif
 	
 	void DumpDataLayers(FOutputDevice& OutputDevice) const;
-	bool ContainsDataLayer(const UDataLayer* InDataLayer) const;
-	const UDataLayer* GetDataLayerFromName(const FName& InDataLayerName) const;
-	const UDataLayer* GetDataLayerFromLabel(const FName& InDataLayerLabel) const;
-	void ForEachDataLayer(TFunctionRef<bool(UDataLayer*)> Func);
-	void ForEachDataLayer(TFunctionRef<bool(UDataLayer*)> Func) const;
+	bool ContainsDataLayer(const UDataLayerInstance* InDataLayer) const;
+	const UDataLayerInstance* GetDataLayerInstance(const FName& InDataLayerInstanceName) const;
+	const UDataLayerInstance* GetDataLayerInstance(const UDataLayerAsset* InDataLayerAsset) const;
+	const UDataLayerInstance* GetDataLayerInstanceFromAssetName(const FName& InDataLayerAssetFullName) const;
+	void ForEachDataLayer(TFunctionRef<bool(UDataLayerInstance*)> Func);
+	void ForEachDataLayer(TFunctionRef<bool(UDataLayerInstance*)> Func) const;
 
 	// DataLayer Runtime State
-	void SetDataLayerRuntimeState(FActorDataLayer InDataLayer, EDataLayerRuntimeState InState, bool bIsRecursive = false);
+	void SetDataLayerRuntimeState(const UDataLayerInstance* InDataLayerInstance, EDataLayerRuntimeState InState, bool bIsRecursive = false);
 	EDataLayerRuntimeState GetDataLayerRuntimeStateByName(FName InDataLayerName) const;
 	EDataLayerRuntimeState GetDataLayerEffectiveRuntimeStateByName(FName InDataLAyerName) const;
 	const TSet<FName>& GetEffectiveActiveDataLayerNames() const { return EffectiveActiveDataLayerNames; }
 	const TSet<FName>& GetEffectiveLoadedDataLayerNames() const { return EffectiveLoadedDataLayerNames; }
 	UFUNCTION(NetMulticast, Reliable)
-	void OnDataLayerRuntimeStateChanged(const UDataLayer* InDataLayer, EDataLayerRuntimeState InState);
+	void OnDataLayerRuntimeStateChanged(const UDataLayerInstance* InDataLayer, EDataLayerRuntimeState InState);
 	static int32 GetDataLayersStateEpoch() { return DataLayersStateEpoch; }
 
 	//~ Begin Deprecated
@@ -94,18 +108,44 @@ public:
 	PRAGMA_DISABLE_DEPRECATION_WARNINGS // Suppress compiler warning on override of deprecated function
 
 	UE_DEPRECATED(5.0, "Use SetDataLayerRuntimeState() instead.")
-	void SetDataLayerState(FActorDataLayer InDataLayer, EDataLayerState InState) { SetDataLayerRuntimeState(InDataLayer, (EDataLayerRuntimeState)InState); }
+	void SetDataLayerState(FActorDataLayer InDataLayer, EDataLayerState InState) { SetDataLayerRuntimeState(GetDataLayerInstance(InDataLayer.Name), (EDataLayerRuntimeState)InState); }
 
 	UE_DEPRECATED(5.0, "Use GetDataLayerRuntimeStateByName() instead.")
 	EDataLayerState GetDataLayerStateByName(FName InDataLayerName) const { return (EDataLayerState)GetDataLayerRuntimeStateByName(InDataLayerName); }
 
-	PRAGMA_ENABLE_DEPRECATION_WARNINGS
-
 	UE_DEPRECATED(5.0, "Use GetEffectiveActiveDataLayerNames() instead.")
 	const TSet<FName>& GetActiveDataLayerNames() const { return GetEffectiveActiveDataLayerNames(); }
 
 	UE_DEPRECATED(5.0, "Use GetEffectiveLoadedDataLayerNames() instead.")
 	const TSet<FName>& GetLoadedDataLayerNames() const { return GetEffectiveLoadedDataLayerNames(); }
+
+#if WITH_EDITOR
+	UE_DEPRECATED(5.1, "Convert DataLayer using UDataLayerToAssetCommandlet and use UDataLayerInstance* overload instead.")
+	bool RemoveDataLayer(const UDEPRECATED_DataLayer* InDataLayer);
+#endif
+
+	UE_DEPRECATED(5.1, "Convert DataLayer using UDataLayerToAssetCommandlet and use UDataLayerInstance* overload instead.")
+	bool ContainsDataLayer(const UDEPRECATED_DataLayer* InDataLayer) const;
+
+	PRAGMA_ENABLE_DEPRECATION_WARNINGS
+
+#if WITH_EDITOR
+	UE_DEPRECATED(5.1, "Label usage is deprecated.")
+	FName GenerateUniqueDataLayerLabel(const FName& InDataLayerLabel) const;
+#endif 
+
+	UE_DEPRECATED(5.1, "Use GetDataLayerInstance with FName or UDataLayerAsset instead")
+	const UDataLayerInstance* GetDataLayerInstance(const FActorDataLayer& InActorDataLayer) const;
+
+	UE_DEPRECATED(5.1, "Use GetDataLayerInstanceNames with UDataLayerAsset instead")
+	TArray<FName> GetDataLayerInstanceNames(const TArray<FActorDataLayer>& InActorDataLayers) const;
+
+	UE_DEPRECATED(5.1, "Use GetDataLayerInstances with FName or UDataLayerAsset instead")
+	TArray<const UDataLayerInstance*> GetDataLayerInstances(const TArray<FActorDataLayer>& InActorDataLayers) const;
+
+	UE_DEPRECATED(5.1, "Convert UDataLayers to UDataLayerAsset and UDataLayerInstance using DataLayerToAssetCommandLet")
+	const UDataLayerInstance* GetDataLayerFromLabel(const FName& InDataLayerLabel) const;
+
 	//~ End Deprecated
 
 protected:
@@ -125,12 +165,17 @@ protected:
 	void OnRep_EffectiveLoadedDataLayerNames();
 
 private:
-	void ResolveEffectiveRuntimeState(const UDataLayer* InDataLayer, bool bInNotifyChange = true);
-	void DumpDataLayerRecursively(const UDataLayer* DataLayer, FString Prefix, FOutputDevice& OutputDevice) const;
+	void ResolveEffectiveRuntimeState(const UDataLayerInstance* InDataLayer, bool bInNotifyChange = true);
+	void DumpDataLayerRecursively(const UDataLayerInstance* DataLayer, FString Prefix, FOutputDevice& OutputDevice) const;
+
+#if WITH_EDITOR
+	void ConvertDataLayerToInstancces();
+	void UpdateContainsDeprecatedDataLayers();
+#endif
 
 #if !WITH_EDITOR
-	TMap<FName, const UDataLayer*> LabelToDataLayer;
-	TMap<FName, const UDataLayer*> NameToDataLayer;
+	TMap<FName, const UDataLayerInstance*> InstanceNameToInstance;
+	TMap<FString, const UDataLayerInstance*> AssetNameToInstance;
 #endif
 
 #if WITH_EDITORONLY_DATA
@@ -145,7 +190,16 @@ private:
 #endif
 
 	UPROPERTY()
-	TSet<TObjectPtr<UDataLayer>> WorldDataLayers;
+	TSet<TObjectPtr<UDataLayerInstance>> DataLayerInstances;
+
+	static_assert(DATALAYER_TO_INSTANCE_RUNTIME_CONVERSION_ENABLED, "DeprecatedDataLayerNameToDataLayerInstance Property is deprecated and needs to be deleted.");
+	UPROPERTY()
+	TMap<FName, TWeakObjectPtr<UDataLayerInstance>> DeprecatedDataLayerNameToDataLayerInstance;
+
+	PRAGMA_DISABLE_DEPRECATION_WARNINGS
+	UPROPERTY(meta = (DeprecatedProperty, DeprecationMessage = "Convert Data layers to DataLayerInstances and DataLayerAssets using DataLayerToAsset CommandLet and use DataLayerInstances instead."))
+	TSet<TObjectPtr<UDEPRECATED_DataLayer>> WorldDataLayers_DEPRECATED;
+	PRAGMA_ENABLE_DEPRECATION_WARNINGS
 
 	UPROPERTY(Transient, Replicated, ReplicatedUsing=OnRep_ActiveDataLayerNames)
 	TArray<FName> RepActiveDataLayerNames;
@@ -169,9 +223,82 @@ private:
 
 	static int32 DataLayersStateEpoch;
 
+	static_assert(DATALAYER_TO_INSTANCE_RUNTIME_CONVERSION_ENABLED, "bHasDeprecatedDataLayers is deprecated and needs to be deleted.");
+	bool bHasDeprecatedDataLayers;
+
 public:
 	DECLARE_DELEGATE_RetVal_ThreeParams(bool, FDataLayersFilterDelegate, FName /*DataLayerName*/, EDataLayerRuntimeState /*CurrentState*/, EDataLayerRuntimeState /*TargetState*/);
 
 	UE_DEPRECATED(5.00, "do not use, will be replaced by another mechanism for initial release.")
 	FDataLayersFilterDelegate DataLayersFilterDelegate;
-};
\ No newline at end of file
+};
+
+#if WITH_EDITOR
+
+template<class DataLayerInstanceType, typename ...CreationsArgs>
+DataLayerInstanceType* AWorldDataLayers::CreateDataLayer(CreationsArgs... InCreationArgs)
+{
+	Modify();
+
+	DataLayerInstanceType* NewDataLayer = NewObject<DataLayerInstanceType>(this, DataLayerInstanceType::MakeName(Forward<CreationsArgs>(InCreationArgs)...), RF_Transactional | RF_NoFlags);
+	check(NewDataLayer != nullptr);
+	NewDataLayer->OnCreated(Forward<CreationsArgs>(InCreationArgs)...);
+	DataLayerInstances.Add(NewDataLayer);
+
+	UpdateContainsDeprecatedDataLayers();
+
+	return NewDataLayer;
+}
+
+template<class IdentifierType>
+void AWorldDataLayers::OverwriteDataLayerRuntimeStates(const TArray<IdentifierType>* InActiveDataLayers, const TArray<IdentifierType>* InLoadedDataLayers)
+{
+	if (GetLocalRole() == ROLE_Authority)
+	{
+		// This should get called before game starts. It doesn't send out events
+		check(!GetWorld()->bMatchStarted);
+
+		if (InActiveDataLayers)
+		{
+			ActiveDataLayerNames.Empty(InActiveDataLayers->Num());
+			for (const IdentifierType& DataLayerIdentitier : *InActiveDataLayers)
+			{
+				const UDataLayerInstance* DataLayerInstance = GetDataLayerInstance(DataLayerIdentitier);
+				if (DataLayerInstance && DataLayerInstance->IsRuntime())
+				{
+					ActiveDataLayerNames.Add(DataLayerInstance->GetDataLayerFName());
+				}
+			}
+			RepActiveDataLayerNames = ActiveDataLayerNames.Array();
+		}
+		
+
+		if(InLoadedDataLayers)
+		{
+			LoadedDataLayerNames.Empty(InLoadedDataLayers->Num());
+			for (const IdentifierType& DataLayerIdentitier : *InLoadedDataLayers)
+			{
+				const UDataLayerInstance* DataLayerInstance = GetDataLayerInstance(DataLayerIdentitier);
+				if (DataLayerInstance && DataLayerInstance->IsRuntime())
+				{
+					LoadedDataLayerNames.Add(DataLayerInstance->GetDataLayerFName());
+				}
+			}
+			RepLoadedDataLayerNames = LoadedDataLayerNames.Array();
+		}
+
+		ForEachDataLayer([this](class UDataLayerInstance* DataLayer)
+		{
+			if (DataLayer && DataLayer->IsRuntime())
+			{
+				ResolveEffectiveRuntimeState(DataLayer, /*bNotifyChange*/false);
+			}
+			return true;
+		});
+
+		RepEffectiveActiveDataLayerNames = EffectiveActiveDataLayerNames.Array();
+		RepEffectiveLoadedDataLayerNames = EffectiveLoadedDataLayerNames.Array();
+	}
+}
+
+#endif // WITH_EDITOR
\ No newline at end of file
diff --git a/Engine/Source/Runtime/Engine/Public/WorldPartition/ErrorHandling/WorldPartitionStreamingGenerationErrorHandler.h b/Engine/Source/Runtime/Engine/Public/WorldPartition/ErrorHandling/WorldPartitionStreamingGenerationErrorHandler.h
--- a/Engine/Source/Runtime/Engine/Public/WorldPartition/ErrorHandling/WorldPartitionStreamingGenerationErrorHandler.h
+++ b/Engine/Source/Runtime/Engine/Public/WorldPartition/ErrorHandling/WorldPartitionStreamingGenerationErrorHandler.h
@@ -5,6 +5,9 @@
 #include "CoreMinimal.h"
 #include "WorldPartition/WorldPartitionActorDescView.h"
 
+class UDataLayerInstance;
+class UDataLayerInstanceWithAsset;
+
 class ENGINE_API IStreamingGenerationErrorHandler
 {
 public:
@@ -39,5 +42,20 @@ public:
 	 * Called when an actor descriptor references an actor using data layers.
 	 */
 	virtual void OnInvalidReferenceLevelScriptDataLayers(const FWorldPartitionActorDescView& ActorDescView) = 0;
+
+	/**
+	 * Called when a data layer instance does not have a data layer asset
+	 */
+	virtual void OnInvalidReferenceDataLayerAsset(const UDataLayerInstanceWithAsset* DataLayerInstance) = 0;
+
+	/**
+	 * Called when a data layer is not of the same type as its parent
+	 */
+	virtual void OnDataLayerHierarchyTypeMismatch(const UDataLayerInstance* DataLayerInstance, const UDataLayerInstance* Parent) = 0;
+
+	/**
+	 * Called when two data layer instances share the same asset
+	 */
+	virtual void OnDataLayerAssetConflict(const UDataLayerInstanceWithAsset* DataLayerInstance, const UDataLayerInstanceWithAsset* ConflictingDataLayerInstance) = 0;
 };
 #endif
diff --git a/Engine/Source/Runtime/Engine/Public/WorldPartition/ErrorHandling/WorldPartitionStreamingGenerationLogErrorHandler.h b/Engine/Source/Runtime/Engine/Public/WorldPartition/ErrorHandling/WorldPartitionStreamingGenerationLogErrorHandler.h
--- a/Engine/Source/Runtime/Engine/Public/WorldPartition/ErrorHandling/WorldPartitionStreamingGenerationLogErrorHandler.h
+++ b/Engine/Source/Runtime/Engine/Public/WorldPartition/ErrorHandling/WorldPartitionStreamingGenerationLogErrorHandler.h
@@ -13,5 +13,8 @@ public:
 	virtual void OnInvalidReferenceLevelScriptStreamed(const FWorldPartitionActorDescView& ActorDescView) override;
 	virtual void OnInvalidReferenceLevelScriptDataLayers(const FWorldPartitionActorDescView& ActorDescView) override;
 	virtual void OnInvalidReferenceRuntimeGrid(const FWorldPartitionActorDescView& ActorDescView, const FWorldPartitionActorDescView& ReferenceActorDescView) override;
+	virtual void OnInvalidReferenceDataLayerAsset(const UDataLayerInstanceWithAsset* DataLayerInstance) override;
+	virtual void OnDataLayerHierarchyTypeMismatch(const UDataLayerInstance* DataLayerInstance, const UDataLayerInstance* Parent) override;
+	virtual void OnDataLayerAssetConflict(const UDataLayerInstanceWithAsset* DataLayerInstance, const UDataLayerInstanceWithAsset* ConflictingDataLayerInstance) override;
 };
 #endif
diff --git a/Engine/Source/Runtime/Engine/Public/WorldPartition/ErrorHandling/WorldPartitionStreamingGenerationNullErrorHandler.h b/Engine/Source/Runtime/Engine/Public/WorldPartition/ErrorHandling/WorldPartitionStreamingGenerationNullErrorHandler.h
--- a/Engine/Source/Runtime/Engine/Public/WorldPartition/ErrorHandling/WorldPartitionStreamingGenerationNullErrorHandler.h
+++ b/Engine/Source/Runtime/Engine/Public/WorldPartition/ErrorHandling/WorldPartitionStreamingGenerationNullErrorHandler.h
@@ -13,5 +13,8 @@ public:
 	virtual void OnInvalidReferenceLevelScriptStreamed(const FWorldPartitionActorDescView& ActorDescView) override {}
 	virtual void OnInvalidReferenceLevelScriptDataLayers(const FWorldPartitionActorDescView& ActorDescView) override {}
 	virtual void OnInvalidReferenceRuntimeGrid(const FWorldPartitionActorDescView& ActorDescView, const FWorldPartitionActorDescView& ReferenceActorDescView) override {}
+	virtual void OnInvalidReferenceDataLayerAsset(const UDataLayerInstanceWithAsset* DataLayerInstance) override {}
+	virtual void OnDataLayerHierarchyTypeMismatch(const UDataLayerInstance* DataLayerInstance, const UDataLayerInstance* Parent) override {};
+	virtual void OnDataLayerAssetConflict(const UDataLayerInstanceWithAsset* DataLayerInstance, const UDataLayerInstanceWithAsset* ConflictingDataLayerInstance) override {};
 };
 #endif
diff --git a/Engine/Source/Runtime/Engine/Public/WorldPartition/ErrorHandling/WorldPartitionStreamingGenerationTokenizedMessageErrorHandler.h b/Engine/Source/Runtime/Engine/Public/WorldPartition/ErrorHandling/WorldPartitionStreamingGenerationTokenizedMessageErrorHandler.h
--- a/Engine/Source/Runtime/Engine/Public/WorldPartition/ErrorHandling/WorldPartitionStreamingGenerationTokenizedMessageErrorHandler.h
+++ b/Engine/Source/Runtime/Engine/Public/WorldPartition/ErrorHandling/WorldPartitionStreamingGenerationTokenizedMessageErrorHandler.h
@@ -14,6 +14,9 @@ class ENGINE_API ITokenizedMessageErrorHandler : public IStreamingGenerationErro
 	virtual void OnInvalidReferenceLevelScriptStreamed(const FWorldPartitionActorDescView& ActorDescView) override;
 	virtual void OnInvalidReferenceLevelScriptDataLayers(const FWorldPartitionActorDescView& ActorDescView) override;
 	virtual void OnInvalidReferenceRuntimeGrid(const FWorldPartitionActorDescView& ActorDescView, const FWorldPartitionActorDescView& ReferenceActorDescView) override;
+	virtual void OnInvalidReferenceDataLayerAsset(const UDataLayerInstanceWithAsset* DataLayerInstance) override;
+	virtual void OnDataLayerHierarchyTypeMismatch(const UDataLayerInstance* DataLayerInstance, const UDataLayerInstance* Parent) override;
+	virtual void OnDataLayerAssetConflict(const UDataLayerInstanceWithAsset* DataLayerInstance, const UDataLayerInstanceWithAsset* ConflictingDataLayerInstance) override;
 
 protected:
 	virtual void AddAdditionalNameToken(TSharedRef<FTokenizedMessage>& InMessage, const FName& InErrorName) {}
@@ -23,11 +26,12 @@ protected:
 
 class ENGINE_API FTokenizedMessageAccumulatorErrorHandler : public ITokenizedMessageErrorHandler
 {
+public:
+	const TArray<TSharedRef<FTokenizedMessage>>& GetErrorMessages() const { return ErrorMessages; }
+
 protected:
 	virtual void HandleTokenizedMessage(TSharedRef<FTokenizedMessage>&& ErrorMessage) override { ErrorMessages.Add(MoveTemp(ErrorMessage)); }
 
-	const TArray<TSharedRef<FTokenizedMessage>>& GetErrorMessages() const { return ErrorMessages; }
-
 private:
 	TArray<TSharedRef<FTokenizedMessage>> ErrorMessages;
 };
diff --git a/Engine/Source/Runtime/Engine/Public/WorldPartition/HLOD/IWorldPartitionHLODUtilities.h b/Engine/Source/Runtime/Engine/Public/WorldPartition/HLOD/IWorldPartitionHLODUtilities.h
--- a/Engine/Source/Runtime/Engine/Public/WorldPartition/HLOD/IWorldPartitionHLODUtilities.h
+++ b/Engine/Source/Runtime/Engine/Public/WorldPartition/HLOD/IWorldPartitionHLODUtilities.h
@@ -55,7 +55,7 @@ public:
 	 * @param	InActors			The actors for which we'll build an HLOD representation
 	 * @param	InDataLayers		The data layers to assign to the newly created HLOD actors
 	 */
-	virtual TArray<AWorldPartitionHLOD*> CreateHLODActors(FHLODCreationContext& InCreationContext, const FHLODCreationParams& InCreationParams, const TSet<FActorInstance>& InActors, const TArray<const UDataLayer*>& InDataLayers) = 0;
+	virtual TArray<AWorldPartitionHLOD*> CreateHLODActors(FHLODCreationContext& InCreationContext, const FHLODCreationParams& InCreationParams, const TSet<FActorInstance>& InActors, const TArray<const UDataLayerInstance*>& InDataLayerInstances) = 0;
 
 	/**
 	 * Build HLOD for the specified AWorldPartitionHLOD actor.
diff --git a/Engine/Source/Runtime/Engine/Public/WorldPartition/WorldPartitionActorCluster.h b/Engine/Source/Runtime/Engine/Public/WorldPartition/WorldPartitionActorCluster.h
--- a/Engine/Source/Runtime/Engine/Public/WorldPartition/WorldPartitionActorCluster.h
+++ b/Engine/Source/Runtime/Engine/Public/WorldPartition/WorldPartitionActorCluster.h
@@ -74,7 +74,7 @@ struct FActorCluster
 	bool						bIsSpatiallyLoaded;
 	FName						RuntimeGrid;
 	FBox						Bounds;
-	TSet<const UDataLayer*>		DataLayers;
+	TSet<const UDataLayerInstance*>		DataLayers;
 	FDataLayersID				DataLayersID;
 
 	FActorCluster(UWorld* InWorld, const FWorldPartitionActorDescView& InActorDescView);
@@ -94,7 +94,7 @@ struct FActorContainerInstance
 	EContainerClusterMode		ClusterMode;
 	const UActorDescContainer*	Container;
 	TMap<FGuid, FWorldPartitionActorDescView> ActorDescViewMap;
-	TSet<const UDataLayer*>		DataLayers;
+	TSet<const UDataLayerInstance*>		DataLayers;
 
 	const FWorldPartitionActorDescView& GetActorDescView(const FGuid& InGuid) const;
 };
@@ -109,7 +109,7 @@ struct FActorClusterInstance
 	FBox							Bounds;
 	const FActorCluster*			Cluster;
 	const FActorContainerInstance*	ContainerInstance;
-	TArray<const UDataLayer*>		DataLayers;
+	TArray<const UDataLayerInstance*>		DataLayers;
 };
 
 /**
diff --git a/Engine/Source/Runtime/Engine/Public/WorldPartition/WorldPartitionActorDesc.h b/Engine/Source/Runtime/Engine/Public/WorldPartition/WorldPartitionActorDesc.h
--- a/Engine/Source/Runtime/Engine/Public/WorldPartition/WorldPartitionActorDesc.h
+++ b/Engine/Source/Runtime/Engine/Public/WorldPartition/WorldPartitionActorDesc.h
@@ -52,7 +52,7 @@ public:
 	inline bool GetLevelBoundsRelevant() const { return bLevelBoundsRelevant; }
 	inline bool GetActorIsHLODRelevant() const { return bActorIsHLODRelevant; }
 	inline FName GetHLODLayer() const { return HLODLayer; }
-	inline const TArray<FName>& GetDataLayers() const { return DataLayers; }
+	inline const TArray<FName>& GetDataLayerInstanceNames() const { return DataLayers; }
 	inline FName GetActorPackage() const { return ActorPackage; }
 	inline FName GetActorPath() const { return ActorPath; }
 	inline FName GetActorLabel() const { return ActorLabel; }
diff --git a/Engine/Source/Runtime/Engine/Public/WorldPartition/WorldPartitionRuntimeCell.h b/Engine/Source/Runtime/Engine/Public/WorldPartition/WorldPartitionRuntimeCell.h
--- a/Engine/Source/Runtime/Engine/Public/WorldPartition/WorldPartitionRuntimeCell.h
+++ b/Engine/Source/Runtime/Engine/Public/WorldPartition/WorldPartitionRuntimeCell.h
@@ -3,7 +3,7 @@
 
 #include "CoreMinimal.h"
 #include "Templates/SubclassOf.h"
-#include "WorldPartition/DataLayer/DataLayer.h"
+#include "WorldPartition/DataLayer/DataLayerInstance.h"
 #include "WorldPartition/WorldPartitionActorDescView.h"
 #include "WorldPartition/WorldPartitionStreamingSource.h"
 #include "WorldPartition/WorldPartitionActorCluster.h"
@@ -159,7 +159,7 @@ class UWorldPartitionRuntimeCell : public UObject
 	void SetClientOnlyVisible(bool bInClientOnlyVisible) { bClientOnlyVisible = bInClientOnlyVisible; }
 	bool GetClientOnlyVisible() const { return bClientOnlyVisible; }
 
-	void SetDataLayers(const TArray<const UDataLayer*>& InDataLayers);
+	void SetDataLayers(const TArray<const UDataLayerInstance*>& InDataLayerInstances);
 	void SetDebugInfo(FIntVector InCoords, FName InGridName);
 	virtual void AddActorToCell(const FWorldPartitionActorDescView& ActorDescView, const FActorContainerID& InContainerID, const FTransform& InContainerTransform, const UActorDescContainer* InContainer) PURE_VIRTUAL(UWorldPartitionRuntimeCell::AddActorToCell,);
 	virtual int32 GetActorCount() const PURE_VIRTUAL(UWorldPartitionRuntimeCell::GetActorCount, return 0;);
diff --git a/Engine/Source/Runtime/Foliage/Private/ProceduralFoliageComponent.cpp b/Engine/Source/Runtime/Foliage/Private/ProceduralFoliageComponent.cpp
--- a/Engine/Source/Runtime/Foliage/Private/ProceduralFoliageComponent.cpp
+++ b/Engine/Source/Runtime/Foliage/Private/ProceduralFoliageComponent.cpp
@@ -15,6 +15,7 @@
 #if WITH_EDITOR
 #include "WorldPartition/WorldPartition.h"
 #include "WorldPartition/WorldPartitionSubsystem.h"
+#include "WorldPartition/DataLayer/DataLayerAsset.h"
 #endif
 
 
diff --git a/Engine/Source/Runtime/Foliage/Public/ProceduralFoliageComponent.h b/Engine/Source/Runtime/Foliage/Public/ProceduralFoliageComponent.h
--- a/Engine/Source/Runtime/Foliage/Public/ProceduralFoliageComponent.h
+++ b/Engine/Source/Runtime/Foliage/Public/ProceduralFoliageComponent.h
@@ -11,6 +11,7 @@
 
 class AVolume;
 class UProceduralFoliageSpawner;
+class UDataLayerAsset;
 struct FBodyInstance;
 
 /** Describes the layout of the tiles used for procedural foliage simulation */
diff --git a/Engine/Source/Runtime/MovieSceneTracks/Private/Systems/MovieSceneDataLayerSystem.cpp b/Engine/Source/Runtime/MovieSceneTracks/Private/Systems/MovieSceneDataLayerSystem.cpp
--- a/Engine/Source/Runtime/MovieSceneTracks/Private/Systems/MovieSceneDataLayerSystem.cpp
+++ b/Engine/Source/Runtime/MovieSceneTracks/Private/Systems/MovieSceneDataLayerSystem.cpp
@@ -13,6 +13,7 @@
 #include "ProfilingDebugging/CsvProfiler.h"
 #include "Sections/MovieSceneDataLayerSection.h"
 #include "WorldPartition/DataLayer/DataLayerSubsystem.h"
+#include "WorldPartition/DataLayer/DataLayerInstance.h"
 #include "WorldPartition/WorldPartitionSubsystem.h"
 #include "WorldPartition/WorldPartitionRuntimeCell.h"
 #include "WorldPartition/WorldPartitionStreamingSource.h"
@@ -43,11 +44,11 @@ ENUM_CLASS_FLAGS(EDataLayerUpdateFlags)
 /** Traits class governing how pre-animated state is (re)stored for data layers */
 struct FPreAnimatedDataLayerStorageTraits
 {
-	using KeyType = TObjectKey<UDataLayer>;
+	using KeyType = TObjectKey<UDataLayerInstance>;
 	using StorageType = EDataLayerRuntimeState;
 
 	/** Called when a previously animated data layer needs to be restored */
-	static void RestorePreAnimatedValue(const TObjectKey<UDataLayer>& InKey, EDataLayerRuntimeState PreviousState, const FRestoreStateParams& Params);
+	static void RestorePreAnimatedValue(const TObjectKey<UDataLayerInstance>& InKey, EDataLayerRuntimeState PreviousState, const FRestoreStateParams& Params);
 };
 
 /** Container class for all pre-animated data layer state */
@@ -63,14 +64,14 @@ struct FPreAnimatedDataLayerStorage
 	void OnGroupDestroyed(FPreAnimatedStorageGroupHandle Group) override;
 
 	/** Make an entry for the specified data layer */
-	FPreAnimatedStateEntry MakeEntry(UDataLayer* InDataLayer);
+	FPreAnimatedStateEntry MakeEntry(UDataLayerInstance* InDataLayer);
 
 	/** Save the value of a data layer. Should only be used for runtime / PIE worlds */
-	void SavePreAnimatedState(UDataLayer* DataLayer, UDataLayerSubsystem* SubSystem);
+	void SavePreAnimatedState(UDataLayerInstance* DataLayer, UDataLayerSubsystem* SubSystem);
 
 #if WITH_EDITOR
 	/** Save the value of a data layer. Should only be used for editor worlds */
-	void SavePreAnimatedStateInEditor(UDataLayer* DataLayer);
+	void SavePreAnimatedStateInEditor(UDataLayerInstance* DataLayer);
 #endif
 
 private:
@@ -113,9 +114,9 @@ struct FDesiredLayerStates
 
 // ---------------------------------------------------------------------
 // FPreAnimatedDataLayerStorageTraits definitions
-void FPreAnimatedDataLayerStorageTraits::RestorePreAnimatedValue(const TObjectKey<UDataLayer>& InKey, EDataLayerRuntimeState PreviousState, const FRestoreStateParams& Params)
+void FPreAnimatedDataLayerStorageTraits::RestorePreAnimatedValue(const TObjectKey<UDataLayerInstance>& InKey, EDataLayerRuntimeState PreviousState, const FRestoreStateParams& Params)
 {
-	UDataLayer* DataLayer = InKey.ResolveObjectPtr();
+	UDataLayerInstance* DataLayer = InKey.ResolveObjectPtr();
 	if (!DataLayer)
 	{
 		return;
@@ -153,7 +154,7 @@ void FPreAnimatedDataLayerStorage::OnGroupDestroyed(FPreAnimatedStorageGroupHand
 	GroupHandle = FPreAnimatedStorageGroupHandle();
 }
 
-FPreAnimatedStateEntry FPreAnimatedDataLayerStorage::MakeEntry(UDataLayer* InDataLayer)
+FPreAnimatedStateEntry FPreAnimatedDataLayerStorage::MakeEntry(UDataLayerInstance* InDataLayer)
 {
 	if (!GroupHandle)
 	{
@@ -163,7 +164,7 @@ FPreAnimatedStateEntry FPreAnimatedDataLayerStorage::MakeEntry(UDataLayer* InDat
 	return FPreAnimatedStateEntry{ GroupHandle, FPreAnimatedStateCachedValueHandle{ StorageID, StorageIndex } };
 }
 
-void FPreAnimatedDataLayerStorage::SavePreAnimatedState(UDataLayer* DataLayer, UDataLayerSubsystem* SubSystem)
+void FPreAnimatedDataLayerStorage::SavePreAnimatedState(UDataLayerInstance* DataLayer, UDataLayerSubsystem* SubSystem)
 {
 	FPreAnimatedStateEntry         Entry              = MakeEntry(DataLayer);
 	EPreAnimatedStorageRequirement StorageRequirement = this->ParentExtension->GetStorageRequirement(Entry);
@@ -179,7 +180,7 @@ void FPreAnimatedDataLayerStorage::SavePreAnimatedState(UDataLayer* DataLayer, U
 }
 
 #if WITH_EDITOR
-void FPreAnimatedDataLayerStorage::SavePreAnimatedStateInEditor(UDataLayer* DataLayer)
+void FPreAnimatedDataLayerStorage::SavePreAnimatedStateInEditor(UDataLayerInstance* DataLayer)
 {
 	FPreAnimatedStateEntry         Entry              = MakeEntry(DataLayer);
 	EPreAnimatedStorageRequirement StorageRequirement = this->ParentExtension->GetStorageRequirement(Entry);
@@ -278,7 +279,7 @@ EDataLayerUpdateFlags FDesiredLayerStates::Apply(FPreAnimatedDataLayerStorage* P
 {
 	EDataLayerUpdateFlags Flags = EDataLayerUpdateFlags::None;
 
-	auto IsDataLayerReady = [WorldPartitionSubsystem](UDataLayer* DataLayer, EDataLayerRuntimeState DesireState, bool bExactState)
+	auto IsDataLayerReady = [WorldPartitionSubsystem](UDataLayerInstance* DataLayer, EDataLayerRuntimeState DesireState, bool bExactState)
 	{
 		EWorldPartitionRuntimeCellState QueryState;
 		switch (DesireState)
@@ -300,7 +301,7 @@ EDataLayerUpdateFlags FDesiredLayerStates::Apply(FPreAnimatedDataLayerStorage* P
 		FWorldPartitionStreamingQuerySource QuerySource;
 		QuerySource.bDataLayersOnly = true;
 		QuerySource.bSpatialQuery = false; // @todo_ow: how would we support spatial query from sequencer?
-		QuerySource.DataLayers.Add(DataLayer->GetFName());
+		QuerySource.DataLayers.Add(DataLayer->GetDataLayerFName());
 		return WorldPartitionSubsystem->IsStreamingCompleted(QueryState, { QuerySource }, bExactState);
 	};
 
@@ -315,7 +316,7 @@ EDataLayerUpdateFlags FDesiredLayerStates::Apply(FPreAnimatedDataLayerStorage* P
 
 		if (TOptional<EDataLayerRuntimeState> DesiredState = StateValue.ComputeDesiredState())
 		{
-			UDataLayer* DataLayer = DataLayerSubsystem->GetDataLayerFromName(It.Key());
+			UDataLayerInstance* DataLayer = DataLayerSubsystem->GetDataLayerInstance(It.Key());
 			if (DataLayer)
 			{
 				if (PreAnimatedStorage)
@@ -333,22 +334,22 @@ EDataLayerUpdateFlags FDesiredLayerStates::Apply(FPreAnimatedDataLayerStorage* P
 					{
 						Flags |= EDataLayerUpdateFlags::FlushStreamingFull;
 						UE_LOG(LogMovieScene, Warning, TEXT("[UMovieSceneDataLayerSystem] Data layer with name '%s' is causing a full streaming flush (%s)"),
-							*DataLayer->GetDataLayerLabel().ToString(),
+							*DataLayer->GetDataLayerShortName(),
 							*StaticEnum<EDataLayerRuntimeState>()->GetDisplayNameTextByValue((int64)DesiredStateValue).ToString());
 							
 						CSV_EVENT_GLOBAL(TEXT("SeqDataLayerFlushFull-%s-%s"),
-							*DataLayer->GetDataLayerLabel().ToString(),
+							*DataLayer->GetDataLayerShortName(),
 							*StaticEnum<EDataLayerRuntimeState>()->GetDisplayNameTextByValue((int64)DesiredStateValue).ToString());
 					}
 					else
 					{
 						Flags |= EDataLayerUpdateFlags::FlushStreamingVisibility;
 						UE_LOG(LogMovieScene, Log, TEXT("[UMovieSceneDataLayerSystem] Data layer with name '%s' is causing a visibility streaming flush (%s)"),
-							*DataLayer->GetDataLayerLabel().ToString(),
+							*DataLayer->GetDataLayerShortName(),
 							*StaticEnum<EDataLayerRuntimeState>()->GetDisplayNameTextByValue((int64)DesiredStateValue).ToString());
 							
 						CSV_EVENT_GLOBAL(TEXT("SeqDataLayerFlushVis-%s-%s"),
-							*DataLayer->GetDataLayerLabel().ToString(),
+							*DataLayer->GetDataLayerShortName(),
 							*StaticEnum<EDataLayerRuntimeState>()->GetDisplayNameTextByValue((int64)DesiredStateValue).ToString());
 					}
 				}
@@ -366,9 +367,9 @@ EDataLayerUpdateFlags FDesiredLayerStates::Apply(FPreAnimatedDataLayerStorage* P
 #if WITH_EDITOR
 void FDesiredLayerStates::ApplyInEditor(FPreAnimatedDataLayerStorage* PreAnimatedStorage, UDataLayerEditorSubsystem* SubSystem)
 {
-	TArray<UDataLayer*> DatalayersNeedingLoad;
-	TArray<UDataLayer*> DatalayersNeedingShow;
-	TArray<UDataLayer*> DatalayersNeedingHide;
+	TArray<UDataLayerInstance*> DatalayersNeedingLoad;
+	TArray<UDataLayerInstance*> DatalayersNeedingShow;
+	TArray<UDataLayerInstance*> DatalayersNeedingHide;
 
 	for (auto It = StatesByLayer.CreateIterator(); It; ++It)
 	{
@@ -381,7 +382,7 @@ void FDesiredLayerStates::ApplyInEditor(FPreAnimatedDataLayerStorage* PreAnimate
 
 		if (TOptional<EDataLayerRuntimeState> DesiredState = StateValue.ComputeDesiredState())
 		{
-			UDataLayer* DataLayer = SubSystem->GetDataLayerFromName(It.Key());
+			UDataLayerInstance* DataLayer = SubSystem->GetDataLayerInstance(It.Key());
 			if (DataLayer)
 			{
 				if (PreAnimatedStorage)
@@ -589,6 +590,13 @@ void UMovieSceneDataLayerSystem::UpdateDesiredStates()
 
 	auto GatherDataLayers = [this, BuiltInComponents](FEntityAllocationIteratorItem Item, const FMovieSceneDataLayerComponentData* ComponentData, const int16* OptHBiases)
 	{
+		UWorld* World = GetWorld();
+		UDataLayerSubsystem* DataLayerSubsystem = World ? World->GetSubsystem<UDataLayerSubsystem>() : nullptr;
+		if (!DataLayerSubsystem)
+		{
+			return;
+		}
+
 		const bool bPreroll = Item.GetAllocationType().Contains(BuiltInComponents->Tags.PreRoll);
 		for (int32 Index = 0; Index < Item.GetAllocation()->Num(); ++Index)
 		{
@@ -601,9 +609,13 @@ void UMovieSceneDataLayerSystem::UpdateDesiredStates()
 			EDataLayerRuntimeState DesiredState = bPreroll ? Section->GetPrerollState() : Section->GetDesiredState();
 			const bool bRequiresStreamingFlush = (DesiredState == EDataLayerRuntimeState::Unloaded) ? Section->GetFlushOnUnload() : !bPreroll;
 
-			for (const FActorDataLayer& ActorDataLayer : Section->GetDataLayers())
+			for (const UDataLayerAsset* DataLayerAsset : Section->GetDataLayerAssets())
 			{
-				this->DesiredLayerStates->ApplyNewState(ActorDataLayer.Name, OptHBiases ? OptHBiases[Index] : 0, DesiredState, bRequiresStreamingFlush);
+				const UDataLayerInstance* DataLayerInstance = DataLayerSubsystem->GetDataLayerFromAsset(DataLayerAsset);
+				if (DataLayerInstance)
+				{
+					this->DesiredLayerStates->ApplyNewState(DataLayerInstance->GetDataLayerFName(), OptHBiases ? OptHBiases[Index] : 0, DesiredState, bRequiresStreamingFlush);
+				}
 			}
 		}
 	};
@@ -653,12 +665,12 @@ void UMovieSceneDataLayerSystem::BeginTrackingEntities()
 				FMovieSceneEntityID EntityID     = EntityIDs[Index];
 				FRootInstanceHandle RootInstance = RootInstanceHandles[Index];
 
-				for (const FActorDataLayer& ActorDataLayer : Section->GetDataLayers())
+				for (const UDataLayerAsset* DataLayerAsset : Section->GetDataLayerAssets())
 				{
-					UDataLayer* DataLayer = DataLayerSubsystem->GetDataLayer(ActorDataLayer);
-					if (DataLayer)
+					UDataLayerInstance* DataLayerInstance = DataLayerSubsystem->GetDataLayerFromAsset(DataLayerAsset);
+					if (DataLayerInstance)
 					{
-						FPreAnimatedStateEntry Entry = PreAnimatedStorage->MakeEntry(DataLayer);
+						FPreAnimatedStateEntry Entry = PreAnimatedStorage->MakeEntry(DataLayerInstance);
 						EntityMetaData->BeginTrackingEntity(Entry, EntityID, RootInstance, bRestoreState);
 					}
 				}
diff --git a/Engine/Source/Runtime/MovieSceneTracks/Private/Systems/MovieSceneDataLayerSystem.h b/Engine/Source/Runtime/MovieSceneTracks/Private/Systems/MovieSceneDataLayerSystem.h
--- a/Engine/Source/Runtime/MovieSceneTracks/Private/Systems/MovieSceneDataLayerSystem.h
+++ b/Engine/Source/Runtime/MovieSceneTracks/Private/Systems/MovieSceneDataLayerSystem.h
@@ -5,7 +5,6 @@
 #include "Templates/SharedPointer.h"
 #include "UObject/ObjectKey.h"
 #include "Containers/Map.h"
-#include "WorldPartition/DataLayer/DataLayer.h"
 #include "EntitySystem/MovieSceneEntitySystem.h"
 #include "EntitySystem/MovieSceneCachedEntityFilterResult.h"
 #include "MovieSceneDataLayerSystem.generated.h"
diff --git a/Engine/Source/Runtime/MovieSceneTracks/Public/Sections/MovieSceneDataLayerSection.h b/Engine/Source/Runtime/MovieSceneTracks/Public/Sections/MovieSceneDataLayerSection.h
--- a/Engine/Source/Runtime/MovieSceneTracks/Public/Sections/MovieSceneDataLayerSection.h
+++ b/Engine/Source/Runtime/MovieSceneTracks/Public/Sections/MovieSceneDataLayerSection.h
@@ -5,6 +5,7 @@
 #include "CoreTypes.h"
 #include "UObject/ObjectMacros.h"
 #include "WorldPartition/DataLayer/ActorDataLayer.h"
+#include "WorldPartition/DataLayer/DataLayerAsset.h"
 #include "EntitySystem/IMovieSceneEntityProvider.h"
 #include "MovieSceneSection.h"
 #include "MovieSceneDataLayerSection.generated.h"
@@ -39,21 +40,31 @@ public:
 	MOVIESCENETRACKS_API void SetFlushOnUnload(bool bFlushOnUnload);
 
 	UFUNCTION(BlueprintPure, Category = "Sequencer|Section")
-	const TArray<FActorDataLayer>& GetDataLayers() const { return DataLayers; }
+	const TArray<UDataLayerAsset*>& GetDataLayerAssets() const { return DataLayerAssets; }
 
 	UFUNCTION(BlueprintCallable, Category = "Sequencer|Section")
-	void SetDataLayers(const TArray<FActorDataLayer>& InDataLayers) { DataLayers = InDataLayers; }
+	void SetDataLayerAssets(const TArray<UDataLayerAsset*>& InDataLayerAssets) { DataLayerAssets = InDataLayerAssets; }
 
 private:
 
 	virtual void ImportEntityImpl(UMovieSceneEntitySystemLinker* EntityLinker, const FEntityImportParams& Params, FImportedEntity* OutImportedEntity) override;
 
 private:
+	UE_DEPRECATED(5.1, "Use GetDataLayerAssets instead")
+	UFUNCTION(BlueprintPure, Category = "Sequencer|Section")
+	const TArray<FActorDataLayer>& GetDataLayers() const { return DataLayers; }
 
-	/** A list of data layers that should be loaded or unloaded by this section */
-	UPROPERTY(EditAnywhere, Category=DataLayer)
+	UE_DEPRECATED(5.1, "Use SetDataLayerAssets instead")
+	UFUNCTION(BlueprintCallable, Category = "Sequencer|Section")
+	void SetDataLayers(const TArray<FActorDataLayer>& InDataLayers) { DataLayers = InDataLayers; }
+
+	UPROPERTY(meta = (DeprecatedProperty, DeprecationMessage = "Use DataLayer Assets Instead"))
 	TArray<FActorDataLayer> DataLayers;
 
+	/** A list of data layers that should be loaded or unloaded by this section */
+	UPROPERTY(EditAnywhere, Category = DataLayer)
+	TArray<TObjectPtr<UDataLayerAsset>> DataLayerAssets;
+
 	/** The desired state for the data layers on this section when the section is actively evaluating. */
 	UPROPERTY(EditAnywhere, Category=DataLayer)
 	EDataLayerRuntimeState DesiredState;
