Limit DMFP range as before for SSProfile due to encoding precision.

diff --git a/Engine/Source/Runtime/Engine/Private/Rendering/SubsurfaceProfile.cpp b/Engine/Source/Runtime/Engine/Private/Rendering/SubsurfaceProfile.cpp
--- a/Engine/Source/Runtime/Engine/Private/Rendering/SubsurfaceProfile.cpp
+++ b/Engine/Source/Runtime/Engine/Private/Rendering/SubsurfaceProfile.cpp
@@ -340,8 +340,10 @@ static float GetNextSmallerPositiveFloat(float x)
 #define ENC_WORLDUNITSCALE_IN_CM_TO_UNIT 0.02f
 #define DEC_UNIT_TO_WORLDUNITSCALE_IN_CM 1/ENC_WORLDUNITSCALE_IN_CM_TO_UNIT
 
-// Make sure UIMax|ClampMax of DiffuseMeanFreePath * 10(cm to mm) * ENC_DIFFUSEMEANFREEPATH_IN_MM_TO_UNIT <= 1
-//
+// Make sure DiffuseMeanFreePath * 10(cm to mm) * ENC_DIFFUSEMEANFREEPATH_IN_MM_TO_UNIT <= 1
+// Although UI switches to mean free path, the diffuse mean free path is maintained to have the 
+// same range as before [0, 50] cm.
+// 1 mfp can map to 1.44 dmfp (surface albedo -> 0.0) or 43.50 dmfp (surface albedo -> 1.0).
 #define ENC_DIFFUSEMEANFREEPATH_IN_MM_TO_UNIT (0.01f*0.2f)
 #define DEC_UNIT_TO_DIFFUSEMEANFREEPATH_IN_MM 1/ENC_DIFFUSEMEANFREEPATH_IN_MM_TO_UNIT
 //------------------------------------------------------------------------------------------
@@ -360,7 +362,7 @@ float DecodeWorldUnitScale(float EncodedWorldUnitScale)
 //in [0,1]
 FLinearColor EncodeDiffuseMeanFreePath(FLinearColor DiffuseMeanFreePath)
 {
-	return DiffuseMeanFreePath * ENC_DIFFUSEMEANFREEPATH_IN_MM_TO_UNIT;
+	return (DiffuseMeanFreePath * ENC_DIFFUSEMEANFREEPATH_IN_MM_TO_UNIT).GetClamped();
 }
 
 FLinearColor DecodeDiffuseMeanFreePath(FLinearColor EncodedDiffuseMeanFreePath)
