Shaders: Converted GlobalShaderMap and MaterialShaderMap to the DDC Value API

#rnx

diff --git a/Engine/Source/Runtime/Engine/Private/Materials/MaterialShader.cpp b/Engine/Source/Runtime/Engine/Private/Materials/MaterialShader.cpp
--- a/Engine/Source/Runtime/Engine/Private/Materials/MaterialShader.cpp
+++ b/Engine/Source/Runtime/Engine/Private/Materials/MaterialShader.cpp
@@ -1,6 +1,9 @@
 // Copyright Epic Games, Inc. All Rights Reserved.
 
 #include "MaterialShader.h"
+#include "Containers/StringConv.h"
+#include "DerivedDataCache.h"
+#include "DerivedDataRequestOwner.h"
 #include "Stats/StatsMisc.h"
 #include "Serialization/MemoryWriter.h"
 #include "Serialization/MemoryReader.h"
@@ -10,7 +13,6 @@
 #include "MeshMaterialShaderType.h"
 #include "MaterialShaderMapLayout.h"
 #include "ShaderCompiler.h"
-#include "DerivedDataCacheInterface.h"
 #include "Interfaces/ITargetPlatformManagerModule.h"
 #include "ProfilingDebugging/CookStats.h"
 #include "UObject/ReleaseObjectVersion.h"
@@ -185,7 +187,13 @@ static FString GetMaterialShaderMapKeyString(const FMaterialShaderMapId& ShaderM
 	FMaterialAttributeDefinitionMap::AppendDDCKeyString(ShaderMapKeyString);
 	FShaderCompileUtilities::AppendGBufferDDCKeyString(Platform, ShaderMapKeyString);
 
-	return FDerivedDataCacheInterface::BuildCacheKey(TEXT("MATSM"), *GetMaterialShaderMapDDCKey(), *ShaderMapKeyString);
+	return FString::Printf(TEXT("%s_%s_%s"), TEXT("MATSM"), *GetMaterialShaderMapDDCKey(), *ShaderMapKeyString);
+}
+
+static UE::DerivedData::FCacheKey GetMaterialShaderMapKey(const FStringView MaterialShaderMapKey)
+{
+	static UE::DerivedData::FCacheBucket Bucket("MaterialShaderMap");
+	return {Bucket, FIoHash::HashBuffer(MakeMemoryView(FTCHARToUTF8(MaterialShaderMapKey)))};
 }
 
 static FString FSHA1_HashString(const FString& Key)
@@ -1342,7 +1350,6 @@ void FMaterialShaderMap::LoadFromDerivedDataCache(const FMaterial* Material, con
 			SCOPE_SECONDS_COUNTER(MaterialDDCTime);
 			COOK_STAT(auto Timer = MaterialShaderCookStats::UsageStats.TimeSyncWork());
 
-			TArray<uint8> CachedData;
 			const FString DataKey = GetMaterialShaderMapKeyString(ShaderMapId, InPlatform);
 			OutDDCKeyDesc = FSHA1_HashString(DataKey);
 
@@ -1381,13 +1388,28 @@ void FMaterialShaderMap::LoadFromDerivedDataCache(const FMaterial* Material, con
 			// Do not check the DD cache if the material isn't persistent, because
 			//   - this results in a lot of DDC requests when editing materials which are almost always going to return nothing.
 			//   - since the get call is synchronous, this can cause a hitch if there's network latency
-			if (CheckCache && Material->IsPersistent() && GetDerivedDataCacheRef().GetSynchronous(*DataKey, CachedData, Material->GetFriendlyName()))
+			FSharedBuffer CachedData;
+			if (CheckCache && Material->IsPersistent())
+			{
+				using namespace UE::DerivedData;
+				FCacheGetValueRequest Request;
+				Request.Name = Material->GetFriendlyName();
+				Request.Key = GetMaterialShaderMapKey(DataKey);
+				FRequestOwner BlockingOwner(EPriority::Blocking);
+				GetCache().GetValue({Request}, BlockingOwner, [&CachedData](FCacheGetValueResponse&& Response)
+				{
+					CachedData = Response.Value.GetData().Decompress();
+				});
+				BlockingOwner.Wait();
+			}
+
+			if (CachedData)
 			{
 				TRACE_COUNTER_INCREMENT(Shaders_FMaterialShaderMapDDCHits);
-				TRACE_COUNTER_ADD(Shaders_FMaterialShaderMapDDCBytesReceived, CachedData.Num());
-				COOK_STAT(Timer.AddHit(CachedData.Num()));
+				TRACE_COUNTER_ADD(Shaders_FMaterialShaderMapDDCBytesReceived, int64(CachedData.GetSize()));
+				COOK_STAT(Timer.AddHit(int64(CachedData.GetSize())));
 				InOutShaderMap = new FMaterialShaderMap();
-				FMemoryReader Ar(CachedData, true);
+				FMemoryReaderView Ar(CachedData, /*bIsPersistent*/ true);
 
 				// Deserialize from the cached data
 				InOutShaderMap->Serialize(Ar);
@@ -1429,21 +1451,21 @@ void FMaterialShaderMap::SaveToDerivedDataCache()
 	FMemoryWriter64 Ar(SaveData, true);
 	Serialize(Ar);
 
-	// Refuse to save too large derived data, as it won't get handled by the cache anyway. It appends some metadata, so for safety we use a lower threshold.
-	const int64 kMaxSaveDataForDDC = MAX_int32 - (4 * 1024 * 1024);
-	if (SaveData.Num() >= kMaxSaveDataForDDC)
-	{
-		UE_LOG(LogMaterial, Warning, TEXT("Not saving %s shaders to DDC as its size (%lld MB) is larger than %lld MB, which may cause problems for the legacy DDC API."), 
-			GetMaterialPath(), SaveData.Num() / (1024 * 1024), kMaxSaveDataForDDC / (1024 * 1024));
-		return;
-	}
-
 	TRACE_COUNTER_ADD(Shaders_FMaterialShaderMapDDCBytesSent, SaveData.Num());
+	COOK_STAT(Timer.AddMiss(SaveData.Num()));
 
-	FString DataKey = GetMaterialShaderMapKeyString(ShaderMapId, GetShaderPlatform());
+	const FString DataKey = GetMaterialShaderMapKeyString(ShaderMapId, GetShaderPlatform());
 	UE_LOG(LogMaterial, Verbose, TEXT("Saved shaders for %s to DDC (key hash: %s)"), GetMaterialPath(), *FSHA1_HashString(DataKey));
-	GetDerivedDataCacheRef().Put(*DataKey, SaveData, FStringView(GetFriendlyName()));
-	COOK_STAT(Timer.AddMiss(SaveData.Num()));
+
+	using namespace UE::DerivedData;
+	FCachePutValueRequest Request;
+	Request.Name = GetFriendlyName();
+	Request.Key = GetMaterialShaderMapKey(DataKey);
+	Request.Value = FValue::Compress(MakeSharedBufferFromArray(MoveTemp(SaveData)));
+	FRequestOwner AsyncOwner(EPriority::Normal);
+	FRequestBarrier AsyncBarrier(AsyncOwner);
+	GetCache().PutValue({Request}, AsyncOwner);
+	AsyncOwner.KeepAlive();
 }
 #endif // WITH_EDITOR
 
diff --git a/Engine/Source/Runtime/Engine/Private/ShaderCompiler/ShaderCompiler.cpp b/Engine/Source/Runtime/Engine/Private/ShaderCompiler/ShaderCompiler.cpp
--- a/Engine/Source/Runtime/Engine/Private/ShaderCompiler/ShaderCompiler.cpp
+++ b/Engine/Source/Runtime/Engine/Private/ShaderCompiler/ShaderCompiler.cpp
@@ -5,6 +5,7 @@
 =============================================================================*/
 
 #include "ShaderCompiler.h"
+#include "Containers/StringConv.h"
 #include "GenericPlatform/GenericPlatformFile.h"
 #include "HAL/PlatformFileManager.h"
 #include "Misc/MessageDialog.h"
@@ -38,7 +39,8 @@
 #include "EngineModule.h"
 #include "ComponentRecreateRenderStateContext.h"
 #include "Components/PrimitiveComponent.h"
-#include "DerivedDataCacheInterface.h"
+#include "DerivedDataCache.h"
+#include "DerivedDataRequestOwner.h"
 #include "UObject/DevObjectVersion.h"
 #include "Misc/FileHelper.h"
 #include "StaticBoundShaderState.h"
@@ -6625,19 +6627,25 @@ static FString GetGlobalShaderCacheFilename(EShaderPlatform Platform)
 	return FString(TEXT("Engine")) / TEXT("GlobalShaderCache-") + LegacyShaderPlatformToShaderFormat(Platform).ToString() + TEXT(".bin");
 }
 
+#if WITH_EDITOR
 /** Creates a string key for the derived data cache entry for the global shader map. */
-static FString GetGlobalShaderMapKeyString(const FGlobalShaderMapId& ShaderMapId, EShaderPlatform Platform, const ITargetPlatform* TargetPlatform, TArray<FShaderTypeDependency> const& Dependencies)
+static UE::DerivedData::FCacheKey GetGlobalShaderMapKey(const FGlobalShaderMapId& ShaderMapId, EShaderPlatform Platform, const ITargetPlatform* TargetPlatform, TArray<FShaderTypeDependency> const& Dependencies)
 {
 	FName Format = LegacyShaderPlatformToShaderFormat(Platform);
 	FString ShaderMapKeyString = Format.ToString() + TEXT("_") + FString(FString::FromInt(GetTargetPlatformManagerRef().ShaderFormatVersion(Format))) + TEXT("_");
 	ShaderMapAppendKeyString(Platform, ShaderMapKeyString);
 	ShaderMapId.AppendKeyString(ShaderMapKeyString, Dependencies);
-	return FDerivedDataCacheInterface::BuildCacheKey(TEXT("GSM"), *GetGlobalShaderMapDDCKey(), *ShaderMapKeyString);
+	const FString DataKey = FString::Printf(TEXT("%s_%s_%s"), TEXT("GSM"), *GetGlobalShaderMapDDCKey(), *ShaderMapKeyString);
+
+	static const UE::DerivedData::FCacheBucket Bucket("GlobalShaderMap");
+	return {Bucket, FIoHash::HashBuffer(MakeMemoryView(FTCHARToUTF8(DataKey)))};
 }
+#endif
 
 /** Saves the platform's shader map to the DDC. It is assumed that the caller will check IsComplete() first before calling the function. */
 static void SaveGlobalShaderMapToDerivedDataCache(EShaderPlatform Platform)
 {
+#if WITH_EDITOR
 	// We've finally built the global shader map, so we can count the miss as we put it in the DDC.
 	COOK_STAT(auto Timer = GlobalShaderCookStats::UsageStats.TimeSyncWork());
 
@@ -6657,11 +6665,20 @@ static void SaveGlobalShaderMapToDerivedDataCache(EShaderPlatform Platform)
 			SaveData.Reset();
 			FMemoryWriter Ar(SaveData, true);
 			Section->Serialize(Ar);
-
-			GetDerivedDataCacheRef().Put(*GetGlobalShaderMapKeyString(ShaderMapId, Platform, TargetPlatform, ShaderFilenameDependencies.Value), SaveData, TEXT("GlobalShaderMap"_SV));
 			COOK_STAT(Timer.AddMiss(SaveData.Num()));
+
+			using namespace UE::DerivedData;
+			FCachePutValueRequest Request;
+			Request.Name = TEXTVIEW("GlobalShaderMap");
+			Request.Key = GetGlobalShaderMapKey(ShaderMapId, Platform, TargetPlatform, ShaderFilenameDependencies.Value);
+			Request.Value = FValue::Compress(MakeSharedBufferFromArray(MoveTemp(SaveData)));
+			FRequestOwner AsyncOwner(EPriority::Normal);
+			FRequestBarrier AsyncBarrier(AsyncOwner);
+			GetCache().PutValue({Request}, AsyncOwner);
+			AsyncOwner.KeepAlive();
 		}
 	}
+#endif // WITH_EDITOR
 }
 
 /** Saves the global shader map as a file for the target platform. */
@@ -6914,6 +6931,7 @@ void CompileGlobalShaderMap(EShaderPlatform Platform, const ITargetPlatform* Tar
 				GGlobalShaderMap[Platform]->LoadFromGlobalArchive(MemoryReader);
 			}
 		}
+#if WITH_EDITOR
 		// Uncooked platform
 		else
 		{
@@ -6922,11 +6940,6 @@ void CompileGlobalShaderMap(EShaderPlatform Platform, const ITargetPlatform* Tar
 			const int32 ShaderFilenameNum = ShaderMapId.GetShaderFilenameToDependeciesMap().Num();
 			const float ProgressStep = 25.0f / ShaderFilenameNum;
 
-			TArray<uint32> AsyncDDCRequestHandles;
-			AsyncDDCRequestHandles.SetNum(ShaderFilenameNum);
-
-			int32 HandleIndex = 0;
-
 			// If NoShaderDDC then don't check for a material the first time we encounter it to simulate
 			// a cold DDC
 			static bool bNoShaderDDC = FParse::Param(FCommandLine::Get(), TEXT("noshaderddc"));
@@ -6936,40 +6949,50 @@ void CompileGlobalShaderMap(EShaderPlatform Platform, const ITargetPlatform* Tar
 			}
 			else
 			{
+				using namespace UE::DerivedData;
+
+				int32 BufferIndex = 0;
+				TArray<FCacheGetValueRequest> Requests;
+
 				// Submit DDC requests.
+				SlowTask.EnterProgressFrame(ProgressStep, LOCTEXT("SubmitDDCRequests", "Submitting global shader DDC Requests..."));
 				for (const auto& ShaderFilenameDependencies : ShaderMapId.GetShaderFilenameToDependeciesMap())
 				{
-					SlowTask.EnterProgressFrame(ProgressStep, LOCTEXT("SubmitDDCRequests", "Submitting global shader DDC Requests..."));
-
-					const FString DataKey = GetGlobalShaderMapKeyString(ShaderMapId, Platform, TargetPlatform, ShaderFilenameDependencies.Value);
-
-					AsyncDDCRequestHandles[HandleIndex] = GetDerivedDataCacheRef().GetAsynchronous(*DataKey, TEXT("GlobalShaderMap"_SV));
-
-					++HandleIndex;
+					FCacheGetValueRequest& Request = Requests.AddDefaulted_GetRef();
+					Request.Name = WriteToString<128>(TEXT("GlobalShaderMap "), LegacyShaderPlatformToShaderFormat(Platform));
+					Request.Key = GetGlobalShaderMapKey(ShaderMapId, Platform, TargetPlatform, ShaderFilenameDependencies.Value);
+					Request.UserData = uint64(BufferIndex);
+					++BufferIndex;
 				}
 
-
-				TArray<uint8> CachedData;
-
-				HandleIndex = 0;
-
 				int32 DDCHits = 0;
 				int32 DDCMisses = 0;
 
 				// Process finished DDC requests.
-				for (const auto& ShaderFilenameDependencies : ShaderMapId.GetShaderFilenameToDependeciesMap())
+				SlowTask.EnterProgressFrame(ProgressStep, LOCTEXT("ProcessDDCRequests", "Processing global shader DDC requests..."));
+				TArray<FValue> GlobalShaderMapBuffers;
+				GlobalShaderMapBuffers.SetNum(Requests.Num());
 				{
-					SlowTask.EnterProgressFrame(ProgressStep, LOCTEXT("ProcessDDCRequests", "Processing global shader DDC requests..."));
-					CachedData.Reset();
 					COOK_STAT(auto Timer = GlobalShaderCookStats::UsageStats.TimeSyncWork());
+					COOK_STAT(Timer.TrackCyclesOnly());
+					FRequestOwner BlockingOwner(EPriority::Blocking);
+					GetCache().GetValue(Requests, BlockingOwner, [&GlobalShaderMapBuffers](FCacheGetValueResponse&& Response)
+					{
+						GlobalShaderMapBuffers[int32(Response.UserData)] = MoveTemp(Response.Value);
+					});
+					BlockingOwner.Wait();
+				}
 
-					GetDerivedDataCacheRef().WaitAsynchronousCompletion(AsyncDDCRequestHandles[HandleIndex]);
-					if (GetDerivedDataCacheRef().GetAsynchronousResults(AsyncDDCRequestHandles[HandleIndex], CachedData))
+				BufferIndex = 0;
+				for (const auto& ShaderFilenameDependencies : ShaderMapId.GetShaderFilenameToDependeciesMap())
+				{
+					COOK_STAT(auto Timer = GlobalShaderCookStats::UsageStats.TimeSyncWork());
+					if (GlobalShaderMapBuffers[BufferIndex].HasData())
 					{
-						COOK_STAT(Timer.AddHit(CachedData.Num()));
-						FMemoryReader MemoryReader(CachedData);
+						COOK_STAT(Timer.AddHit(int64(GlobalShaderMapBuffers[BufferIndex].GetRawSize())));
+						const FSharedBuffer CachedData = GlobalShaderMapBuffers[BufferIndex].GetData().Decompress();
+						FMemoryReaderView MemoryReader(CachedData);
 						GGlobalShaderMap[Platform]->AddSection(FGlobalShaderMapSection::CreateFromArchive(MemoryReader));
-
 						DDCHits++;
 					}
 					else
@@ -6979,14 +7002,14 @@ void CompileGlobalShaderMap(EShaderPlatform Platform, const ITargetPlatform* Tar
 						bShaderMapIsBeingCompiled = true;
 						DDCMisses++;
 					}
-
-					++HandleIndex;
+					++BufferIndex;
 				}
 
 				GShaderCompilerStats->AddDDCHit(DDCHits);
 				GShaderCompilerStats->AddDDCMiss(DDCMisses);
 			}
 		}
+#endif // WITH_EDITOR
 
 		// If any shaders weren't loaded, compile them now.
 		VerifyGlobalShaders(Platform, TargetPlatform, bLoadedFromCacheFile);
@@ -7793,7 +7816,6 @@ FShaderJobCache::FJobCachedOutput* FShaderJobCache::Find(const FJobInputHash& Ha
 				TRACE_COUNTER_INCREMENT(Shaders_JobCacheDDCRequests);
 
 				UE::DerivedData::FRequestOwner RequestOwner(UE::DerivedData::EPriority::Blocking);
-				RequestOwner.KeepAlive();
 
 				UE::DerivedData::FCacheGetRequest Request;
 				Request.Name = TEXT("FShaderJobCache");
@@ -7930,6 +7952,7 @@ void FShaderJobCache::Add(const FJobInputHash& Hash, const FJobCachedOutput& Con
 			TRACE_COUNTER_ADD(Shaders_JobCacheDDCBytesSent, Contents.GetAllocatedSize());
 
 			UE::DerivedData::FRequestOwner RequestOwner(UE::DerivedData::EPriority::Normal);
+			UE::DerivedData::FRequestBarrier RequestBarrier(RequestOwner);
 			RequestOwner.KeepAlive();
 			UE::DerivedData::GetCache().Put(
 				{ {{TEXT("FShaderJobCache")}, RecordBuilder.Build(), IsShaderJobCacheDDCRemotePolicyEnabled() ? UE::DerivedData::ECachePolicy::Default : UE::DerivedData::ECachePolicy::Local } },
