Make sure to cleanup levels that are pending unload when cleaning up the world.

diff --git a/Engine/Source/Runtime/Engine/Classes/Engine/World.h b/Engine/Source/Runtime/Engine/Classes/Engine/World.h
--- a/Engine/Source/Runtime/Engine/Classes/Engine/World.h
+++ b/Engine/Source/Runtime/Engine/Classes/Engine/World.h
@@ -420,6 +420,8 @@ public:
  */
 struct ENGINE_API FLevelStreamingGCHelper
 {
+	friend UWorld;
+
 	/** Called when streamed out levels are going to be garbage collected  */
 	DECLARE_MULTICAST_DELEGATE(FOnGCStreamedOutLevelsEvent);
 	static FOnGCStreamedOutLevelsEvent OnGCStreamedOutLevels;
diff --git a/Engine/Source/Runtime/Engine/Private/World.cpp b/Engine/Source/Runtime/Engine/Private/World.cpp
--- a/Engine/Source/Runtime/Engine/Private/World.cpp
+++ b/Engine/Source/Runtime/Engine/Private/World.cpp
@@ -5047,25 +5047,35 @@ void UWorld::CleanupWorldInternal(bool bSessionEnded, bool bCleanupResources, UW
 		}
 	}
 
-	// Cleanup Persistent level outside of following loop because uninitialized worlds don't have a valid Levels array
-	// StreamingLevels are not initialized.
-	if (PersistentLevel)
-	{
-		PersistentLevel->CleanupLevel(bCleanupResources, bUnloadFromEditor);
-		PersistentLevel->CleanupReferences();
-	}
+		auto CleanupLevelResourcesAndReferences = [](ULevel* Level, bool bCleanupResources, bool bUnloadFromEditor)
+		{
+			if (Level)
+			{
+				Level->CleanupLevel(bCleanupResources, bUnloadFromEditor);
+				Level->CleanupReferences();
+			}
+		};
 
-	if (GetNumLevels() > 1)
-	{
-		check(GetLevel(0) == PersistentLevel);
-		for (int32 LevelIndex = 1; LevelIndex < GetNumLevels(); ++LevelIndex)
+		// Cleanup Persistent level outside of following loop because uninitialized worlds don't have a valid Levels array
+		// StreamingLevels are not initialized.
+		CleanupLevelResourcesAndReferences(PersistentLevel, bCleanupResources, bUnloadFromEditor);
+
+		if (GetNumLevels() > 1)
 		{
-			ULevel* Level = GetLevel(LevelIndex);
-			Level->CleanupLevel(bCleanupResources, bUnloadFromEditor);
-			Level->CleanupReferences();
+			check(GetLevel(0) == PersistentLevel);
+			for (int32 LevelIndex = 1; LevelIndex < GetNumLevels(); ++LevelIndex)
+			{
+				CleanupLevelResourcesAndReferences(GetLevel(LevelIndex), bCleanupResources, bUnloadFromEditor);
+			}
 		}
-	}
 
+		// Also cleanup levels pending a GC purge
+		for (int32 LevelIndex = 0; LevelIndex < FLevelStreamingGCHelper::LevelsPendingUnload.Num(); ++LevelIndex)
+		{
+			CleanupLevelResourcesAndReferences(FLevelStreamingGCHelper::LevelsPendingUnload[LevelIndex].Get(), bCleanupResources, bUnloadFromEditor);
+		}
+		FLevelStreamingGCHelper::LevelsPendingUnload.Empty();
+	}
 #else
 	if (PersistentLevel)
 	{
