fix bad merge

diff --git a/Engine/Source/Runtime/Engine/Private/World.cpp b/Engine/Source/Runtime/Engine/Private/World.cpp
--- a/Engine/Source/Runtime/Engine/Private/World.cpp
+++ b/Engine/Source/Runtime/Engine/Private/World.cpp
@@ -5047,35 +5047,34 @@ void UWorld::CleanupWorldInternal(bool bSessionEnded, bool bCleanupResources, UW
 		}
 	}
 
-		auto CleanupLevelResourcesAndReferences = [](ULevel* Level, bool bCleanupResources, bool bUnloadFromEditor)
+	auto CleanupLevelResourcesAndReferences = [](ULevel* Level, bool bCleanupResources, bool bUnloadFromEditor)
+	{
+		if (Level)
 		{
-			if (Level)
-			{
-				Level->CleanupLevel(bCleanupResources, bUnloadFromEditor);
-				Level->CleanupReferences();
-			}
-		};
+			Level->CleanupLevel(bCleanupResources, bUnloadFromEditor);
+			Level->CleanupReferences();
+		}
+	};
 
-		// Cleanup Persistent level outside of following loop because uninitialized worlds don't have a valid Levels array
-		// StreamingLevels are not initialized.
-		CleanupLevelResourcesAndReferences(PersistentLevel, bCleanupResources, bUnloadFromEditor);
+	// Cleanup Persistent level outside of following loop because uninitialized worlds don't have a valid Levels array
+	// StreamingLevels are not initialized.
+	CleanupLevelResourcesAndReferences(PersistentLevel, bCleanupResources, bUnloadFromEditor);
 
-		if (GetNumLevels() > 1)
+	if (GetNumLevels() > 1)
+	{
+		check(GetLevel(0) == PersistentLevel);
+		for (int32 LevelIndex = 1; LevelIndex < GetNumLevels(); ++LevelIndex)
 		{
-			check(GetLevel(0) == PersistentLevel);
-			for (int32 LevelIndex = 1; LevelIndex < GetNumLevels(); ++LevelIndex)
-			{
-				CleanupLevelResourcesAndReferences(GetLevel(LevelIndex), bCleanupResources, bUnloadFromEditor);
-			}
+			CleanupLevelResourcesAndReferences(GetLevel(LevelIndex), bCleanupResources, bUnloadFromEditor);
 		}
+	}
 
-		// Also cleanup levels pending a GC purge
-		for (int32 LevelIndex = 0; LevelIndex < FLevelStreamingGCHelper::LevelsPendingUnload.Num(); ++LevelIndex)
-		{
-			CleanupLevelResourcesAndReferences(FLevelStreamingGCHelper::LevelsPendingUnload[LevelIndex].Get(), bCleanupResources, bUnloadFromEditor);
-		}
-		FLevelStreamingGCHelper::LevelsPendingUnload.Empty();
+	// Also cleanup levels pending a GC purge
+	for (int32 LevelIndex = 0; LevelIndex < FLevelStreamingGCHelper::LevelsPendingUnload.Num(); ++LevelIndex)
+	{
+		CleanupLevelResourcesAndReferences(FLevelStreamingGCHelper::LevelsPendingUnload[LevelIndex].Get(), bCleanupResources, bUnloadFromEditor);
 	}
+	FLevelStreamingGCHelper::LevelsPendingUnload.Empty();
 #else
 	if (PersistentLevel)
 	{
