Downgrade seek ensure to log
- Add more logging around failures in ADPCM decode
- Remove superfluous logs
#rnx

diff --git a/Engine/Plugins/Runtime/Metasound/Source/MetasoundEngine/Private/MetasoundWaveProxyReader.cpp b/Engine/Plugins/Runtime/Metasound/Source/MetasoundEngine/Private/MetasoundWaveProxyReader.cpp
--- a/Engine/Plugins/Runtime/Metasound/Source/MetasoundEngine/Private/MetasoundWaveProxyReader.cpp
+++ b/Engine/Plugins/Runtime/Metasound/Source/MetasoundEngine/Private/MetasoundWaveProxyReader.cpp
@@ -410,7 +410,14 @@ namespace Metasound
 		}
 
 		// Seek input to start time.
-		DecoderInput->SeekToTime(InStartTimeInSeconds);
+		const bool bSeekSucceeded = DecoderInput->SeekToTime(InStartTimeInSeconds);
+		if (!bSeekSucceeded)
+		{
+			UE_LOG(LogMetaSound, Warning, TEXT("Failed to seek decoder input during initialization: (format:%s) for wave (package:%s) to time '%.6f'"),
+				*Format.ToString(),
+				*WaveProxy->GetPackageName().ToString(),
+				InStartTimeInSeconds);
+		}
 		CurrentFrameIndex = InStartTimeInSeconds * GetSampleRate();
 
 		// Get codec ptr by reading the header info from the decoder input.
diff --git a/Engine/Source/Runtime/AudioCodecEngine/Private/DecoderInputBackCompat.cpp b/Engine/Source/Runtime/AudioCodecEngine/Private/DecoderInputBackCompat.cpp
--- a/Engine/Source/Runtime/AudioCodecEngine/Private/DecoderInputBackCompat.cpp
+++ b/Engine/Source/Runtime/AudioCodecEngine/Private/DecoderInputBackCompat.cpp
@@ -79,13 +79,16 @@ namespace Audio
 		return -1;
 	}
 
-	void FBackCompatInput::SeekToTime(const float InSeconds)
+	bool FBackCompatInput::SeekToTime(const float InSeconds)
 	{
 		ICompressedAudioInfo* Info = GetInfo();
-		if (ensure(Info))
+		if (Info)
 		{
 			Info->SeekToTime(InSeconds);
+			return true;
 		}
+
+		return false;
 	}
 
 	bool FBackCompatInput::IsEndOfStream() const
diff --git a/Engine/Source/Runtime/AudioCodecEngine/Private/DecoderInputBackCompat.h b/Engine/Source/Runtime/AudioCodecEngine/Private/DecoderInputBackCompat.h
--- a/Engine/Source/Runtime/AudioCodecEngine/Private/DecoderInputBackCompat.h
+++ b/Engine/Source/Runtime/AudioCodecEngine/Private/DecoderInputBackCompat.h
@@ -34,7 +34,7 @@ namespace Audio
 		bool FindSection(FEncodedSectionBase& OutSection) override;
 		int64 Tell() const override;
 
-		void SeekToTime(const float InSeconds) override;
+		bool SeekToTime(const float InSeconds) override;
 		
 		TArrayView<const uint8> PeekNextPacket(
 			int32 InMaxPacketLength) const override;
diff --git a/Engine/Source/Runtime/AudioExtensions/Public/IAudioCodec.h b/Engine/Source/Runtime/AudioExtensions/Public/IAudioCodec.h
--- a/Engine/Source/Runtime/AudioExtensions/Public/IAudioCodec.h
+++ b/Engine/Source/Runtime/AudioExtensions/Public/IAudioCodec.h
@@ -116,7 +116,7 @@ namespace Audio
 
 		virtual int64 Tell() const = 0;
 
-		virtual void SeekToTime(const float InSeconds) { }
+		virtual bool SeekToTime(const float InSeconds) { return false; }
 
 		virtual TArrayView<const uint8> PeekNextPacket(
 			int32 InMaxPacketLength) const = 0;
diff --git a/Engine/Source/Runtime/Engine/Classes/Sound/SoundWave.h b/Engine/Source/Runtime/Engine/Classes/Sound/SoundWave.h
--- a/Engine/Source/Runtime/Engine/Classes/Sound/SoundWave.h
+++ b/Engine/Source/Runtime/Engine/Classes/Sound/SoundWave.h
@@ -1151,8 +1151,15 @@ public:
 		uint32 MaxSizeInCache = 0;
 	};
 	
-	// ctor
-	FSoundWaveData() = default;
+	FSoundWaveData()
+		: bIsLooping(0)
+		, bIsTemplate(0)
+		, bIsStreaming(0)
+		, bIsSeekable(0)
+		, bShouldUseStreamCaching(0)
+		, bLoadingBehaviorOverridden(0)
+	{
+	}
 
 	// dtor
 	~FSoundWaveData();
@@ -1225,7 +1232,7 @@ private:
 
 	FFormatContainer CompressedFormatData;
 
-	int32 ResourceSize;
+	int32 ResourceSize = 0;
 
 	FBulkDataBuffer<uint8> ResourceData;
 
@@ -1242,12 +1249,12 @@ private:
 	TArray<FSoundWaveCuePoint> CuePoints;
 	ESoundAssetCompressionType SoundAssetCompressionType;
 
-	float SampleRate;
-	float Duration;
+	float SampleRate = 0;
+	float Duration = 0;
 
-	uint32 NumChannels;
-	uint32 NumChunks;
-	int32 NumFrames;
+	uint32 NumChannels = 0;
+	uint32 NumChunks = 0;
+	int32 NumFrames = 0;
 
 	// shared flags
 	uint8 bIsLooping : 1;
diff --git a/Engine/Source/Runtime/Engine/Private/ADPCMAudioInfo.cpp b/Engine/Source/Runtime/Engine/Private/ADPCMAudioInfo.cpp
--- a/Engine/Source/Runtime/Engine/Private/ADPCMAudioInfo.cpp
+++ b/Engine/Source/Runtime/Engine/Private/ADPCMAudioInfo.cpp
@@ -485,22 +485,24 @@ bool FADPCMAudioInfo::StreamCompressedInfoInternal(const FSoundWaveProxyPtr& InW
 
 	if (ChunkData == nullptr)
 	{
+		UE_LOG(LogAudio, Warning, TEXT("FADPCMAudioInfo::StreamCompressedInfoInternal: Failed to get loaded chunk at index '%u'"), CurrentChunkIndex);
 		return false;
 	}
 
 	SrcBufferData = nullptr;
 	SrcBufferDataSize = 0;
 
-	void* FormatHeader;
-
-	if (!WaveInfo.ReadWaveInfo((uint8*)ChunkData, CurrentChunkDataSize, nullptr, true, &FormatHeader))
+	void* FormatHeader = nullptr;
+	FString ErrorMsg;
+	constexpr bool bHeaderDataOnly = true;
+	if (!WaveInfo.ReadWaveInfo(static_cast<const uint8*>(ChunkData), CurrentChunkDataSize, &ErrorMsg, bHeaderDataOnly, &FormatHeader))
 	{
-		UE_LOG(LogAudio, Warning, TEXT("WaveInfo.ReadWaveInfo Failed"));
+		UE_LOG(LogAudio, Warning, TEXT("FADPCMAudioInfo::StreamCompressedInfoInternal: Failed to read WaveInfo: %s"), *ErrorMsg);
 		return false;
 	}
 	
 	// if we only included the header in the zeroth chunk, skip to the next chunk.
-	int32 SampleDataOffset = WaveInfo.SampleDataStart - ChunkData;
+	const int32 SampleDataOffset = WaveInfo.SampleDataStart - ChunkData;
 	check(SampleDataOffset > 0);
 	if (((uint32)SampleDataOffset) >= CurrentChunkDataSize)
 	{
@@ -529,6 +531,7 @@ bool FADPCMAudioInfo::StreamCompressedInfoInternal(const FSoundWaveProxyPtr& InW
 
 	if (Format == WAVE_FORMAT_ADPCM)
 	{
+		check(FormatHeader);
 		ADPCM::ADPCMFormatHeader* ADPCMHeader = (ADPCM::ADPCMFormatHeader*)FormatHeader;
 		TotalSamplesPerChannel = ADPCMHeader->SamplesPerChannel;
 		SamplesPerBlock = ADPCMHeader->wSamplesPerBlock;
@@ -561,7 +564,7 @@ bool FADPCMAudioInfo::StreamCompressedInfoInternal(const FSoundWaveProxyPtr& InW
 	}
 	else
 	{
-		UE_LOG(LogAudio, Error, TEXT("Unsupported wave format"));
+		UE_LOG(LogAudio, Error, TEXT("FADPCMAudioInfo::StreamCompressedInfoInternal: Unsupported wave format (Neither ADPCM nor LPCM)"));
 		return false;
 	}
 
diff --git a/Engine/Source/Runtime/Engine/Private/SoundWave.cpp b/Engine/Source/Runtime/Engine/Private/SoundWave.cpp
--- a/Engine/Source/Runtime/Engine/Private/SoundWave.cpp
+++ b/Engine/Source/Runtime/Engine/Private/SoundWave.cpp
@@ -98,7 +98,6 @@ static void DumpBakedAnalysisData(const TArray<FString>& Args)
 		if (Args.Num() == 1)
 		{
 			const FString& SoundWaveToDump = Args[0];
-			UE_LOG(LogTemp, Log, TEXT("Foo"));
 			for (TObjectIterator<USoundWave> It; It; ++It)
 			{
 				if (It->IsTemplate(RF_ClassDefaultObject))
@@ -108,7 +107,6 @@ static void DumpBakedAnalysisData(const TArray<FString>& Args)
 
 				if (SoundWaveToDump.Equals(It->GetName()))
 				{
-					UE_LOG(LogTemp, Log, TEXT("Foo"));
 #if WITH_EDITOR
 					It->LogBakedData();
 #endif // WITH_EDITOR
