Fix ActorPosition material node to work with WPO velocity vectors correctly.

diff --git a/Engine/Source/Runtime/Engine/Private/Materials/HLSLMaterialTranslator.cpp b/Engine/Source/Runtime/Engine/Private/Materials/HLSLMaterialTranslator.cpp
--- a/Engine/Source/Runtime/Engine/Private/Materials/HLSLMaterialTranslator.cpp
+++ b/Engine/Source/Runtime/Engine/Private/Materials/HLSLMaterialTranslator.cpp
@@ -5436,7 +5436,7 @@ int32 FHLSLMaterialTranslator::ActorWorldPosition()
 		// material node is used in VS
 		Result = AddInlinedCodeChunkZeroDeriv(
 			MCT_LWCVector3,
-			TEXT("LWCMultiply(LWCMultiply(GetActorWorldPosition(Parameters), GetPrimitiveData(Parameters).WorldToLocal), Parameters.PrevFrameLocalToWorld)"));
+			TEXT("LWCMultiply(LWCMultiply(GetActorWorldPosition(Parameters), GetWorldToInstance(Parameters)), Parameters.PrevFrameLocalToWorld)"));
 	}
 	else
 	{
