Fix Build (resolve error)

#rnx

diff --git a/Engine/Source/Runtime/Engine/Private/WorldPartition/HLOD/HLODSubsystem.cpp b/Engine/Source/Runtime/Engine/Private/WorldPartition/HLOD/HLODSubsystem.cpp
--- a/Engine/Source/Runtime/Engine/Private/WorldPartition/HLOD/HLODSubsystem.cpp
+++ b/Engine/Source/Runtime/Engine/Private/WorldPartition/HLOD/HLODSubsystem.cpp
@@ -247,7 +247,7 @@ void UHLODSubsystem::OnCellHidden(const UWorldPartitionRuntimeCell* InCell)
 	const UWorldPartition* WorldPartition = InCell->GetTypedOuter<UWorldPartition>();
 	FWorldPartitionHLODRuntimeData& WorldPartitionHLODRuntimeData = WorldPartitionsHLODRuntimeData.FindChecked(WorldPartition);
 
-	FCellData& CellData = WorldPartitionHLODRuntimeData.CellsData.FindChecked(InCell->GetFName());
+	FCellData& CellData = WorldPartitionHLODRuntimeData->CellsData.FindChecked(InCell->GetFName());
 	CellData.bIsCellVisible = false;
 
 #if WITH_EDITOR
