Backout 19356865 and history index size changes from 18938566

diff --git a/Engine/Source/Runtime/CoreUObject/Public/UObject/CoreNet.h b/Engine/Source/Runtime/CoreUObject/Public/UObject/CoreNet.h
--- a/Engine/Source/Runtime/CoreUObject/Public/UObject/CoreNet.h
+++ b/Engine/Source/Runtime/CoreUObject/Public/UObject/CoreNet.h
@@ -233,7 +233,7 @@ struct FPropertyRetirement
 	TSharedPtr<class INetDeltaBaseState> DynamicState;
 
 	FPacketIdRange OutPacketIdRange;
-	uint16 FastArrayChangelistHistory;
+	uint32 FastArrayChangelistHistory;
 
 	FPropertyRetirement() :
 		 Next(nullptr)
@@ -426,15 +426,15 @@ public:
 	 */
 	virtual void CountBytes(FArchive& Ar) const {}
 
-	uint16 GetLastAckedHistory() const { return LastAckedHistory; }
-	void SetLastAckedHistory(uint16 InAckedHistory) { LastAckedHistory = InAckedHistory; }
+	uint32 GetLastAckedHistory() const { return LastAckedHistory; }
+	void SetLastAckedHistory(uint32 InAckedHistory) { LastAckedHistory = InAckedHistory; }
 
-	uint16 GetChangelistHistory() const { return ChangelistHistory; }
-	void SetChangelistHistory(uint16 InChangelistHistory) { ChangelistHistory = InChangelistHistory; }
+	uint32 GetChangelistHistory() const { return ChangelistHistory; }
+	void SetChangelistHistory(uint32 InChangelistHistory) { ChangelistHistory = InChangelistHistory; }
 
 private:
-	uint16 LastAckedHistory;
-	uint16 ChangelistHistory;
+	uint32 LastAckedHistory;
+	uint32 ChangelistHistory;
 };
 
 struct FNetDeltaSerializeInfo;
diff --git a/Engine/Source/Runtime/Engine/Private/DataReplication.cpp b/Engine/Source/Runtime/Engine/Private/DataReplication.cpp
--- a/Engine/Source/Runtime/Engine/Private/DataReplication.cpp
+++ b/Engine/Source/Runtime/Engine/Private/DataReplication.cpp
@@ -862,9 +862,9 @@ void FObjectReplicator::ReceivedNak( int32 NakPacketId )
 		if (FSendingRepState* SendingRepState = RepState.IsValid() ? RepState->GetSendingRepState() : nullptr)
 		{
 			// Go over properties tracked with histories, and mark them as needing to be resent.
-			for (uint16 i = SendingRepState->HistoryStart; i < SendingRepState->HistoryEnd; ++i)
+			for (int32 i = SendingRepState->HistoryStart; i < SendingRepState->HistoryEnd; ++i)
 			{
-				const uint16 HistoryIndex = i % FSendingRepState::MAX_CHANGE_HISTORY;
+				const int32 HistoryIndex = i % FSendingRepState::MAX_CHANGE_HISTORY;
 
 				FRepChangedHistory& HistoryItem = SendingRepState->ChangeHistory[HistoryIndex];
 
@@ -884,7 +884,7 @@ void FObjectReplicator::ReceivedNak( int32 NakPacketId )
 
 				// If this is a dynamic array property, we have to look through the list of retirement records to see if we need to reset the base state
 				FPropertyRetirement* Rec = Retirement.Next; // Retirement[i] is head and not actually used in this case
-				uint16 LastAcknowledged = 0;
+				uint32 LastAcknowledged = 0;
 				while (Rec != nullptr)
 				{
 					if (NakPacketId > Rec->OutPacketIdRange.Last)
@@ -1504,7 +1504,7 @@ void FObjectReplicator::PostReceivedBunch()
 
 static FORCEINLINE FPropertyRetirement** UpdateAckedRetirements(
 	FPropertyRetirement& Retire,
-	uint16& LastAcknowledged,
+	uint32& LastAcknowledged,
 	const int32 OutAckPacketId,
 	const UObject* Object)
 {
@@ -1641,7 +1641,7 @@ void FObjectReplicator::ReplicateCustomDeltaProperties( FNetBitWriter & Bunch, F
 
 			// Update Retirement records with this new state so we can handle packet drops.
 			// LastNext will be pointer to the last "Next" pointer in the list (so pointer to a pointer)
-			uint16 LastAcknowledged = 0;
+			uint32 LastAcknowledged = 0;
 			LastNext = UpdateAckedRetirements(Retire, LastAcknowledged, Connection->OutAckPacketId, Object);
 
 			if (LastAcknowledged != 0 && OldState.IsValid())
@@ -1895,9 +1895,9 @@ void FObjectReplicator::PostSendBunch( FPacketIdRange & PacketRange, uint8 bReli
 			// Don't call if reliable, since the bunch will be resent. We dont want this to end up in the changelist history
 			// But is that enough? How does it know to delta against this latest state?
 
-			for (uint16 i = SendingRepState->HistoryStart; i < SendingRepState->HistoryEnd; ++i)
+			for (int32 i = SendingRepState->HistoryStart; i < SendingRepState->HistoryEnd; ++i)
 			{
-				const uint16 HistoryIndex = i % FSendingRepState::MAX_CHANGE_HISTORY;
+				const int32 HistoryIndex = i % FSendingRepState::MAX_CHANGE_HISTORY;
 
 				FRepChangedHistory & HistoryItem = SendingRepState->ChangeHistory[HistoryIndex];
 
@@ -1926,7 +1926,7 @@ void FObjectReplicator::PostSendBunch( FPacketIdRange & PacketRange, uint8 bReli
 			// possible that we either didn't send anything, or the last thing we sent wasn't reliable.
 			if (SendingRepState->HistoryEnd > SendingRepState->HistoryStart)
 			{
-				const uint16 HistoryIndex = (SendingRepState->HistoryEnd - 1) % FSendingRepState::MAX_CHANGE_HISTORY;
+				const int32 HistoryIndex = (SendingRepState->HistoryEnd - 1) % FSendingRepState::MAX_CHANGE_HISTORY;
 				FRepChangedHistory& HistoryItem = SendingRepState->ChangeHistory[HistoryIndex];
 	
 				if (!HistoryItem.WasSent())
diff --git a/Engine/Source/Runtime/Engine/Private/RepLayout.cpp b/Engine/Source/Runtime/Engine/Private/RepLayout.cpp
--- a/Engine/Source/Runtime/Engine/Private/RepLayout.cpp
+++ b/Engine/Source/Runtime/Engine/Private/RepLayout.cpp
@@ -471,7 +471,7 @@ struct FDeltaArrayHistoryItem
 struct FDeltaArrayHistoryState
 {
 	/** The maximum number of individual changelists allowed.*/
-	static const uint16 MAX_CHANGE_HISTORY = 32;
+	static const uint32 MAX_CHANGE_HISTORY = 32;
 
 	/** Circular buffer of changelists. */
 	TArray<FDeltaArrayHistoryItem> ChangeHistory;
@@ -484,10 +484,10 @@ struct FDeltaArrayHistoryState
 	int32 ArrayReplicationKey = INDEX_NONE;
 
 	/** Index in the buffer where changelist history starts (i.e., the Oldest changelist). */
-	uint16 HistoryStart = 0;
+	uint32 HistoryStart = 0;
 
 	/** Index in the buffer where changelist history ends (i.e., the Newest changelist). */
-	uint16 HistoryEnd = 0;
+	uint32 HistoryEnd = 0;
 
 	void InitHistory()
 	{
@@ -1692,20 +1692,10 @@ ERepLayoutResult FRepLayout::CompareProperties(
 		return ERepLayoutResult::Empty;
 	}
 
-	static_assert(sizeof(RepChangelistState->CompareIndex) == sizeof(uint16), "CompareIndex should be a uint16");
-
-	if (RepChangelistState->CompareIndex == MAX_uint16)
-	{
-		// other systems test for index > 1 to indicate that we've ever done a comparison
-		RepChangelistState->CompareIndex = 2;
-	}
-	else
-	{
-		RepChangelistState->CompareIndex++;
-	}
+	RepChangelistState->CompareIndex++;
 
 	check((RepChangelistState->HistoryEnd - RepChangelistState->HistoryStart) < FRepChangelistState::MAX_CHANGE_HISTORY);
-	const uint16 HistoryIndex = RepChangelistState->HistoryEnd % FRepChangelistState::MAX_CHANGE_HISTORY;
+	const int32 HistoryIndex = RepChangelistState->HistoryEnd % FRepChangelistState::MAX_CHANGE_HISTORY;
 
 	FRepChangedHistory& NewHistoryItem = RepChangelistState->ChangeHistory[HistoryIndex];
 
@@ -1812,24 +1802,16 @@ ERepLayoutResult FRepLayout::CompareProperties(
 	// If we're full, merge the oldest up, so we always have room for a new entry
 	if ((RepChangelistState->HistoryEnd - RepChangelistState->HistoryStart) == FRepChangelistState::MAX_CHANGE_HISTORY)
 	{
-		const uint16 FirstHistoryIndex = RepChangelistState->HistoryStart % FRepChangelistState::MAX_CHANGE_HISTORY;
+		const int32 FirstHistoryIndex = RepChangelistState->HistoryStart % FRepChangelistState::MAX_CHANGE_HISTORY;
 
 		RepChangelistState->HistoryStart++;
 
-		const uint16 SecondHistoryIndex = RepChangelistState->HistoryStart % FRepChangelistState::MAX_CHANGE_HISTORY;
+		const int32 SecondHistoryIndex = RepChangelistState->HistoryStart % FRepChangelistState::MAX_CHANGE_HISTORY;
 
 		TArray<uint16>& FirstChangelistRef = RepChangelistState->ChangeHistory[FirstHistoryIndex].Changed;
 		TArray<uint16> SecondChangelistCopy = MoveTemp(RepChangelistState->ChangeHistory[SecondHistoryIndex].Changed);
 
 		MergeChangeList(Data, FirstChangelistRef, SecondChangelistCopy, RepChangelistState->ChangeHistory[SecondHistoryIndex].Changed);
-
-		// Remove any tiling in the history markers to keep them from wrapping over time
-		const uint16 NewHistoryCount = RepChangelistState->HistoryEnd - RepChangelistState->HistoryStart;
-
-		check(NewHistoryCount < FRepChangelistState::MAX_CHANGE_HISTORY);
-
-		RepChangelistState->HistoryStart = RepChangelistState->HistoryStart % FRepChangelistState::MAX_CHANGE_HISTORY;
-		RepChangelistState->HistoryEnd = RepChangelistState->HistoryStart + NewHistoryCount;
 	}
 
 	return Result;
@@ -1939,9 +1921,9 @@ bool FRepLayout::ReplicateProperties(
 		if (RepState->LifetimeChangelist.Num() == 0)
 		{
 			// If this object was dormant, the lifetime list will be empty if the replicator was discarded, so rebuild it from the changelist state.
-			for (uint16 i = RepChangelistState->HistoryStart; i < RepChangelistState->HistoryEnd; ++i)
+			for (int32 i = RepChangelistState->HistoryStart; i < RepChangelistState->HistoryEnd; ++i)
 			{
-				const uint16 HistoryIndex = i % FRepChangelistState::MAX_CHANGE_HISTORY;
+				const int32 HistoryIndex = i % FRepChangelistState::MAX_CHANGE_HISTORY;
 
 				FRepChangedHistory& HistoryItem = RepChangelistState->ChangeHistory[HistoryIndex];
 
@@ -2003,7 +1985,7 @@ bool FRepLayout::ReplicateProperties(
 		RepState->LastChangelistIndex = RepChangelistState->HistoryStart;
 	}
 
-	const uint16 PossibleNewHistoryIndex = RepState->HistoryEnd % FSendingRepState::MAX_CHANGE_HISTORY;
+	const int32 PossibleNewHistoryIndex = RepState->HistoryEnd % FSendingRepState::MAX_CHANGE_HISTORY;
 
 	FRepChangedHistory& PossibleNewHistoryItem = RepState->ChangeHistory[PossibleNewHistoryIndex];
 
@@ -2013,9 +1995,9 @@ bool FRepLayout::ReplicateProperties(
 	check(Changed.Num() == 0);
 
 	// Gather all change lists that are new since we last looked, and merge them all together into a single CL
-	for (uint16 i = RepState->LastChangelistIndex; i < RepChangelistState->HistoryEnd; ++i)
+	for (int32 i = RepState->LastChangelistIndex; i < RepChangelistState->HistoryEnd; ++i)
 	{
-		const uint16 HistoryIndex = i % FRepChangelistState::MAX_CHANGE_HISTORY;
+		const int32 HistoryIndex = i % FRepChangelistState::MAX_CHANGE_HISTORY;
 
 		FRepChangedHistory& HistoryItem = RepChangelistState->ChangeHistory[HistoryIndex];
 
@@ -2124,7 +2106,7 @@ void FRepLayout::UpdateChangelistHistory(
 {
 	check(RepState->HistoryEnd >= RepState->HistoryStart);
 
-	const uint16 HistoryCount = RepState->HistoryEnd - RepState->HistoryStart;
+	const int32 HistoryCount = RepState->HistoryEnd - RepState->HistoryStart;
 	const bool bDumpHistory = HistoryCount == FSendingRepState::MAX_CHANGE_HISTORY;
 	const int32 AckPacketId = Connection->OutAckPacketId;
 
@@ -2138,9 +2120,9 @@ void FRepLayout::UpdateChangelistHistory(
 
 	if (bDeltaCheckpoint)
 	{
-		for (uint16 i = RepState->HistoryStart; i < RepState->HistoryEnd - 1; i++)
+		for (int32 i = RepState->HistoryStart; i < RepState->HistoryEnd - 1; i++)
 		{
-			const uint16 HistoryIndex = i % FSendingRepState::MAX_CHANGE_HISTORY;
+			const int32 HistoryIndex = i % FSendingRepState::MAX_CHANGE_HISTORY;
 
 			FRepChangedHistory& HistoryItem = RepState->ChangeHistory[HistoryIndex];
 
@@ -2153,9 +2135,9 @@ void FRepLayout::UpdateChangelistHistory(
 	}
 	else
 	{
-		for (uint16 i = RepState->HistoryStart; i < RepState->HistoryEnd; i++)
+		for (int32 i = RepState->HistoryStart; i < RepState->HistoryEnd; i++)
 		{
-			const uint16 HistoryIndex = i % FSendingRepState::MAX_CHANGE_HISTORY;
+			const int32 HistoryIndex = i % FSendingRepState::MAX_CHANGE_HISTORY;
 
 			FRepChangedHistory& HistoryItem = RepState->ChangeHistory[HistoryIndex];
 
@@ -2195,7 +2177,7 @@ void FRepLayout::UpdateChangelistHistory(
 	}
 
 	// Remove any tiling in the history markers to keep them from wrapping over time
-	const uint16 NewHistoryCount = RepState->HistoryEnd - RepState->HistoryStart;
+	const int32 NewHistoryCount	= RepState->HistoryEnd - RepState->HistoryStart;
 
 	check(NewHistoryCount < FSendingRepState::MAX_CHANGE_HISTORY);
 
@@ -7343,8 +7325,8 @@ void FRepLayout::PreSendCustomDeltaProperties(
 							{
 								FastArrayHistoryState.InitHistory();
 
-								const uint16 HistoryDelta = FastArrayHistoryState.HistoryEnd - FastArrayHistoryState.HistoryStart;
-								const uint16 CurrentHistoryIndex = FastArrayHistoryState.HistoryEnd % FDeltaArrayHistoryState::MAX_CHANGE_HISTORY;
+								const uint32 HistoryDelta = FastArrayHistoryState.HistoryEnd - FastArrayHistoryState.HistoryStart;
+								const uint32 CurrentHistoryIndex = FastArrayHistoryState.HistoryEnd % FDeltaArrayHistoryState::MAX_CHANGE_HISTORY;
 								const FDeltaArrayHistoryItem& CurrentHistory = FastArrayHistoryState.ChangeHistory[CurrentHistoryIndex];
 								const bool bCurrentHistoryUpdated = FastArrayHistoryState.ChangeHistoryUpdated[CurrentHistoryIndex];
 
@@ -7362,7 +7344,7 @@ void FRepLayout::PreSendCustomDeltaProperties(
 									}
 
 									++FastArrayHistoryState.HistoryEnd;
-									const uint16 NewHistory = FastArrayHistoryState.HistoryEnd % FDeltaArrayHistoryState::MAX_CHANGE_HISTORY;
+									const uint32 NewHistory = FastArrayHistoryState.HistoryEnd % FDeltaArrayHistoryState::MAX_CHANGE_HISTORY;
 									FastArrayHistoryState.ChangeHistory[NewHistory].Reset();
 									FastArrayHistoryState.ChangeHistoryUpdated[NewHistory] = false;
 								}
@@ -7489,9 +7471,9 @@ ERepLayoutResult FRepLayout::DeltaSerializeFastArrayProperty(FFastArrayDeltaSeri
 
 				// Cache off the newest history, the last history sent to this connection, and then update the state
 				// to notify that we're going to send it the newest history.
-				const uint16 NewChangelistHistory = FastArrayState.HistoryEnd;
-				const uint16 LastAckedHistory = NewArrayDeltaState->GetLastAckedHistory();
-				const uint16 LastAckedChangelistDelta = NewChangelistHistory - LastAckedHistory;
+				const uint32 NewChangelistHistory = FastArrayState.HistoryEnd;
+				const uint32 LastAckedHistory = NewArrayDeltaState->GetLastAckedHistory();
+				const uint32 LastAckedChangelistDelta = NewChangelistHistory - LastAckedHistory;
 
 				NewArrayDeltaState->SetChangelistHistory(NewChangelistHistory);
 
@@ -7506,8 +7488,8 @@ ERepLayoutResult FRepLayout::DeltaSerializeFastArrayProperty(FFastArrayDeltaSeri
 
 				// Check to see whether or not we need to update the global changelist shared between connections.
 				{
-					const uint16 RelativeNewHistory = NewChangelistHistory % FDeltaArrayHistoryState::MAX_CHANGE_HISTORY;
-					const uint16 CompareChangelistDelta = NewChangelistHistory - FastArrayState.HistoryStart;
+					const uint32 RelativeNewHistory = NewChangelistHistory % FDeltaArrayHistoryState::MAX_CHANGE_HISTORY;
+					const uint32 CompareChangelistDelta = NewChangelistHistory - FastArrayState.HistoryStart;
 					FDeltaArrayHistoryItem& HistoryItem = FastArrayState.ChangeHistory[RelativeNewHistory];
 					const bool bHistoryItemUpdated = FastArrayState.ChangeHistoryUpdated[RelativeNewHistory];
 
diff --git a/Engine/Source/Runtime/Engine/Public/Net/RepLayout.h b/Engine/Source/Runtime/Engine/Public/Net/RepLayout.h
--- a/Engine/Source/Runtime/Engine/Public/Net/RepLayout.h
+++ b/Engine/Source/Runtime/Engine/Public/Net/RepLayout.h
@@ -508,7 +508,7 @@ public:
 	~FRepChangelistState();
 
 	/** The maximum number of individual changelists allowed.*/
-	static const uint16 MAX_CHANGE_HISTORY = 64;
+	static const int32 MAX_CHANGE_HISTORY = 64;
 
 	/** Circular buffer of changelists. */
 	FRepChangedHistory ChangeHistory[MAX_CHANGE_HISTORY];
@@ -517,13 +517,13 @@ public:
 	TUniquePtr<struct FCustomDeltaChangelistState> CustomDeltaChangelistState;
 
 	/** Index in the buffer where changelist history starts (i.e., the Oldest changelist). */
-	uint16 HistoryStart;
+	int32 HistoryStart;
 
 	/** Index in the buffer where changelist history ends (i.e., the Newest changelist). */
-	uint16 HistoryEnd;
+	int32 HistoryEnd;
 
 	/** Number of times that properties have been compared */
-	uint16 CompareIndex;
+	int32 CompareIndex;
 
 	/** Latest state of all property data. Not used on Clients, only used on Servers if Shadow State is enabled. */
 	FRepStateStaticBuffer StaticBuffer;
@@ -654,19 +654,19 @@ public:
 	TEnumAsByte<ENetRole> SavedRole = ROLE_MAX;
 
 	/** Index in the buffer where changelist history starts (i.e., the Oldest changelist). */
-	uint16 HistoryStart;
+	int32 HistoryStart;
 
 	/** Index in the buffer where changelist history ends (i.e., the Newest changelist). */
-	uint16 HistoryEnd;
+	int32 HistoryEnd;
 
 	/** Number of Changelist history entries that have outstanding Naks. */
-	uint16 NumNaks;
+	int32 NumNaks;
 
 	/**
 	 * The last change list history item we replicated from FRepChangelistState.
 	 * (If we are caught up to FRepChangelistState::HistoryEnd, there are no new changelists to replicate).
 	 */
-	uint16 LastChangelistIndex;
+	int32 LastChangelistIndex;
 
 	/**
 	 * Tracks the last time this RepState actually replicated data.
@@ -677,14 +677,14 @@ public:
 	 *
 	 * Note, we can't solely rely on on LastChangelistIndex, since changelists are stored in circular buffers.
 	 */
-	uint16 LastCompareIndex;
+	int32 LastCompareIndex;
 
 	FReplicationFlags RepFlags;
 
 	TSharedPtr<FRepChangedPropertyTracker> RepChangedPropertyTracker;
 
 	/** The maximum number of individual changelists allowed.*/
-	static constexpr uint16 MAX_CHANGE_HISTORY = 32;
+	static constexpr int32 MAX_CHANGE_HISTORY = 32;
 
 	/** Circular buffer of changelists. */
 	FRepChangedHistory ChangeHistory[MAX_CHANGE_HISTORY];
diff --git a/Engine/Source/Runtime/Net/Core/Classes/Net/Serialization/FastArraySerializer.h b/Engine/Source/Runtime/Net/Core/Classes/Net/Serialization/FastArraySerializer.h
--- a/Engine/Source/Runtime/Net/Core/Classes/Net/Serialization/FastArraySerializer.h
+++ b/Engine/Source/Runtime/Net/Core/Classes/Net/Serialization/FastArraySerializer.h
@@ -1601,8 +1601,8 @@ bool FFastArraySerializer::FastArrayDeltaSerialize_DeltaSerializeStructs(TArray<
 
 		// We set ChangelistHistory and LastAckedHistory to 0, as that's the default RepLayout expects.
 		// Setting them to anything else may cause issues with initial sends.
-		uint16 OldChangelistHistory = 0;
-		uint16 OldLastAckedHistory = 0;
+		int32 OldChangelistHistory = 0;
+		int32 OldLastAckedHistory = 0;
 
 		// See if the array changed at all. If the ArrayReplicationKey matches we can skip checking individual items
 		if (Parms.OldState)
