Merge r.SkinnedMesh.UpdateBoundsNotifyStreaming and r.SkinnedMesh.UpdateBoundsNotifyStreamingRadiusChangeRatio into a single cvar

#tests verified no perf regression on a less powerful platform via replay a/b tests

diff --git a/Engine/Source/Runtime/Engine/Classes/Components/SkinnedMeshComponent.h b/Engine/Source/Runtime/Engine/Classes/Components/SkinnedMeshComponent.h
--- a/Engine/Source/Runtime/Engine/Classes/Components/SkinnedMeshComponent.h
+++ b/Engine/Source/Runtime/Engine/Classes/Components/SkinnedMeshComponent.h
@@ -280,7 +280,7 @@ protected:
 
 protected:
 	/** The bounds radius at the point we last notified the streamer of a bounds radius change */
-	double LastStreamerUpdateBoundsRadius;
+	float LastStreamerUpdateBoundsRadius;
 
 	/** The index for the ComponentSpaceTransforms buffer we can currently write to */
 	int32 CurrentEditableComponentTransforms;
diff --git a/Engine/Source/Runtime/Engine/Private/Components/SkinnedMeshComponent.cpp b/Engine/Source/Runtime/Engine/Private/Components/SkinnedMeshComponent.cpp
--- a/Engine/Source/Runtime/Engine/Private/Components/SkinnedMeshComponent.cpp
+++ b/Engine/Source/Runtime/Engine/Private/Components/SkinnedMeshComponent.cpp
@@ -58,19 +58,11 @@ static TAutoConsoleVariable<int32> CVarAnimVisualizeLODs(
 	0,
 	TEXT("Visualize SkelMesh LODs"));
 
-int32 GUpdateBoundsNotifyStreaming = 1;
-FAutoConsoleVariableRef CVarUpdateBoundsNotifyStreaming(
-	TEXT("r.SkinnedMesh.UpdateBoundsNotifyStreaming"),
-	GUpdateBoundsNotifyStreaming,
-	TEXT("Update the streaming manager when the bounds change significantly"),
-	ECVF_Default
-);
-
 float GUpdateBoundsNotifyStreamingRadiusChangeRatio = 0.1f;
 FAutoConsoleVariableRef CVarUpdateBoundsNotifyStreamingRadiusChangeRatio(
 	TEXT("r.SkinnedMesh.UpdateBoundsNotifyStreamingRadiusChangeRatio"),
 	GUpdateBoundsNotifyStreamingRadiusChangeRatio,
-	TEXT("Update the streaming manager when the radius changes by more than this ratio since the last update. Only relevant with r.SkinnedMesh.UpdateBoundsNotifyStreaming=1"),
+	TEXT("Update the streaming manager when the radius changes by more than this ratio since the last update. A negative value will disable the update."),
 	ECVF_Default
 );
 
@@ -1186,45 +1178,20 @@ FBoxSphereBounds USkinnedMeshComponent::CalcBounds(const FTransform& LocalToWorl
 
 void USkinnedMeshComponent::UpdateBounds()
 {
-	if (GUpdateBoundsNotifyStreaming)
-	{
-		const auto OldRadius = Bounds.SphereRadius;
-
-		Super::UpdateBounds();
-
-		// Avoid updating the streamer for small changes in size
-		if (!FMath::IsNearlyEqual(OldRadius, Bounds.SphereRadius, 0.1f) && GetForcedLOD() == 0)
-		{
-			bool bNotify = true;
-
-		    // If a threshold is specified, don't notify the streamer if the radius has changed by less than the threshold since the last update
-		    if (GUpdateBoundsNotifyStreamingRadiusChangeRatio > 0.0f)
-		    {
-			    if (LastStreamerUpdateBoundsRadius >= 0.0 && LastStreamerUpdateBoundsRadius != Bounds.SphereRadius)
-			    {
-				    // Compute a change ratio based on the change in radius relative to the min of the two radii
-				    auto MinRadius = FMath::Min(Bounds.SphereRadius, LastStreamerUpdateBoundsRadius);
-				    auto MaxRadius = FMath::Max(Bounds.SphereRadius, LastStreamerUpdateBoundsRadius);
-				    if (MinRadius > 0.0)
-				    {
-					    float ChangeRatio = float(MaxRadius / MinRadius)-1.0f;
-					    if (ChangeRatio < GUpdateBoundsNotifyStreamingRadiusChangeRatio)
-					    {
-						    bNotify = false;
-					    }
-				    }
-			    }
-		    }
-		    if (bNotify)
-		    {
-			    IStreamingManager::Get().NotifyPrimitiveUpdated_Concurrent(this);
-			    LastStreamerUpdateBoundsRadius = Bounds.SphereRadius;
-		    }
-		}
+	if (LastStreamerUpdateBoundsRadius < 0.f)
+	{
+		LastStreamerUpdateBoundsRadius = float(Bounds.SphereRadius);
 	}
-	else
+
+	Super::UpdateBounds();
+
+	// Only notify the streamer if the bounds radius has changed enough
+	if (GUpdateBoundsNotifyStreamingRadiusChangeRatio >= 0.f
+		&& FMath::Abs(float(Bounds.SphereRadius) - LastStreamerUpdateBoundsRadius) > GUpdateBoundsNotifyStreamingRadiusChangeRatio * FMath::Min(float(Bounds.SphereRadius), LastStreamerUpdateBoundsRadius)
+		&& GetForcedLOD() == 0)
 	{
-		Super::UpdateBounds();
+		IStreamingManager::Get().NotifyPrimitiveUpdated_Concurrent(this);
+		LastStreamerUpdateBoundsRadius = float(Bounds.SphereRadius);
 	}
 }
 
