- Fix smooth client position update visual jittering when trying to smooth meshes with scaled parent components by dividing the unrotated mesh translation by the parent's component scale

diff --git a/Engine/Source/Runtime/Engine/Private/Components/CharacterMovementComponent.cpp b/Engine/Source/Runtime/Engine/Private/Components/CharacterMovementComponent.cpp
--- a/Engine/Source/Runtime/Engine/Private/Components/CharacterMovementComponent.cpp
+++ b/Engine/Source/Runtime/Engine/Private/Components/CharacterMovementComponent.cpp
@@ -7952,7 +7952,15 @@ void UCharacterMovementComponent::SmoothClientPosition_UpdateVisuals()
 		{
 			// Adjust capsule rotation and mesh location. Optimized to trigger only one transform chain update.
 			// If we know the rotation is changing that will update children, so it's sufficient to set RelativeLocation directly on the mesh.
-			const FVector NewRelLocation = ClientData->MeshRotationOffset.UnrotateVector(ClientData->MeshTranslationOffset) + CharacterOwner->GetBaseTranslationOffset();
+			const USceneComponent* MeshParent = Mesh->GetAttachParent();
+
+			FVector MeshParentScale = MeshParent != nullptr ? MeshParent->GetComponentScale() : FVector(1.0f, 1.0f, 1.0f);
+
+			MeshParentScale.X = FMath::IsNearlyZero(MeshParentScale.X) ? 1.0f : MeshParentScale.X;
+			MeshParentScale.Y = FMath::IsNearlyZero(MeshParentScale.Y) ? 1.0f : MeshParentScale.Y;
+			MeshParentScale.Z = FMath::IsNearlyZero(MeshParentScale.Z) ? 1.0f : MeshParentScale.Z;
+
+			const FVector NewRelLocation = (ClientData->MeshRotationOffset.UnrotateVector(ClientData->MeshTranslationOffset) / MeshParentScale) + CharacterOwner->GetBaseTranslationOffset();
 			if (!UpdatedComponent->GetComponentQuat().Equals(ClientData->MeshRotationOffset, 1e-6f))
 			{
 				const FVector OldLocation = Mesh->GetRelativeLocation();
@@ -7974,8 +7982,16 @@ void UCharacterMovementComponent::SmoothClientPosition_UpdateVisuals()
 		}
 		else if (NetworkSmoothingMode == ENetworkSmoothingMode::Exponential)
 		{
+			const USceneComponent* MeshParent = Mesh->GetAttachParent();
+
+			FVector MeshParentScale = MeshParent != nullptr ? MeshParent->GetComponentScale() : FVector(1.0f, 1.0f, 1.0f);
+
+			MeshParentScale.X = FMath::IsNearlyZero(MeshParentScale.X) ? 1.0f : MeshParentScale.X;
+			MeshParentScale.Y = FMath::IsNearlyZero(MeshParentScale.Y) ? 1.0f : MeshParentScale.Y;
+			MeshParentScale.Z = FMath::IsNearlyZero(MeshParentScale.Z) ? 1.0f : MeshParentScale.Z;
+
 			// Adjust mesh location and rotation
-			const FVector NewRelTranslation = UpdatedComponent->GetComponentToWorld().InverseTransformVectorNoScale(ClientData->MeshTranslationOffset) + CharacterOwner->GetBaseTranslationOffset();
+			const FVector NewRelTranslation = (UpdatedComponent->GetComponentToWorld().InverseTransformVectorNoScale(ClientData->MeshTranslationOffset) / MeshParentScale) + CharacterOwner->GetBaseTranslationOffset();
 			const FQuat NewRelRotation = ClientData->MeshRotationOffset * CharacterOwner->GetBaseRotationOffset();
 			Mesh->SetRelativeLocationAndRotation(NewRelTranslation, NewRelRotation, false, nullptr, GetTeleportType());
 		}
