Merge Actors: Add missing members to `FMaterialProxySettings::operator==` implemention

#rnx

diff --git a/Engine/Source/Runtime/Engine/Classes/Engine/MaterialMerging.h b/Engine/Source/Runtime/Engine/Classes/Engine/MaterialMerging.h
--- a/Engine/Source/Runtime/Engine/Classes/Engine/MaterialMerging.h
+++ b/Engine/Source/Runtime/Engine/Classes/Engine/MaterialMerging.h
@@ -228,26 +228,33 @@ struct FMaterialProxySettings
 			&& MeshMaxScreenSizePercent == Other.MeshMaxScreenSizePercent
 			&& MeshMinDrawDistance == Other.MeshMinDrawDistance
 			&& GutterSpace == Other.GutterSpace
+			&& MetallicConstant == Other.MetallicConstant
+			&& RoughnessConstant == Other.RoughnessConstant
+			&& AnisotropyConstant == Other.AnisotropyConstant
+			&& SpecularConstant == Other.SpecularConstant
+			&& OpacityConstant == Other.OpacityConstant
+			&& OpacityMaskConstant == Other.OpacityMaskConstant
+			&& AmbientOcclusionConstant == Other.AmbientOcclusionConstant
+			&& MaterialMergeType == Other.MaterialMergeType
+			&& BlendMode == Other.BlendMode
+			&& bAllowTwoSidedMaterial == Other.bAllowTwoSidedMaterial
 			&& bNormalMap == Other.bNormalMap
 			&& bTangentMap == Other.bTangentMap
-			&& MetallicConstant == Other.MetallicConstant
 			&& bMetallicMap == Other.bMetallicMap
-			&& RoughnessConstant == Other.RoughnessConstant
 			&& bRoughnessMap == Other.bRoughnessMap
-			&& AnisotropyConstant == Other.AnisotropyConstant
 			&& bAnisotropyMap == Other.bAnisotropyMap
-			&& SpecularConstant == Other.SpecularConstant
 			&& bSpecularMap == Other.bSpecularMap
 			&& bEmissiveMap == Other.bEmissiveMap
 			&& bOpacityMap == Other.bOpacityMap
 			&& bOpacityMaskMap == Other.bOpacityMaskMap
 			&& bAmbientOcclusionMap == Other.bAmbientOcclusionMap
-			&& AmbientOcclusionConstant == Other.AmbientOcclusionConstant
 			&& DiffuseTextureSize == Other.DiffuseTextureSize
 			&& NormalTextureSize == Other.NormalTextureSize
+			&& TangentTextureSize == Other.TangentTextureSize
 			&& MetallicTextureSize == Other.MetallicTextureSize
 			&& RoughnessTextureSize == Other.RoughnessTextureSize
 			&& AnisotropyTextureSize == Other.AnisotropyTextureSize
+			&& SpecularTextureSize == Other.SpecularTextureSize
 			&& EmissiveTextureSize == Other.EmissiveTextureSize
 			&& OpacityTextureSize == Other.OpacityTextureSize
 			&& OpacityMaskTextureSize == Other.OpacityMaskTextureSize
