Apply DMFP clamping to Separable SSS and transmission as well.

diff --git a/Engine/Source/Runtime/Engine/Private/Rendering/SubsurfaceProfile.cpp b/Engine/Source/Runtime/Engine/Private/Rendering/SubsurfaceProfile.cpp
--- a/Engine/Source/Runtime/Engine/Private/Rendering/SubsurfaceProfile.cpp
+++ b/Engine/Source/Runtime/Engine/Private/Rendering/SubsurfaceProfile.cpp
@@ -362,7 +362,7 @@ float DecodeWorldUnitScale(float EncodedWorldUnitScale)
 //in [0,1]
 FLinearColor EncodeDiffuseMeanFreePath(FLinearColor DiffuseMeanFreePath)
 {
-	return (DiffuseMeanFreePath * ENC_DIFFUSEMEANFREEPATH_IN_MM_TO_UNIT).GetClamped();
+	return DiffuseMeanFreePath * ENC_DIFFUSEMEANFREEPATH_IN_MM_TO_UNIT;
 }
 
 FLinearColor DecodeDiffuseMeanFreePath(FLinearColor EncodedDiffuseMeanFreePath)
@@ -384,6 +384,9 @@ void SetupSurfaceAlbedoAndDiffuseMeanFreePath(FLinearColor& SurfaceAlbedo, FLine
 
 	SurfaceAlbedo.A = SurfaceAlbedo.Component(IndexOfMaxComp);
 	Dmfp.A = Dmfp.Component(IndexOfMaxComp);
+
+	// Apply clamping so that dmfp is within encoding range.
+	Dmfp = Dmfp.GetClamped(0.0f, DEC_UNIT_TO_DIFFUSEMEANFREEPATH_IN_MM);
 }
 
 float Sqrt2(float X)
