Compile time switched versions of FindUFieldOrFProperty that resolve to either FindUField or FindFProperty based on their template type

#rnx

diff --git a/Engine/Source/Runtime/CoreUObject/Public/UObject/UnrealType.h b/Engine/Source/Runtime/CoreUObject/Public/UObject/UnrealType.h
--- a/Engine/Source/Runtime/CoreUObject/Public/UObject/UnrealType.h
+++ b/Engine/Source/Runtime/CoreUObject/Public/UObject/UnrealType.h
@@ -6701,6 +6701,30 @@ inline FFieldVariant FindUFieldOrFProperty(const UStruct* Owner, const TCHAR* Fi
 	return FindUFieldOrFProperty(Owner, Name, IterationFlags);
 }
 
+template <class T>
+typename TEnableIf<TIsDerivedFrom<T, UField>::IsDerived, T*>::Type FindUFieldOrFProperty(const UStruct* Owner, FName FieldName, EFieldIterationFlags IterationFlags = EFieldIterationFlags::Default)
+{
+	return FindUField<T>(Owner, FieldName, IterationFlags);
+}
+
+template <class T>
+typename TEnableIf<TIsDerivedFrom<T, UField>::IsDerived, T*>::Type FindUFieldOrFProperty(const UStruct* Owner, const TCHAR* FieldName, EFieldIterationFlags IterationFlags = EFieldIterationFlags::Default)
+{
+	return FindUField<T>(Owner, FieldName, IterationFlags);
+}
+
+template <class T>
+typename TEnableIf<TIsDerivedFrom<T, FField>::IsDerived, T*>::Type FindUFieldOrFProperty(const UStruct* Owner, FName FieldName, EFieldIterationFlags IterationFlags = EFieldIterationFlags::Default)
+{
+	return FindFProperty<T>(Owner, FieldName, IterationFlags);
+}
+
+template <class T>
+typename TEnableIf<TIsDerivedFrom<T, FField>::IsDerived, T*>::Type FindUFieldOrFProperty(const UStruct* Owner, const TCHAR* FieldName, EFieldIterationFlags IterationFlags = EFieldIterationFlags::Default)
+{
+	return FindFProperty<T>(Owner, FieldName, IterationFlags);
+}
+
 /**
  * Search for the named field within the specified scope, including any Outer classes; assert on failure.
  *
diff --git a/Engine/Source/Runtime/Engine/Classes/Engine/MemberReference.h b/Engine/Source/Runtime/Engine/Classes/Engine/MemberReference.h
--- a/Engine/Source/Runtime/Engine/Classes/Engine/MemberReference.h
+++ b/Engine/Source/Runtime/Engine/Classes/Engine/MemberReference.h
@@ -372,7 +372,7 @@ public:
 			UStruct* MemberScopeStruct = FindUField<UStruct>(SelfScope, *MemberScope);
 
 			// Find in target scope
-			ReturnField = FindUFieldOrFProperty(MemberScopeStruct, MemberName, EFieldIterationFlags::IncludeAll).Get<TFieldType>();
+			ReturnField = FindUFieldOrFProperty<TFieldType>(MemberScopeStruct, MemberName, EFieldIterationFlags::IncludeAll);
 
 #if WITH_EDITOR
 			if(ReturnField == nullptr)
@@ -381,7 +381,7 @@ public:
 				const FName RenamedMemberName = RefreshLocalVariableName(SelfScope);
 				if (RenamedMemberName != NAME_None)
 				{
-					ReturnField = FindUFieldOrFProperty(MemberScopeStruct, MemberName, EFieldIterationFlags::IncludeAll).Get<TFieldType>();
+					ReturnField = FindUFieldOrFProperty<TFieldType>(MemberScopeStruct, MemberName, EFieldIterationFlags::IncludeAll);
 				}
 			}
 #endif
@@ -428,12 +428,12 @@ public:
 #endif
 				{
 					// Find in target scope or in the sparse class data
-					ReturnField = FindUFieldOrFProperty(TargetScope, MemberName, EFieldIterationFlags::IncludeAll).Get<TFieldType>();
+					ReturnField = FindUFieldOrFProperty<TFieldType>(TargetScope, MemberName, EFieldIterationFlags::IncludeAll);
 					if (!ReturnField)
 					{
 						if (UScriptStruct* SparseClassDataStruct = TargetScope->GetSparseClassDataStruct())
 						{
-							ReturnField = FindUFieldOrFProperty(SparseClassDataStruct, MemberName, EFieldIterationFlags::IncludeAll).Get<TFieldType>();
+							ReturnField = FindUFieldOrFProperty<TFieldType>(SparseClassDataStruct, MemberName, EFieldIterationFlags::IncludeAll);
 						}
 					}
 				}
@@ -453,7 +453,7 @@ public:
 					if (RenamedMemberName != NAME_None)
 					{
 						MemberName = RenamedMemberName;
-						ReturnField = FindUFieldOrFProperty(TargetScope, MemberName, EFieldIterationFlags::IncludeAll).Get<TFieldType>();
+						ReturnField = FindUFieldOrFProperty<TFieldType>(TargetScope, MemberName, EFieldIterationFlags::IncludeAll);
 					}
 				}
 #endif
diff --git a/Engine/Source/Runtime/Engine/Private/MemberReference.cpp b/Engine/Source/Runtime/Engine/Private/MemberReference.cpp
--- a/Engine/Source/Runtime/Engine/Private/MemberReference.cpp
+++ b/Engine/Source/Runtime/Engine/Private/MemberReference.cpp
@@ -284,8 +284,7 @@ TFieldType* FindRemappedFieldImpl(FName FieldClassOutermostName, FName FieldClas
 	FMemberReference::InitFieldRedirectMap();
 
 	// In the case of a bifurcation of a variable (e.g. moved from a parent into certain children), verify that we don't also define the variable in the current scope first
-	FFieldVariant ExistingField = FindUFieldOrFProperty(InitialScope, InitialName);
-	if (ExistingField.Get<TFieldType>())
+	if (FindUFieldOrFProperty<TFieldType>(InitialScope, InitialName))
 	{
 		return nullptr;
 	}
@@ -329,7 +328,7 @@ TFieldType* FindRemappedFieldImpl(FName FieldClassOutermostName, FName FieldClas
 		if (NewFieldName != NAME_None)
 		{
 			// Find the actual field specified by the redirector, so we can return it and update the node that uses it
-			TFieldType* NewField = FindUFieldOrFProperty(SearchClass, NewFieldName).Get<TFieldType>();
+			TFieldType* NewField = FindUFieldOrFProperty<TFieldType>(SearchClass, NewFieldName);
 			if (NewField != nullptr)
 			{
 				if (bInitialScopeMustBeOwnerOfField && !InitialScope->IsChildOf(SearchClass))
