- Remove ensure in case first chunk fails to load and log descriptive error in the case that the asset is corrupt or was failed to be retrieved
- Rename EnsureZerothChunkIsLoaded (because it doesn't actually ensure) to GetOrLoadFirstChunk
- Misc error-state logging fix-ups
- Report if seek failed for decoder input  in more robust log statement and remove ensure
- Fix modified 0th chunk size not invalidating compressed data in OnPropertyChanged
#rnx

diff --git a/Engine/Source/Editor/AudioEditor/Private/Factories/SoundFactory.cpp b/Engine/Source/Editor/AudioEditor/Private/Factories/SoundFactory.cpp
--- a/Engine/Source/Editor/AudioEditor/Private/Factories/SoundFactory.cpp
+++ b/Engine/Source/Editor/AudioEditor/Private/Factories/SoundFactory.cpp
@@ -574,7 +574,7 @@ UObject* USoundFactory::CreateObject
 		// If stream caching is enabled, we need to make sure this asset is ready for playback.
 		if (bRebuildStreamingChunks && Sound->IsStreaming(nullptr))
 		{
-			Sound->EnsureZerothChunkIsLoaded();
+			Sound->LoadZerothChunk();
 		}
 
 		GEditor->GetEditorSubsystem<UImportSubsystem>()->BroadcastAssetPostImport(this, Sound);
diff --git a/Engine/Source/Runtime/Engine/Classes/Sound/SoundWave.h b/Engine/Source/Runtime/Engine/Classes/Sound/SoundWave.h
--- a/Engine/Source/Runtime/Engine/Classes/Sound/SoundWave.h
+++ b/Engine/Source/Runtime/Engine/Classes/Sound/SoundWave.h
@@ -792,8 +792,10 @@ public:
 	virtual void FinishDestroy() override;
 	virtual void PostLoad() override;
 
-	// When stream caching is enabled, this is called after we've successfully compressed and split the streamed audio for this file.
-	void EnsureZerothChunkIsLoaded();
+	// Returns true if the zeroth chunk is loaded, or attempts to load it if not already loaded,
+	// returning true if the load was successful. Can return false if either an error was encountered
+	// in attempting to load the chunk or if stream caching is not enabled for the given sound.
+	bool LoadZerothChunk();
 
 	// Returns the amount of chunks this soundwave contains if it's streaming,
 	// or zero if it is not a streaming source.
@@ -1203,7 +1205,11 @@ public:
  
  	bool IsZerothChunkDataLoaded() const;
  	const TArrayView<uint8> GetZerothChunkDataView() const;
- 	void EnsureZerothChunkIsLoaded();
+
+	// Returns true if the zeroth chunk is loaded, or attempts to load it if not already loaded,
+	// returning true if the load was successful. Can return false if either an error was encountered
+	// in attempting to load the chunk or if stream caching is not enabled for the given sound.
+	bool LoadZerothChunk();
  
  #if WITH_EDITOR
  	int32 GetCurrentChunkRevision() const;
@@ -1287,7 +1293,11 @@ public:
 
 	// USoundWave Interface
 	void ReleaseCompressedAudio();
-	void EnsureZerothChunkIsLoaded();
+
+	// Returns true if the zeroth chunk is loaded, or attempts to load it if not already loaded,
+	// returning true if the load was successful. Can return false if either an error was encountered
+	// in attempting to load the chunk or if stream caching is not enabled for the given sound.
+	bool LoadZerothChunk();
 	bool GetChunkData(int32 ChunkIndex, uint8** OutChunkData, bool bMakeSureChunkIsLoaded = false);
 
 	// Getters
@@ -1310,7 +1320,6 @@ public:
 	bool IsLooping() const;
 	bool IsTemplate() const;
 	bool IsStreaming() const;
-	bool UseBinkAudio() const;
 	bool IsRetainingAudio() const;
 	bool ShouldUseStreamCaching() const;
 	bool IsSeekable() const;
diff --git a/Engine/Source/Runtime/Engine/Private/SoundWave.cpp b/Engine/Source/Runtime/Engine/Private/SoundWave.cpp
--- a/Engine/Source/Runtime/Engine/Private/SoundWave.cpp
+++ b/Engine/Source/Runtime/Engine/Private/SoundWave.cpp
@@ -134,6 +134,7 @@ namespace SoundWaveCookStats
 }
 #endif
 
+
 FSoundWaveData::~FSoundWaveData()
 {
 	// at this point, first chunk data should have been cleared by the last reference holder
@@ -189,7 +190,7 @@ void FSoundWaveData::InitializeDataFromSoundWave(USoundWave& InWave)
 uint32 FSoundWaveData::GetNumChunks() const
 {
 	// This function shouldn't be called on audio marked "ForceInline."
-	ensureMsgf(LoadingBehavior != ESoundWaveLoadingBehavior::ForceInline, TEXT("Calling GetNumChunks without RunnigPlatformData is not allowed! SoundWave: %s - %s")
+	ensureMsgf(LoadingBehavior != ESoundWaveLoadingBehavior::ForceInline, TEXT("Calling GetNumChunks without RunningPlatformData is not allowed! SoundWave: %s - %s")
 		, *GetFName().ToString(), EnumToString(LoadingBehavior));
 
 	return RunningPlatformData.Chunks.Num();
@@ -214,7 +215,7 @@ uint32 FSoundWaveData::GetSizeOfChunk(uint32 ChunkIndex) const
 
 	// This function shouldn't be called on audio marked "ForceInline."
 	ensureMsgf(RunningPlatformData.NumChunks
-		, TEXT("Calling GetNumChunks() without RunnigPlatformData is not allowed! SoundWave: %s - %s")
+		, TEXT("Calling GetNumChunks() without RunningPlatformData is not allowed! SoundWave: %s - %s")
 		, *GetFName().ToString()
 		, EnumToString(LoadingBehavior));
 
@@ -341,40 +342,51 @@ const TArrayView<uint8> FSoundWaveData::GetZerothChunkDataView() const
 	return TArrayView<uint8>(View.GetData(), static_cast<int32>(View.Num()));
 }
 
-void FSoundWaveData::EnsureZerothChunkIsLoaded()
+bool FSoundWaveData::LoadZerothChunk()
 {
 	LLM_SCOPE(ELLMTag::AudioSoundWaves);
 
 	// If the zeroth chunk is already loaded, early exit.
-	if (ZerothChunkData.GetView().Num() > 0 || !bShouldUseStreamCaching)
+	if (IsZerothChunkDataLoaded())
 	{
-		return;
+		return true;
 	}
 
-	// #if WITH_EDITOR
-#ifdef WITH_EDITOR
-	// If we're running the editor, we'll need to retrieve the chunked audio from the DDC:
+	if (!ShouldUseStreamCaching())
+	{
+		return false;
+	}
+
+#if WITH_EDITOR
+	// If running editor, retrieve chunked audio from the DDC
 	uint8* TempChunkBuffer = nullptr;
 	int32 ChunkSizeInBytes = GetChunkFromDDC(0, &TempChunkBuffer, true);
-	// Since we block for the DDC in the previous call we should always have the chunk loaded.
+
 	if (ChunkSizeInBytes == 0)
 	{
-		return;
+		UE_LOG(LogAudio, Warning, TEXT("FSoundWaveData::LoadZerothChunk: Unsuccessful load of zeroth chunk from DDC. Asset requires manual re-cook.\n\tAsset: '%s'\n\tDerivedDataKey: '%s'"),
+			*GetFName().ToString(),
+			*RunningPlatformData.DerivedDataKey);
+		return false;
 	}
-
 	ZerothChunkData.Reset(TempChunkBuffer, ChunkSizeInBytes);
 
 #else // WITH_EDITOR
-	// Otherwise, the zeroth chunk is cooked out to SharedRuntimeData::RunningPlatformData, and we just need to retrieve it.
+	// Otherwise, the zeroth chunk is cooked out to RunningPlatformData, so retrieve it.
 	check(GetNumChunks() > 0);
-	const FStreamedAudioChunk& ZerothChunk = GetChunk(0);
-	// Some sanity checks to ensure that the bulk size set up
+	FStreamedAudioChunk& ZerothChunk = GetChunk(0);
+
+	// Sanity check to ensure bulk size is set up
 	UE_CLOG(ZerothChunk.BulkData.GetBulkDataSize() != ZerothChunk.DataSize, LogAudio, Warning
-		, TEXT("Bulk data serialized out had a mismatched size with the DataSize field. Soundwave: %s Bulk Data Reported Size: %d Bulk Data Actual Size: %ld")
-		, *GetFullName(), ZerothChunk.DataSize, ZerothChunk.BulkData.GetBulkDataSize());
+		, TEXT("SoundWave '%s' bulk data serialized out had a mismatched size with the DataSize field."
+			"\nBulk Data Reported Size: %d"
+			"\nBulk Data Actual Size: %ld")
+		, *GetFName().ToString(), ZerothChunk.DataSize, ZerothChunk.BulkData.GetBulkDataSize());
 
 	ZerothChunkData = ZerothChunk.BulkData.GetCopyAsBuffer(ZerothChunk.AudioDataSize, true);
 #endif // WITH_EDITOR
+
+	return true;
 }
 
 #if WITH_EDITOR
@@ -388,7 +400,7 @@ FStreamedAudioChunk& FSoundWaveData::GetChunk(uint32 ChunkIndex)
 {
 	// This function shouldn't be called on audio marked "ForceInline."
 	ensureMsgf(RunningPlatformData.NumChunks
-		, TEXT("Calling GetNumChunks() without RunnigPlatformData is not allowed! SoundWave: %s - %s")
+		, TEXT("Calling GetNumChunks() without RunningPlatformData is not allowed! SoundWave: %s - %s")
 		, *GetFName().ToString()
 		, EnumToString(LoadingBehavior));
 
@@ -401,7 +413,7 @@ int32 FSoundWaveData::GetChunkFromDDC(int32 ChunkIndex, uint8** OutChunkData, bo
 	LLM_SCOPE(ELLMTag::AudioSoundWaves);
 	// This function shouldn't be called on audio marked "ForceInline."
 	ensureMsgf(RunningPlatformData.NumChunks
-		, TEXT("Calling GetNumChunks on a FSoundWaveProxy without RunnigPlatformData is not allowed! SoundWave: %s - %s")
+		, TEXT("Calling GetNumChunks on a FSoundWaveProxy without RunningPlatformData is not allowed! SoundWave: %s - %s")
 		, *GetFName().ToString()
 		, EnumToString(LoadingBehavior));
 
@@ -415,7 +427,7 @@ FString FSoundWaveData::GetDerivedDataKey() const
 
 	// This function shouldn't be called on audio marked "ForceInline."
 	ensureMsgf(RunningPlatformData.NumChunks
-		, TEXT("Calling GetNumChunks on a FSoundWaveProxy without RunnigPlatformData is not allowed! SoundWave: %s - %s")
+		, TEXT("Calling GetNumChunks on a FSoundWaveProxy without RunningPlatformData is not allowed! SoundWave: %s - %s")
 		, *GetFName().ToString()
 		, EnumToString(LoadingBehavior));
 
@@ -913,7 +925,6 @@ void USoundWave::Serialize( FArchive& Ar )
 		{
 			ESoundWaveLoadingBehavior CurrentLoadingBehavior = GetLoadingBehavior(false);
 
-			//EnsureZerothChunkIsLoaded();
 			const bool bHasFirstChunk = GetNumChunks() > 1;
 			
 			if (!bHasFirstChunk)
@@ -1572,7 +1583,7 @@ void USoundWave::PostLoad()
 
 	if (ShouldUseStreamCaching() && bDoesSoundWaveHaveStreamingAudioData)
 	{
-		EnsureZerothChunkIsLoaded();
+		LoadZerothChunk();
 	}
 
 #if WITH_EDITORONLY_DATA
@@ -1599,14 +1610,19 @@ void USoundWave::PostLoad()
 	}
 }
 
-void USoundWave::EnsureZerothChunkIsLoaded()
+bool USoundWave::LoadZerothChunk()
 {
 	check(SoundWaveDataPtr);
 
 	// If the zeroth chunk is already loaded, early exit.
 	if (SoundWaveDataPtr->ZerothChunkData.GetView().Num() > 0 || !ShouldUseStreamCaching())
 	{
-		return;
+		return true;
+	}
+
+	if (!ShouldUseStreamCaching())
+	{
+		return true;
 	}
 
 #if WITH_EDITOR 
@@ -1614,7 +1630,7 @@ void USoundWave::EnsureZerothChunkIsLoaded()
 	if (!SoundWaveDataPtr.IsValid())
 	{
 		// we may be in the middle of garbage collection, don't access RunningPlatformData
-		return;
+		return false;
 	}
 
 	CachePlatformData(false);
@@ -1622,10 +1638,12 @@ void USoundWave::EnsureZerothChunkIsLoaded()
 	// If we're running the editor, we'll need to retrieve the chunked audio from the DDC:
 	uint8* TempChunkBuffer = nullptr;
 	int32 ChunkSizeInBytes = SoundWaveDataPtr->RunningPlatformData.GetChunkFromDDC(0, &TempChunkBuffer, true);
-	// Since we block for the DDC in the previous call we should always have the chunk loaded.
 	if (ChunkSizeInBytes == 0)
 	{
-		return;
+		UE_LOG(LogAudio, Warning, TEXT("USoundWave::LoadZerothChunk: Unsuccessful load of zeroth chunk from DDC. Asset requires manual re-cook.\n\tAsset: '%s'\n\tDerivedDataKey: '%s'"),
+			*GetFName().ToString(),
+			*SoundWaveDataPtr->RunningPlatformData.DerivedDataKey);
+		return false;
 	}
 
 	SoundWaveDataPtr->ZerothChunkData.Reset(TempChunkBuffer, ChunkSizeInBytes);
@@ -1639,6 +1657,8 @@ void USoundWave::EnsureZerothChunkIsLoaded()
 
 	SoundWaveDataPtr->ZerothChunkData = ZerothChunk.BulkData.GetCopyAsBuffer(ZerothChunk.AudioDataSize, true);
 #endif // WITH_EDITOR
+
+	return true;
 }
 
 uint32 USoundWave::GetNumChunks() const
@@ -1825,7 +1845,7 @@ void USoundWave::InvalidateSoundWaveIfNeccessary()
 		// If stream caching is now turned on, recook the streaming audio if neccessary.
 		if (bIsStreamCachingEnabled && IsStreaming(nullptr))
 		{
-			EnsureZerothChunkIsLoaded();
+			LoadZerothChunk();
 		}
 	}
 }
@@ -2244,6 +2264,7 @@ void USoundWave::PostEditChangeProperty(FPropertyChangedEvent& PropertyChangedEv
 	static const FName StreamingFName = GET_MEMBER_NAME_CHECKED(USoundWave, bStreaming);
 	static const FName SoundAssetCompressionTypeFName = GET_MEMBER_NAME_CHECKED(USoundWave, SoundAssetCompressionType);
 	static const FName LoadingBehaviorFName = GET_MEMBER_NAME_CHECKED(USoundWave, LoadingBehavior);
+	static const FName InitialChunkSizeFName = GET_MEMBER_NAME_CHECKED(USoundWave, InitialChunkSize);
 
 	// force proxy flags to be up to date
 	SoundWaveDataPtr->bIsSeekable = IsSeekable();
@@ -2276,7 +2297,12 @@ void USoundWave::PostEditChangeProperty(FPropertyChangedEvent& PropertyChangedEv
 				}
 			}
 
-			if (Name == CompressionQualityFName || Name == SampleRateFName || Name == StreamingFName || Name == SoundAssetCompressionTypeFName || Name == LoadingBehaviorFName)
+			if (Name == CompressionQualityFName
+				|| Name == SampleRateFName
+				|| Name == StreamingFName
+				|| Name == SoundAssetCompressionTypeFName
+				|| Name == LoadingBehaviorFName
+				|| Name == InitialChunkSizeFName)
 			{
 				UpdateAsset();
 			}
@@ -2930,7 +2956,7 @@ TArrayView<const uint8> USoundWave::GetZerothChunk(bool bForImmediatePlayback)
 		// In editor, we actually don't have a zeroth chunk until we try to play an audio file.
 		if (GIsEditor)
 		{
-			EnsureZerothChunkIsLoaded();
+			LoadZerothChunk();
 		}
 
 		check(SoundWaveDataPtr->ZerothChunkData.GetView().Num() > 0);
@@ -3013,7 +3039,7 @@ int32 USoundWave::GetCompressionQuality() const
 FSoundWaveProxyPtr USoundWave::CreateSoundWaveProxy()
 {
 #if WITH_EDITORONLY_DATA
-	EnsureZerothChunkIsLoaded();
+	LoadZerothChunk();
 #endif // #if WITH_EDITORONLY_DATA
 
 	LLM_SCOPE(ELLMTag::AudioSoundWaveProxies);
@@ -3029,7 +3055,10 @@ TUniquePtr<Audio::IProxyData> USoundWave::CreateNewProxyData(const Audio::FProxy
 	check(SoundWaveDataPtr);
 
 #if WITH_EDITORONLY_DATA
-	EnsureZerothChunkIsLoaded();
+	if (!LoadZerothChunk())
+	{
+		return nullptr;
+	}
 #endif // #if WITH_EDITORONLY_DATA
 
 	LLM_SCOPE(ELLMTag::AudioSoundWaveProxies);
@@ -3808,10 +3837,10 @@ void FSoundWaveProxy::ReleaseCompressedAudio()
 	return SoundWaveDataPtr->ReleaseCompressedAudio();
 }
 
-void FSoundWaveProxy::EnsureZerothChunkIsLoaded()
+bool FSoundWaveProxy::LoadZerothChunk()
 {
 	check(SoundWaveDataPtr);
-	return SoundWaveDataPtr->EnsureZerothChunkIsLoaded();
+	return SoundWaveDataPtr->LoadZerothChunk();
 }
 
 bool FSoundWaveProxy::GetChunkData(int32 ChunkIndex, uint8** OutChunkData, bool bMakeSureChunkIsLoaded)
@@ -3971,10 +4000,10 @@ TArrayView<const uint8> FSoundWaveProxy::GetZerothChunk(const FSoundWaveProxyPtr
 	{
 		if (GIsEditor)
 		{
-			SoundWaveProxy->EnsureZerothChunkIsLoaded();
+			SoundWaveProxy->LoadZerothChunk();
 		}
 
-		if (ensure(SoundWaveProxy->IsZerothChunkDataLoaded()))
+		if (SoundWaveProxy->IsZerothChunkDataLoaded())
 		{
 			if (SoundWaveProxy->SoundWaveDataPtr->ShouldUseStreamCaching())
 			{
