Track both input/output types for material graph connections
Pins are now color-coded by type in addition to wires

diff --git a/Engine/Source/Editor/GraphEditor/Private/MaterialGraphConnectionDrawingPolicy.cpp b/Engine/Source/Editor/GraphEditor/Private/MaterialGraphConnectionDrawingPolicy.cpp
--- a/Engine/Source/Editor/GraphEditor/Private/MaterialGraphConnectionDrawingPolicy.cpp
+++ b/Engine/Source/Editor/GraphEditor/Private/MaterialGraphConnectionDrawingPolicy.cpp
@@ -104,15 +104,18 @@ bool FMaterialGraphConnectionDrawingPolicy::ShouldChangeTangentForKnot(UMaterial
 		return bPinReversed;
 	}
 }
+
 void FMaterialGraphConnectionDrawingPolicy::DetermineWiringStyle(UEdGraphPin* OutputPin, UEdGraphPin* InputPin, /*inout*/ FConnectionParams& Params)
 {
 	Params.AssociatedPin1 = OutputPin;
 	Params.AssociatedPin2 = InputPin;
 	Params.WireColor = MaterialGraphSchema->ActivePinColor;
 
-	FMaterialConnectionKey ConnectionKey;
+	UE::Shader::EValueType InputType = UE::Shader::EValueType::Void;
+	UE::Shader::EValueType OutputType = UE::Shader::EValueType::Void;
 	bool bInactivePin = false;
 	bool bExecPin = false;
+	bool bValueTypePin = false;
 
 	// Have to consider both pins as the input will be an 'output' when previewing a connection
 	if (OutputPin)
@@ -121,10 +124,20 @@ void FMaterialGraphConnectionDrawingPolicy::DetermineWiringStyle(UEdGraphPin* Ou
 		{
 			bInactivePin = true;
 		}
+
 		if (OutputPin->PinType.PinCategory == UMaterialGraphSchema::PC_Exec)
 		{
 			bExecPin = true;
 		}
+		else if (OutputPin->PinType.PinCategory == UMaterialGraphSchema::PC_Void)
+		{
+			bValueTypePin = true;
+		}
+		else if (OutputPin->PinType.PinCategory == UMaterialGraphSchema::PC_ValueType)
+		{
+			OutputType = UE::Shader::FindValueType(OutputPin->PinType.PinSubCategory);
+			bValueTypePin = true;
+		}
 
 		UEdGraphNode* OutputNode = OutputPin->GetOwningNode();
 		if (UMaterialGraphNode_Knot* OutputKnotNode = Cast<UMaterialGraphNode_Knot>(OutputNode))
@@ -138,11 +151,6 @@ void FMaterialGraphConnectionDrawingPolicy::DetermineWiringStyle(UEdGraphPin* Ou
 		{
 			bInactivePin = true;
 		}
-		else if(UMaterialGraphNode_Base* OutputMaterialNode = Cast<UMaterialGraphNode_Base>(OutputNode))
-		{
-			ConnectionKey.OutputObject = OutputMaterialNode->GetMaterialNodeOwner();
-			ConnectionKey.OutputIndex = OutputPin->SourceIndex;
-		}
 	}
 	if (InputPin)
 	{
@@ -150,10 +158,20 @@ void FMaterialGraphConnectionDrawingPolicy::DetermineWiringStyle(UEdGraphPin* Ou
 		{
 			bInactivePin = true;
 		}
+
 		if (InputPin->PinType.PinCategory == UMaterialGraphSchema::PC_Exec)
 		{
 			bExecPin = true;
 		}
+		else if (InputPin->PinType.PinCategory == UMaterialGraphSchema::PC_Void)
+		{
+			bValueTypePin = true;
+		}
+		else if (InputPin->PinType.PinCategory == UMaterialGraphSchema::PC_ValueType)
+		{
+			InputType = UE::Shader::FindValueType(InputPin->PinType.PinSubCategory);
+			bValueTypePin = true;
+		}
 
 		UEdGraphNode* InputNode = InputPin->GetOwningNode();
 		if (UMaterialGraphNode_Knot* InputKnotNode = Cast<UMaterialGraphNode_Knot>(InputNode))
@@ -167,11 +185,6 @@ void FMaterialGraphConnectionDrawingPolicy::DetermineWiringStyle(UEdGraphPin* Ou
 		{
 			bInactivePin = true;
 		}
-		else if (UMaterialGraphNode_Base* InputMaterialNode = Cast<UMaterialGraphNode_Base>(InputNode))
-		{
-			ConnectionKey.InputObject = InputMaterialNode->GetMaterialNodeOwner();
-			ConnectionKey.InputIndex = InputMaterialNode->GetInputIndexForPin(InputPin);
-		}
 	}
 
 	if (bInactivePin)
@@ -183,43 +196,12 @@ void FMaterialGraphConnectionDrawingPolicy::DetermineWiringStyle(UEdGraphPin* Ou
 		Params.WireColor = Settings->ExecutionPinTypeColor;
 		Params.WireThickness = Settings->DefaultExecutionWireThickness;
 	}
-	else
+	else if (bValueTypePin)
 	{
-		const UE::Shader::EValueType ConnectionType = MaterialGraph->GetConnectionType(ConnectionKey);
-		if (ConnectionType == UE::Shader::EValueType::Struct)
-		{
-			Params.WireColor = Settings->StructPinTypeColor;
-		}
-		else
+		// TODO - how to handle InputPin being nullptr, when does this occur?
+		if (InputPin)
 		{
-			const UE::Shader::FValueTypeDescription TypeDesc = UE::Shader::GetValueTypeDescription(ConnectionType);
-			if (UE::Shader::IsTextureType(TypeDesc.ComponentType))
-			{
-				Params.WireColor = Settings->ObjectPinTypeColor;
-			}
-			else if (TypeDesc.ComponentType == UE::Shader::EValueComponentType::Float)
-			{
-				if (TypeDesc.NumComponents == 1)
-				{
-					Params.WireColor = Settings->FloatPinTypeColor;
-				}
-				else
-				{
-					Params.WireColor = Settings->VectorPinTypeColor;
-				}
-			}
-			else if (TypeDesc.ComponentType == UE::Shader::EValueComponentType::Double)
-			{
-				Params.WireColor = Settings->DoublePinTypeColor;
-			}
-			else if (TypeDesc.ComponentType == UE::Shader::EValueComponentType::Bool)
-			{
-				Params.WireColor = Settings->BooleanPinTypeColor;
-			}
-			else if (TypeDesc.ComponentType == UE::Shader::EValueComponentType::Int)
-			{
-				Params.WireColor = Settings->IntPinTypeColor;
-			}
+			Params.WireColor = MaterialGraphSchema->GetPinTypeColor(InputPin->PinType);
 		}
 	}
 
diff --git a/Engine/Source/Editor/MaterialEditor/Private/MaterialEditor.cpp b/Engine/Source/Editor/MaterialEditor/Private/MaterialEditor.cpp
--- a/Engine/Source/Editor/MaterialEditor/Private/MaterialEditor.cpp
+++ b/Engine/Source/Editor/MaterialEditor/Private/MaterialEditor.cpp
@@ -2500,6 +2500,8 @@ void FMaterialEditor::UpdatePreviewMaterial( bool bForce )
 		Material->PostEditChange();
 	}
 
+	Material->MaterialGraph->UpdatePinTypes();
+
 	if (!PreviewExpression)
 	{
 		UpdateStatsMaterials();
@@ -5760,6 +5762,8 @@ void FMaterialEditor::UpdateMaterialAfterGraphChange()
 
 		HideUnrelatedNodes();
 	}
+
+	Material->MaterialGraph->UpdatePinTypes();
 }
 
 void FMaterialEditor::JumpToHyperlink(const UObject* ObjectReference)
diff --git a/Engine/Source/Editor/UnrealEd/Classes/MaterialGraph/MaterialGraph.h b/Engine/Source/Editor/UnrealEd/Classes/MaterialGraph/MaterialGraph.h
--- a/Engine/Source/Editor/UnrealEd/Classes/MaterialGraph/MaterialGraph.h
+++ b/Engine/Source/Editor/UnrealEd/Classes/MaterialGraph/MaterialGraph.h
@@ -11,8 +11,6 @@
 
 class UMaterialExpressionComment;
 class UMaterialExpressionComposite;
-struct FMaterialGraphCachedConnections;
-struct FMaterialConnectionKey;
 
 namespace UE::Shader
 {
@@ -215,7 +213,7 @@ public:
 	 */
 	void GetUnusedExpressions(TArray<class UEdGraphNode*>& UnusedNodes) const;
 
-	UE::Shader::EValueType GetConnectionType(const FMaterialConnectionKey& Key);
+	void UpdatePinTypes();
 
 private:
 	/**
@@ -231,6 +229,4 @@ private:
 	 * @param	Input	Input we are finding an output index for
 	 */
 	int32 GetValidOutputIndex(FExpressionInput* Input) const;
-
-	TUniquePtr<FMaterialGraphCachedConnections> CachedConnections;
 };
diff --git a/Engine/Source/Editor/UnrealEd/Classes/MaterialGraph/MaterialGraphNode_Base.h b/Engine/Source/Editor/UnrealEd/Classes/MaterialGraph/MaterialGraphNode_Base.h
--- a/Engine/Source/Editor/UnrealEd/Classes/MaterialGraph/MaterialGraphNode_Base.h
+++ b/Engine/Source/Editor/UnrealEd/Classes/MaterialGraph/MaterialGraphNode_Base.h
@@ -23,8 +23,8 @@ class UMaterialGraphNode_Base : public UEdGraphNode
 	virtual bool IsRootNode() const {return false;}
 	/** Gets the object that owns this node, typically either a UMaterial, UMaterialFunction, UMaterialExpression */
 	virtual UObject* GetMaterialNodeOwner() const { return nullptr; }
-	/** Returns the InputIndex associated with the given pin, maps to FMaterialConnectionKey::InputIndex */
-	virtual int32 GetInputIndexForPin(const UEdGraphPin* Pin) const;
+	/** Returns the SourceIndex associated with a particular FMaterialConnectionKey::InputIndex */
+	virtual int32 GetSourceIndexForInputIndex(int32 InputIndex) const;
 	/** Get a single Input Pin via its index */
 	UNREALED_API class UEdGraphPin* GetInputPin(int32 InputIndex) const;
 	/** Get a single Output Pin via its index */
diff --git a/Engine/Source/Editor/UnrealEd/Classes/MaterialGraph/MaterialGraphNode_Root.h b/Engine/Source/Editor/UnrealEd/Classes/MaterialGraph/MaterialGraphNode_Root.h
--- a/Engine/Source/Editor/UnrealEd/Classes/MaterialGraph/MaterialGraphNode_Root.h
+++ b/Engine/Source/Editor/UnrealEd/Classes/MaterialGraph/MaterialGraphNode_Root.h
@@ -28,8 +28,8 @@ class UMaterialGraphNode_Root : public UMaterialGraphNode_Base
 	//~ End UEdGraphNode Interface.
 
 	//~ Begin UMaterialGraphNode_Base Interface
-	virtual UObject* GetMaterialNodeOwner() const override { return Material; }
-	virtual int32 GetInputIndexForPin(const UEdGraphPin* Pin) const override;
+	virtual UObject* GetMaterialNodeOwner() const override;
+	virtual int32 GetSourceIndexForInputIndex(int32 InputIndex) const override;
 	virtual uint32 GetPinMaterialType(const UEdGraphPin* Pin) const override;
 	virtual void CreateInputPins() override;
 	virtual bool IsRootNode() const override {return true;}
diff --git a/Engine/Source/Editor/UnrealEd/Classes/MaterialGraph/MaterialGraphSchema.h b/Engine/Source/Editor/UnrealEd/Classes/MaterialGraph/MaterialGraphSchema.h
--- a/Engine/Source/Editor/UnrealEd/Classes/MaterialGraph/MaterialGraphSchema.h
+++ b/Engine/Source/Editor/UnrealEd/Classes/MaterialGraph/MaterialGraphSchema.h
@@ -183,6 +183,8 @@ class UMaterialGraphSchema : public UEdGraphSchema
 	UNREALED_API static const FName PC_Optional;
 	UNREALED_API static const FName PC_MaterialInput;
 	UNREALED_API static const FName PC_Exec;
+	UNREALED_API static const FName PC_Void;
+	UNREALED_API static const FName PC_ValueType;
 
 	// Common PinType.PinSubCategory values
 	UNREALED_API static const FName PSC_Red;
diff --git a/Engine/Source/Editor/UnrealEd/Private/MaterialGraph.cpp b/Engine/Source/Editor/UnrealEd/Private/MaterialGraph.cpp
--- a/Engine/Source/Editor/UnrealEd/Private/MaterialGraph.cpp
+++ b/Engine/Source/Editor/UnrealEd/Private/MaterialGraph.cpp
@@ -31,11 +31,6 @@
 
 #define LOCTEXT_NAMESPACE "MaterialGraph"
 
-struct FMaterialGraphCachedConnections
-{
-	TMap<FMaterialConnectionKey, UE::Shader::EValueType> ConnectionMap;
-};
-
 UMaterialGraph::UMaterialGraph(const FObjectInitializer& ObjectInitializer)
 	: Super(ObjectInitializer)
 {
@@ -99,7 +94,16 @@ static UMaterialGraphNode* InitExpressionNewNode(UMaterialGraph* Graph, UMateria
 
 namespace Private
 {
+class FOwnedNodeMaterial : public UE::HLSLTree::FOwnedNode
+{
+public:
+	FOwnedNodeMaterial(UObject* InMaterial) : Material(InMaterial) {}
+	virtual TConstArrayView<UObject*> GetOwners() const final { return MakeArrayView(&Material, 1); }
+	UObject* Material;
+};
+
 void PrepareHLSLTree(UE::HLSLTree::FEmitContext& EmitContext,
+	UMaterial* Material,
 	const FMaterialCachedHLSLTree& CachedTree,
 	EShaderFrequency ShaderFrequency)
 {
@@ -111,57 +115,107 @@ void PrepareHLSLTree(UE::HLSLTree::FEmitContext& EmitContext,
 	EmitContext.bMarkLiveValues = false;
 	FEmitScope* EmitResultScope = EmitContext.PrepareScope(CachedTree.GetResultScope());
 
-	FRequestedType RequestedAttributesType;
+	FRequestedType RequestedAttributesType(CachedTree.GetMaterialAttributesType(), false);
 	CachedTree.SetRequestedFields(ShaderFrequency, RequestedAttributesType);
 
+	FOwnedNodeMaterial MaterialOwner(Material);
+	FEmitOwnerScope OwnerScope(EmitContext, &MaterialOwner);
 	const FPreparedType& ResultType = EmitContext.PrepareExpression(CachedTree.GetResultExpression(), *EmitResultScope, RequestedAttributesType);
 }
 
-FMaterialGraphCachedConnections* CreateConnections(const FMaterialCachedHLSLTree& CachedTree)
+UMaterialGraphNode_Base* FindGraphNodeForObject(const UObject* Object)
+{
+	if (const UMaterialExpression* MaterialExpression = Cast<UMaterialExpression>(Object))
+	{
+		return CastChecked<UMaterialGraphNode_Base>(MaterialExpression->GraphNode);
+	}
+	else if (const UMaterial* Material = Cast<UMaterial>(Object))
+	{
+		const UMaterialGraph* MaterialGraph = Material->MaterialGraph;
+		if (MaterialGraph)
+		{
+			return MaterialGraph->RootNode;
+		}
+	}
+	return nullptr;
+}
+} // namespace Private
+
+void UMaterialGraph::UpdatePinTypes()
 {
 	using namespace UE::HLSLTree;
 
+	if (!Material->IsUsingNewHLSLGenerator())
+	{
+		return;
+	}
+
+	for (UEdGraphNode* Node : Nodes)
+	{
+		for (UEdGraphPin* Pin : Node->Pins)
+		{
+			if (Pin->PinType.PinCategory != UMaterialGraphSchema::PC_Exec)
+			{
+				Pin->PinType.PinCategory = UMaterialGraphSchema::PC_Void;
+			}
+		}
+	}
+
 	FNullErrorHandler NullErrorHandler;
 	FMemStackBase Allocator;
+	const FMaterialCachedHLSLTree& CachedTree = Material->GetCachedHLSLTree();
 	FEmitContext EmitContext(Allocator, FTargetParameters(), NullErrorHandler, CachedTree.GetTypeRegistry());
 
 	Material::FEmitData& EmitMaterialData = EmitContext.AcquireData<Material::FEmitData>();
 
-	PrepareHLSLTree(EmitContext, CachedTree, SF_Pixel);
-	PrepareHLSLTree(EmitContext, CachedTree, SF_Vertex);
+	::Private::PrepareHLSLTree(EmitContext, Material, CachedTree, SF_Pixel);
+	::Private::PrepareHLSLTree(EmitContext, Material, CachedTree, SF_Vertex);
 
-	FMaterialGraphCachedConnections* Connections = new FMaterialGraphCachedConnections();
 	for (const auto& It : CachedTree.GetConnections())
 	{
-		const UE::Shader::FType Type = EmitContext.GetType(It.Value);
-		if (!Type.IsVoid())
+		const FMaterialConnectionKey& Key = It.Key;
+		const FExpression* OutputExpression = It.Value;
+		const UE::Shader::FType OutputType = EmitContext.GetType(OutputExpression);
+		if (!OutputType.IsVoid())
 		{
-			Connections->ConnectionMap.Add(It.Key, Type.ValueType);
-		}
-	}
-
-	return Connections;
-}
+			const FConnectionKey InputKey(Key.InputObject, OutputExpression);
+			const UE::Shader::FType* InputType = EmitContext.ConnectionMap.Find(InputKey);
+			if (Cast<UMaterialExpressionMaterialFunctionCall>(Key.InputObject) || Cast<UMaterialExpressionMaterialFunctionCall>(Key.OutputObject))
+			{
+				int a = 0;
+			}
 
-} // namespace Private
+			if (InputType)
+			{
+				UMaterialGraphNode_Base* InputNode = ::Private::FindGraphNodeForObject(Key.InputObject);
+				if (InputNode)
+				{
+					const int32 InputIndex = InputNode->GetSourceIndexForInputIndex(Key.InputIndex);
+					UEdGraphPin* InputPin = InputNode->GetInputPin(InputIndex);
+					if (InputPin)
+					{
+						const UE::Shader::FValueTypeDescription InputTypeDesc = UE::Shader::GetValueTypeDescription(*InputType);
+						ensure(InputPin->PinType.PinCategory == UMaterialGraphSchema::PC_Void);
+						InputPin->PinType.PinCategory = UMaterialGraphSchema::PC_ValueType;
+						InputPin->PinType.PinSubCategory = InputTypeDesc.Name;
+					}
+				}
+			}
 
-UE::Shader::EValueType UMaterialGraph::GetConnectionType(const FMaterialConnectionKey& Key)
-{
-	FMaterialGraphCachedConnections* LocalConnections = CachedConnections.Get();
-	if (!LocalConnections && Material->IsUsingNewHLSLGenerator())
-	{
-		LocalConnections = Private::CreateConnections(Material->GetCachedHLSLTree());
-		CachedConnections.Reset(LocalConnections);
-	}
-	if (LocalConnections)
-	{
-		const UE::Shader::EValueType* Result = LocalConnections->ConnectionMap.Find(Key);
-		if (Result)
-		{
-			return *Result;
+			UMaterialGraphNode_Base* OutputNode = ::Private::FindGraphNodeForObject(Key.OutputObject);
+			if (OutputNode)
+			{
+				UEdGraphPin* OutputPin = OutputNode->GetOutputPin(Key.OutputIndex);
+				if (OutputPin)
+				{
+					// A single output pin may connect to multiple inputs, so it's possible to set the same value multiple times here
+					const UE::Shader::FValueTypeDescription OutputTypeDesc = UE::Shader::GetValueTypeDescription(OutputType);
+					OutputPin->PinType.PinCategory = UMaterialGraphSchema::PC_ValueType;
+					OutputPin->PinType.PinSubCategory = OutputTypeDesc.Name;
+				}
+			}
 		}
 	}
-	return UE::Shader::EValueType::Void;
 }
 
 void UMaterialGraph::RebuildGraphInternal(const TMap<UMaterialExpression*, TArray<UMaterialExpression*>>& SubgraphExpressionMap, const TMap<UMaterialExpression*, TArray<UMaterialExpressionComment*>>& SubgraphCommentMap)
@@ -287,6 +341,7 @@ UMaterialGraphNode* UMaterialGraph::AddExpression(UMaterialExpression* Expressio
 {
 	// Node for UMaterialExpressionExecBegin is explicitly placed if needed
 	// We don't created any node for UMaterialExpressionExecEnd, it's handled as part of the root node
+	UMaterialGraphNode* Node = nullptr;
 	if (Expression &&
 		!Expression->IsA(UMaterialExpressionExecBegin::StaticClass()) &&
 		!Expression->IsA(UMaterialExpressionExecEnd::StaticClass()))
@@ -295,23 +350,28 @@ UMaterialGraphNode* UMaterialGraph::AddExpression(UMaterialExpression* Expressio
 
 		if (Expression->IsA(UMaterialExpressionReroute::StaticClass()))
 		{
-			return InitExpressionNewNode<UMaterialGraphNode_Knot>(this, Expression, false);
+			Node = InitExpressionNewNode<UMaterialGraphNode_Knot>(this, Expression, false);
 		}
 		else if (Expression->IsA(UMaterialExpressionComposite::StaticClass()))
 		{
-			return InitExpressionNewNode<UMaterialGraphNode_Composite>(this, Expression, false);
+			Node = InitExpressionNewNode<UMaterialGraphNode_Composite>(this, Expression, false);
 		}
 		else if (Expression->IsA(UMaterialExpressionPinBase::StaticClass()))
 		{
-			return InitExpressionNewNode<UMaterialGraphNode_PinBase>(this, Expression, false);
+			Node = InitExpressionNewNode<UMaterialGraphNode_PinBase>(this, Expression, false);
 		}
 		else 
 		{
-			return InitExpressionNewNode<UMaterialGraphNode>(this, Expression, bUserInvoked);
+			Node = InitExpressionNewNode<UMaterialGraphNode>(this, Expression, bUserInvoked);
 		}
 	}
 
-	return nullptr;
+	if (Node && bUserInvoked)
+	{
+		UpdatePinTypes();
+	}
+
+	return Node;
 }
 
 UMaterialGraphNode_Comment* UMaterialGraph::AddComment(UMaterialExpressionComment* Comment, bool bIsUserInvoked)
@@ -495,6 +555,8 @@ void UMaterialGraph::LinkGraphNodesFromMaterial()
 	}
 
 	NotifyGraphChanged();
+
+	UpdatePinTypes();
 }
 
 void UMaterialGraph::LinkMaterialExpressionsFromGraph()
@@ -674,9 +736,6 @@ void UMaterialGraph::LinkMaterialExpressionsFromGraph()
 	{
 		CastChecked<UMaterialGraph>(SubGraph)->LinkMaterialExpressionsFromGraph();
 	}
-
-	// This is called when graph connections have been modified, so we need to reset any cached connection types
-	CachedConnections.Reset();
 }
 
 bool UMaterialGraph::IsInputActive(UEdGraphPin* GraphPin) const
@@ -808,13 +867,11 @@ void UMaterialGraph::GetUnusedExpressions(TArray<UEdGraphNode*>& UnusedNodes) co
 
 void UMaterialGraph::NotifyGraphChanged(const FEdGraphEditAction& Action)
 {
-	CachedConnections.Reset();
 	Super::NotifyGraphChanged(Action);
 }
 
 void UMaterialGraph::NotifyGraphChanged()
 {
-	CachedConnections.Reset();
 	Super::NotifyGraphChanged();
 }
 
diff --git a/Engine/Source/Editor/UnrealEd/Private/MaterialGraphNode_Base.cpp b/Engine/Source/Editor/UnrealEd/Private/MaterialGraphNode_Base.cpp
--- a/Engine/Source/Editor/UnrealEd/Private/MaterialGraphNode_Base.cpp
+++ b/Engine/Source/Editor/UnrealEd/Private/MaterialGraphNode_Base.cpp
@@ -17,10 +17,10 @@ UMaterialGraphNode_Base::UMaterialGraphNode_Base(const FObjectInitializer& Objec
 {
 }
 
-int32 UMaterialGraphNode_Base::GetInputIndexForPin(const UEdGraphPin* Pin) const
+int32 UMaterialGraphNode_Base::GetSourceIndexForInputIndex(int32 InputIndex) const
 {
-	// For most node types, the pin's SourceIndex will be the input index
-	return Pin ? Pin->SourceIndex : INDEX_NONE;
+	// For most node types, SourceIndex==InputIndex
+	return InputIndex;
 }
 
 UEdGraphPin* UMaterialGraphNode_Base::GetInputPin(int32 InputIndex) const
diff --git a/Engine/Source/Editor/UnrealEd/Private/MaterialGraphNode_Root.cpp b/Engine/Source/Editor/UnrealEd/Private/MaterialGraphNode_Root.cpp
--- a/Engine/Source/Editor/UnrealEd/Private/MaterialGraphNode_Root.cpp
+++ b/Engine/Source/Editor/UnrealEd/Private/MaterialGraphNode_Root.cpp
@@ -45,16 +45,21 @@ void UMaterialGraphNode_Root::PostPlacedNewNode()
 	}
 }
 
-int32 UMaterialGraphNode_Root::GetInputIndexForPin(const UEdGraphPin* Pin) const
+UObject* UMaterialGraphNode_Root::GetMaterialNodeOwner() const
 {
-	if (Pin)
+	return Material;
+}
+
+int32 UMaterialGraphNode_Root::GetSourceIndexForInputIndex(int32 InputIndex) const
+{
+	const UMaterialGraph* MaterialGraph = CastChecked<UMaterialGraph>(GetGraph());
+	for (int32 SourceIndex = 0; SourceIndex < MaterialGraph->MaterialInputs.Num(); ++SourceIndex)
 	{
-		// Pin->SourceIndex will be an index into the graph's MaterialInputs array
-		// the 'InputIndex' for root nodes is equal to the material property value
-		const UMaterialGraph* MaterialGraph = CastChecked<UMaterialGraph>(GetGraph());
-		if (MaterialGraph->MaterialInputs.IsValidIndex(Pin->SourceIndex))
+		// InputIndex will be the EMaterialProperty of the input
+		// SourceIndex will be an index into the 'MaterialInputs' array
+		if (InputIndex == (int32)MaterialGraph->MaterialInputs[SourceIndex].GetProperty())
 		{
-			return (int32)MaterialGraph->MaterialInputs[Pin->SourceIndex].GetProperty();
+			return SourceIndex;
 		}
 	}
 	return INDEX_NONE;
diff --git a/Engine/Source/Editor/UnrealEd/Private/MaterialGraphSchema.cpp b/Engine/Source/Editor/UnrealEd/Private/MaterialGraphSchema.cpp
--- a/Engine/Source/Editor/UnrealEd/Private/MaterialGraphSchema.cpp
+++ b/Engine/Source/Editor/UnrealEd/Private/MaterialGraphSchema.cpp
@@ -39,6 +39,7 @@
 #include "ScopedTransaction.h"
 #include "MaterialEditorUtilities.h"
 #include "GraphEditorActions.h"
+#include "GraphEditorSettings.h"
 #include "AssetRegistryModule.h"
 #include "MaterialEditorActions.h"
 #include "MaterialGraphNode_Knot.h"
@@ -229,6 +230,8 @@ const FName UMaterialGraphSchema::PC_Required(TEXT("required"));
 const FName UMaterialGraphSchema::PC_Optional(TEXT("optional"));
 const FName UMaterialGraphSchema::PC_MaterialInput(TEXT("materialinput"));
 const FName UMaterialGraphSchema::PC_Exec(TEXT("exec"));
+const FName UMaterialGraphSchema::PC_Void(TEXT("void"));
+const FName UMaterialGraphSchema::PC_ValueType(TEXT("value"));
 
 const FName UMaterialGraphSchema::PSC_Red(TEXT("red"));
 const FName UMaterialGraphSchema::PSC_Green(TEXT("green"));
@@ -621,6 +624,50 @@ bool UMaterialGraphSchema::TryCreateConnection(UEdGraphPin* A, UEdGraphPin* B) c
 	return bModified;
 }
 
+namespace Private
+{
+FLinearColor GetColorForConnectionType(const UGraphEditorSettings* Settings, UE::Shader::EValueType ConnectionType)
+{
+	if (ConnectionType == UE::Shader::EValueType::Struct)
+	{
+		return Settings->StructPinTypeColor;
+	}
+	else
+	{
+		const UE::Shader::FValueTypeDescription TypeDesc = UE::Shader::GetValueTypeDescription(ConnectionType);
+		if (UE::Shader::IsTextureType(TypeDesc.ComponentType))
+		{
+			return Settings->ObjectPinTypeColor;
+		}
+		else if (TypeDesc.ComponentType == UE::Shader::EValueComponentType::Float)
+		{
+			if (TypeDesc.NumComponents == 1)
+			{
+				return Settings->FloatPinTypeColor;
+			}
+			else
+			{
+				return Settings->VectorPinTypeColor;
+			}
+		}
+		else if (TypeDesc.ComponentType == UE::Shader::EValueComponentType::Double)
+		{
+			return Settings->DoublePinTypeColor;
+		}
+		else if (TypeDesc.ComponentType == UE::Shader::EValueComponentType::Bool)
+		{
+			return Settings->BooleanPinTypeColor;
+		}
+		else if (TypeDesc.ComponentType == UE::Shader::EValueComponentType::Int)
+		{
+			return Settings->IntPinTypeColor;
+		}
+	}
+
+	return Settings->DefaultPinTypeColor;
+}
+} // namespace Private
+
 FLinearColor UMaterialGraphSchema::GetPinTypeColor(const FEdGraphPinType& PinType) const
 {
 	if (PinType.PinCategory == PC_Mask)
@@ -650,6 +697,16 @@ FLinearColor UMaterialGraphSchema::GetPinTypeColor(const FEdGraphPinType& PinTyp
 	{
 		return InactivePinColor;
 	}
+	else if (PinType.PinCategory == PC_ValueType)
+	{
+		const UE::Shader::EValueType ValueType = UE::Shader::FindValueType(PinType.PinSubCategory);
+		const UGraphEditorSettings* Settings = GetDefault<UGraphEditorSettings>();
+		return Private::GetColorForConnectionType(Settings, ValueType);
+	}
+	else if (PinType.PinCategory == PC_Void)
+	{
+		return InactivePinColor;
+	}
 
 	return ActivePinColor;
 }
diff --git a/Engine/Source/Runtime/Engine/Private/HLSLTree/HLSLTree.cpp b/Engine/Source/Runtime/Engine/Private/HLSLTree/HLSLTree.cpp
--- a/Engine/Source/Runtime/Engine/Private/HLSLTree/HLSLTree.cpp
+++ b/Engine/Source/Runtime/Engine/Private/HLSLTree/HLSLTree.cpp
@@ -247,7 +247,7 @@ bool FExpressionLocalPHI::PrepareValue(FEmitContext& Context, FEmitScope& Scope,
 					const FPreparedType MergedType = MergePreparedTypes(CurrentType, ValueType);
 					if (MergedType.IsVoid())
 					{
-						return Context.Errors->AddErrorf(TEXT("Mismatched types for local variable %s and %s"),
+						return Context.Errorf(TEXT("Mismatched types for local variable %s and %s"),
 							CurrentType.GetName(),
 							ValueType.GetName());
 					}
@@ -284,7 +284,7 @@ bool FExpressionLocalPHI::PrepareValue(FEmitContext& Context, FEmitScope& Scope,
 		}
 		if (NumValidTypes < NumLiveScopes)
 		{
-			return Context.Errors->AddError(TEXT("Failed to compute all types for LocalPHI"));
+			return Context.Error(TEXT("Failed to compute all types for LocalPHI"));
 		}
 
 		if (CurrentType != InitialType)
@@ -302,7 +302,7 @@ bool FExpressionLocalPHI::PrepareValue(FEmitContext& Context, FEmitScope& Scope,
 				// Don't expect types to change *again*
 				if (ValueType.IsVoid() || MergePreparedTypes(CurrentType, ValueType) != CurrentType)
 				{
-					return Context.Errors->AddErrorf(TEXT("Mismatched types for local variable %s and %s"),
+					return Context.Errorf(TEXT("Mismatched types for local variable %s and %s"),
 						CurrentType.GetName(),
 						ValueType.GetName());
 				}
@@ -348,7 +348,7 @@ bool GetLiveScopes(FEmitContext& Context, const FExpressionLocalPHI& Expression,
 			OutLiveScopes.EmitDeclarationScope = FEmitScope::FindSharedParent(OutLiveScopes.EmitDeclarationScope, EmitValueScope);
 			if (!OutLiveScopes.EmitDeclarationScope)
 			{
-				return Context.Errors->AddError(TEXT("Invalid LocalPHI"));
+				return Context.Error(TEXT("Invalid LocalPHI"));
 			}
 		}
 	}
@@ -522,22 +522,40 @@ void FExpressionFunctionCall::EmitValuePreshader(FEmitContext& Context, FEmitSco
 	OutResult.Type = Function->OutputExpressions[OutputIndex]->GetValuePreshader(Context, Scope, RequestedType, OutResult.Preshader);
 }
 
-FRequestedType::FRequestedType(ERequestedType InType)
+FRequestedType::FRequestedType(Shader::EValueType InType, bool bDefaultRequest)
 {
-	const int32 NumComponents = (uint8)InType;
-	check(NumComponents >= 0 && NumComponents <= 16);
-	RequestedComponents.Init(true, NumComponents);
+	const Shader::FValueTypeDescription TypeDesc = Shader::GetValueTypeDescription(InType);
+	ValueComponentType = TypeDesc.ComponentType;
+	if (bDefaultRequest)
+	{
+		RequestedComponents.Init(true, TypeDesc.NumComponents);
+	}
 }
 
-FRequestedType::FRequestedType(Shader::EValueType InType)
+FRequestedType::FRequestedType(const Shader::FType& InType, bool bDefaultRequest)
 {
-	const Shader::FValueTypeDescription TypeDesc = Shader::GetValueTypeDescription(InType);
-	RequestedComponents.Init(true, TypeDesc.NumComponents);
+	if (InType.IsStruct())
+	{
+		StructType = InType.StructType;
+	}
+	else
+	{
+		const Shader::FValueTypeDescription TypeDesc = Shader::GetValueTypeDescription(InType);
+		ValueComponentType = TypeDesc.ComponentType;
+	}
+	if (bDefaultRequest)
+	{
+		RequestedComponents.Init(true, InType.GetNumComponents());
+	}
 }
 
-FRequestedType::FRequestedType(const Shader::FType& InType)
+Shader::FType FRequestedType::GetType() const
 {
-	RequestedComponents.Init(true, InType.GetNumComponents());
+	if (StructType)
+	{
+		return StructType;
+	}
+	return Shader::MakeValueType(ValueComponentType, GetNumComponents());
 }
 
 int32 FRequestedType::GetNumComponents() const
@@ -582,7 +600,7 @@ void FRequestedType::SetField(const Shader::FStructField* Field, const FRequeste
 
 FRequestedType FRequestedType::GetField(const Shader::FStructField* Field) const
 {
-	FRequestedType Result;
+	FRequestedType Result(Field->Type, false);
 	const int32 NumComponents = Field->GetNumComponents();
 	for (int32 Index = 0; Index < NumComponents; ++Index)
 	{
@@ -699,6 +717,8 @@ Shader::FType FPreparedType::GetType() const
 FRequestedType FPreparedType::GetRequestedType() const
 {
 	FRequestedType Result;
+	Result.ValueComponentType = ValueComponentType;
+	Result.StructType = StructType;
 	for (int32 Index = 0; Index < PreparedComponents.Num(); ++Index)
 	{
 		const FPreparedComponent& Component = PreparedComponents[Index];
@@ -982,7 +1002,7 @@ bool FPrepareValueResult::TryMergePreparedType(FEmitContext& Context, const Shad
 		check(ComponentType == Shader::EValueComponentType::Void);
 		if (StructType != PreparedType.StructType)
 		{
-			return Context.Errors->AddError(TEXT("Invalid type"));
+			return Context.Error(TEXT("Invalid type"));
 		}
 	}
 	else
@@ -1059,6 +1079,11 @@ bool FPrepareValueResult::SetType(FEmitContext& Context, const FRequestedType& R
 	return false;
 }
 
+void FStatement::EmitShader(FEmitContext& Context, FEmitScope& Scope) const
+{
+	check(false);
+}
+
 void FStatement::EmitPreshader(FEmitContext& Context, FEmitScope& Scope, const FRequestedType& RequestedType, TArrayView<const FEmitPreshaderScope> Scopes, Shader::FPreshaderData& OutPreshader) const
 {
 	check(false);
@@ -1086,7 +1111,7 @@ void FExpression::EmitValuePreshader(FEmitContext& Context, FEmitScope& Scope, c
 
 FEmitShaderExpression* FExpression::GetValueShader(FEmitContext& Context, FEmitScope& Scope, const FRequestedType& RequestedType, const FPreparedType& PreparedType, const Shader::FType& ResultType) const
 {
-	FOwnerScope OwnerScope(*Context.Errors, GetOwner());
+	FEmitOwnerScope OwnerScope(Context, this);
 
 	const EExpressionEvaluation Evaluation = PreparedType.GetEvaluation(Scope, RequestedType);
 	check(Evaluation != EExpressionEvaluation::None && Evaluation != EExpressionEvaluation::Unknown);
@@ -1140,7 +1165,7 @@ FEmitShaderExpression* FExpression::GetValueShader(FEmitContext& Context, FEmitS
 
 Shader::FType FExpression::GetValuePreshader(FEmitContext& Context, FEmitScope& Scope, const FRequestedType& RequestedType, const FPreparedType& PreparedType, const Shader::FType& ResultType, Shader::FPreshaderData& OutPreshader) const
 {
-	FOwnerScope OwnerScope(*Context.Errors, GetOwner());
+	FEmitOwnerScope OwnerScope(Context, this);
 	
 	const int32 PrevStackPosition = Context.PreshaderStackPosition;
 	const EExpressionEvaluation Evaluation = PreparedType.GetEvaluation(Scope, RequestedType);
@@ -1184,7 +1209,7 @@ Shader::FType FExpression::GetValuePreshader(FEmitContext& Context, FEmitScope&
 
 Shader::FValue FExpression::GetValueConstant(FEmitContext& Context, FEmitScope& Scope, const FRequestedType& RequestedType, const FPreparedType& PreparedType, const Shader::FType& ResultType) const
 {
-	FOwnerScope OwnerScope(*Context.Errors, GetOwner());
+	FEmitOwnerScope OwnerScope(Context, this);
 
 	const EExpressionEvaluation Evaluation = PreparedType.GetEvaluation(Scope, RequestedType);
 	if (Evaluation == EExpressionEvaluation::ConstantZero)
@@ -1282,17 +1307,17 @@ void FTree::Destroy(FTree* Tree)
 	}
 }
 
-void FOwnerContext::PushOwner(UObject* Owner)
+void FTree::PushOwner(UObject* Owner)
 {
 	OwnerStack.Add(Owner);
 }
 
-UObject* FOwnerContext::PopOwner()
+UObject* FTree::PopOwner()
 {
 	return OwnerStack.Pop(false);
 }
 
-UObject* FOwnerContext::GetCurrentOwner() const
+UObject* FTree::GetCurrentOwner() const
 {
 	return (OwnerStack.Num() > 0) ? OwnerStack.Last() : nullptr;
 }
@@ -1369,14 +1394,13 @@ bool FTree::EmitShader(FEmitContext& Context, FStringBuilderBase& OutCode) const
 
 void FTree::RegisterNode(FNode* Node)
 {
-	Node->Owner = GetCurrentOwner();
 	Node->NextNode = Nodes;
 	Nodes = Node;
 }
 
-const FExpression* FTree::FindExpression(FXxHash64 Hash) const
+FExpression* FTree::FindExpression(FXxHash64 Hash)
 {
-	FExpression const* const* FoundExpression = ExpressionMap.Find(Hash);
+	FExpression* const* FoundExpression = ExpressionMap.Find(Hash);
 	if (FoundExpression)
 	{
 		return *FoundExpression;
@@ -1386,6 +1410,7 @@ const FExpression* FTree::FindExpression(FXxHash64 Hash) const
 
 void FTree::RegisterExpression(FExpression* Expression, FXxHash64 Hash)
 {
+	Expression->Owners.Add(GetCurrentOwner());
 	ExpressionMap.Add(Hash, Expression);
 }
 
@@ -1395,10 +1420,16 @@ void FTree::RegisterExpression(FExpressionLocalPHI* Expression, FXxHash64 Hash)
 	RegisterExpression(static_cast<FExpression*>(Expression), Hash);
 }
 
+void FTree::AddCurrentOwner(FExpression* Expression)
+{
+	Expression->Owners.AddUnique(GetCurrentOwner());
+}
+
 void FTree::RegisterStatement(FScope& Scope, FStatement* Statement)
 {
 	check(!Scope.ContainedStatement)
 	check(!Statement->ParentScope);
+	Statement->Owner = GetCurrentOwner();
 	Statement->ParentScope = &Scope;
 	Scope.ContainedStatement = Statement;
 }
@@ -1449,7 +1480,7 @@ FExpressionDerivatives FTree::GetAnalyticDerivatives(const FExpression* InExpres
 	FExpressionDerivatives Derivatives;
 	if (InExpression)
 	{
-		FOwnerScope OwnerScope(*this, InExpression->GetOwner()); // Associate any newly created nodes with the same owner as the input expression
+		//FOwnerScope OwnerScope(*this, InExpression->GetOwner()); // Associate any newly created nodes with the same owner as the input expression
 		InExpression->ComputeAnalyticDerivatives(*this, Derivatives);
 	}
 	return Derivatives;
@@ -1460,7 +1491,7 @@ const FExpression* FTree::GetPreviousFrame(const FExpression* InExpression, cons
 	const FExpression* Result = InExpression;
 	if (Result && !RequestedType.IsVoid())
 	{
-		FOwnerScope OwnerScope(*this, InExpression->GetOwner()); // Associate any newly created nodes with the same owner as the input expression
+		//FOwnerScope OwnerScope(*this, InExpression->GetOwner()); // Associate any newly created nodes with the same owner as the input expression
 		const FExpression* PrevFrameExpression = InExpression->ComputePreviousFrame(*this, RequestedType);
 		if (PrevFrameExpression)
 		{
diff --git a/Engine/Source/Runtime/Engine/Private/HLSLTree/HLSLTreeCommon.cpp b/Engine/Source/Runtime/Engine/Private/HLSLTree/HLSLTreeCommon.cpp
--- a/Engine/Source/Runtime/Engine/Private/HLSLTree/HLSLTreeCommon.cpp
+++ b/Engine/Source/Runtime/Engine/Private/HLSLTree/HLSLTreeCommon.cpp
@@ -72,7 +72,7 @@
 
 bool FExpressionError::PrepareValue(FEmitContext& Context, FEmitScope&, const FRequestedType&, FPrepareValueResult&) const
 {
-	return Context.Errors->AddError(ErrorMessage);
+	return Context.Error(ErrorMessage);
 }
 
 void FExpressionForward::ComputeAnalyticDerivatives(FTree& Tree, FExpressionDerivatives& OutResult) const
@@ -87,7 +87,8 @@
 
 bool FExpressionForward::PrepareValue(FEmitContext& Context, FEmitScope& Scope, const FRequestedType& RequestedType, FPrepareValueResult& OutResult) const
 {
-	return Expression->PrepareValue(Context, Scope, RequestedType, OutResult);
+	const FPreparedType& PreparedType = Context.PrepareExpression(Expression, Scope, RequestedType);
+	return OutResult.SetType(Context, RequestedType, PreparedType);
 }
 
 void FExpressionForward::EmitValueShader(FEmitContext& Context, FEmitScope& Scope, const FRequestedType& RequestedType, FEmitValueShaderResult& OutResult) const
@@ -163,10 +164,11 @@
 
 bool FExpressionTextureSample::PrepareValue(FEmitContext& Context, FEmitScope& Scope, const FRequestedType& RequestedType, FPrepareValueResult& OutResult) const
 {
-	const FPreparedType& TextureType = Context.PrepareExpression(TextureExpression, Scope, ERequestedType::Texture);
+	// TODO - need the correct texture type?
+	const FPreparedType& TextureType = Context.PrepareExpression(TextureExpression, Scope, Shader::EValueType::Texture2D);
 	if (!Shader::IsTextureType(TextureType.ValueComponentType))
 	{
-		return Context.Errors->AddError(TEXT("Expected texture"));
+		return Context.Error(TEXT("Expected texture"));
 	}
 
 	const FRequestedType RequestedTexCoordType = Private::GetTexCoordType(TextureType.GetType());
@@ -184,7 +186,7 @@
 	}
 	else if (MipValueMode == TMVM_MipLevel || MipValueMode == TMVM_MipBias)
 	{
-		Context.PrepareExpression(MipValueExpression, Scope, ERequestedType::Scalar);
+		Context.PrepareExpression(MipValueExpression, Scope, Shader::EValueType::Float1);
 	}
 
 	return OutResult.SetType(Context, RequestedType, EExpressionEvaluation::Shader, Shader::EValueType::Float4);
@@ -309,20 +311,20 @@
 
 const FExpression* FExpressionGetStructField::ComputePreviousFrame(FTree& Tree, const FRequestedType& RequestedType) const
 {
-	FRequestedType RequestedStructType;
+	FRequestedType RequestedStructType(StructType, false);
 	RequestedStructType.SetField(Field, RequestedType);
 	return Tree.NewExpression<FExpressionGetStructField>(StructType, Field, Tree.GetPreviousFrame(StructExpression, RequestedStructType));
 }
 
 bool FExpressionGetStructField::PrepareValue(FEmitContext& Context, FEmitScope& Scope, const FRequestedType& RequestedType, FPrepareValueResult& OutResult) const
 {
-	FRequestedType RequestedStructType;
+	FRequestedType RequestedStructType(StructType, false);
 	RequestedStructType.SetField(Field, RequestedType);
 
 	FPreparedType StructPreparedType = Context.PrepareExpression(StructExpression, Scope, RequestedStructType);
 	if (!StructPreparedType.IsVoid() && StructPreparedType.StructType != StructType)
 	{
-		return Context.Errors->AddErrorf(TEXT("Expected type %s"), StructType->Name);
+		return Context.Errorf(TEXT("Expected type %s"), StructType->Name);
 	}
 
 	return OutResult.SetType(Context, RequestedType, StructPreparedType.GetFieldType(Field));
@@ -330,7 +332,7 @@
 
 void FExpressionGetStructField::EmitValueShader(FEmitContext& Context, FEmitScope& Scope, const FRequestedType& RequestedType, FEmitValueShaderResult& OutResult) const
 {
-	FRequestedType RequestedStructType;
+	FRequestedType RequestedStructType(StructType, false);
 	RequestedStructType.SetField(Field, RequestedType);
 
 	FEmitShaderExpression* StructValue = StructExpression->GetValueShader(Context, Scope, RequestedStructType);
@@ -342,7 +344,7 @@
 
 void FExpressionGetStructField::EmitValuePreshader(FEmitContext& Context, FEmitScope& Scope, const FRequestedType& RequestedType, FEmitValuePreshaderResult& OutResult) const
 {
-	FRequestedType RequestedStructType;
+	FRequestedType RequestedStructType(StructType, false);
 	RequestedStructType.SetField(Field, RequestedType);
 
 	StructExpression->GetValuePreshader(Context, Scope, RequestedStructType, OutResult.Preshader);
@@ -387,7 +389,7 @@
 	FPreparedType StructPreparedType = Context.PrepareExpression(StructExpression, Scope, RequestedStructType);
 	if (!StructPreparedType.IsVoid() && StructPreparedType.StructType != StructType)
 	{
-		return Context.Errors->AddErrorf(TEXT("Expected type %s"), StructType->Name);
+		return Context.Errorf(TEXT("Expected type %s"), StructType->Name);
 	}
 
 	const FRequestedType RequestedFieldType = RequestedType.GetField(Field);
@@ -467,15 +469,15 @@
 const FExpression* FExpressionSelect::ComputePreviousFrame(FTree& Tree, const FRequestedType& RequestedType) const
 {
 	return Tree.NewExpression<FExpressionSelect>(
-		Tree.GetPreviousFrame(ConditionExpression, ERequestedType::Scalar),
+		Tree.GetPreviousFrame(ConditionExpression, Shader::EValueType::Bool1),
 		Tree.GetPreviousFrame(TrueExpression, RequestedType),
 		Tree.GetPreviousFrame(FalseExpression, RequestedType));
 }
 
 bool FExpressionSelect::PrepareValue(FEmitContext& Context, FEmitScope& Scope, const FRequestedType& RequestedType, FPrepareValueResult& OutResult) const
 {
-	const FPreparedType& ConditionType = Context.PrepareExpression(ConditionExpression, Scope, ERequestedType::Scalar);
-	const EExpressionEvaluation ConditionEvaluation = ConditionType.GetEvaluation(Scope, ERequestedType::Scalar);
+	const FPreparedType& ConditionType = Context.PrepareExpression(ConditionExpression, Scope, Shader::EValueType::Bool1);
+	const EExpressionEvaluation ConditionEvaluation = ConditionType.GetEvaluation(Scope, Shader::EValueType::Bool1);
 	if (IsConstantEvaluation(ConditionEvaluation))
 	{
 		const bool bCondition = ConditionExpression->GetValueConstant(Context, Scope, ConditionType, Shader::EValueType::Bool1).AsBoolScalar();
@@ -490,7 +492,7 @@
 	if (LhsType.ValueComponentType != RhsType.ValueComponentType ||
 		LhsType.StructType != RhsType.StructType)
 	{
-		return Context.Errors->AddError(TEXT("Type mismatch"));
+		return Context.Error(TEXT("Type mismatch"));
 	}
 
 	FPreparedType ResultType = MergePreparedTypes(LhsType, RhsType);
@@ -502,7 +504,7 @@
 void FExpressionSelect::EmitValueShader(FEmitContext& Context, FEmitScope& Scope, const FRequestedType& RequestedType, FEmitValueShaderResult& OutResult) const
 {
 	const FPreparedType& ConditionType = Context.GetPreparedType(ConditionExpression);
-	const EExpressionEvaluation ConditionEvaluation = ConditionType.GetEvaluation(Scope, ERequestedType::Scalar);
+	const EExpressionEvaluation ConditionEvaluation = ConditionType.GetEvaluation(Scope, Shader::EValueType::Bool1);
 	if (IsConstantEvaluation(ConditionEvaluation))
 	{
 		const bool bCondition = ConditionExpression->GetValueConstant(Context, Scope, ConditionType, Shader::EValueType::Bool1).AsBoolScalar();
@@ -585,6 +587,7 @@
 FRequestedType GetRequestedSwizzleType(const FSwizzleParameters& Params, const FRequestedType& RequestedType)
 {
 	FRequestedType RequestedInputType;
+	RequestedInputType.ValueComponentType = RequestedType.ValueComponentType;
 	for (int32 Index = 0; Index < Params.NumComponents; ++Index)
 	{
 		if (RequestedType.IsComponentRequested(Index))
@@ -667,6 +670,7 @@
 
 	const int32 NumComponents = FMath::Min<int32>(RequestedType.GetNumComponents(), Parameters.NumComponents);
 	FRequestedType RequestedInputType;
+	RequestedInputType.ValueComponentType = RequestedType.ValueComponentType;
 	for (int32 ComponentIndex = 0; ComponentIndex < NumComponents; ++ComponentIndex)
 	{
 		// If component wasn't requested, we just refer to 'x' component of input, since that should always be present
@@ -746,6 +750,50 @@
 	}
 }
 
+namespace Private
+{
+struct FAppendTypes
+{
+	Shader::EValueType ResultType;
+	Shader::EValueType LhsType;
+	Shader::EValueType RhsType;
+	FRequestedType LhsRequestedType;
+	FRequestedType RhsRequestedType;
+	bool bIsLWC;
+};
+FAppendTypes GetAppendTypes(const FRequestedType& RequestedType, Shader::EValueType LhsType, Shader::EValueType RhsType)
+{
+	const Shader::FValueTypeDescription LhsTypeDesc = Shader::GetValueTypeDescription(LhsType);
+	const Shader::FValueTypeDescription RhsTypeDesc = Shader::GetValueTypeDescription(RhsType);
+	const Shader::EValueComponentType ComponentType = Shader::CombineComponentTypes(LhsTypeDesc.ComponentType, RhsTypeDesc.ComponentType);
+	const int32 NumComponents = FMath::Min(LhsTypeDesc.NumComponents + RhsTypeDesc.NumComponents, 4);
+
+	FAppendTypes Types;
+	Types.LhsRequestedType.ValueComponentType = ComponentType;
+	Types.RhsRequestedType.ValueComponentType = ComponentType;
+	for (int32 Index = 0; Index < LhsTypeDesc.NumComponents; ++Index)
+	{
+		if (RequestedType.IsComponentRequested(Index))
+		{
+			Types.LhsRequestedType.SetComponentRequest(Index);
+		}
+	}
+	for (int32 Index = LhsTypeDesc.NumComponents; Index < NumComponents; ++Index)
+	{
+		if (RequestedType.IsComponentRequested(Index))
+		{
+			Types.RhsRequestedType.SetComponentRequest(Index - LhsTypeDesc.NumComponents);
+		}
+	}
+
+	Types.ResultType = Shader::MakeValueType(ComponentType, NumComponents);
+	Types.LhsType = Shader::MakeValueType(ComponentType, LhsTypeDesc.NumComponents);
+	Types.RhsType = Shader::MakeValueType(ComponentType, NumComponents - LhsTypeDesc.NumComponents);
+	Types.bIsLWC = ComponentType == Shader::EValueComponentType::Double;
+	return Types;
+}
+} // namespace Private
+
 void FExpressionAppend::ComputeAnalyticDerivatives(FTree& Tree, FExpressionDerivatives& OutResult) const
 {
 	const FExpressionDerivatives LhsDerivatives = Tree.GetAnalyticDerivatives(Lhs);
@@ -776,17 +824,19 @@
 	}
 
 	FRequestedType RhsRequestedType;
+	RhsRequestedType.ValueComponentType = RequestedType.ValueComponentType;
 	for (int32 Index = NumLhsComponents; Index < NumRequestedComponents; ++Index)
 	{
 		RhsRequestedType.SetComponentRequest(Index - NumLhsComponents, RequestedType.IsComponentRequested(Index));
 	}
 
+	FPreparedType RhsType;
 	if (!RhsRequestedType.IsVoid())
 	{
-		const FPreparedType& RhsType = Context.PrepareExpression(Rhs, Scope, RhsRequestedType);
+		RhsType = Context.PrepareExpression(Rhs, Scope, RhsRequestedType);
 		if (LhsType.ValueComponentType != RhsType.ValueComponentType)
 		{
-			return Context.Errors->AddError(TEXT("Type mismatch"));
+			return Context.Error(TEXT("Type mismatch"));
 		}
 
 		const int32 NumRhsComponents = FMath::Min(RhsType.GetNumComponents(), NumRequestedComponents - NumLhsComponents);
@@ -796,50 +846,12 @@
 		}
 	}
 
+	const Private::FAppendTypes Types = Private::GetAppendTypes(RequestedType, LhsType.GetType(), RhsType.GetType());
+	Context.MarkInputType(Lhs, Types.LhsType);
+	Context.MarkInputType(Rhs, Types.RhsType);
+
 	return OutResult.SetType(Context, RequestedType, ResultType);
 }
 
-namespace Private
-{
-struct FAppendTypes
-{
-	Shader::EValueType ResultType;
-	Shader::EValueType LhsType;
-	Shader::EValueType RhsType;
-	FRequestedType LhsRequestedType;
-	FRequestedType RhsRequestedType;
-	bool bIsLWC;
-};
-FAppendTypes GetAppendTypes(const FRequestedType& RequestedType, Shader::EValueType LhsType, Shader::EValueType RhsType)
-{
-	const Shader::FValueTypeDescription LhsTypeDesc = Shader::GetValueTypeDescription(LhsType);
-	const Shader::FValueTypeDescription RhsTypeDesc = Shader::GetValueTypeDescription(RhsType);
-	const Shader::EValueComponentType ComponentType = Shader::CombineComponentTypes(LhsTypeDesc.ComponentType, RhsTypeDesc.ComponentType);
-	const int32 NumComponents = FMath::Min(LhsTypeDesc.NumComponents + RhsTypeDesc.NumComponents, 4);
-
-	FAppendTypes Types;
-	for (int32 Index = 0; Index < LhsTypeDesc.NumComponents; ++Index)
-	{
-		if (RequestedType.IsComponentRequested(Index))
-		{
-			Types.LhsRequestedType.SetComponentRequest(Index);
-		}
-	}
-	for (int32 Index = LhsTypeDesc.NumComponents; Index < NumComponents; ++Index)
-	{
-		if (RequestedType.IsComponentRequested(Index))
-		{
-			Types.RhsRequestedType.SetComponentRequest(Index - LhsTypeDesc.NumComponents);
-		}
-	}
-
-	Types.ResultType = Shader::MakeValueType(ComponentType, NumComponents);
-	Types.LhsType = Shader::MakeValueType(ComponentType, LhsTypeDesc.NumComponents);
-	Types.RhsType = Shader::MakeValueType(ComponentType, NumComponents - LhsTypeDesc.NumComponents);
-	Types.bIsLWC = ComponentType == Shader::EValueComponentType::Double;
-	return Types;
-}
-} // namespace Private
-
 void FExpressionAppend::EmitValueShader(FEmitContext& Context, FEmitScope& Scope, const FRequestedType& RequestedType, FEmitValueShaderResult& OutResult) const
 {
@@ -995,7 +1007,8 @@
 {
 	for (const FCustomHLSLInput& Input : Inputs)
 	{
-		const FPreparedType& InputType = Context.PrepareExpression(Input.Expression, Scope, ERequestedType::Vector4);
+		// TODO - do we need to support untyped input here?  Add explicit custom input types?
+		const FPreparedType& InputType = Context.PrepareExpression(Input.Expression, Scope, Shader::EValueType::Float4);
 		if (InputType.IsVoid())
 		{
 			return false;
@@ -1010,6 +1023,11 @@
 	OutResult.Code = Context.EmitCustomHLSL(Scope, DeclarationCode, FunctionCode, Inputs, OutputStructType);
 }
 
+bool FStatementError::Prepare(FEmitContext& Context, FEmitScope& Scope) const
+{
+	return Context.Error(ErrorMessage);
+}
+
 bool FStatementBreak::Prepare(FEmitContext& Context, FEmitScope& Scope) const
 {
 	return true;
@@ -1030,13 +1048,13 @@
 
 bool FStatementIf::Prepare(FEmitContext& Context, FEmitScope& Scope) const
 {
-	const FPreparedType& ConditionType = Context.PrepareExpression(ConditionExpression, Scope, ERequestedType::Scalar);
+	const FPreparedType& ConditionType = Context.PrepareExpression(ConditionExpression, Scope, Shader::EValueType::Bool1);
 	if (ConditionType.IsVoid())
 	{
 		return false;
 	}
 
-	const EExpressionEvaluation ConditionEvaluation = ConditionType.GetEvaluation(Scope, ERequestedType::Scalar);
+	const EExpressionEvaluation ConditionEvaluation = ConditionType.GetEvaluation(Scope, Shader::EValueType::Bool1);
 	check(ConditionEvaluation != EExpressionEvaluation::None);
 	if (IsConstantEvaluation(ConditionEvaluation))
 	{
@@ -1065,7 +1083,7 @@
 {
 	FEmitShaderNode* Dependency = nullptr;
 	const FPreparedType& ConditionType = Context.GetPreparedType(ConditionExpression);
-	const EExpressionEvaluation ConditionEvaluation = ConditionType.GetEvaluation(Scope, ERequestedType::Scalar);
+	const EExpressionEvaluation ConditionEvaluation = ConditionType.GetEvaluation(Scope, Shader::EValueType::Bool1);
 	if (IsConstantEvaluation(ConditionEvaluation))
 	{
 		const bool bCondition = ConditionExpression->GetValueConstant(Context, Scope, ConditionType, Shader::EValueType::Bool1).AsBoolScalar();
@@ -1089,7 +1107,7 @@
 
 void FStatementIf::EmitPreshader(FEmitContext& Context, FEmitScope& Scope, const FRequestedType& RequestedType, TArrayView<const FEmitPreshaderScope> Scopes, Shader::FPreshaderData& OutPreshader) const
 {
-	ConditionExpression->GetValuePreshader(Context, Scope, ERequestedType::Scalar, OutPreshader);
+	ConditionExpression->GetValuePreshader(Context, Scope, Shader::EValueType::Bool1, OutPreshader);
 
 	check(Context.PreshaderStackPosition > 0);
 	Context.PreshaderStackPosition--;
diff --git a/Engine/Source/Runtime/Engine/Private/HLSLTree/HLSLTreeEmit.cpp b/Engine/Source/Runtime/Engine/Private/HLSLTree/HLSLTreeEmit.cpp
--- a/Engine/Source/Runtime/Engine/Private/HLSLTree/HLSLTreeEmit.cpp
+++ b/Engine/Source/Runtime/Engine/Private/HLSLTree/HLSLTreeEmit.cpp
@@ -213,6 +213,21 @@ FEmitContext::~FEmitContext()
 	}
 }
 
+bool FEmitContext::InternalError(FStringView ErrorMessage)
+{
+	if (Errors)
+	{
+		TConstArrayView<UObject*> CurrentOwners;
+		if (OwnerStack.Num() > 0)
+		{
+			CurrentOwners = OwnerStack.Last()->GetOwners();
+		}
+		Errors->AddError(CurrentOwners, ErrorMessage);
+	}
+	NumErrors++;
+	return false;
+}
+
 void FEmitContext::InternalRegisterData(FXxHash64 Hash, FCustomDataWrapper* Data)
 {
 	CustomDataMap.Add(Hash, Data);
@@ -416,9 +431,12 @@ FPreparedType FEmitContext::PrepareExpression(const FExpression* InExpression, F
 		return Result->PreparedType;
 	}
 
+	MarkInputType(InExpression, RequestedType.GetType());
+
 	bool bResult = false;
 	{
-		FOwnerScope OwnerScope(*Errors, InExpression->GetOwner());
+		FEmitOwnerScope OwnerScope(*this, InExpression);
+
 		Result->bPreparingValue = true;
 		bResult = InExpression->PrepareValue(*this, Scope, RequestedType, *Result);
 		Result->bPreparingValue = false;
@@ -433,6 +451,19 @@ FPreparedType FEmitContext::PrepareExpression(const FExpression* InExpression, F
 	return ResultType;
 }
 
+void FEmitContext::MarkInputType(const FExpression* InExpression, const Shader::FType& Type)
+{
+	if (OwnerStack.Num() > 0)
+	{
+		const FOwnedNode* InputOwner = OwnerStack.Last();
+		for (UObject* InputObject : InputOwner->GetOwners())
+		{
+			const FConnectionKey Key(InputObject, InExpression);
+			ConnectionMap.Add(Key, Type);
+		}
+	}
+}
+
 FEmitScope* FEmitContext::InternalPrepareScope(FScope* InScope, FScope* InParentScope)
 {
 	if (!InScope)
@@ -500,11 +531,14 @@ FEmitScope* FEmitContext::InternalPrepareScope(FScope* InScope, FScope* InParent
 		}
 		else if (EmitScope->State == EEmitScopeState::Initializing)
 		{
-			if (EmitScope->OwnerStatement)
+			FStatement* Statement = EmitScope->OwnerStatement;
+			if (Statement)
 			{
+				FEmitOwnerScope OwnerScope(*this, Statement);
+
 				FEmitScope* EmitParentScope = EmitScope->ParentScope;
 				check(EmitParentScope);
-				if (EmitScope->OwnerStatement->Prepare(*this, *EmitParentScope))
+				if (Statement->Prepare(*this, *EmitParentScope))
 				{
 					if (EmitScope->State == EEmitScopeState::Initializing)
 					{
@@ -1205,7 +1239,7 @@ FEmitShaderExpression* FEmitContext::EmitCast(FEmitScope& Scope, FEmitShaderExpr
 	}
 	else
 	{
-		Errors->AddErrorf(TEXT("Cannot cast between non-numeric types %s to %s."), SourceTypeDesc.Name, DestTypeDesc.Name);
+		Errorf(TEXT("Cannot cast between non-numeric types %s to %s."), SourceTypeDesc.Name, DestTypeDesc.Name);
 		FormattedCode.Appendf(TEXT("((%s)0)"), DestType.GetName());
 	}
 
diff --git a/Engine/Source/Runtime/Engine/Private/HLSLTree/HLSLTreeOperation.cpp b/Engine/Source/Runtime/Engine/Private/HLSLTree/HLSLTreeOperation.cpp
--- a/Engine/Source/Runtime/Engine/Private/HLSLTree/HLSLTreeOperation.cpp
+++ b/Engine/Source/Runtime/Engine/Private/HLSLTree/HLSLTreeOperation.cpp
@@ -114,10 +114,10 @@ const FExpression* FTree::NewCross(const FExpression* Lhs, const FExpression* Rh
 	//c_P[0] = v_A[1] * v_B[2] - v_A[2] * v_B[1];
 	//c_P[1] = -(v_A[0] * v_B[2] - v_A[2] * v_B[0]);
 	//c_P[2] = v_A[0] * v_B[1] - v_A[1] * v_B[0];
-	const FExpression* Lhs0 = NewExpression<FExpressionSwizzle>(FSwizzleParameters(1, 0, 0), Lhs);
-	const FExpression* Lhs1 = NewExpression<FExpressionSwizzle>(FSwizzleParameters(2, 2, 1), Lhs);
-	const FExpression* Rhs0 = NewExpression<FExpressionSwizzle>(FSwizzleParameters(2, 2, 1), Rhs);
-	const FExpression* Rhs1 = NewExpression<FExpressionSwizzle>(FSwizzleParameters(1, 0, 0), Rhs);
+	const FExpression* Lhs0 = NewSwizzle(FSwizzleParameters(1, 0, 0), Lhs);
+	const FExpression* Lhs1 = NewSwizzle(FSwizzleParameters(2, 2, 1), Lhs);
+	const FExpression* Rhs0 = NewSwizzle(FSwizzleParameters(2, 2, 1), Rhs);
+	const FExpression* Rhs1 = NewSwizzle(FSwizzleParameters(1, 0, 0), Rhs);
 	return NewSub(NewMul(NewMul(Lhs0, Rhs0), NewConstant(FVector3f(1.0f, -1.0f, 1.0f))), NewMul(Lhs1, Rhs1));
 }
 
@@ -161,24 +161,24 @@ FOperationRequestedTypes GetOperationRequestedTypes(EOperation Op, const FReques
 	case EOperation::Length:
 	case EOperation::Normalize:
 	case EOperation::Sum:
-		Types.InputType[0] = ERequestedType::Vector4;
+		Types.InputType[0] = Shader::MakeValueType(RequestedType.ValueComponentType, 4);
 		break;
 	case EOperation::VecMulMatrix3:
 		Types.bIsMatrixOperation = true;
-		Types.InputType[0] = ERequestedType::Vector3;
-		Types.InputType[1] = ERequestedType::Matrix4x4;
+		Types.InputType[0] = Shader::MakeValueType(RequestedType.ValueComponentType, 3);
+		Types.InputType[1] = Shader::MakeValueType(RequestedType.ValueComponentType, 16);
 		break;
 	case EOperation::VecMulMatrix4:
 		Types.bIsMatrixOperation = true;
-		Types.InputType[0] = ERequestedType::Vector3;
-		Types.InputType[1] = ERequestedType::Matrix4x4;
+		Types.InputType[0] = Shader::MakeValueType(RequestedType.ValueComponentType, 3);
+		Types.InputType[1] = Shader::MakeValueType(RequestedType.ValueComponentType, 16);
 		break;
 	case EOperation::Matrix3MulVec:
 	case EOperation::Matrix4MulVec:
 		// No LWC for transpose matrices
 		Types.bIsMatrixOperation = true;
-		Types.InputType[0] = ERequestedType::Matrix4x4;
-		Types.InputType[1] = ERequestedType::Vector3;
+		Types.InputType[0] = Shader::MakeValueType(RequestedType.ValueComponentType, 16);
+		Types.InputType[1] = Shader::MakeValueType(RequestedType.ValueComponentType, 3);
 		break;
 	default:
 		break;
@@ -253,7 +253,8 @@ FOperationTypes GetOperationTypes(EOperation Op, TConstArrayView<FPreparedType>
 					break;
 				}
 			}
-			Types.InputType[Index] = Shader::MakeValueType(InputComponentType, NumIntermediateComponents);
+			const int32 NumInputComponents = InputPreparedType[Index].GetNumComponents();
+			Types.InputType[Index] = Shader::MakeValueType(InputComponentType, (NumInputComponents == 1) ? 1 : NumIntermediateComponents);
 		}
 		Types.ResultType = IntermediateType;
 		Types.bIsLWC = (IntermediateComponentType == Shader::EValueComponentType::Double);
@@ -557,7 +558,7 @@ bool FExpressionOperation::PrepareValue(FEmitContext& Context, FEmitScope& Scope
 
 		if (!InputPreparedType[Index].IsNumeric())
 		{
-			return Context.Errors->AddError(TEXT("Invalid arithmetic between non-numeric types"));
+			return Context.Error(TEXT("Invalid arithmetic between non-numeric types"));
 		}
 		
 		const EExpressionEvaluation InputEvaluation = InputPreparedType[Index].GetEvaluation(Scope, RequestedType);
@@ -575,6 +576,11 @@ bool FExpressionOperation::PrepareValue(FEmitContext& Context, FEmitScope& Scope
 		Types.ResultType.SetEvaluation(EExpressionEvaluation::Shader);
 	}
 
+	for (int32 Index = 0; Index < OpDesc.NumInputs; ++Index)
+	{
+		Context.MarkInputType(Inputs[Index], Types.InputType[Index]);
+	}
+
 	switch (Op)
 	{
 	case EOperation::Mul:
diff --git a/Engine/Source/Runtime/Engine/Private/Materials/Material.cpp b/Engine/Source/Runtime/Engine/Private/Materials/Material.cpp
--- a/Engine/Source/Runtime/Engine/Private/Materials/Material.cpp
+++ b/Engine/Source/Runtime/Engine/Private/Materials/Material.cpp
@@ -3930,7 +3930,7 @@ void UMaterial::PostEditChangePropertyInternal(FPropertyChangedEvent& PropertyCh
 	// not interested in PostEditChange calls though as the graph may have instigated it
 	if (PropertyThatChanged && MaterialGraph)
 	{
-		MaterialGraph->NotifyGraphChanged();
+		MaterialGraph->RebuildGraph();
 	}
 } 
 
diff --git a/Engine/Source/Runtime/Engine/Private/Materials/MaterialCachedData.cpp b/Engine/Source/Runtime/Engine/Private/Materials/MaterialCachedData.cpp
--- a/Engine/Source/Runtime/Engine/Private/Materials/MaterialCachedData.cpp
+++ b/Engine/Source/Runtime/Engine/Private/Materials/MaterialCachedData.cpp
@@ -470,7 +470,7 @@ void PrepareHLSLTree(UE::HLSLTree::FEmitContext& EmitContext,
 	EmitContext.bMarkLiveValues = false;
 	FEmitScope* EmitResultScope = EmitContext.PrepareScope(CachedTree.GetResultScope());
 
-	FRequestedType RequestedAttributesType;
+	FRequestedType RequestedAttributesType(CachedTree.GetMaterialAttributesType(), false);
 	CachedTree.SetRequestedFields(ShaderFrequency, RequestedAttributesType);
 
 	const FPreparedType& ResultType = EmitContext.PrepareExpression(CachedTree.GetResultExpression(), *EmitResultScope, RequestedAttributesType);
diff --git a/Engine/Source/Runtime/Engine/Private/Materials/MaterialExpressionHLSL.cpp b/Engine/Source/Runtime/Engine/Private/Materials/MaterialExpressionHLSL.cpp
--- a/Engine/Source/Runtime/Engine/Private/Materials/MaterialExpressionHLSL.cpp
+++ b/Engine/Source/Runtime/Engine/Private/Materials/MaterialExpressionHLSL.cpp
@@ -713,7 +713,7 @@ bool UMaterialExpressionTextureSample::GenerateHLSLExpressionBase(FMaterialHLSLG
 		return Generator.Error(TEXT("Missing input texture"));
 	}
 
-	const FExpression* TexCoordExpression = Coordinates.GetTracedInput().Expression ? Coordinates.TryAcquireHLSLExpression(Generator, Scope) : Generator.NewTexCoord(ConstCoordinate);
+	const FExpression* TexCoordExpression = Coordinates.AcquireHLSLExpressionOrExternalInput(Generator, Scope, Material::MakeInputTexCoord(ConstCoordinate));
 	const FExpression* MipLevelExpression = nullptr;
 	FExpressionDerivatives TexCoordDerivatives;
 	switch (MipValueMode)
@@ -878,11 +878,7 @@ bool UMaterialExpressionNoise::GenerateHLSLExpression(FMaterialHLSLGenerator& Ge
 {
 	using namespace UE::HLSLTree;
 
-	const FExpression* ExpressionPosition = Position.TryAcquireHLSLExpression(Generator, Scope);
-	if (!ExpressionPosition)
-	{
-		ExpressionPosition = Generator.GetTree().NewExpression<Material::FExpressionExternalInput>(Material::EExternalInput::WorldPosition);
-	}
+	const FExpression* ExpressionPosition = Position.AcquireHLSLExpressionOrExternalInput(Generator, Scope, Material::EExternalInput::WorldPosition);
 	const FExpression* ExpressionFilterWidth = FilterWidth.AcquireHLSLExpressionOrConstant(Generator, Scope, 0.0f);
 
 	Material::FNoiseParameters NoiseParameters;
@@ -2013,12 +2009,7 @@ bool UMaterialExpressionFresnel::GenerateHLSLExpression(FMaterialHLSLGenerator&
 		return false;
 	}
 
-	const FExpression* ExpressionNormal = Normal.TryAcquireHLSLExpression(Generator, Scope);
-	if (!ExpressionNormal)
-	{
-		ExpressionNormal = Generator.GetTree().NewExpression<Material::FExpressionExternalInput>(Material::EExternalInput::WorldNormal);
-	}
-
+	const FExpression* ExpressionNormal = Normal.AcquireHLSLExpressionOrExternalInput(Generator, Scope, Material::EExternalInput::WorldNormal);
 	const FExpression* ExpressionNdotV = Generator.GetTree().NewDot(ExpressionNormal, Generator.GetTree().NewExpression<Material::FExpressionExternalInput>(Material::EExternalInput::CameraVector));
 	const FExpression* ExpressionMax = Generator.GetTree().NewMax(Generator.NewConstant(0.0f), ExpressionNdotV);
 	const FExpression* ExpressionMinus = Generator.GetTree().NewSub(Generator.NewConstant(1.0f), ExpressionMax);
@@ -2115,11 +2106,7 @@ bool UMaterialExpressionBumpOffset::GenerateHLSLExpression(FMaterialHLSLGenerato
 
 	const FExpression* ExpressionHeightRatio = HeightRatioInput.AcquireHLSLExpressionOrConstant(Generator, Scope, HeightRatio);
 	const FExpression* ExpressionHeight = Height.AcquireHLSLExpression(Generator, Scope);
-	const FExpression* ExpressionCoordinate = Coordinate.AcquireHLSLExpression(Generator, Scope);
-	if (!ExpressionCoordinate)
-	{
-		ExpressionCoordinate = Generator.NewTexCoord(ConstCoordinate);
-	}
+	const FExpression* ExpressionCoordinate = Coordinate.AcquireHLSLExpressionOrExternalInput(Generator, Scope, Material::MakeInputTexCoord(ConstCoordinate));
 	if (!ExpressionHeightRatio || !ExpressionHeight || !ExpressionCoordinate)
 	{
 		return false;
@@ -2138,23 +2125,14 @@ bool UMaterialExpressionPanner::GenerateHLSLExpression(FMaterialHLSLGenerator& G
 {
 	using namespace UE::HLSLTree;
 
-	const FExpression* ExpressionTime = Time.TryAcquireHLSLExpression(Generator, Scope);
-	if (!ExpressionTime)
-	{
-		ExpressionTime = Generator.GetTree().NewExpression<Material::FExpressionExternalInput>(Material::EExternalInput::GameTime);
-	}
+	const FExpression* ExpressionTime = Time.AcquireHLSLExpressionOrExternalInput(Generator, Scope, Material::EExternalInput::GameTime);
 	const FExpression* ExpressionSpeed = Speed.AcquireHLSLExpressionOrConstant(Generator, Scope, FVector2f(SpeedX, SpeedY));
 	const FExpression* ExpressionOffset = Generator.GetTree().NewMul(ExpressionSpeed, ExpressionTime);
 	if (bFractionalPart)
 	{
 		ExpressionOffset = Generator.GetTree().NewFrac(ExpressionOffset);
 	}
-	const FExpression* ExpressionTexCoord = Coordinate.TryAcquireHLSLExpression(Generator, Scope);
-	if (!ExpressionTexCoord)
-	{
-		ExpressionTexCoord = Generator.NewTexCoord(ConstCoordinate);
-	}
-
+	const FExpression* ExpressionTexCoord = Coordinate.AcquireHLSLExpressionOrExternalInput(Generator, Scope, Material::MakeInputTexCoord(ConstCoordinate));
 	OutExpression = Generator.GetTree().NewAdd(ExpressionTexCoord, ExpressionOffset);
 	return true;
 }
@@ -2163,11 +2141,7 @@ bool UMaterialExpressionRotator::GenerateHLSLExpression(FMaterialHLSLGenerator&
 {
 	using namespace UE::HLSLTree;
 
-	const FExpression* ExpressionTime = Time.TryAcquireHLSLExpression(Generator, Scope);
-	if (!ExpressionTime)
-	{
-		ExpressionTime = Generator.GetTree().NewExpression<Material::FExpressionExternalInput>(Material::EExternalInput::GameTime);
-	}
+	const FExpression* ExpressionTime = Time.AcquireHLSLExpressionOrExternalInput(Generator, Scope, Material::EExternalInput::GameTime);
 	ExpressionTime = Generator.GetTree().NewMul(ExpressionTime, Generator.NewConstant(Speed));
 
 	const FExpression* ExpressionCos = Generator.GetTree().NewUnaryOp(EOperation::Cos, ExpressionTime);
@@ -2176,11 +2150,7 @@ bool UMaterialExpressionRotator::GenerateHLSLExpression(FMaterialHLSLGenerator&
 	const FExpression* ExpressionRowY = Generator.GetTree().NewExpression<FExpressionAppend>(ExpressionSin, ExpressionCos);
 	const FExpression* ExpressionOrigin = Generator.GetTree().NewConstant(FVector2f(CenterX, CenterY));
 
-	const FExpression* ExpressionCoord = Coordinate.TryAcquireHLSLExpression(Generator, Scope);
-	if (!ExpressionCoord)
-	{
-		ExpressionCoord = Generator.NewTexCoord(ConstCoordinate);
-	}
+	const FExpression* ExpressionCoord = Coordinate.AcquireHLSLExpressionOrExternalInput(Generator, Scope, Material::MakeInputTexCoord(ConstCoordinate));
 	const FExpression* ExpressionCoordXY = Generator.GetTree().NewSub(Generator.NewSwizzle(FSwizzleParameters(0, 1), ExpressionCoord), ExpressionOrigin);
 
 	const FExpression* ExpressionResultX = Generator.GetTree().NewDot(ExpressionRowX, ExpressionCoordXY);
@@ -2435,13 +2405,16 @@ bool UMaterialExpressionExecEnd::GenerateHLSLStatements(FMaterialHLSLGenerator&
 
 bool UMaterialExpressionSetLocal::GenerateHLSLStatements(FMaterialHLSLGenerator& Generator, UE::HLSLTree::FScope& Scope) const
 {
-	const UE::HLSLTree::FExpression* ValueExpression = Value.AcquireHLSLExpression(Generator, Scope);
+	using namespace UE::HLSLTree;
+	const FExpression* ValueExpression = Value.AcquireHLSLExpression(Generator, Scope);
 	if (!ValueExpression)
 	{
 		return false;
 	}
 
-	Generator.GetTree().AssignLocal(Scope, LocalName, ValueExpression);
+	// Wrap the value being assigned in a new expression, so we can properly track input type for the UMaterialExpressionSetLocal node
+	const FExpression* LocalExpression = Generator.GetTree().NewExpression<FExpressionForward>(ValueExpression);
+	Generator.GetTree().AssignLocal(Scope, LocalName, LocalExpression);
 	Exec.GenerateHLSLStatements(Generator, Scope);
 	return true;
 }
diff --git a/Engine/Source/Runtime/Engine/Private/Materials/MaterialHLSLEmitter.cpp b/Engine/Source/Runtime/Engine/Private/Materials/MaterialHLSLEmitter.cpp
--- a/Engine/Source/Runtime/Engine/Private/Materials/MaterialHLSLEmitter.cpp
+++ b/Engine/Source/Runtime/Engine/Private/Materials/MaterialHLSLEmitter.cpp
@@ -703,39 +703,42 @@ public:
 		Material->ErrorExpressions.Reset();
 	}
 
-	virtual void AddErrorInternal(UObject* InOwner, FStringView InError) override
+	virtual void AddErrorInternal(TConstArrayView<UObject*> InOwners, FStringView InError) override
 	{
-		UMaterialExpression* MaterialExpressionOwner = Cast<UMaterialExpression>(InOwner);
-		UMaterialExpression* ExpressionToError = nullptr;
-		TStringBuilder<1024> FormattedError;
-
-		if (MaterialExpressionOwner)
+		for (UObject* Owner : InOwners)
 		{
-			if (MaterialExpressionOwner->GetClass() != UMaterialExpressionMaterialFunctionCall::StaticClass()
-				&& MaterialExpressionOwner->GetClass() != UMaterialExpressionFunctionInput::StaticClass()
-				&& MaterialExpressionOwner->GetClass() != UMaterialExpressionFunctionOutput::StaticClass())
+			UMaterialExpression* MaterialExpressionOwner = Cast<UMaterialExpression>(Owner);
+			UMaterialExpression* ExpressionToError = nullptr;
+			TStringBuilder<1024> FormattedError;
+
+			if (MaterialExpressionOwner)
 			{
-				// Add the expression currently being compiled to ErrorExpressions so we can draw it differently
-				ExpressionToError = MaterialExpressionOwner;
+				if (MaterialExpressionOwner->GetClass() != UMaterialExpressionMaterialFunctionCall::StaticClass()
+					&& MaterialExpressionOwner->GetClass() != UMaterialExpressionFunctionInput::StaticClass()
+					&& MaterialExpressionOwner->GetClass() != UMaterialExpressionFunctionOutput::StaticClass())
+				{
+					// Add the expression currently being compiled to ErrorExpressions so we can draw it differently
+					ExpressionToError = MaterialExpressionOwner;
 
-				const int32 ChopCount = FCString::Strlen(TEXT("MaterialExpression"));
-				const FString ErrorClassName = MaterialExpressionOwner->GetClass()->GetName();
+					const int32 ChopCount = FCString::Strlen(TEXT("MaterialExpression"));
+					const FString ErrorClassName = MaterialExpressionOwner->GetClass()->GetName();
 
-				// Add the node type to the error message
-				FormattedError.Appendf(TEXT("(Node %s) "), *ErrorClassName.Right(ErrorClassName.Len() - ChopCount));
+					// Add the node type to the error message
+					FormattedError.Appendf(TEXT("(Node %s) "), *ErrorClassName.Right(ErrorClassName.Len() - ChopCount));
+				}
 			}
-		}
 
-		FormattedError.Append(InError);
-		const FString Error(FormattedError.ToView());
+			FormattedError.Append(InError);
+			const FString Error(FormattedError.ToView());
 
-		// Standard error handling, immediately append one-off errors and signal failure
-		Material->CompileErrors.AddUnique(Error);
+			// Standard error handling, immediately append one-off errors and signal failure
+			Material->CompileErrors.AddUnique(Error);
 
-		if (ExpressionToError)
-		{
-			Material->ErrorExpressions.Add(ExpressionToError);
-			ExpressionToError->LastErrorText = Error;
+			if (ExpressionToError)
+			{
+				Material->ErrorExpressions.Add(ExpressionToError);
+				ExpressionToError->LastErrorText = Error;
+			}
 		}
 	}
 
@@ -806,7 +809,7 @@ bool MaterialEmitHLSL(const FMaterialCompileTargetParameters& InCompilerTarget,
 		FEmitScope* EmitResultScope = EmitContext.PrepareScope(CachedTree->GetResultScope());
 
 		// Prepare all fields *except* normal
-		FRequestedType RequestedPixelAttributesType;
+		FRequestedType RequestedPixelAttributesType(CachedTree->GetMaterialAttributesType(), false);
 		CachedTree->SetRequestedFields(SF_Pixel, RequestedPixelAttributesType);
 		RequestedPixelAttributesType.SetFieldRequested(NormalField, false);
 
@@ -881,7 +884,7 @@ bool MaterialEmitHLSL(const FMaterialCompileTargetParameters& InCompilerTarget,
 	// Prepare vertex shader code
 	FStringBuilderMemstack VertexCode(Allocator, 128 * 1024);
 	{
-		FRequestedType RequestedVertexAttributesType;
+		FRequestedType RequestedVertexAttributesType(CachedTree->GetMaterialAttributesType(), false);
 		CachedTree->SetRequestedFields(SF_Vertex, RequestedVertexAttributesType);
 		RequestedVertexAttributesType.SetFieldRequested(CachedTree->GetMaterialAttributesType()->FindFieldByName(TEXT("PrevWorldPositionOffset")));
 
@@ -910,7 +913,7 @@ bool MaterialEmitHLSL(const FMaterialCompileTargetParameters& InCompilerTarget,
 		CachedTree->GetTree().EmitShader(EmitContext, VertexCode);
 	}
 
-	if (EmitContext.Errors->HasErrors())
+	if (EmitContext.NumErrors > 0)
 	{
 		return false;
 	}
diff --git a/Engine/Source/Runtime/Engine/Private/Materials/MaterialHLSLGenerator.cpp b/Engine/Source/Runtime/Engine/Private/Materials/MaterialHLSLGenerator.cpp
--- a/Engine/Source/Runtime/Engine/Private/Materials/MaterialHLSLGenerator.cpp
+++ b/Engine/Source/Runtime/Engine/Private/Materials/MaterialHLSLGenerator.cpp
@@ -59,6 +59,17 @@ UMaterialExpression* FMaterialHLSLGenerator::GetCurrentExpression() const
 }
 
 bool FMaterialHLSLGenerator::Generate()
+{
+	const bool bResult = InternalGenerate();
+	if (!bResult)
+	{
+		CachedTree.ResultScope = &GetTree().GetRootScope();
+		CachedTree.ResultExpression = GetTree().NewExpression<UE::HLSLTree::FExpressionError>(AcquireError());
+	}
+	return bResult;
+}
+
+bool FMaterialHLSLGenerator::InternalGenerate()
 {
 	using namespace UE::HLSLTree;
 
@@ -164,7 +175,7 @@ bool FMaterialHLSLGenerator::GenerateResult(UE::HLSLTree::FScope& Scope)
 		for (int32 OutputIndex = 0; OutputIndex < FunctionEntry->FunctionOutputs.Num(); ++OutputIndex)
 		{
 			UMaterialExpressionFunctionOutput* ExpressionOutput = FunctionEntry->FunctionOutputs[OutputIndex];
-			FOwnerScope TreeOwnerScope(GetTree(), ExpressionOutput);
+			//FOwnerScope TreeOwnerScope(GetTree(), ExpressionOutput);
 			HLSLFunction->OutputExpressions.Add(ExpressionOutput->A.TryAcquireHLSLExpression(*this, Scope, 0));
 		}
 		FunctionEntry->bGeneratedResult = true;
@@ -229,7 +240,7 @@ bool FMaterialHLSLGenerator::GenerateResult(UE::HLSLTree::FScope& Scope)
 								InputExpression);
 							if (Property == MP_WorldPositionOffset)
 							{
-								const FExpression* PrevWPOExpression = GetTree().GetPreviousFrame(InputExpression, ERequestedType::Vector3);
+								const FExpression* PrevWPOExpression = GetTree().GetPreviousFrame(InputExpression, EValueType::Float3);
 								ensure(PrevWPOExpression);
 								AttributesExpression = GetTree().NewExpression<FExpressionSetStructField>(CachedTree.GetMaterialAttributesType(), PrevWPOField, AttributesExpression, PrevWPOExpression);
 							}
@@ -364,6 +375,11 @@ const UE::Shader::FTextureValue* FMaterialHLSLGenerator::AcquireTextureValue(con
 	return Value;
 }
 
+const UE::HLSLTree::FExpression* FMaterialHLSLGenerator::InternalNewErrorExpression(FStringView Error)
+{
+	return GetTree().NewExpression<UE::HLSLTree::FExpressionError>(UE::MemStack::AllocateStringView(GetTree().GetAllocator(), Error));
+}
+
 bool FMaterialHLSLGenerator::InternalError(FStringView ErrorMessage)
 {
 	if (CurrentErrorMessage.Len() > 0)
@@ -374,6 +390,13 @@ bool FMaterialHLSLGenerator::InternalError(FStringView ErrorMessage)
 	return false;
 }
 
+FStringView FMaterialHLSLGenerator::AcquireError()
+{
+	const FStringView ErrorMessage = UE::MemStack::AllocateStringView(GetTree().GetAllocator(), CurrentErrorMessage);
+	CurrentErrorMessage.Reset();
+	return ErrorMessage;
+}
+
 int32 FMaterialHLSLGenerator::FindInputIndex(const FExpressionInput* Input) const
 {
 	UMaterialExpression* OwnerMaterialExpression = GetCurrentExpression();
@@ -386,6 +409,36 @@ int32 FMaterialHLSLGenerator::FindInputIndex(const FExpressionInput* Input) cons
 	return Index;
 }
 
+const UE::HLSLTree::FExpression* FMaterialHLSLGenerator::NewDefaultInputConstant(int32 InputIndex, const UE::Shader::FValue& Value)
+{
+	using namespace UE::HLSLTree;
+	
+	const FExpression* Expression = GetTree().NewConstant(Value);
+
+	UObject* InputOwner = GetTree().GetCurrentOwner();
+	if (InputOwner && InputIndex != INDEX_NONE)
+	{
+		const FMaterialConnectionKey Key{ InputOwner, nullptr, InputIndex, INDEX_NONE };
+		CachedTree.ConnectionMap.Add(Key, Expression);
+	}
+	return Expression;
+}
+
+const UE::HLSLTree::FExpression* FMaterialHLSLGenerator::NewDefaultInputExternal(int32 InputIndex, UE::HLSLTree::Material::EExternalInput Input)
+{
+	using namespace UE::HLSLTree;
+	
+	const FExpression* Expression = NewExternalInput(Input);
+
+	UObject* InputOwner = GetTree().GetCurrentOwner();
+	if (InputOwner && InputIndex != INDEX_NONE)
+	{
+		const FMaterialConnectionKey Key{ InputOwner, nullptr, InputIndex, INDEX_NONE };
+		CachedTree.ConnectionMap.Add(Key, Expression);
+	}
+	return Expression;
+}
+
 const UE::HLSLTree::FExpression* FMaterialHLSLGenerator::AcquireExpression(UE::HLSLTree::FScope& Scope,
 	int32 InputIndex,
 	UMaterialExpression* MaterialExpression,
@@ -394,7 +447,7 @@ const UE::HLSLTree::FExpression* FMaterialHLSLGenerator::AcquireExpression(UE::H
 {
 	using namespace UE::HLSLTree;
 	UObject* InputOwner = GetTree().GetCurrentOwner();
-	FOwnerScope TreeOwnerScope(GetTree(), MaterialExpression);
+	FOwnerScope OwnerScope(GetTree(), MaterialExpression, NeedToPushOwnerExpression());
 
 	const FExpression* Expression = nullptr;
 	if (MaterialExpression->GenerateHLSLExpression(*this, Scope, OutputIndex, Expression))
@@ -410,14 +463,14 @@ const UE::HLSLTree::FExpression* FMaterialHLSLGenerator::AcquireExpression(UE::H
 			const FMaterialConnectionKey Key{ InputOwner, MaterialExpression, InputIndex, OutputIndex };
 			CachedTree.ConnectionMap.Add(Key, Expression);
 		}
-
-		return Expression;
 	}
-	
-	check(!Expression);
-	const FStringView ErrorMessage = UE::MemStack::AllocateStringView(GetTree().GetAllocator(), CurrentErrorMessage);
-	CurrentErrorMessage.Reset();
-	return GetTree().NewExpression<UE::HLSLTree::FExpressionError>(ErrorMessage);
+	else
+	{
+		check(!Expression);
+		Expression = GetTree().NewExpression<UE::HLSLTree::FExpressionError>(AcquireError());
+	}
+
+	return Expression;
 }
 
 const UE::HLSLTree::FExpression* FMaterialHLSLGenerator::AcquireFunctionInputExpression(UE::HLSLTree::FScope& Scope, const UMaterialExpressionFunctionInput* MaterialExpression)
@@ -517,8 +570,12 @@ bool FMaterialHLSLGenerator::GenerateStatements(UE::HLSLTree::FScope& Scope, UMa
 			}
 		}
 
-		FOwnerScope TreeOwnerScope(GetTree(), MaterialExpression);
+		FOwnerScope OwnerScope(GetTree(), MaterialExpression, NeedToPushOwnerExpression());
 		bResult = MaterialExpression->GenerateHLSLStatements(*this, *ScopeToUse);
+		if (!bResult)
+		{
+			GetTree().NewStatement<FStatementError>(*ScopeToUse, AcquireError());
+		}
 	}
 
 	return bResult;
diff --git a/Engine/Source/Runtime/Engine/Private/Materials/MaterialHLSLTree.cpp b/Engine/Source/Runtime/Engine/Private/Materials/MaterialHLSLTree.cpp
--- a/Engine/Source/Runtime/Engine/Private/Materials/MaterialHLSLTree.cpp
+++ b/Engine/Source/Runtime/Engine/Private/Materials/MaterialHLSLTree.cpp
@@ -605,7 +605,7 @@ bool FExpressionMaterialLayers::PrepareValue(FEmitContext& Context, FEmitScope&
 
 bool FExpressionSceneTexture::PrepareValue(FEmitContext& Context, FEmitScope& Scope, const FRequestedType& RequestedType, FPrepareValueResult& OutResult) const
 {
-	Context.PrepareExpression(TexCoordExpression, Scope, ERequestedType::Vector2);
+	Context.PrepareExpression(TexCoordExpression, Scope, Shader::EValueType::Float2);
 	if (Context.bMarkLiveValues)
 	{
 		Context.MaterialCompilationOutput->bNeedsSceneTextures = true;
@@ -652,8 +652,9 @@ void FExpressionSceneTexture::EmitValueShader(FEmitContext& Context, FEmitScope&
 
 bool FExpressionNoise::PrepareValue(FEmitContext& Context, FEmitScope& Scope, const FRequestedType& RequestedType, FPrepareValueResult& OutResult) const
 {
-	const FPreparedType& PositionType = Context.PrepareExpression(PositionExpression, Scope, ERequestedType::Vector3);
-	const FPreparedType& FilterWidthType = Context.PrepareExpression(FilterWidthExpression, Scope, ERequestedType::Scalar);
+	// TODO - we support Float3 or Double3 position input
+	const FPreparedType& PositionType = Context.PrepareExpression(PositionExpression, Scope, Shader::EValueType::Float3);
+	const FPreparedType& FilterWidthType = Context.PrepareExpression(FilterWidthExpression, Scope, Shader::EValueType::Float1);
 	if (PositionType.IsVoid() || FilterWidthType.IsVoid())
 	{
 		return false;
@@ -741,6 +742,8 @@ void FExpressionVertexInterpolator::EmitValueShader(FEmitContext& Context, FEmit
 		const FVertexInterpolator& Interpolator = EmitMaterialData.VertexInterpolators[InterpolatorIndex];
 
 		FRequestedType RequestedPreshaderType;
+		RequestedPreshaderType.ValueComponentType = RequestedType.ValueComponentType;
+		RequestedPreshaderType.StructType = RequestedType.StructType;
 		for (int32 Index = 0; Index < RequestedType.RequestedComponents.Num(); ++Index)
 		{
 			// Requested components that are *not* requested by the interpolator will be either constant or preshader evaluation
diff --git a/Engine/Source/Runtime/Engine/Private/Materials/MaterialShared.cpp b/Engine/Source/Runtime/Engine/Private/Materials/MaterialShared.cpp
--- a/Engine/Source/Runtime/Engine/Private/Materials/MaterialShared.cpp
+++ b/Engine/Source/Runtime/Engine/Private/Materials/MaterialShared.cpp
@@ -309,11 +309,11 @@ const UE::HLSLTree::FExpression* FExpressionInput::AcquireHLSLExpression(FMateri
 		}
 		if (!LocalInputName.IsNone())
 		{
-			Generator.Errorf(TEXT("Missing input '%s'"), *LocalInputName.ToString());
+			return Generator.NewErrorExpressionf(TEXT("Missing input '%s'"), *LocalInputName.ToString());
 		}
 		else
 		{
-			Generator.Error(TEXT("Missing input"));
+			Generator.NewErrorExpression(TEXT("Missing input"));
 		}
 		return nullptr;
 	}
@@ -325,14 +325,34 @@ const UE::HLSLTree::FExpression* FExpressionInput::AcquireHLSLExpression(FMateri
 	return AcquireHLSLExpression(Generator, Scope, Generator.FindInputIndex(this));
 }
 
+const UE::HLSLTree::FExpression* FExpressionInput::AcquireHLSLExpressionOrConstant(FMaterialHLSLGenerator& Generator, UE::HLSLTree::FScope& Scope, const UE::Shader::FValue& ConstantValue, int32 InputIndex) const
+{
+	const FExpressionInput TracedInput = GetTracedInput();
+	if (!TracedInput.Expression)
+	{
+		return Generator.NewDefaultInputConstant(InputIndex, ConstantValue);
+	}
+	return TryAcquireHLSLExpression(Generator, Scope, InputIndex);
+}
+
 const UE::HLSLTree::FExpression* FExpressionInput::AcquireHLSLExpressionOrConstant(FMaterialHLSLGenerator& Generator, UE::HLSLTree::FScope& Scope, const UE::Shader::FValue& ConstantValue) const
+{
+	return AcquireHLSLExpressionOrConstant(Generator, Scope, ConstantValue, Generator.FindInputIndex(this));
+}
+
+const UE::HLSLTree::FExpression* FExpressionInput::AcquireHLSLExpressionOrExternalInput(FMaterialHLSLGenerator& Generator, UE::HLSLTree::FScope& Scope, UE::HLSLTree::Material::EExternalInput Input, int32 InputIndex) const
 {
 	const FExpressionInput TracedInput = GetTracedInput();
 	if (!TracedInput.Expression)
 	{
-		return Generator.NewConstant(ConstantValue);
+		return Generator.NewDefaultInputExternal(InputIndex, Input);
 	}
-	return TryAcquireHLSLExpression(Generator, Scope);
+	return TryAcquireHLSLExpression(Generator, Scope, InputIndex);
+}
+
+const UE::HLSLTree::FExpression* FExpressionInput::AcquireHLSLExpressionOrExternalInput(FMaterialHLSLGenerator& Generator, UE::HLSLTree::FScope& Scope, UE::HLSLTree::Material::EExternalInput Input) const
+{
+	return AcquireHLSLExpressionOrExternalInput(Generator, Scope, Input, Generator.FindInputIndex(this));
 }
 
 void FExpressionInput::Connect( int32 InOutputIndex, class UMaterialExpression* InExpression )
diff --git a/Engine/Source/Runtime/Engine/Private/Shader/ShaderValue.cpp b/Engine/Source/Runtime/Engine/Private/Shader/ShaderValue.cpp
--- a/Engine/Source/Runtime/Engine/Private/Shader/ShaderValue.cpp
+++ b/Engine/Source/Runtime/Engine/Private/Shader/ShaderValue.cpp
@@ -666,6 +666,20 @@ FValueTypeDescription GetValueTypeDescription(EValueType Type)
 	}
 }
 
+EValueType FindValueType(FName Name)
+{
+	for (int32 TypeIndex = 1; TypeIndex < NumValueTypes; ++TypeIndex)
+	{
+		const EValueType Type = (EValueType)TypeIndex;
+		const FValueTypeDescription TypeDesc = GetValueTypeDescription(Type);
+		if (Name == TypeDesc.Name)
+		{
+			return Type;
+		}
+	}
+	return EValueType::Void;
+}
+
 EValueType MakeValueType(EValueComponentType ComponentType, int32 NumComponents)
 {
 	if (ComponentType == EValueComponentType::Void || NumComponents == 0)
diff --git a/Engine/Source/Runtime/Engine/Public/HLSLTree/HLSLTree.h b/Engine/Source/Runtime/Engine/Public/HLSLTree/HLSLTree.h
--- a/Engine/Source/Runtime/Engine/Public/HLSLTree/HLSLTree.h
+++ b/Engine/Source/Runtime/Engine/Public/HLSLTree/HLSLTree.h
@@ -47,75 +47,31 @@ class FEmitShaderExpression;
 
 static constexpr int32 MaxNumPreviousScopes = 2;
 
-class FOwnerContext
-{
-public:
-	virtual ~FOwnerContext() {}
-
-	void PushOwner(UObject* Owner);
-	UObject* PopOwner();
-	UObject* GetCurrentOwner() const;
-
-private:
-	TArray<UObject*, TInlineAllocator<8>> OwnerStack;
-};
-
-struct FOwnerScope
-{
-	FOwnerScope(FOwnerContext& InContext, UObject* InOwner) : Context(&InContext), Owner(InOwner)
-	{
-		Context->PushOwner(InOwner);
-	}
-
-	~FOwnerScope()
-	{
-		verify(Context->PopOwner() == Owner);
-	}
-
-	FOwnerContext* Context;
-	UObject* Owner;
-};
-
-class FErrorHandlerInterface : public FOwnerContext
+class FErrorHandlerInterface
 {
 public:
 	virtual ~FErrorHandlerInterface() {}
 
-	int32 GetNumErrors() const { return NumErrors; }
-	inline bool HasErrors() const { return NumErrors > 0; }
-
-	template<typename StringType>
-	inline bool AddError(const StringType& InError)
-	{
-		return AddErrorInternalWrapper(FStringView(InError));
-	}
-
-	template<typename FormatType, typename... Types>
-	inline bool AddErrorf(const FormatType& Format, Types... Args)
+	bool AddError(UObject* InOwner, FStringView InError)
 	{
-		TStringBuilder<2048> String;
-		String.Appendf(Format, Forward<Types>(Args)...);
-		return AddErrorInternalWrapper(FStringView(String.ToString(), String.Len()));
+		AddErrorInternal(MakeArrayView(&InOwner, 1), InError);
+		return false;
 	}
 
-protected:
-	bool AddErrorInternalWrapper(FStringView InError)
+	bool AddError(TConstArrayView<UObject*> InOwners, FStringView InError)
 	{
-		NumErrors++;
-		AddErrorInternal(GetCurrentOwner(), InError);
+		AddErrorInternal(InOwners, InError);
 		return false;
 	}
 
-	virtual void AddErrorInternal(UObject* InOwner, FStringView InError) = 0;
-
-private:
-	int32 NumErrors = 0;
+protected:
+	virtual void AddErrorInternal(TConstArrayView<UObject*> InOwners, FStringView InError) = 0;
 };
 
 class FNullErrorHandler : public FErrorHandlerInterface
 {
 public:
-	virtual void AddErrorInternal(UObject* InOwner, FStringView InError) override
+	virtual void AddErrorInternal(TConstArrayView<UObject*> InOwners, FStringView InError) override
 	{
 	}
 };
@@ -126,18 +82,19 @@ class FNode
 public:
 	virtual ~FNode() {}
 
-	UObject* GetOwner() const { return Owner; }
-
 private:
 	/** Next node in the FTree's list of all nodes */
 	FNode* NextNode = nullptr;
 
-	/** The UObject that owns/created this node, allows errors generated by this node to be properly associated with the owner */
-	UObject* Owner = nullptr;
-
 	friend class FTree;
 };
 
+class FOwnedNode : public FNode
+{
+public:
+	virtual TConstArrayView<UObject*> GetOwners() const = 0;
+};
+
 struct FEmitPreshaderScope
 {
 	FEmitPreshaderScope() = default;
@@ -152,19 +109,24 @@ struct FEmitPreshaderScope
  * Examples include assigning a value, or various control flow structures (if, for, while, etc)
  * This is an abstract base class, with derived classes representing various types of statements
  */
-class FStatement : public FNode
+class FStatement : public FOwnedNode
 {
 public:
+	UObject* GetOwner() const { return Owner; }
 	FScope& GetParentScope() const { return *ParentScope; }
 
+	virtual TConstArrayView<UObject*> GetOwners() const final { return MakeArrayView(&Owner, 1); }
 	virtual bool IsLoop() const { return false; }
 
 protected:
 	virtual bool Prepare(FEmitContext& Context, FEmitScope& Scope) const = 0;
-	virtual void EmitShader(FEmitContext& Context, FEmitScope& Scope) const = 0;
+	virtual void EmitShader(FEmitContext& Context, FEmitScope& Scope) const;
 	virtual void EmitPreshader(FEmitContext& Context, FEmitScope& Scope, const FRequestedType& RequestedType, TArrayView<const FEmitPreshaderScope> Scopes, Shader::FPreshaderData& OutPreshader) const;
 
 private:
+	/** The UObject that owns/created this node, allows errors generated by this node to be properly associated with the owner */
+	UObject* Owner = nullptr;
+
 	FScope* ParentScope = nullptr;
 
 	friend class FTree;
@@ -172,32 +134,18 @@ private:
 	friend class FExpressionLocalPHI;
 };
 
-/**
- * Used to initialize FRequestedType with a fixed number of components
- * enum value is the number of parameters
- */
-enum class ERequestedType : uint8
-{
-	Void = 0u,
-	Scalar = 1u,
-	Texture = 1u, // Currently no different than scalar, just makes code more clear
-	Vector2 = 2u,
-	Vector3 = 3u,
-	Vector4 = 4u,
-	Matrix4x4 = 16u,
-};
-
 /**
  * Represents a request for an arbitrary set of untyped components
  */
-class FRequestedType
+class ENGINE_API FRequestedType
 {
 public:
 	FRequestedType() = default;
-	FRequestedType(ERequestedType InType);
-	FRequestedType(Shader::EValueType InType);
-	FRequestedType(const Shader::FType& InType);
+	FRequestedType(Shader::EValueType InType, bool bDefaultRequest = true);
+	FRequestedType(const Shader::FType& InType, bool bDefaultRequest = true);
 
+	Shader::FType GetType() const;
+	const TCHAR* GetName() const { return GetType().GetName(); }
 	int32 GetNumComponents() const;
 	bool IsComponentRequested(int32 Index) const { return RequestedComponents.IsValidIndex(Index) ? (bool)RequestedComponents[Index] : false; }
 	bool IsVoid() const { return RequestedComponents.Find(true) == INDEX_NONE; }
@@ -222,12 +170,17 @@ public:
 	/** Returns the requested type of the given field */
 	FRequestedType GetField(const Shader::FStructField* Field) const;
 
+	const Shader::FStructType* StructType = nullptr;
+	Shader::EValueComponentType ValueComponentType = Shader::EValueComponentType::Void;
+
 	/** 1 bit per component, a value of 'true' means the specified component is requsted */
 	TBitArray<> RequestedComponents;
 };
 inline bool operator==(const FRequestedType& Lhs, const FRequestedType& Rhs)
 {
-	return Lhs.RequestedComponents == Rhs.RequestedComponents;
+	return Lhs.ValueComponentType == Rhs.ValueComponentType &&
+		Lhs.StructType == Rhs.StructType &&
+		Lhs.RequestedComponents == Rhs.RequestedComponents;
 }
 inline bool operator!=(const FRequestedType& Lhs, const FRequestedType& Rhs)
 {
@@ -308,7 +261,6 @@ public:
 	void SetComponentBounds(int32 Index, const Shader::FComponentBounds Bounds);
 	void MergeComponent(int32 Index, const FPreparedComponent& InComponent);
 
-	/** Unlike FRequestedType, one of these should be set */
 	const Shader::FStructType* StructType = nullptr;
 	Shader::EValueComponentType ValueComponentType = Shader::EValueComponentType::Void;
 
@@ -386,9 +338,11 @@ struct FExpressionDerivatives
  * Examples include constant literals, variable accessors, and various types of math operations
  * This is an abstract base class, with derived classes representing various types of expression
  */
-class FExpression : public FNode
+class FExpression : public FOwnedNode
 {
 public:
+	virtual TConstArrayView<UObject*> GetOwners() const final { return Owners; }
+
 	/**
 	 * Get Shader/Preshader/Constant value for the expression
 	 * @param RequestedType the componets of the result that are needed
@@ -422,6 +376,9 @@ protected:
 	virtual void EmitValueShader(FEmitContext& Context, FEmitScope& Scope, const FRequestedType& RequestedType, FEmitValueShaderResult& OutResult) const;
 	virtual void EmitValuePreshader(FEmitContext& Context, FEmitScope& Scope, const FRequestedType& RequestedType, FEmitValuePreshaderResult& OutResult) const;
 
+private:
+	TArray<UObject*, TInlineAllocator<2>> Owners;
+
 	friend class FTree;
 	friend class FEmitContext;
 	friend class FExpressionForward;
@@ -475,7 +432,7 @@ private:
 /**
  * The HLSL AST.  Basically a wrapper around the root scope, with some helper methods
  */
-class FTree : public FOwnerContext
+class FTree
 {
 public:
 	static FTree* Create(FMemStackBase& Allocator);
@@ -485,6 +442,10 @@ public:
 
 	bool Finalize();
 
+	void PushOwner(UObject* Owner);
+	UObject* PopOwner();
+	UObject* GetCurrentOwner() const;
+
 	bool EmitShader(FEmitContext& Context, FStringBuilderBase& OutCode) const;
 
 	FScope& GetRootScope() const { return *RootScope; }
@@ -496,13 +457,17 @@ public:
 		AppendHash(Hasher, GetGeneratedTypeName<T>());
 		AppendHashes(Hasher, Forward<ArgTypes>(Args)...);
 		const FXxHash64 Hash = Hasher.Finalize();
-		const FExpression* Expression = FindExpression(Hash);
+		FExpression* Expression = FindExpression(Hash);
 		if (!Expression)
 		{
 			T* TypedExpression = NewNode<T>(Forward<ArgTypes>(Args)...);
 			RegisterExpression(TypedExpression, Hash);
 			Expression = TypedExpression;
 		}
+		else
+		{
+			AddCurrentOwner(Expression);
+		}
 		return Expression;
 	}
 
@@ -578,18 +543,42 @@ private:
 	void RegisterNode(FNode* Node);
 	void RegisterExpression(FExpression* Expression, FXxHash64 Hash);
 	void RegisterExpression(FExpressionLocalPHI* Expression, FXxHash64 Hash);
+	void AddCurrentOwner(FExpression* Expression);
 	void RegisterStatement(FScope& Scope, FStatement* Statement);
-	const FExpression* FindExpression(FXxHash64 Hash) const;
+	FExpression* FindExpression(FXxHash64 Hash);
 
 	FMemStackBase* Allocator = nullptr;
 	FNode* Nodes = nullptr;
 	FScope* RootScope = nullptr;
-	TMap<FXxHash64, const FExpression*> ExpressionMap;
+	TArray<UObject*, TInlineAllocator<8>> OwnerStack;
+	TMap<FXxHash64, FExpression*> ExpressionMap;
 	TArray<FExpressionLocalPHI*> PHIExpressions;
 
 	friend class FExpressionLocalPHI;
 };
 
+struct FOwnerScope : private FNoncopyable
+{
+	FOwnerScope(FTree& InTree, UObject* InOwner, bool bPushOwner = true) : Tree(bPushOwner ? &InTree : nullptr), Owner(InOwner)
+	{
+		if (bPushOwner)
+		{
+			Tree->PushOwner(InOwner);
+		}
+	}
+
+	~FOwnerScope()
+	{
+		if (Tree)
+		{
+			verify(Tree->PopOwner() == Owner);
+		}
+	}
+
+	FTree* Tree;
+	UObject* Owner;
+};
+
 } // namespace UE::HLSLTree
 
 #endif // WITH_EDITOR
\ No newline at end of file
diff --git a/Engine/Source/Runtime/Engine/Public/HLSLTree/HLSLTreeCommon.h b/Engine/Source/Runtime/Engine/Public/HLSLTree/HLSLTreeCommon.h
--- a/Engine/Source/Runtime/Engine/Public/HLSLTree/HLSLTreeCommon.h
+++ b/Engine/Source/Runtime/Engine/Public/HLSLTree/HLSLTreeCommon.h
@@ -285,6 +285,18 @@ public:
 	const Shader::FStructType* OutputStructType = nullptr;
 };
 
+class FStatementError : public FStatement
+{
+public:
+	explicit FStatementError(FStringView InErrorMessage)
+		: ErrorMessage(InErrorMessage)
+	{}
+
+	virtual bool Prepare(FEmitContext& Context, FEmitScope& Scope) const override;
+
+	FStringView ErrorMessage;
+};
+
 class FStatementBreak : public FStatement
 {
 public:
diff --git a/Engine/Source/Runtime/Engine/Public/HLSLTree/HLSLTreeEmit.h b/Engine/Source/Runtime/Engine/Public/HLSLTree/HLSLTreeEmit.h
--- a/Engine/Source/Runtime/Engine/Public/HLSLTree/HLSLTreeEmit.h
+++ b/Engine/Source/Runtime/Engine/Public/HLSLTree/HLSLTreeEmit.h
@@ -28,6 +28,7 @@ namespace UE::HLSLTree
 
 class FErrorHandlerInterface;
 class FNode;
+class FOwnedNode;
 class FScope;
 class FExpression;
 class FFunction;
@@ -264,6 +265,28 @@ struct FTargetParameters
 	const ITargetPlatform* TargetPlatform = nullptr;
 };
 
+struct FConnectionKey
+{
+	FConnectionKey() = default;
+	FConnectionKey(const UObject* InInputObject, const FExpression* InOutputExpression) : InputObject(InInputObject), OutputExpression(InOutputExpression) {}
+
+	const UObject* InputObject = nullptr;
+	const FExpression* OutputExpression = nullptr;
+};
+inline uint32 GetTypeHash(const FConnectionKey& Key)
+{
+	return HashCombine(::GetTypeHash(Key.InputObject), ::GetTypeHash(Key.OutputExpression));
+}
+inline bool operator==(const FConnectionKey& Lhs, const FConnectionKey& Rhs)
+{
+	return Lhs.InputObject == Rhs.InputObject &&
+		Lhs.OutputExpression == Rhs.OutputExpression;
+}
+inline bool operator!=(const FConnectionKey& Lhs, const FConnectionKey& Rhs)
+{
+	return !operator==(Lhs, Rhs);
+}
+
 /** Tracks shared state while emitting HLSL code */
 class ENGINE_API FEmitContext
 {
@@ -275,6 +298,20 @@ public:
 
 	~FEmitContext();
 
+	template<typename StringType>
+	inline bool Error(const StringType& InError)
+	{
+		return InternalError(FStringView(InError));
+	}
+
+	template<typename FormatType, typename... Types>
+	inline bool Errorf(const FormatType& Format, Types... Args)
+	{
+		TStringBuilder<1024> String;
+		String.Appendf(Format, Forward<Types>(Args)...);
+		return InternalError(FStringView(String.ToString(), String.Len()));
+	}
+
 	template<typename T>
 	T& AcquireData()
 	{
@@ -314,6 +351,7 @@ public:
 	EExpressionEvaluation GetEvaluation(const FExpression* Expression, const FEmitScope& Scope, const FRequestedType& RequestedType) const;
 
 	FPreparedType PrepareExpression(const FExpression* InExpression, FEmitScope& Scope, const FRequestedType& RequestedType);
+	void MarkInputType(const FExpression* InExpression, const Shader::FType& Type);
 
 	FEmitScope* InternalPrepareScope(FScope* Scope, FScope* ParentScope);
 	FEmitScope* PrepareScope(FScope* Scope);
@@ -494,6 +532,7 @@ public:
 		T Data;
 	};
 
+	bool InternalError(FStringView ErrorMessage);
 	void InternalRegisterData(FXxHash64 Hash, FCustomDataWrapper* Data);
 	FCustomDataWrapper* InternalFindData(FXxHash64 Hash) const;
 
@@ -505,6 +544,7 @@ public:
 	bool bMarkLiveValues = false;
 	bool bUseAnalyticDerivatives = false;
 
+	TArray<const FOwnedNode*, TInlineAllocator<32>> OwnerStack;
 	TArray<FEmitShaderNode*> EmitNodes;
 	TMap<const FScope*, FEmitScope*> EmitScopeMap;
 	TMap<const FExpression*, FPrepareValueResult*> PrepareValueMap;
@@ -517,7 +557,9 @@ public:
 	TArray<struct FPreshaderLoopScope*> PreshaderLoopScopes;
 	TArray<const FPreshaderLocalPHIScope*> PreshaderLocalPHIScopes;
 	TMap<FXxHash64, FCustomDataWrapper*> CustomDataMap;
+	TMap<FConnectionKey, Shader::FType> ConnectionMap;
 	int32 PreshaderStackPosition = 0;
+	int32 NumErrors = 0;
 
 	int32 NumExpressionLocals = 0;
 	int32 NumExpressionLocalPHIs = 0;
@@ -530,6 +572,22 @@ public:
 	uint32 CurrentNumBoolComponents = 32u;
 };
 
+struct FEmitOwnerScope
+{
+	FEmitOwnerScope(FEmitContext& InContext, const FOwnedNode* InNode) : Context(InContext), Node(InNode)
+	{
+		InContext.OwnerStack.Add(InNode);
+	}
+
+	~FEmitOwnerScope()
+	{
+		verify(Context.OwnerStack.Pop(false) == Node);
+	}
+
+	FEmitContext& Context;
+	const FOwnedNode* Node;
+};
+
 } // namespace UE::HLSLTree
 
 #endif // WITH_EDITOR
diff --git a/Engine/Source/Runtime/Engine/Public/MaterialExpressionIO.h b/Engine/Source/Runtime/Engine/Public/MaterialExpressionIO.h
--- a/Engine/Source/Runtime/Engine/Public/MaterialExpressionIO.h
+++ b/Engine/Source/Runtime/Engine/Public/MaterialExpressionIO.h
@@ -8,19 +8,20 @@
 
 class FMaterialHLSLGenerator;
 
-namespace UE
-{
-namespace HLSLTree
+namespace UE::HLSLTree
 {
 class FScope;
 class FExpression;
 }
-namespace Shader
+namespace UE::HLSLTree::Material
+{
+enum class EExternalInput : uint8;
+}
+namespace UE::Shader
 {
 enum class EValueType : uint8;
 struct FValue;
 }
-}
 
 //
 //	FExpressionInput
@@ -79,10 +80,13 @@ struct FExpressionInput
 	ENGINE_API int32 Compile(class FMaterialCompiler* Compiler);
 	ENGINE_API const UE::HLSLTree::FExpression* TryAcquireHLSLExpression(FMaterialHLSLGenerator& Generator, UE::HLSLTree::FScope& Scope, int32 InputIndex) const;
 	ENGINE_API const UE::HLSLTree::FExpression* AcquireHLSLExpression(FMaterialHLSLGenerator& Generator, UE::HLSLTree::FScope& Scope, int32 InputIndex) const;
+	ENGINE_API const UE::HLSLTree::FExpression* AcquireHLSLExpressionOrConstant(FMaterialHLSLGenerator& Generator, UE::HLSLTree::FScope& Scope, const UE::Shader::FValue& ConstantValue, int32 InputIndex) const;
+	ENGINE_API const UE::HLSLTree::FExpression* AcquireHLSLExpressionOrExternalInput(FMaterialHLSLGenerator& Generator, UE::HLSLTree::FScope& Scope, UE::HLSLTree::Material::EExternalInput Input, int32 InputIndex) const;
 
 	ENGINE_API const UE::HLSLTree::FExpression* TryAcquireHLSLExpression(FMaterialHLSLGenerator& Generator, UE::HLSLTree::FScope& Scope) const;
 	ENGINE_API const UE::HLSLTree::FExpression* AcquireHLSLExpression(FMaterialHLSLGenerator& Generator, UE::HLSLTree::FScope& Scope) const;
 	ENGINE_API const UE::HLSLTree::FExpression* AcquireHLSLExpressionOrConstant(FMaterialHLSLGenerator& Generator, UE::HLSLTree::FScope& Scope, const UE::Shader::FValue& ConstantValue) const;
+	ENGINE_API const UE::HLSLTree::FExpression* AcquireHLSLExpressionOrExternalInput(FMaterialHLSLGenerator& Generator, UE::HLSLTree::FScope& Scope, UE::HLSLTree::Material::EExternalInput Input) const;
 #endif // WITH_EDITOR
 
 	/**
diff --git a/Engine/Source/Runtime/Engine/Public/MaterialHLSLGenerator.h b/Engine/Source/Runtime/Engine/Public/MaterialHLSLGenerator.h
--- a/Engine/Source/Runtime/Engine/Public/MaterialHLSLGenerator.h
+++ b/Engine/Source/Runtime/Engine/Public/MaterialHLSLGenerator.h
@@ -87,7 +87,7 @@ public:
 	template<typename FormatType, typename... Types>
 	inline bool Errorf(const FormatType& Format, Types... Args)
 	{
-		TStringBuilder<2048> String;
+		TStringBuilder<1024> String;
 		String.Appendf(Format, Forward<Types>(Args)...);
 		return InternalError(FStringView(String.ToString(), String.Len()));
 	}
@@ -106,10 +106,27 @@ public:
 	const UE::HLSLTree::FExpression* NewExternalInput(UE::HLSLTree::Material::EExternalInput Input);
 	const UE::HLSLTree::FExpression* NewSwizzle(const UE::HLSLTree::FSwizzleParameters& Params, const UE::HLSLTree::FExpression* Input);
 
+	template<typename StringType>
+	inline const UE::HLSLTree::FExpression* NewErrorExpression(const StringType& InError)
+	{
+		return InternalNewErrorExpression(FStringView(InError));
+	}
+
+	template<typename FormatType, typename... Types>
+	const UE::HLSLTree::FExpression* NewErrorExpressionf(const FormatType& Format, Types... Args)
+	{
+		TStringBuilder<1024> String;
+		String.Appendf(Format, Forward<Types>(Args)...);
+		return InternalNewErrorExpression(String.ToView());
+	}
+
 	const UE::Shader::FTextureValue* AcquireTextureValue(const UE::Shader::FTextureValue& Value);
 
 	int32 FindInputIndex(const FExpressionInput* Input) const;
 
+	const UE::HLSLTree::FExpression* NewDefaultInputConstant(int32 InputIndex, const UE::Shader::FValue& Value);
+	const UE::HLSLTree::FExpression* NewDefaultInputExternal(int32 InputIndex, UE::HLSLTree::Material::EExternalInput Input);
+
 	/**
 	 * Returns the appropriate HLSLNode representing the given UMaterialExpression.
 	 * The node will be created if it doesn't exist. Otherwise, the tree will be updated to ensure the node is visible in the given scope
@@ -206,11 +223,26 @@ private:
 		int32 NumInputs = 0;
 	};
 
+	bool InternalGenerate();
+
+	const UE::HLSLTree::FExpression* InternalNewErrorExpression(FStringView Error);
 	bool InternalError(FStringView ErrorMessage);
 
 	void InternalRegisterExpressionData(const FName& Type, const UMaterialExpression* MaterialExpression, void* Data);
 	void* InternalFindExpressionData(const FName& Type, const UMaterialExpression* MaterialExpression);
 
+	FStringView AcquireError();
+
+	bool NeedToPushOwnerExpression() const
+	{
+		// If we're inside a material function, don't change the owner to the current expression
+		// Instead, we let the initial UMaterialExpressionFunctionCall own *all* the expressions generated within the function call
+		// This allows errors generated within the function to be attributed back to the function call,
+		// as well as allow wire types to be properly tracked for wires going into the function
+		const FFunctionCallEntry* FunctionEntry = FunctionCallStack.Last();
+		return !FunctionEntry->MaterialFunction;
+	}
+
 	UMaterial* TargetMaterial;
 	const FMaterialLayersFunctions* LayerOverrides;
 	UMaterialExpression* PreviewExpression;
diff --git a/Engine/Source/Runtime/Engine/Public/Shader/ShaderTypes.h b/Engine/Source/Runtime/Engine/Public/Shader/ShaderTypes.h
--- a/Engine/Source/Runtime/Engine/Public/Shader/ShaderTypes.h
+++ b/Engine/Source/Runtime/Engine/Public/Shader/ShaderTypes.h
@@ -69,7 +69,10 @@ enum class EValueComponentType : uint8
 	TextureCubeArray,
 	Texture3D,
 	TextureExternal,
+
+	Num,
 };
+static constexpr int32 NumValueComponentTypes = (int32)EValueComponentType::Num;
 
 struct FValueComponentTypeDescription
 {
@@ -133,7 +136,10 @@ enum class EValueType : uint8
 	TextureCubeArray,
 	Texture3D,
 	TextureExternal,
+
+	Num,
 };
+static constexpr int32 NumValueTypes = (int32)EValueType::Num;
 
 struct FValueTypeDescription
 {
@@ -146,6 +152,7 @@ struct FValueTypeDescription
 };
 
 ENGINE_API FValueTypeDescription GetValueTypeDescription(EValueType Type);
+ENGINE_API EValueType FindValueType(FName Name);
 inline bool IsLWCType(EValueType Type) { return IsLWCType(GetValueTypeDescription(Type).ComponentType); }
 ENGINE_API EValueType MakeValueType(EValueComponentType ComponentType, int32 NumComponents);
 ENGINE_API EValueType MakeValueType(EValueType BaseType, int32 NumComponents);
