Texture properly deprecate bPNGCompressed
remove unused ETextureSourceArtType
add GetGammaSpace on TextureSource

diff --git a/Engine/Source/Runtime/Engine/Classes/Engine/Texture.h b/Engine/Source/Runtime/Engine/Classes/Engine/Texture.h
--- a/Engine/Source/Runtime/Engine/Classes/Engine/Texture.h
+++ b/Engine/Source/Runtime/Engine/Classes/Engine/Texture.h
@@ -79,18 +79,6 @@ enum ETextureMipCount
 	TMC_MAX,
 };
 
-UENUM()
-enum ETextureSourceArtType
-{
-	/** FColor Data[SrcWidth * SrcHeight]. */
-	TSAT_Uncompressed,
-	/** PNG compresed version of FColor Data[SrcWidth * SrcHeight]. */
-	TSAT_PNGCompressed,
-	/** DDS file with header. */
-	TSAT_DDSFile,
-	TSAT_MAX,
-};
-
 // TextureCompressionQuality is used for ASTC
 UENUM()
 enum ETextureCompressionQuality
@@ -341,7 +329,10 @@ struct FTextureSource
 	FString GetSourceCompressionAsString() const;
 
 	/** Returns the compression format of the source data in enum format. */
-	ETextureSourceCompressionFormat GetSourceCompression() const;
+	FORCEINLINE ETextureSourceCompressionFormat GetSourceCompression() const { return CompressionFormat; }
+	
+	/** Get GammaSpace for this Format */
+	ENGINE_API EGammaSpace GetGammaSpace() const;
 
 	/** Support for copy/paste */
 	void ExportCustomProperties(FOutputDevice& Out, uint32 Indent);
@@ -357,7 +348,7 @@ struct FTextureSource
 	FORCEINLINE int32 GetNumLayers() const { return NumLayers; }
 	FORCEINLINE int32 GetNumBlocks() const { return Blocks.Num() + 1; }
 	FORCEINLINE ETextureSourceFormat GetFormat(int32 LayerIndex = 0) const { return (LayerIndex == 0) ? Format : LayerFormat[LayerIndex]; }
-	FORCEINLINE bool IsPNGCompressed() const { return bPNGCompressed; }
+	FORCEINLINE bool IsPNGCompressed() const { return GetSourceCompression() == ETextureSourceCompressionFormat::TSCF_PNG; }
 	FORCEINLINE bool IsLongLatCubemap() const { return bLongLatCubemap; }
 	FORCEINLINE int64 GetSizeOnDisk() const { return BulkData.GetPayloadSize(); }
 	inline bool HasPayloadData() const { return BulkData.HasPayloadData(); }
@@ -632,9 +623,10 @@ private:
 	int32 NumLayers;
 
 	/** RGBA8 source data is optionally compressed as PNG. 
-	Deprecated, use CompressionFormat instead.  To be removed. */
-	UPROPERTY(VisibleAnywhere, Category=TextureSource)
-	bool bPNGCompressed;
+	Deprecated, use CompressionFormat instead.  To be removed.
+	Deprecated uproperties are loaded but not saved. */
+	UPROPERTY()
+	bool bPNGCompressed_DEPRECATED;
 
 	/**
 	 * Source represents a cubemap in long/lat format, will have only 1 slice per cube, rather than 6 slices.
diff --git a/Engine/Source/Runtime/Engine/Classes/Engine/TextureDefines.h b/Engine/Source/Runtime/Engine/Classes/Engine/TextureDefines.h
--- a/Engine/Source/Runtime/Engine/Classes/Engine/TextureDefines.h
+++ b/Engine/Source/Runtime/Engine/Classes/Engine/TextureDefines.h
@@ -295,10 +295,10 @@ enum ETextureSourceFormat
 
 	TSF_MAX,
 
-	// error in UHT, enable when fixed
 	// provide aliases to the old names with deprecation warnings
-	//TSF_RGBA8 UE_DEPRECATED(5.1,"Legacy ETextureSourceFormat not supported, use BGRA8") = TSF_RGBA8_DEPRECATED,
-	//TSF_RGBE8 UE_DEPRECATED(5.1,"Legacy ETextureSourceFormat not supported, use BGRE8") = TSF_RGBE8_DEPRECATED
+	//  remove these someday
+	TSF_RGBA8 UE_DEPRECATED(5.1,"Legacy ETextureSourceFormat not supported, use BGRA8") = TSF_RGBA8_DEPRECATED,
+	TSF_RGBE8 UE_DEPRECATED(5.1,"Legacy ETextureSourceFormat not supported, use BGRE8") = TSF_RGBE8_DEPRECATED
 };
 
 // This needs to be mirrored in EditorFactories.cpp.
diff --git a/Engine/Source/Runtime/Engine/Private/Texture.cpp b/Engine/Source/Runtime/Engine/Private/Texture.cpp
--- a/Engine/Source/Runtime/Engine/Private/Texture.cpp
+++ b/Engine/Source/Runtime/Engine/Private/Texture.cpp
@@ -695,6 +695,15 @@ void UTexture::PostLoad()
 		AssetImportData->SourceData = MoveTemp(Info);
 	}
 
+	if ( Source.bPNGCompressed_DEPRECATED )
+	{
+		// loaded with deprecated "bPNGCompressed"
+		// change to CompressionFormat PNG
+		check( Source.CompressionFormat == TSCF_None || Source.CompressionFormat == TSCF_PNG );
+		Source.CompressionFormat = TSCF_PNG;
+		Source.bPNGCompressed_DEPRECATED = false;
+	}
+
 	if ( Source.GetFormat() == TSF_RGBA8_DEPRECATED
 		|| Source.GetFormat() == TSF_RGBE8_DEPRECATED )
 	{
@@ -1235,7 +1244,7 @@ FTextureSource::FTextureSource()
 	, NumSlices(0)
 	, NumMips(0)
 	, NumLayers(1) // Default to 1 so old data has the correct value
-	, bPNGCompressed(false)
+	, bPNGCompressed_DEPRECATED(false)
 	, bLongLatCubemap(false)
 	, CompressionFormat(TSCF_None)
 	, bGuidIsHash(false)
@@ -1487,7 +1496,7 @@ void FTextureSource::Compress()
 	//	should be faster to load and also smaller files (than traditional PNG+zlib)
 	bool bUseOodleOnPNGz0 = true;
 
-	// may already have bPNGCompressed or "CompressionFormat" set
+	// may already have "CompressionFormat" set
 
 	if (CanPNGCompress()) // Note that this will return false if the data is already a compressed PNG
 	{
@@ -1512,18 +1521,11 @@ void FTextureSource::Compress()
 			{
 				BulkData.UpdatePayload(MakeSharedBufferFromArray(MoveTemp(CompressedData)), Owner);
 
-				bPNGCompressed = true;
 				CompressionFormat = TSCF_PNG;
 			}
 		}
 	}
 
-	// Fix up for packages that were saved before CompressionFormat was introduced. Can removed this when we deprecate bPNGCompressed!
-	if (bPNGCompressed)
-	{
-		CompressionFormat = TSCF_PNG;
-	}
-
 	if ( ( CompressionFormat == TSCF_PNG && bUseOodleOnPNGz0 ) ||
 		CompressionFormat == TSCF_None )
 	{
@@ -1541,7 +1543,7 @@ FSharedBuffer FTextureSource::Decompress(IImageWrapperModule* ) const
 	
 	// ImageWrapperModule argument ignored, not drilled through DecompressImage
 
-	if (CompressionFormat != TSCF_None || bPNGCompressed)
+	if (CompressionFormat != TSCF_None )
 	{
 		return TryDecompressData();
 	}
@@ -1572,16 +1574,7 @@ FTextureSource::FMipLock::FMipLock(ELockState InLockState,FTextureSource * InTex
 		Image.SizeY = FMath::Max(Block.SizeY >> MipIndex, 1);
 		Image.NumSlices = Block.NumSlices;
 		Image.Format = FImageCoreUtils::ConvertToRawImageFormat(TextureSource->Format);
-
-		if ( TextureSource->Owner != nullptr )
-		{
-			Image.GammaSpace = TextureSource->Owner->GetGammaSpace();
-		}
-		else
-		{
-			// TextureSource does not have gamma information, just guess it ?
-			Image.GammaSpace = ERawImageFormat::GetDefaultGammaSpace(Image.Format);
-		}
+		Image.GammaSpace = TextureSource->GetGammaSpace();
 		
 		const int64 MipSizeBytes = TextureSource->CalcMipSize(BlockIndex, LayerIndex, MipIndex);
 
@@ -1699,7 +1692,6 @@ void FTextureSource::UnlockMip(int32 BlockIndex, int32 LayerIndex, int32 MipInde
 			BulkData.UpdatePayload(LockedMipData.Release(), Owner);
 			BulkData.SetCompressionOptions(UE::Serialization::ECompressionOptions::Default);
 
-			bPNGCompressed = false;
 			CompressionFormat = TSCF_None;
 
 			// Need to unlock before calling UseHashAsGuid
@@ -1732,16 +1724,7 @@ bool FTextureSource::GetMipImage(FImage & OutImage, int32 BlockIndex, int32 Laye
 	OutImage.SizeY = FMath::Max(Block.SizeY >> MipIndex, 1);
 	OutImage.NumSlices = Block.NumSlices;
 	OutImage.Format = FImageCoreUtils::ConvertToRawImageFormat(Format);
-
-	if ( Owner != nullptr )
-	{
-		OutImage.GammaSpace = Owner->GetGammaSpace();
-	}
-	else
-	{
-		// TextureSource does not have gamma information, just guess it ?
-		OutImage.GammaSpace = ERawImageFormat::GetDefaultGammaSpace(OutImage.Format);
-	}
+	OutImage.GammaSpace = GetGammaSpace();
 
 	check( OutImage.GetImageSizeBytes() == MipSizeBytes );
 	return true;
@@ -1900,17 +1883,18 @@ FString FTextureSource::GetIdString() const
 	return GuidString;
 }
 
-ETextureSourceCompressionFormat FTextureSource::GetSourceCompression() const
+EGammaSpace FTextureSource::GetGammaSpace() const
 {
-	// Until we deprecate bPNGCompressed it might not be 100% in sync with CompressionFormat
-	// so if it is set we should use that rather than the enum.
-	// @todo fix this, remove bPNGCompressed
-	if (bPNGCompressed)
+	// TextureSource does not know its own gamma, but its owning Texture does :
+	if ( Owner != nullptr )
 	{
-		return ETextureSourceCompressionFormat::TSCF_PNG;
+		return Owner->GetGammaSpace();
+	}
+	else
+	{
+		// TextureSource does not have gamma information, just guess it ?
+		return ERawImageFormat::GetDefaultGammaSpace( FImageCoreUtils::ConvertToRawImageFormat(Format) );
 	}
-
-	return CompressionFormat;
 }
 
 FString FTextureSource::GetSourceCompressionAsString() const
@@ -1940,12 +1924,10 @@ FSharedBuffer FTextureSource::TryDecompressData() const
 			// this shouldn't ever happen currently
 			UE_LOG(LogTexture, Warning, TEXT("TryDecompressData unexpected format conversion?"));
 
-			// TextureSource does not have gamma information, so just guess:
-			//  (gamma conversion should never actually happen here so this should be unused)
-			Image.ChangeFormat(RawFormat, ERawImageFormat::GetDefaultGammaSpace(RawFormat));
+			Image.ChangeFormat(RawFormat, GetGammaSpace());
 		}
 			
-		if ( (CompressionFormat == TSCF_PNG || bPNGCompressed) && Image.Format == ERawImageFormat::BGRA8 )
+		if ( CompressionFormat == TSCF_PNG && Image.Format == ERawImageFormat::BGRA8 )
 		{
 			// Legacy bug, must be matched in Compress & Decompress
 			// see FTextureSource::Compress
@@ -2058,7 +2040,7 @@ bool FTextureSource::CanPNGCompress() const
 {
 	bool bCanPngCompressFormat = (Format == TSF_G8 || Format == TSF_G16 || Format == TSF_BGRA8 || Format == TSF_RGBA16 );
 
-	if (!bPNGCompressed &&
+	if (
 		NumLayers == 1 &&
 		NumMips == 1 &&
 		NumSlices == 1 &&
@@ -2097,7 +2079,6 @@ void FTextureSource::RemoveSourceData()
 	LayerFormat.Empty();
 	Blocks.Empty();
 	BlockDataOffsets.Empty();
-	bPNGCompressed = false;
 	CompressionFormat = TSCF_None;
 	LockedMipData.Reset();
 	NumLockedMips = 0u;
