New HLSL material translator, support automatic mip bias input

diff --git a/Engine/Source/Runtime/Engine/Private/HLSLTree/HLSLTreeCommon.cpp b/Engine/Source/Runtime/Engine/Private/HLSLTree/HLSLTreeCommon.cpp
--- a/Engine/Source/Runtime/Engine/Private/HLSLTree/HLSLTreeCommon.cpp
+++ b/Engine/Source/Runtime/Engine/Private/HLSLTree/HLSLTreeCommon.cpp
@@ -178,6 +178,15 @@ bool FExpressionTextureSample::PrepareValue(FEmitContext& Context, FEmitScope& S
 		return false;
 	}
 
+	if (AutomaticMipBiasExpression)
+	{
+		const FPreparedType& AutomaticMipBiasType = Context.PrepareExpression(AutomaticMipBiasExpression, Scope, Shader::EValueType::Bool1);
+		if (!IsConstantEvaluation(AutomaticMipBiasType.GetEvaluation(Scope, Shader::EValueType::Bool1)))
+		{
+			return Context.Error(TEXT("Automatic Mip Bias input must be constant"));
+		}
+	}
+
 	const bool bUseAnalyticDerivatives = Context.bUseAnalyticDerivatives && (MipValueMode != TMVM_MipLevel) && TexCoordDerivatives.IsValid();
 	if (MipValueMode == TMVM_Derivative || bUseAnalyticDerivatives)
 	{
@@ -227,7 +236,7 @@ void FExpressionTextureSample::EmitValueShader(FEmitContext& Context, FEmitScope
 		break;
 	}
 
-	const bool AutomaticViewMipBias = false; // TODO
+	const bool AutomaticViewMipBias = AutomaticMipBiasExpression ? AutomaticMipBiasExpression->GetValueConstant(Context, Scope, Shader::EValueType::Bool1).AsBoolScalar() : false;
 	TStringBuilder<256> FormattedSampler;
 	switch (SamplerSource)
 	{
diff --git a/Engine/Source/Runtime/Engine/Private/Materials/MaterialExpressionHLSL.cpp b/Engine/Source/Runtime/Engine/Private/Materials/MaterialExpressionHLSL.cpp
--- a/Engine/Source/Runtime/Engine/Private/Materials/MaterialExpressionHLSL.cpp
+++ b/Engine/Source/Runtime/Engine/Private/Materials/MaterialExpressionHLSL.cpp
@@ -351,7 +351,7 @@ bool UMaterialExpressionCurveAtlasRowParameter::GenerateHLSLExpression(FMaterial
 	const FExpression* ExpressionCoordV = Tree.NewMul(Tree.NewAdd(ExpressionSlot, Tree.NewConstant(0.5f)), Tree.NewRcp(ExpressionTextureHeight));
 	const FExpression* ExpressionUV = Tree.NewExpression<FExpressionAppend>(ExpressionCoordU, ExpressionCoordV);
 
-	OutExpression = Generator.GetTree().NewExpression<FExpressionTextureSample>(ExpressionTexture, ExpressionUV, nullptr, FExpressionDerivatives(), SSM_Clamp_WorldGroupSettings, TMVM_None);
+	OutExpression = Generator.GetTree().NewExpression<FExpressionTextureSample>(ExpressionTexture, ExpressionUV, nullptr, nullptr, FExpressionDerivatives(), SSM_Clamp_WorldGroupSettings, TMVM_None);
 	return true;
 }
 
@@ -741,8 +741,9 @@ bool UMaterialExpressionTextureSample::GenerateHLSLExpressionBase(FMaterialHLSLG
 		checkNoEntry();
 		break;
 	}
-	
-	OutExpression = Generator.GetTree().NewExpression<FExpressionTextureSample>(TextureExpression, TexCoordExpression, MipLevelExpression, TexCoordDerivatives, SamplerSource, MipValueMode);
+
+	const FExpression* AutomaticMipBiasExpression = AutomaticViewMipBiasValue.AcquireHLSLExpressionOrConstant(Generator, Scope, (bool)AutomaticViewMipBias);
+	OutExpression = Generator.GetTree().NewExpression<FExpressionTextureSample>(TextureExpression, TexCoordExpression, MipLevelExpression, AutomaticMipBiasExpression, TexCoordDerivatives, SamplerSource, MipValueMode);
 	return true;
 }
 
@@ -843,7 +844,7 @@ bool UMaterialExpressionFontSample::GenerateHLSLExpression(FMaterialHLSLGenerato
 	const FExpression* TexCoordExpression = Generator.NewTexCoord(0);
 	const FExpressionDerivatives TexCoordDerivatives = Generator.GetTree().GetAnalyticDerivatives(TextureExpression);
 
-	OutExpression = Generator.GetTree().NewExpression<FExpressionTextureSample>(TextureExpression, TexCoordExpression, nullptr, TexCoordDerivatives, SSM_FromTextureAsset, TMVM_None);
+	OutExpression = Generator.GetTree().NewExpression<FExpressionTextureSample>(TextureExpression, TexCoordExpression, nullptr, nullptr, TexCoordDerivatives, SSM_FromTextureAsset, TMVM_None);
 	return true;
 }
 
diff --git a/Engine/Source/Runtime/Engine/Private/Materials/MaterialHLSLGenerator.cpp b/Engine/Source/Runtime/Engine/Private/Materials/MaterialHLSLGenerator.cpp
--- a/Engine/Source/Runtime/Engine/Private/Materials/MaterialHLSLGenerator.cpp
+++ b/Engine/Source/Runtime/Engine/Private/Materials/MaterialHLSLGenerator.cpp
@@ -275,14 +275,19 @@ bool FMaterialHLSLGenerator::GenerateResult(UE::HLSLTree::FScope& Scope)
 				const int32 OutputIndex = 0;
 				PreviewExpressionResult = AcquireExpression(Scope, INDEX_NONE, PreviewExpression, OutputIndex, FSwizzleParameters());
 			}
-			const FString& EmissiveColorName = FMaterialAttributeDefinitionMap::GetAttributeName(MP_EmissiveColor);
-			const FStructField* EmissiveColorField = CachedTree.GetMaterialAttributesType()->FindFieldByName(*EmissiveColorName);
 
-			// Get back into gamma corrected space, as DrawTile does not do this adjustment.
-			const FExpression* ExpressionEmissive = GetTree().NewPowClamped(PreviewExpressionResult, NewConstant(1.f / 2.2f));
+			const FExpression* ExpressionEmissive = PreviewExpressionResult;
+			if (ExpressionEmissive)
+			{
+				const FString& EmissiveColorName = FMaterialAttributeDefinitionMap::GetAttributeName(MP_EmissiveColor);
+				const FStructField* EmissiveColorField = CachedTree.GetMaterialAttributesType()->FindFieldByName(*EmissiveColorName);
+
+				// Get back into gamma corrected space, as DrawTile does not do this adjustment.
+				ExpressionEmissive = GetTree().NewPowClamped(ExpressionEmissive, NewConstant(1.f / 2.2f));
 
-			AttributesExpression = GetTree().NewExpression<FExpressionConstant>(CachedTree.GetMaterialAttributesDefaultValue());
-			AttributesExpression = GetTree().NewExpression<FExpressionSetStructField>(CachedTree.GetMaterialAttributesType(), EmissiveColorField, AttributesExpression, ExpressionEmissive);
+				AttributesExpression = GetTree().NewExpression<FExpressionConstant>(CachedTree.GetMaterialAttributesDefaultValue());
+				AttributesExpression = GetTree().NewExpression<FExpressionSetStructField>(CachedTree.GetMaterialAttributesType(), EmissiveColorField, AttributesExpression, ExpressionEmissive);
+			}
 		}
 
 		if (AttributesExpression)
diff --git a/Engine/Source/Runtime/Engine/Public/HLSLTree/HLSLTreeCommon.h b/Engine/Source/Runtime/Engine/Public/HLSLTree/HLSLTreeCommon.h
--- a/Engine/Source/Runtime/Engine/Public/HLSLTree/HLSLTreeCommon.h
+++ b/Engine/Source/Runtime/Engine/Public/HLSLTree/HLSLTreeCommon.h
@@ -59,12 +59,14 @@ public:
 	FExpressionTextureSample(const FExpression* InTextureExpression,
 		const FExpression* InTexCoordExpression,
 		const FExpression* InMipValueExpression,
+		const FExpression* InAutomaticMipBiasExpression,
 		const FExpressionDerivatives& InTexCoordDerivatives,
 		ESamplerSourceMode InSamplerSource,
 		ETextureMipValueMode InMipValueMode)
 		: TextureExpression(InTextureExpression)
 		, TexCoordExpression(InTexCoordExpression)
 		, MipValueExpression(InMipValueExpression)
+		, AutomaticMipBiasExpression(InAutomaticMipBiasExpression)
 		, TexCoordDerivatives(InTexCoordDerivatives)
 		, SamplerSource(InSamplerSource)
 		, MipValueMode(InMipValueMode)
@@ -73,6 +75,7 @@ public:
 	const FExpression* TextureExpression;
 	const FExpression* TexCoordExpression;
 	const FExpression* MipValueExpression;
+	const FExpression* AutomaticMipBiasExpression;
 	FExpressionDerivatives TexCoordDerivatives;
 	ESamplerSourceMode SamplerSource;
 	ETextureMipValueMode MipValueMode;
