Move call of GetDBufferReprojectedWorldNormal() into scene texture samplnig shader code and out of material HLSL generation.
This allows us to select based on DBUFFER define.
This fixes emissive decals that sample Normal when calculating Opacity, since they never need the reprojection path.
Previously they were taking a stubbed version of reprojection when DBuffer was enabled on the project.

diff --git a/Engine/Shaders/Private/DBufferNormalReprojection.ush b/Engine/Shaders/Private/DBufferNormalReprojection.ush
--- a/Engine/Shaders/Private/DBufferNormalReprojection.ush
+++ b/Engine/Shaders/Private/DBufferNormalReprojection.ush
@@ -65,7 +65,7 @@ float3 ReconstructNormalFromDepthBuffer(int2 BaseCoord)
 	return normalize(cross(Right, Down));
 }
 
-float4 GetDBufferReprojectedWorldNormal(float2 UV, float3 FallbackNormal)
+float4 GetDBufferReprojectedWorldNormal(float2 UV)
 {
 	int2 BaseCoord = (int2)trunc(UV * View.BufferSizeAndInvSize.xy);
 
@@ -115,12 +115,12 @@ float4 GetDBufferReprojectedWorldNormal(float2 UV, float3 FallbackNormal)
 	return float4(N,1.0f);
 }
 
-#else
+#else // IS_DBUFFER_DECAL && FEATURE_LEVEL >= FEATURE_LEVEL_SM5
 
-// for mobile, just return the geometry normal of the surface on top as a fallback
-float4 GetDBufferReprojectedWorldNormal(float2 UV, float3 FallbackNormal)
+float4 GetDBufferReprojectedWorldNormal(float2 UV)
 {
-	return float4(FallbackNormal.xyz,1.0);
+	// Stub function. Shouldn't ever be called.
+	return float4(0, 0, 1, 1);
 }
 
-#endif
\ No newline at end of file
+#endif // IS_DBUFFER_DECAL && FEATURE_LEVEL >= FEATURE_LEVEL_SM5
\ No newline at end of file
diff --git a/Engine/Shaders/Private/MaterialTemplate.ush b/Engine/Shaders/Private/MaterialTemplate.ush
--- a/Engine/Shaders/Private/MaterialTemplate.ush
+++ b/Engine/Shaders/Private/MaterialTemplate.ush
@@ -2319,7 +2319,12 @@ float4 SceneTextureLookup(float2 UV, int SceneTextureIndex, bool bFiltered)
 		case PPI_Metallic:
 			return ScreenSpaceData.GBuffer.Metallic;
 		case PPI_WorldNormal:
+#if IS_DBUFFER_DECAL
+			// DBuffer uses depth buffer gradients or reprojection of last frame normal.
+			return GetDBufferReprojectedWorldNormal(UV);
+#else
 			return float4(ScreenSpaceData.GBuffer.WorldNormal, 0);
+#endif
 		case PPI_SeparateTranslucency:
 			return float4(1, 1, 1, 1);	// todo
 		case PPI_Opacity:
diff --git a/Engine/Source/Runtime/Engine/Private/Materials/HLSLMaterialTranslator.cpp b/Engine/Source/Runtime/Engine/Private/Materials/HLSLMaterialTranslator.cpp
--- a/Engine/Source/Runtime/Engine/Private/Materials/HLSLMaterialTranslator.cpp
+++ b/Engine/Source/Runtime/Engine/Private/Materials/HLSLMaterialTranslator.cpp
@@ -6603,71 +6603,10 @@ int32 FHLSLMaterialTranslator::SceneTextureLookup(int32 ViewportUV, uint32 InSce
 		// we can relax this later if needed
 		return NonPixelShaderExpressionError();
 	}
-		
-	// Special case for DBuffer normals.
-	// Note that full reprojection only happens with r.Decal.NormalReprojectionEnabled 1.
-	// Otherwise we use a depth buffer approximation.
-	bool bIsUsingReprojectedNormal = false;
-	if (Material->GetMaterialDomain() == MD_DeferredDecal)
-	{
-		if (SceneTextureId == PPI_WorldNormal)
-		{
-			const bool bUsingDBuffer = IsUsingDBuffers(Platform);
-			const bool bIsDBufferOutput = 
-				MaterialProperty == MP_BaseColor ||
-				MaterialProperty == MP_Normal || 
-				MaterialProperty == MP_Roughness || 
-				MaterialProperty == MP_Specular || 
-				MaterialProperty == MP_Metallic ||
-				MaterialProperty == MP_Opacity;
-
-			if (bUsingDBuffer && bIsDBufferOutput)
-			{
-				bIsUsingReprojectedNormal = true;
-			}
-		}
-	}
-
-	if (bIsUsingReprojectedNormal)
-	{
-		int32 BufferUV;
-		if (ViewportUV != INDEX_NONE)
-		{
-			BufferUV = AddCodeChunk(MCT_Float2,
-				TEXT("ClampSceneTextureUV(ViewportUVToSceneTextureUV(%s, %d), %d)"),
-				*CoerceParameter(ViewportUV, MCT_Float2), (int)SceneTextureId, (int)SceneTextureId);
-		}
-		else
-		{
-			BufferUV = AddInlinedCodeChunk(MCT_Float2, TEXT("GetDefaultSceneTextureUV(Parameters, %d)"), (int)SceneTextureId);
-		}
-
-		int32 PixelNormal = PixelNormalWS();
-		if (PixelNormal == INDEX_NONE)
-		{
-			return INDEX_NONE;
-		}
-
-		if (FeatureLevel >= ERHIFeatureLevel::SM5)
-		{
-			return AddCodeChunk(MCT_Float4,
-				TEXT("GetDBufferReprojectedWorldNormal(%s,%s)"),
-				*CoerceParameter(BufferUV, MCT_Float2),
-				*CoerceParameter(PixelNormalWS(),MCT_Float3));
-		}
-		else
-		{
-			// Return the original geometry normal as a fall back.
-			// Shouldn't get here since DBuffer isn't enabled on mobile rendering paths?
-			return AddCodeChunk(MCT_Float4,
-				TEXT("float4(%s,1.0f)"),
-				*CoerceParameter(PixelNormalWS(),MCT_Float3));
-		}
-	}
 
 	if (SceneTextureId == PPI_DecalMask)
 	{
-		return Error(TEXT("Decal Mask bit was move out of GBuffer to the stencil buffer for performance optimisation and is therefor no longer available"));
+		return Error(TEXT("Decal Mask bit was moved from GBuffer to the Stencil Buffer for performance optimisation so therefore no longer available."));
 	}
 
 	UseSceneTextureId(SceneTextureId, true);
