GameFramework - In the case that you've got a pawn that is set to AdjustIfPossibleButDontSpawnIfColliding.  The pawns are always spawned but then Destroy() is called on them, previously we were still setting those pawns on the Controllers, however, giving a controller a PendingKill pawn can result in a lot of undesirable behavior as well as trigger a lot of checks/ensures (like when hooking delegates).  Additionally there is a codepath that was checking if the pawn was null, but not performing the newer IsValid() call which checks things like null, and is pending kill, that code path now does, and if so calls a new function AGameModeBase::FailedToRestartPlayer, which calls the existing function AController::FailedToSpawnPawn.  Previously AController::FailedToSpawnPawn wasn't actually being called in a lot of circumstances because the pawn wasn't null - just pending kill.  Additionally - we no longer call SetPawn if the spawned pawn IsValid fails, this avoids a lot of game code being triggered for pawns that are already consigned to oblivion.

Lyra - Updating Lyra's game mode to handle FailedToRestartPlayer which then re-attempts to spawn the controller's pawn next frame if it failed this frame.

diff --git a/Engine/Source/Runtime/Engine/Classes/GameFramework/GameModeBase.h b/Engine/Source/Runtime/Engine/Classes/GameFramework/GameModeBase.h
--- a/Engine/Source/Runtime/Engine/Classes/GameFramework/GameModeBase.h
+++ b/Engine/Source/Runtime/Engine/Classes/GameFramework/GameModeBase.h
@@ -539,6 +539,12 @@ protected:
 	/** Handles second half of RestartPlayer */
 	virtual void FinishRestartPlayer(AController* NewPlayer, const FRotator& StartRotation);
 
+	/**
+	 * Called in the event that we fail to spawn a controller's pawn, which maybe it didn't have one or maybe it tried
+	 * to spawn and was destroyed due to collision.
+	 */
+	virtual void FailedToRestartPlayer(AController* NewPlayer);
+
 	/**
 	 * Notifies all clients to travel to the specified URL.
 	 *
diff --git a/Engine/Source/Runtime/Engine/Private/GameModeBase.cpp b/Engine/Source/Runtime/Engine/Private/GameModeBase.cpp
--- a/Engine/Source/Runtime/Engine/Private/GameModeBase.cpp
+++ b/Engine/Source/Runtime/Engine/Private/GameModeBase.cpp
@@ -1246,12 +1246,16 @@ void AGameModeBase::RestartPlayerAtPlayerStart(AController* NewPlayer, AActor* S
 	else if (GetDefaultPawnClassForController(NewPlayer) != nullptr)
 	{
 		// Try to create a pawn to use of the default class for this player
-		NewPlayer->SetPawn(SpawnDefaultPawnFor(NewPlayer, StartSpot));
+		APawn* NewPawn = SpawnDefaultPawnFor(NewPlayer, StartSpot);
+		if (IsValid(NewPawn))
+		{
+			NewPlayer->SetPawn(NewPawn);
+		}
 	}
-
-	if (NewPlayer->GetPawn() == nullptr)
+	
+	if (!IsValid(NewPlayer->GetPawn()))
 	{
-		NewPlayer->FailedToSpawnPawn();
+		FailedToRestartPlayer(NewPlayer);
 	}
 	else
 	{
@@ -1287,12 +1291,16 @@ void AGameModeBase::RestartPlayerAtTransform(AController* NewPlayer, const FTran
 	else if (GetDefaultPawnClassForController(NewPlayer) != nullptr)
 	{
 		// Try to create a pawn to use of the default class for this player
-		NewPlayer->SetPawn(SpawnDefaultPawnAtTransform(NewPlayer, SpawnTransform));
+		APawn* NewPawn = SpawnDefaultPawnAtTransform(NewPlayer, SpawnTransform);
+		if (IsValid(NewPawn))
+		{
+			NewPlayer->SetPawn(NewPawn);
+		}
 	}
 
-	if (NewPlayer->GetPawn() == nullptr)
+	if (IsValid(NewPlayer->GetPawn()))
 	{
-		NewPlayer->FailedToSpawnPawn();
+		FailedToRestartPlayer(NewPlayer);
 	}
 	else
 	{
@@ -1300,14 +1308,19 @@ void AGameModeBase::RestartPlayerAtTransform(AController* NewPlayer, const FTran
 	}
 }
 
+void AGameModeBase::FailedToRestartPlayer(AController* NewPlayer)
+{
+	NewPlayer->FailedToSpawnPawn();
+}
+
 void AGameModeBase::FinishRestartPlayer(AController* NewPlayer, const FRotator& StartRotation)
 {
 	NewPlayer->Possess(NewPlayer->GetPawn());
 
 	// If the Pawn is destroyed as part of possession we have to abort
-	if (NewPlayer->GetPawn() == nullptr)
+	if (!IsValid(NewPlayer->GetPawn()))
 	{
-		NewPlayer->FailedToSpawnPawn();
+		FailedToRestartPlayer(NewPlayer);
 	}
 	else
 	{
