Fix for crash in duplicate/retargeted blendspace from missing samples (when using graph sample points).

diff --git a/Engine/Source/Editor/UnrealEd/Private/EditorAnimUtils.cpp b/Engine/Source/Editor/UnrealEd/Private/EditorAnimUtils.cpp
--- a/Engine/Source/Editor/UnrealEd/Private/EditorAnimUtils.cpp
+++ b/Engine/Source/Editor/UnrealEd/Private/EditorAnimUtils.cpp
@@ -515,6 +515,14 @@ namespace EditorAnimUtils
 				FFindAnimAssetRefs AnimRefFinderBlueprint(Node, AnimationAssets);
 			}
 		}
+		
+		// removes blendspaces that are embedded in the graph so that we don't duplicate them to make new assets
+		AnimationAssets.RemoveAll([](UAnimationAsset*& AnimationAsset)
+		{
+			const bool bIsBlendspace = IsValid(Cast<UBlendSpace>(AnimationAsset));
+			const bool bIsEmbedded = !AnimationAsset->GetSkeleton();
+			return bIsBlendspace && bIsEmbedded;
+		});
 	}
 
 	void ReplaceReferredAnimationsInBlueprint(UAnimBlueprint* AnimBlueprint, const TMap<UAnimationAsset*, UAnimationAsset*>& AnimAssetReplacementMap)
diff --git a/Engine/Source/Runtime/Engine/Private/Animation/BlendSpace.cpp b/Engine/Source/Runtime/Engine/Private/Animation/BlendSpace.cpp
--- a/Engine/Source/Runtime/Engine/Private/Animation/BlendSpace.cpp
+++ b/Engine/Source/Runtime/Engine/Private/Animation/BlendSpace.cpp
@@ -640,27 +640,25 @@ void UBlendSpace::ReplaceReferredAnimations(const TMap<UAnimationAsset*, UAnimat
 	Super::ReplaceReferredAnimations(ReplacementMap);
 
 	TArray<FBlendSample> NewSamples;
-	for (auto Iter = SampleData.CreateIterator(); Iter; ++Iter)
+	for (FBlendSample& Sample : SampleData)
 	{
-		FBlendSample& Sample = (*Iter);
-		UAnimSequence* Anim = Sample.Animation;
-
-		if (Anim)
+		// replace the referenced animation sequence (if there was one)
+		if (Sample.Animation)
 		{
-			UAnimSequence* const* ReplacementAsset = (UAnimSequence* const*)ReplacementMap.Find(Anim);
-			if (ReplacementAsset)
+			if (UAnimSequence* const* ReplacementAsset = (UAnimSequence* const*)ReplacementMap.Find(Sample.Animation))
 			{
 				Sample.Animation = *ReplacementAsset;
 				Sample.Animation->ReplaceReferredAnimations(ReplacementMap);
-				NewSamples.Add(Sample);
 			}
 		}
+
+		NewSamples.Add(Sample);
 	}
 
+	// replace preview base pose sequence asset (if there was one)
 	if (PreviewBasePose)
 	{
-		UAnimSequence* const* ReplacementAsset = (UAnimSequence* const*)ReplacementMap.Find(PreviewBasePose);
-		if (ReplacementAsset)
+		if (UAnimSequence* const* ReplacementAsset = (UAnimSequence* const*)ReplacementMap.Find(PreviewBasePose))
 		{
 			PreviewBasePose = *ReplacementAsset;
 			PreviewBasePose->ReplaceReferredAnimations(ReplacementMap);
