fix Engine dependency on ImageCore module

diff --git a/Engine/Source/Runtime/Engine/Engine.Build.cs b/Engine/Source/Runtime/Engine/Engine.Build.cs
--- a/Engine/Source/Runtime/Engine/Engine.Build.cs
+++ b/Engine/Source/Runtime/Engine/Engine.Build.cs
@@ -18,8 +18,7 @@ public class Engine : ModuleRules
 			"PacketHandler", 
 			"AudioMixer", 
 			"AudioMixerCore",
-			"AnimationCore",
-			"ImageCore"
+			"AnimationCore"
 			}
 		);
 
@@ -79,6 +78,7 @@ public class Engine : ModuleRules
 				"CoreUObject",
 				"NetCore",
 				"ApplicationCore",
+				"ImageCore",
 				"Json",
 				"JsonUtilities",
 				"SlateCore",
@@ -259,7 +259,6 @@ public class Engine : ModuleRules
 
 			PrivateDependencyModuleNames.AddRange(
 				new string[] {
-					"ImageCore",
 					"RawMesh"
 				}
 			);
