-Force bPrecomputeVisibility to false & bForceNoPrecomputedLighting to true on all WP maps (until 5.1 where it'll be supported)
-Remove the previous warning about WorldPartition maps not supporting static lighting, we do it with this flag now.
-Modify MapCheck to stop reporting warnings about static lighting not being up to date when it's disabled

diff --git a/Engine/Source/Editor/UnrealEd/Private/StaticLightingSystem/StaticLightingSystem.cpp b/Engine/Source/Editor/UnrealEd/Private/StaticLightingSystem/StaticLightingSystem.cpp
--- a/Engine/Source/Editor/UnrealEd/Private/StaticLightingSystem/StaticLightingSystem.cpp
+++ b/Engine/Source/Editor/UnrealEd/Private/StaticLightingSystem/StaticLightingSystem.cpp
@@ -291,14 +291,6 @@ void FStaticLightingManager::CreateStaticLightingSystem(const FLightingBuildOpti
 		bBuildReflectionCapturesOnFinish = !Options.bOnlyBuildVisibility;
 
 		UWorld* World = GWorld;
-
-		if (World->IsPartitionedWorld())
-		{
-			// No StaticLighting or Precomputed Visibility on WorldPartition Levels so we can early out here 
-			FStaticLightingManager::Get()->FailLightingBuild(
-				LOCTEXT("LightBuildWPFail", "Static Lighting on World Partition maps is not supported."));
-			return;
-		}
 		
 		for (ULevel* Level : World->GetLevels())
 		{
diff --git a/Engine/Source/Runtime/Engine/Classes/GameFramework/WorldSettings.h b/Engine/Source/Runtime/Engine/Classes/GameFramework/WorldSettings.h
--- a/Engine/Source/Runtime/Engine/Classes/GameFramework/WorldSettings.h
+++ b/Engine/Source/Runtime/Engine/Classes/GameFramework/WorldSettings.h
@@ -792,6 +792,7 @@ public:
 	UWorldPartition* GetWorldPartition() const;
 	bool IsPartitionedWorld() const { return GetWorldPartition() != nullptr; }
 	void SetWorldPartition(UWorldPartition* InWorldPartition);
+	void ApplyWorldPartitionForcedSettings();
 
 	/**
 	 * Returns the Z component of the current world gravity and initializes it to the default
diff --git a/Engine/Source/Runtime/Engine/Private/WorldSettings.cpp b/Engine/Source/Runtime/Engine/Private/WorldSettings.cpp
--- a/Engine/Source/Runtime/Engine/Private/WorldSettings.cpp
+++ b/Engine/Source/Runtime/Engine/Private/WorldSettings.cpp
@@ -227,7 +227,14 @@ void AWorldSettings::SetWorldPartition(UWorldPartition* InWorldPartition)
 	check(!IsValid(WorldPartition));
 	check(InWorldPartition);
 	WorldPartition = InWorldPartition;
+	ApplyWorldPartitionForcedSettings();
+}
+
+void AWorldSettings::ApplyWorldPartitionForcedSettings()
+{
 	bEnableWorldComposition = false;
+	bForceNoPrecomputedLighting = true;
+	bPrecomputeVisibility = false;
 }
 
 float AWorldSettings::GetGravityZ() const
@@ -543,6 +550,11 @@ void AWorldSettings::PostLoad()
 		Entry.MergeSetting.PostLoadDeprecated();
 	}
 
+	if (WorldPartition)
+	{
+		// Force to re-apply WorldPartition restrictions on WorldSettings (in case they changed)
+		ApplyWorldPartitionForcedSettings();
+	}
 #endif// WITH_EDITOR
 
 PRAGMA_DISABLE_DEPRECATION_WARNINGS
@@ -599,24 +611,26 @@ void AWorldSettings::CheckForErrors()
 			->AddToken(FMapErrorToken::Create(FMapErrors::DuplicateLevelInfo));
 	}
 
-	int32 NumLightingScenariosEnabled = 0;
-
-	for (int32 LevelIndex = 0; LevelIndex < World->GetNumLevels(); LevelIndex++)
+	if (!World->GetWorldSettings()->bForceNoPrecomputedLighting)
 	{
-		ULevel* Level = World->GetLevels()[LevelIndex];
+		int32 NumLightingScenariosEnabled = 0;
 
-		if (Level->bIsLightingScenario && Level->bIsVisible)
+		for (int32 LevelIndex = 0; LevelIndex < World->GetNumLevels(); LevelIndex++)
 		{
-			NumLightingScenariosEnabled++;
-		}
-	}
+			ULevel* Level = World->GetLevels()[LevelIndex];
 
-	if( World->NumLightingUnbuiltObjects > 0 && NumLightingScenariosEnabled <= 1 )
-	{
-		FMessageLog("MapCheck").Error()
-			->AddToken(FUObjectToken::Create(this))
-			->AddToken(FTextToken::Create(LOCTEXT( "MapCheck_Message_RebuildLighting", "Maps need lighting rebuilt" ) ))
-			->AddToken(FMapErrorToken::Create(FMapErrors::RebuildLighting));
+			if (Level->bIsLightingScenario && Level->bIsVisible)
+			{
+				NumLightingScenariosEnabled++;
+			}
+		}
+		if( World->NumLightingUnbuiltObjects > 0 && NumLightingScenariosEnabled <= 1 )
+		{
+			FMessageLog("MapCheck").Error()
+				->AddToken(FUObjectToken::Create(this))
+				->AddToken(FTextToken::Create(LOCTEXT( "MapCheck_Message_RebuildLighting", "Maps need lighting rebuilt" ) ))
+				->AddToken(FMapErrorToken::Create(FMapErrors::RebuildLighting));
+		}
 	}
 }
 
@@ -657,7 +671,9 @@ bool AWorldSettings::CanEditChange(const FProperty* InProperty) const
 				return LightmassSettings.EnvironmentIntensity > 0;
 			}
 		}
-		else if (PropertyName == GET_MEMBER_NAME_STRING_CHECKED(AWorldSettings, bEnableWorldComposition))
+		else if (PropertyName == GET_MEMBER_NAME_STRING_CHECKED(AWorldSettings, bEnableWorldComposition ) ||
+				 PropertyName == GET_MEMBER_NAME_STRING_CHECKED(AWorldSettings, bForceNoPrecomputedLighting ) ||
+				 PropertyName == GET_MEMBER_NAME_STRING_CHECKED(AWorldSettings, bPrecomputeVisibility))
 		{
 			return !IsPartitionedWorld();
 		}
