fix old PNG textures not loading correctly when referenced from an object PostLoad
migrating deprecated properties needs to be at the end of Serialize, not in PostLoad
fixes PIE in ShooterGame with -ddc=cold
caused by 19409594

diff --git a/Engine/Source/Runtime/Engine/Private/Texture.cpp b/Engine/Source/Runtime/Engine/Private/Texture.cpp
--- a/Engine/Source/Runtime/Engine/Private/Texture.cpp
+++ b/Engine/Source/Runtime/Engine/Private/Texture.cpp
@@ -581,42 +581,43 @@ void UTexture::Serialize(FArchive& Ar)
 
 	/** Legacy serialization. */
 #if WITH_EDITORONLY_DATA
-	if (!StripFlags.IsEditorDataStripped())
-	{
-#if WITH_EDITOR
-		FWriteScopeLock BulkDataExclusiveScope(Source.BulkDataLock.Get());
-#endif
 
-		if (Ar.IsLoading() && Ar.CustomVer(FUE5MainStreamObjectVersion::GUID) < FUE5MainStreamObjectVersion::TextureDoScaleMipsForAlphaCoverage)
-		{
-			// bDoScaleMipsForAlphaCoverage was not transmitted in old versions
-			//	and AlphaCoverageThresholds was being incorrectly set to (0,0,0,1)
-			check( bDoScaleMipsForAlphaCoverage == false );
+	if (Ar.IsLoading() && Ar.CustomVer(FUE5MainStreamObjectVersion::GUID) < FUE5MainStreamObjectVersion::TextureDoScaleMipsForAlphaCoverage)
+	{
+		// bDoScaleMipsForAlphaCoverage was not transmitted in old versions
+		//	and AlphaCoverageThresholds was being incorrectly set to (0,0,0,1)
+		check( bDoScaleMipsForAlphaCoverage == false );
 	
-			if ( AlphaCoverageThresholds != FVector4(0,0,0,0) && AlphaCoverageThresholds != FVector4(0,0,0,1) )
-			{
-				// AlphaCoverageThresholds is a non-default value, assume that means they wanted it on
-				bDoScaleMipsForAlphaCoverage = true;
-			}
-			else if ( AlphaCoverageThresholds == FVector4(0,0,0,1) )
-			{
-				// if value is (0,0,0,1)
-				//	that was previously incorrectly being set by default and enabling alpha coverage processing
-				// we don't want that, but to optionally preserve old behavior you can set a config option :
+		if ( AlphaCoverageThresholds != FVector4(0,0,0,0) && AlphaCoverageThresholds != FVector4(0,0,0,1) )
+		{
+			// AlphaCoverageThresholds is a non-default value, assume that means they wanted it on
+			bDoScaleMipsForAlphaCoverage = true;
+		}
+		else if ( AlphaCoverageThresholds == FVector4(0,0,0,1) )
+		{
+			// if value is (0,0,0,1)
+			//	that was previously incorrectly being set by default and enabling alpha coverage processing
+			// we don't want that, but to optionally preserve old behavior you can set a config option :
 
-				static struct ReadConfigOnce
+			static struct ReadConfigOnce
+			{
+				bool bBool;
+				ReadConfigOnce()
 				{
-					bool bBool;
-					ReadConfigOnce()
-					{
-						bBool = false;
-						GConfig->GetBool(TEXT("Texture"), TEXT("EnableLegacyAlphaCoverageThresholdScaling"), bBool, GEditorIni);
-					}
-				} ConfigValue;
+					bBool = false;
+					GConfig->GetBool(TEXT("Texture"), TEXT("EnableLegacyAlphaCoverageThresholdScaling"), bBool, GEditorIni);
+				}
+			} ConfigValue;
 
-				bDoScaleMipsForAlphaCoverage = ConfigValue.bBool;
-			}
+			bDoScaleMipsForAlphaCoverage = ConfigValue.bBool;
 		}
+	}
+
+	if (!StripFlags.IsEditorDataStripped())
+	{
+#if WITH_EDITOR
+		FWriteScopeLock BulkDataExclusiveScope(Source.BulkDataLock.Get());
+#endif
 
 		if (Ar.IsLoading() && Ar.CustomVer(FUE5MainStreamObjectVersion::GUID) < FUE5MainStreamObjectVersion::VirtualizedBulkDataHaveUniqueGuids)
 		{
@@ -659,7 +660,45 @@ void UTexture::Serialize(FArchive& Ar)
 			UE_LOG(LogTexture, Display, TEXT("%s is marked for virtual streaming but virtual texture streaming is not available."), *GetPathName());
 		}
 	}
+	
+	if (Ar.IsLoading())
+	{
+		if ( Source.bPNGCompressed_DEPRECATED )
+		{
+			// loaded with deprecated "bPNGCompressed"
+			// change to CompressionFormat PNG
+			check( Source.CompressionFormat == TSCF_None || Source.CompressionFormat == TSCF_PNG );
+			Source.CompressionFormat = TSCF_PNG;
+			Source.bPNGCompressed_DEPRECATED = false;
+		}
+
+		if ( Source.GetFormat() == TSF_RGBA8_DEPRECATED
+			|| Source.GetFormat() == TSF_RGBE8_DEPRECATED )
+		{
+			// ensure that later code doesn't ever see the _DEPRECATED formats
 
+			// needs RB swap
+			// force BulkData to become resident
+			// do the swap on the bits
+			// change format to swapped version
+			// these formats are incredibly rare and old
+			// just warn and change the enum but don't swap the bits
+			// will appear RB swapped until reimported
+			UE_LOG(LogTexture, Warning, TEXT("TextureSource is a deprecated RB swapped format, needs reimport!: %s"), *GetPathName());
+		
+			for(int i=0;i<Source.LayerFormat.Num();i++)
+			{
+				if ( Source.LayerFormat[i] == TSF_RGBA8_DEPRECATED )
+				{
+					Source.LayerFormat[i] = TSF_BGRA8;
+				}
+				else if ( Source.LayerFormat[i] == TSF_RGBE8_DEPRECATED )
+				{
+					Source.LayerFormat[i] = TSF_BGRE8;
+				}
+			}
+		}
+	}
 #endif // #if WITH_EDITORONLY_DATA
 }
 
@@ -695,42 +734,6 @@ void UTexture::PostLoad()
 		AssetImportData->SourceData = MoveTemp(Info);
 	}
 
-	if ( Source.bPNGCompressed_DEPRECATED )
-	{
-		// loaded with deprecated "bPNGCompressed"
-		// change to CompressionFormat PNG
-		check( Source.CompressionFormat == TSCF_None || Source.CompressionFormat == TSCF_PNG );
-		Source.CompressionFormat = TSCF_PNG;
-		Source.bPNGCompressed_DEPRECATED = false;
-	}
-
-	if ( Source.GetFormat() == TSF_RGBA8_DEPRECATED
-		|| Source.GetFormat() == TSF_RGBE8_DEPRECATED )
-	{
-		// ensure that later code doesn't ever see the _DEPRECATED formats
-
-		// needs RB swap
-		// force BulkData to become resident
-		// do the swap on the bits
-		// change format to swapped version
-		// these formats are incredibly rare and old
-		// just warn and change the enum but don't swap the bits
-		// will appear RB swapped until reimported
-		UE_LOG(LogTexture, Warning, TEXT("TextureSource is a deprecated RB swapped format, needs reimport!: %s"), *GetPathName());
-		
-		for(int i=0;i<Source.LayerFormat.Num();i++)
-		{
-			if ( Source.LayerFormat[i] == TSF_RGBA8_DEPRECATED )
-			{
-				Source.LayerFormat[i] = TSF_BGRA8;
-			}
-			else if ( Source.LayerFormat[i] == TSF_RGBE8_DEPRECATED )
-			{
-				Source.LayerFormat[i] = TSF_BGRE8;
-			}
-		}
-	}
-
 #endif
 
 	if( !IsTemplate() )
