Removing ULevel reference from replay actor rollback state
Adding -skipreplayrollback to not generate startup actor rollback state, assumes no scrubbing

diff --git a/Engine/Source/Runtime/Engine/Classes/Engine/DemoNetDriver.h b/Engine/Source/Runtime/Engine/Classes/Engine/DemoNetDriver.h
--- a/Engine/Source/Runtime/Engine/Classes/Engine/DemoNetDriver.h
+++ b/Engine/Source/Runtime/Engine/Classes/Engine/DemoNetDriver.h
@@ -68,8 +68,7 @@ struct FRollbackNetStartupActorInfo
 	FVector		Location;
 	FRotator	Rotation;
 	FVector		Scale3D;
-	UPROPERTY()
-	TObjectPtr<ULevel>		Level = nullptr;
+	FName		LevelName;
 
 	TSharedPtr<FRepState> RepState;
 	TMap<FString, TSharedPtr<FRepState>> SubObjRepState;
@@ -720,5 +719,8 @@ private:
 
 	FReplayHelper ReplayHelper;
 
+	/** Enabled via -skipreplayrollback and causes rollback data to not be generated, with the assumption that there will be no scrubbing. */
+	bool bSkipStartupActorRollback = false;
+
 	friend class UDemoNetConnection;
 };
\ No newline at end of file
diff --git a/Engine/Source/Runtime/Engine/Private/DemoNetDriver.cpp b/Engine/Source/Runtime/Engine/Private/DemoNetDriver.cpp
--- a/Engine/Source/Runtime/Engine/Private/DemoNetDriver.cpp
+++ b/Engine/Source/Runtime/Engine/Private/DemoNetDriver.cpp
@@ -696,6 +696,7 @@ void UDemoNetDriver::InitDefaults()
 	MaxArchiveReadPos = 0;
 	bNeverApplyNetworkEmulationSettings = true;
 	bSkipServerReplicateActors = true;
+	bSkipStartupActorRollback = false;
 
 	if (!HasAnyFlags(RF_ClassDefaultObject))
 	{
@@ -826,6 +827,11 @@ bool UDemoNetDriver::InitBase(bool bInitAsClient, FNetworkNotify* InNotify, cons
 		PlaybackPacketIndex				= 0;
 		CheckpointSaveMaxMSPerFrame		= -1.0f;
 
+		if (FParse::Param(FCommandLine::Get(), TEXT("skipreplayrollback")))
+		{
+			bSkipStartupActorRollback = true;
+		}
+
 		RecordBuildConsiderAndPrioritizeTimeSlice = CVarDemoMaximumRepPrioritizeTime.GetValueOnAnyThread();
 		RecordDestructionInfoReplicationTimeSlice = DemoNetDriverRecordingPrivate::CVarDemoMaximumRecDestructionInfoTime.GetValueOnAnyThread();
 
@@ -1056,6 +1062,7 @@ bool UDemoNetDriver::InitConnectInternal(FString& Error)
 	else
 	{
 		ReplayHelper.ResetLevelStatuses();
+		ReplayHelper.ResetLevelMap();
 	}
 
 	return true;
@@ -1100,10 +1107,12 @@ void UDemoNetDriver::NotifyStreamingLevelUnload( ULevel* InLevel )
 {
 	if (InLevel && !InLevel->bClientOnlyVisible && HasLevelStreamingFixes() && IsPlaying())
 	{
+		const FName FilterLevelName = InLevel->GetOutermost()->GetFName();
+
 		// We can't just iterate over the levels actors, because the ones in the queue will already have been destroyed.
 		for (TMap<FString, FRollbackNetStartupActorInfo>::TIterator RollbackIt = RollbackNetStartupActors.CreateIterator(); RollbackIt; ++RollbackIt)
 		{
-			if (RollbackIt.Value().Level == InLevel)
+			if (RollbackIt.Value().LevelName == FilterLevelName)
 			{
 				RollbackIt.RemoveCurrent();
 			}
@@ -1115,15 +1124,27 @@ void UDemoNetDriver::NotifyStreamingLevelUnload( ULevel* InLevel )
 
 void UDemoNetDriver::OnPostLoadMapWithWorld(UWorld* InWorld)
 {
-	if (InWorld != nullptr && InWorld == World && HasLevelStreamingFixes())
+	if (InWorld != nullptr && InWorld == World)
 	{
+		if (HasLevelStreamingFixes())
+		{
+			if (IsPlaying())
+			{
+				ReplayHelper.ResetLevelStatuses();
+			}
+			else
+			{
+				ReplayHelper.ClearLevelStreamingState();
+			}
+		}
+
 		if (IsPlaying())
 		{
-			ReplayHelper.ResetLevelStatuses();
+			ReplayHelper.ResetLevelMap();
 		}
 		else
 		{
-			ReplayHelper.ClearLevelStreamingState();
+			ReplayHelper.ClearLevelMap();
 		}
 	}
 }
@@ -3183,8 +3204,12 @@ bool UDemoNetDriver::SetExternalDataForObject(UObject* OwningObject, const uint8
 
 void UDemoNetDriver::RespawnNecessaryNetStartupActors(TArray<AActor*>& SpawnedActors, ULevel* Level /* = nullptr */)
 {
+	QUICK_SCOPE_CYCLE_COUNTER(STAT_RespawnNecessaryNetStartupActors);
+
 	TGuardValue<bool> RestoringStartupActors(bIsRestoringStartupActors, true);
 
+	const FName FilterLevelName = Level ? Level->GetOutermost()->GetFName() : NAME_None;
+
 	for (auto RollbackIt = RollbackNetStartupActors.CreateIterator(); RollbackIt; ++RollbackIt)
 	{
 		if (ReplayHelper.PlaybackDeletedNetStartupActors.Contains(RollbackIt.Key()))
@@ -3196,19 +3221,14 @@ void UDemoNetDriver::RespawnNecessaryNetStartupActors(TArray<AActor*>& SpawnedAc
 		FRollbackNetStartupActorInfo& RollbackActor = RollbackIt.Value();
 
 		// filter to a specific level
-		if ((Level != nullptr) && (RollbackActor.Level != Level))
+		if ((Level != nullptr) && (RollbackActor.LevelName != FilterLevelName))
 		{
 			continue;
 		}
 
 		if (HasLevelStreamingFixes())
 		{
-			if (!ensureMsgf(RollbackActor.Level, TEXT("RespawnNecessaryNetStartupActors: Rollback actor level is nullptr: %s"), *RollbackActor.Name.ToString()))
-			{
-				continue;
-			}
-
-			const FString LevelPackageName = ReplayHelper.GetLevelPackageName(*RollbackActor.Level);
+			const FString LevelPackageName = UWorld::RemovePIEPrefix(RollbackActor.LevelName.ToString());
 
 			// skip rollback actors in streamed out levels (pending gc)
 			if (!ReplayHelper.LevelStatusesByName.Contains(LevelPackageName))
@@ -3223,7 +3243,14 @@ void UDemoNetDriver::RespawnNecessaryNetStartupActors(TArray<AActor*>& SpawnedAc
 			}
 		}
 
-		AActor* ExistingActor = FindObjectFast<AActor>(RollbackActor.Level, RollbackActor.Name);
+		ULevel* RollbackActorLevel = ReplayHelper.WeakLevelsByName.FindRef(RollbackActor.LevelName).Get();
+
+		if (!ensureMsgf(RollbackActorLevel, TEXT("RespawnNecessaryNetStartupActors: Rollback actor level is nullptr: %s"), *RollbackActor.Name.ToString()))
+		{
+			continue;
+		}
+
+		AActor* ExistingActor = FindObjectFast<AActor>(RollbackActorLevel, RollbackActor.Name);
 		if (ExistingActor)
 		{
 			ensureMsgf((!IsValidChecked(ExistingActor) || ExistingActor->IsUnreachable()), TEXT("RespawnNecessaryNetStartupActors: Renaming rollback actor that wasn't destroyed: %s"), *GetFullNameSafe(ExistingActor));
@@ -3236,7 +3263,7 @@ void UDemoNetDriver::RespawnNecessaryNetStartupActors(TArray<AActor*>& SpawnedAc
 		SpawnInfo.SpawnCollisionHandlingOverride	= ESpawnActorCollisionHandlingMethod::AlwaysSpawn;
 		SpawnInfo.bNoFail							= true;
 		SpawnInfo.Name								= RollbackActor.Name;
-		SpawnInfo.OverrideLevel						= RollbackActor.Level;
+		SpawnInfo.OverrideLevel						= RollbackActorLevel;
 		SpawnInfo.bDeferConstruction				= true;
 
 		const FTransform SpawnTransform = FTransform(RollbackActor.Rotation, RollbackActor.Location, RollbackActor.Scale3D);
@@ -3823,6 +3850,8 @@ bool UDemoNetDriver::LoadCheckpoint(const FGotoResult& GotoResult)
 		ReplayHelper.ResetLevelStatuses();
 	}
 
+	ReplayHelper.ResetLevelMap();
+
 	LastProcessedPacketTime = 0.f;
 	ReplayHelper.LatestReadFrameTime = 0.f;
 
@@ -5101,6 +5130,11 @@ void UDemoNetDriver::QueueNetStartupActorForRollbackViaDeletion(AActor* Actor)
 		return;		// We should only be doing this at runtime while playing a replay
 	}
 
+	if (bSkipStartupActorRollback)
+	{
+		return;
+	}
+
 	check(Actor != nullptr);
 
 	if (!Actor->IsNetStartupActor())
@@ -5127,7 +5161,11 @@ void UDemoNetDriver::QueueNetStartupActorForRollbackViaDeletion(AActor* Actor)
 	RollbackActor.Location	= Actor->GetActorLocation();
 	RollbackActor.Rotation	= Actor->GetActorRotation();
 	RollbackActor.Scale3D	= Actor->GetActorScale3D();
-	RollbackActor.Level		= Actor->GetLevel();
+
+	if (ULevel* ActorLevel = Actor->GetLevel())
+	{
+		RollbackActor.LevelName = ActorLevel->GetOutermost()->GetFName();
+	}
 
 	if (GDemoSaveRollbackActorState != 0)
 	{
diff --git a/Engine/Source/Runtime/Engine/Private/ReplayHelper.cpp b/Engine/Source/Runtime/Engine/Private/ReplayHelper.cpp
--- a/Engine/Source/Runtime/Engine/Private/ReplayHelper.cpp
+++ b/Engine/Source/Runtime/Engine/Private/ReplayHelper.cpp
@@ -1208,6 +1208,31 @@ void FReplayHelper::ResetLevelStatuses()
 	}
 }
 
+void FReplayHelper::ClearLevelMap()
+{
+	WeakLevelsByName.Reset();
+}
+
+void FReplayHelper::ResetLevelMap()
+{
+	ClearLevelMap();
+
+	check(World.Get());
+
+	WeakLevelsByName.Add(World->PersistentLevel->GetOutermost()->GetFName(), World->PersistentLevel);
+
+	for (ULevelStreaming* LevelStreaming : World->GetStreamingLevels())
+	{
+		if (LevelStreaming && LevelStreaming->IsLevelVisible())
+		{
+			if (ULevel* Level = LevelStreaming->GetLoadedLevel())
+			{
+				WeakLevelsByName.Add(Level->GetOutermost()->GetFName(), Level);
+			}
+		}
+	}
+}
+
 void FReplayHelper::WriteDemoFrame(UNetConnection* Connection, FArchive& Ar, TArray<FQueuedDemoPacket>& QueuedPackets, float FrameTime, EWriteDemoFrameFlags Flags)
 {
 	DECLARE_SCOPE_CYCLE_COUNTER(TEXT("Replay write frame time"), STAT_ReplayWriteDemoFrame, STATGROUP_Net);
@@ -1729,6 +1754,8 @@ void FReplayHelper::Serialize(FArchive& Ar)
 				Ar << ActorString;
 			}
 		);
+
+		GRANULAR_NETWORK_MEMORY_TRACKING_TRACK("WeakLevelsByName", WeakLevelsByName.CountBytes(Ar));
 	}
 }
 
@@ -2119,9 +2146,9 @@ void FReplayHelper::OnLevelAddedToWorld(ULevel* InLevel, UWorld* InWorld)
 {
 	LLM_SCOPE(ELLMTag::Replays);
 
-	if (InLevel && !InLevel->bClientOnlyVisible && (World == InWorld) && HasLevelStreamingFixes() && InWorld->IsPlayingReplay())
+	if (InLevel && !InLevel->bClientOnlyVisible && (World == InWorld) && InWorld->IsPlayingReplay())
 	{
-		if (!NewStreamingLevelsThisFrame.Contains(InLevel) && !LevelsPendingFastForward.Contains(InLevel))
+		if (HasLevelStreamingFixes() && !NewStreamingLevelsThisFrame.Contains(InLevel) && !LevelsPendingFastForward.Contains(InLevel))
 		{
 			FLevelStatus& LevelStatus = FindOrAddLevelStatus(*InLevel);
 
@@ -2138,23 +2165,30 @@ void FReplayHelper::OnLevelAddedToWorld(ULevel* InLevel, UWorld* InWorld)
 				NewStreamingLevelsThisFrame.Add(InLevel);
 			}
 		}
+
+		WeakLevelsByName.Add(InLevel->GetOutermost()->GetFName(), InLevel);
 	}
 }
 
 void FReplayHelper::OnLevelRemovedFromWorld(ULevel* InLevel, UWorld* InWorld)
 {
-	if (InLevel && !InLevel->bClientOnlyVisible && (World == InWorld) && HasLevelStreamingFixes() && InWorld->IsPlayingReplay())
+	if (InLevel && !InLevel->bClientOnlyVisible && (World == InWorld) && InWorld->IsPlayingReplay())
 	{
-		const FString LevelPackageName = GetLevelPackageName(*InLevel);
-		if (LevelStatusesByName.Contains(LevelPackageName))
+		if (HasLevelStreamingFixes())
 		{
-			FLevelStatus& LevelStatus = GetLevelStatus(LevelPackageName);
-			LevelStatus.bIsReady = false;
+			const FString LevelPackageName = GetLevelPackageName(*InLevel);
+			if (LevelStatusesByName.Contains(LevelPackageName))
+			{
+				FLevelStatus& LevelStatus = GetLevelStatus(LevelPackageName);
+				LevelStatus.bIsReady = false;
 
-			// Make sure we don't try to fast-forward this level later.
-			LevelsPendingFastForward.Remove(InLevel);
-			NewStreamingLevelsThisFrame.Remove(InLevel);
+				// Make sure we don't try to fast-forward this level later.
+				LevelsPendingFastForward.Remove(InLevel);
+				NewStreamingLevelsThisFrame.Remove(InLevel);
+			}
 		}
+
+		WeakLevelsByName.Remove(InLevel->GetOutermost()->GetFName());
 	}
 
 	// always invalidate cache since it uses pointers
diff --git a/Engine/Source/Runtime/Engine/Public/ReplayHelper.h b/Engine/Source/Runtime/Engine/Public/ReplayHelper.h
--- a/Engine/Source/Runtime/Engine/Public/ReplayHelper.h
+++ b/Engine/Source/Runtime/Engine/Public/ReplayHelper.h
@@ -433,6 +433,13 @@ private:
 	// Maintain a quick lookup for loaded levels directly to LevelStatus
 	TMap<const ULevel*, int32> LevelStatusIndexByLevel;
 
+	// Map of ULevel GetFName to weak object pointer to the level
+	// Populated during playback, not using NetworkRemapPath because it is never serialized
+	TMap<FName, TWeakObjectPtr<ULevel>> WeakLevelsByName;
+
+	void ResetLevelMap();
+	void ClearLevelMap();
+
 	// List of seen level statuses indices (in AllLevelStatuses).
 	TArray<int32> SeenLevelStatuses;
 
