Add a missing task scope. Without the scope, UTexture::GetResource returns a nullptr when called in some ParallelFor tasks initiated by an async task dispatched from the game thread and eventually causes incorrect render times to be returned

diff --git a/Engine/Source/Runtime/Engine/Private/Streaming/StreamingManagerTexture.cpp b/Engine/Source/Runtime/Engine/Private/Streaming/StreamingManagerTexture.cpp
--- a/Engine/Source/Runtime/Engine/Private/Streaming/StreamingManagerTexture.cpp
+++ b/Engine/Source/Runtime/Engine/Private/Streaming/StreamingManagerTexture.cpp
@@ -1224,6 +1224,8 @@ void FRenderAssetStreamingManager::UpdateStreamingRenderAssets(int32 StageIndex,
 
 		ParallelFor(Packets.Num(), [this, &Packets, &bWaitForMipFading, &bAsync](int32 PacketIndex)
 		{
+			FOptionalTaskTagScope Scope(ETaskTag::EParallelGameThread);
+
 			for (int32 Index = Packets[PacketIndex].StartIndex; Index < Packets[PacketIndex].EndIndex; ++Index)
 			{
 				FStreamingRenderAsset& StreamingRenderAsset = Packets[PacketIndex].StreamingRenderAssets[Index];
