Add in Only show objects that are disregarded for GC to Obj List command.

diff --git a/Engine/Source/Runtime/Engine/Private/UnrealEngine.cpp b/Engine/Source/Runtime/Engine/Private/UnrealEngine.cpp
--- a/Engine/Source/Runtime/Engine/Private/UnrealEngine.cpp
+++ b/Engine/Source/Runtime/Engine/Private/UnrealEngine.cpp
@@ -8581,6 +8581,7 @@ bool UEngine::HandleObjCommand( const TCHAR* Cmd, FOutputDevice& Ar )
 			const bool bShouldIncludeDefaultObjects = FParse::Param(Cmd, TEXT("INCLUDEDEFAULTS"));
 			const bool bOnlyListDefaultObjects = FParse::Param(Cmd, TEXT("DEFAULTSONLY"));
 			const bool bShowDetailedObjectInfo = FParse::Param(Cmd, TEXT("NODETAILEDINFO")) == false && bTrackDetailedObjectInfo;
+			const bool bOnlyShowDisregardGC = FParse::Param(Cmd, TEXT("GCDISREGARDONLY"));
 
 			for( FThreadSafeObjectIterator It; It; ++It )
 			{
@@ -8600,6 +8601,11 @@ bool UEngine::HandleObjCommand( const TCHAR* Cmd, FOutputDevice& Ar )
 					continue;
 				}
 
+				if ( bOnlyShowDisregardGC && !GUObjectArray.IsDisregardForGC(*It))
+				{
+					continue;
+				}
+
 				if ( bOnlyListGCObjects && GUObjectArray.IsDisregardForGC(*It) )
 				{
 					continue;
