Fix missing PIEInstanceID when RemapGuid is called from loading a new package.

diff --git a/Engine/Source/Runtime/Engine/Private/UnrealEngine.cpp b/Engine/Source/Runtime/Engine/Private/UnrealEngine.cpp
--- a/Engine/Source/Runtime/Engine/Private/UnrealEngine.cpp
+++ b/Engine/Source/Runtime/Engine/Private/UnrealEngine.cpp
@@ -1077,6 +1077,7 @@ namespace
 		UPackage* NewPackage = CreatePackage( *PIEPackageName);
 		if (NewPackage != nullptr && WorldContext.WorldType == EWorldType::PIE)
 		{
+			NewPackage->SetPIEInstanceID(WorldContext.PIEInstance);
 			NewPackage->SetPackageFlags(PKG_PlayInEditor);
 			LoadFlags |= LOAD_PackageForPIE;
 		}
