Converted FMemberReference::IsSameReference to a const method.

diff --git a/Engine/Source/Runtime/Engine/Classes/Engine/MemberReference.h b/Engine/Source/Runtime/Engine/Classes/Engine/MemberReference.h
--- a/Engine/Source/Runtime/Engine/Classes/Engine/MemberReference.h
+++ b/Engine/Source/Runtime/Engine/Classes/Engine/MemberReference.h
@@ -321,7 +321,7 @@ public:
 	}
 
 	/** Compares with another MemberReference to see if they are identical */
-	bool IsSameReference(FMemberReference& InReference)
+	bool IsSameReference(const FMemberReference& InReference) const
 	{
 		return 
 			bSelfContext == InReference.bSelfContext &&
