Added defensive check and early out for negative array index.
#rnx

diff --git a/Engine/Source/Runtime/Engine/Private/InstancedStaticMesh.cpp b/Engine/Source/Runtime/Engine/Private/InstancedStaticMesh.cpp
--- a/Engine/Source/Runtime/Engine/Private/InstancedStaticMesh.cpp
+++ b/Engine/Source/Runtime/Engine/Private/InstancedStaticMesh.cpp
@@ -1464,7 +1464,7 @@ void FInstancedStaticMeshSceneProxy::SetupProxy(UInstancedStaticMeshComponent* I
 		for (int32 InstanceIndex = 0; InstanceIndex < InComponent->GetInstanceCount(); ++InstanceIndex)
 		{
 			const int32 RenderInstanceIndex = InstanceReorderTable.IsValidIndex(InstanceIndex) ? InstanceReorderTable[InstanceIndex] : InstanceIndex;
-			if (!ensure(RenderInstanceIndex < NumRenderInstances))
+			if (!ensure(InstanceSceneData.IsValidIndex(RenderInstanceIndex)))
 			{
 				continue;
 			}
