Fix references to new material textures not being added after setting a new material on a mesh component

#tests verified no perf regression using ReplayRun a/b tests

diff --git a/Engine/Source/Runtime/Engine/Private/Components/MeshComponent.cpp b/Engine/Source/Runtime/Engine/Private/Components/MeshComponent.cpp
--- a/Engine/Source/Runtime/Engine/Private/Components/MeshComponent.cpp
+++ b/Engine/Source/Runtime/Engine/Private/Components/MeshComponent.cpp
@@ -65,7 +65,12 @@ void UMeshComponent::SetMaterial(int32 ElementIndex, UMaterialInterface* Materia
 
 			// Set the material and invalidate things
 			OverrideMaterials[ElementIndex] = Material;
-			MarkRenderStateDirty();			
+			MarkRenderStateDirty();
+			// If MarkRenderStateDirty didn't notify the streamer, do it now
+			if (!bIgnoreStreamingManagerUpdate && OwnerLevelHasRegisteredStaticComponentsInStreamingManager(GetOwner()))
+			{
+				IStreamingManager::Get().NotifyPrimitiveUpdated_Concurrent(this);
+			}
 			if (Material)
 			{
 				Material->AddToCluster(this, true);
diff --git a/Engine/Source/Runtime/Engine/Private/Components/PrimitiveComponent.cpp b/Engine/Source/Runtime/Engine/Private/Components/PrimitiveComponent.cpp
--- a/Engine/Source/Runtime/Engine/Private/Components/PrimitiveComponent.cpp
+++ b/Engine/Source/Runtime/Engine/Private/Components/PrimitiveComponent.cpp
@@ -558,20 +558,6 @@ void UPrimitiveComponent::GetUsedTextures(TArray<UTexture*>& OutTextures, EMater
 //////////////////////////////////////////////////////////////////////////
 // Render
 
-// Helper to access the level bStaticComponentsRegisteredInStreamingManager flag.
-FORCEINLINE_DEBUGGABLE bool OwnerLevelHasRegisteredStaticComponentsInStreamingManager(const AActor* Owner)
-{
-	if (Owner)
-	{
-		const ULevel* Level = Owner->GetLevel();
-		if (Level)
-		{
-			return Level->bStaticComponentsRegisteredInStreamingManager;
-		}
-	}
-	return false;
-}
-
 void UPrimitiveComponent::CreateRenderState_Concurrent(FRegisterComponentContext* Context)
 {
 	// Make sure cached cull distance is up-to-date if its zero and we have an LD cull distance
diff --git a/Engine/Source/Runtime/Engine/Private/Streaming/TextureStreamingHelpers.cpp b/Engine/Source/Runtime/Engine/Private/Streaming/TextureStreamingHelpers.cpp
--- a/Engine/Source/Runtime/Engine/Private/Streaming/TextureStreamingHelpers.cpp
+++ b/Engine/Source/Runtime/Engine/Private/Streaming/TextureStreamingHelpers.cpp
@@ -448,3 +448,16 @@ int64 GetAverageRequiredTexturePoolSize()
 	return 0;
 #endif
 }
+
+bool OwnerLevelHasRegisteredStaticComponentsInStreamingManager(const AActor* Owner)
+{
+	if (Owner)
+	{
+		const ULevel* Level = Owner->GetLevel();
+		if (Level)
+		{
+			return Level->bStaticComponentsRegisteredInStreamingManager;
+		}
+	}
+	return false;
+}
diff --git a/Engine/Source/Runtime/Engine/Private/Streaming/TextureStreamingHelpers.h b/Engine/Source/Runtime/Engine/Private/Streaming/TextureStreamingHelpers.h
--- a/Engine/Source/Runtime/Engine/Private/Streaming/TextureStreamingHelpers.h
+++ b/Engine/Source/Runtime/Engine/Private/Streaming/TextureStreamingHelpers.h
@@ -181,3 +181,6 @@ struct FRenderAssetStreamingStats
 	int64 ResidentMeshMem;		// Total memory in bytes of resident mesh LODs
 	int64 EvictedMeshMem;		// Total memory in bytes of evicted mesh LODs
 };
+
+// Helper to access the level bStaticComponentsRegisteredInStreamingManager flag.
+extern bool OwnerLevelHasRegisteredStaticComponentsInStreamingManager(const AActor* Owner);
