Fixing content worker "assume non-streaming" audio serialization hack to only impact loading and not serialization of new content.

diff --git a/Engine/Source/Runtime/Engine/Private/SoundWave.cpp b/Engine/Source/Runtime/Engine/Private/SoundWave.cpp
--- a/Engine/Source/Runtime/Engine/Private/SoundWave.cpp
+++ b/Engine/Source/Runtime/Engine/Private/SoundWave.cpp
@@ -829,11 +829,14 @@ void USoundWave::Serialize( FArchive& Ar )
 	if (bCooked)
 	{
 #if WITH_EDITOR
-		// Temporary workaround for allowing editors to load data that was saved for platforms that had streaming disabled. There is nothing int
-		// the serialized data that lets us know what is actually stored on disc, so we have to be explicitly told. Ideally, we'd just store something
-		// on disc to say how the serialized data is arranged, but doing so would cause a major patch delta.
-		static const bool bSoundWaveDataHasStreamingDisabled = FParse::Param(FCommandLine::Get(), TEXT("SoundWaveDataHasStreamingDisabled"));
-		bShouldStreamSound = bShouldStreamSound && !bSoundWaveDataHasStreamingDisabled;
+		if (Ar.IsLoading())
+		{
+			// Temporary workaround for allowing editors to load data that was saved for platforms that had streaming disabled. There is nothing int
+			// the serialized data that lets us know what is actually stored on disc, so we have to be explicitly told. Ideally, we'd just store something
+			// on disc to say how the serialized data is arranged, but doing so would cause a major patch delta.
+			static const bool bSoundWaveDataHasStreamingDisabled = FParse::Param(FCommandLine::Get(), TEXT("SoundWaveDataHasStreamingDisabled"));
+			bShouldStreamSound = bShouldStreamSound && !bSoundWaveDataHasStreamingDisabled;
+		}
 #endif
 
 		// Only want to cook/load full data if we don't support streaming
