Alt Skin Weights: need a way to trigger them from the AnimBP
#feature added ability to override the behaviour for picking which skin weight profile to load, this enables picking different profiles according to user defined state/behaviour

diff --git a/Engine/Source/Runtime/Engine/Private/Animation/SkinWeightProfile.cpp b/Engine/Source/Runtime/Engine/Private/Animation/SkinWeightProfile.cpp
--- a/Engine/Source/Runtime/Engine/Private/Animation/SkinWeightProfile.cpp
+++ b/Engine/Source/Runtime/Engine/Private/Animation/SkinWeightProfile.cpp
@@ -241,6 +241,7 @@ FSkinWeightProfilesData::~FSkinWeightProfilesData()
 	DefaultProfileName = NAME_None;
 }
 
+FSkinWeightProfilesData::FOnPickOverrideSkinWeightProfile FSkinWeightProfilesData::OnPickOverrideSkinWeightProfile;
 #if !WITH_EDITOR
 void FSkinWeightProfilesData::OverrideBaseBufferSkinWeightData(USkeletalMesh* Mesh, int32 LODIndex)
 {
@@ -248,23 +249,23 @@ void FSkinWeightProfilesData::OverrideBaseBufferSkinWeightData(USkeletalMesh* Me
 	{
 		const TArray<FSkinWeightProfileInfo>& Profiles = Mesh->GetSkinWeightProfiles();
 		// Try and find a default buffer and whether or not it is set for this LOD index 
-		const int32 DefaultProfileIndex = Profiles.IndexOfByPredicate([LODIndex](FSkinWeightProfileInfo ProfileInfo)
-		{
-			// Setup to not apply any skin weight profiles at this LOD level
-			if (LODIndex < GSkinWeightProfilesAllowedFromLOD)
-			{
-				return false;
-			}
+		int32 DefaultProfileIndex = INDEX_NONE;
 
-			// In case the default LOD index has been overridden check against that
-			if (GSkinWeightProfilesDefaultLODOverride >= 0)
+		// Setup to not apply any skin weight profiles at this LOD level
+		if (LODIndex >= GSkinWeightProfilesAllowedFromLOD)
+		{
+			DefaultProfileIndex = OnPickOverrideSkinWeightProfile.IsBound() ? OnPickOverrideSkinWeightProfile.Execute(Mesh, MakeArrayView(Profiles), LODIndex) : Profiles.IndexOfByPredicate([LODIndex](FSkinWeightProfileInfo ProfileInfo)
 			{
-				return (ProfileInfo.DefaultProfile.Default && LODIndex >= GSkinWeightProfilesDefaultLODOverride);
-			}
+				// In case the default LOD index has been overridden check against that
+				if (GSkinWeightProfilesDefaultLODOverride >= 0)
+				{
+					return (ProfileInfo.DefaultProfile.Default && LODIndex >= GSkinWeightProfilesDefaultLODOverride);
+				}
 
-			// Otherwise check if this profile is set as default and the current LOD index is applicable
-			return (ProfileInfo.DefaultProfile.Default && LODIndex >= ProfileInfo.DefaultProfileFromLODIndex.Default);
-		});
+				// Otherwise check if this profile is set as default and the current LOD index is applicable
+				return (ProfileInfo.DefaultProfile.Default && LODIndex >= ProfileInfo.DefaultProfileFromLODIndex.Default);
+			});
+		}
 
 		bool bProfileSet = false;
 		// If we found a profile try and find the override skin weights and apply if found
diff --git a/Engine/Source/Runtime/Engine/Public/Animation/SkinWeightProfile.h b/Engine/Source/Runtime/Engine/Public/Animation/SkinWeightProfile.h
--- a/Engine/Source/Runtime/Engine/Public/Animation/SkinWeightProfile.h
+++ b/Engine/Source/Runtime/Engine/Public/Animation/SkinWeightProfile.h
@@ -131,6 +131,9 @@ struct ENGINE_API FSkinWeightProfilesData
 	void Init(FSkinWeightVertexBuffer* InBaseBuffer);
 
 	~FSkinWeightProfilesData();
+	
+	DECLARE_DELEGATE_RetVal_ThreeParams(int32 /** Index into Profiles ArrayView */, FOnPickOverrideSkinWeightProfile, const USkeletalMesh* /** Skeletal Mesh to pick the profile for */, const TArrayView<const FSkinWeightProfileInfo> /** Available skin weight profiles to pick from */, int32 /** LOD Index */);
+	static FOnPickOverrideSkinWeightProfile OnPickOverrideSkinWeightProfile;
 
 #if !WITH_EDITOR
 	// Mark this as non-editor only to prevent mishaps from users
