Fix jump resimulation issue on charactermovementcomponent when we do a correction back to walking and resimulate jump input. The PressedJump command from the character was cleaned up due to exiting falling state but not properly set back when replaying saved move.

diff --git a/Engine/Source/Runtime/Engine/Private/Components/CharacterMovementComponent.cpp b/Engine/Source/Runtime/Engine/Private/Components/CharacterMovementComponent.cpp
--- a/Engine/Source/Runtime/Engine/Private/Components/CharacterMovementComponent.cpp
+++ b/Engine/Source/Runtime/Engine/Private/Components/CharacterMovementComponent.cpp
@@ -8055,7 +8055,7 @@ bool UCharacterMovementComponent::ClientUpdatePositionAfterServerUpdate()
 	const float SavedAnalogInputModifier = AnalogInputModifier;
 	const FRootMotionMovementParams BackupRootMotionParams = RootMotionParams; // For animation root motion
 	const FRootMotionSourceGroup BackupRootMotion = CurrentRootMotion;
-	const bool bRealJump = CharacterOwner->bPressedJump;
+	const bool bRealPressedJump = CharacterOwner->bPressedJump;
 	const bool bRealCrouch = bWantsToCrouch;
 	const bool bRealForceMaxAccel = bForceMaxAccel;
 	CharacterOwner->bClientWasFalling = (MovementMode == MOVE_Falling);
@@ -8085,6 +8085,7 @@ bool UCharacterMovementComponent::ClientUpdatePositionAfterServerUpdate()
 		CurrentMove->PostUpdate(CharacterOwner, FSavedMove_Character::PostUpdate_Replay);
 		SetCurrentNetworkMoveData(nullptr);
 	}
+	const bool bPostReplayPressedJump = CharacterOwner->bPressedJump;
 
 	if (FSavedMove_Character* const PendingMove = ClientData->PendingMove.Get())
 	{
@@ -8107,7 +8108,7 @@ bool UCharacterMovementComponent::ClientUpdatePositionAfterServerUpdate()
 	CharacterOwner->SavedRootMotion.Clear();
 	CharacterOwner->bClientResimulateRootMotion = false;
 	CharacterOwner->bClientUpdating = false;
-	CharacterOwner->bPressedJump = bRealJump;
+	CharacterOwner->bPressedJump = bRealPressedJump || bPostReplayPressedJump;
 	bWantsToCrouch = bRealCrouch;
 	bForceMaxAccel = bRealForceMaxAccel;
 	bForceNextFloorCheck = true;
