Strata - compiler is now safer and does not crash when uniplemented strata functions or edge cases are encountered.

diff --git a/Engine/Source/Runtime/Engine/Classes/Materials/MaterialExpression.h b/Engine/Source/Runtime/Engine/Classes/Materials/MaterialExpression.h
--- a/Engine/Source/Runtime/Engine/Classes/Materials/MaterialExpression.h
+++ b/Engine/Source/Runtime/Engine/Classes/Materials/MaterialExpression.h
@@ -427,7 +427,7 @@ class ENGINE_API UMaterialExpression : public UObject
 	 * This recursively parse nodes outputing strata material in order to gather the maximum distance to any leaves. 
 	 * This is used to drive the bottom up order processing of those nodes.
 	 */
-	virtual FStrataOperator* StrataGenerateMaterialTopologyTree(class FMaterialCompiler* Compiler, class UMaterialExpression* Parent, int32 OutputIndex) { return nullptr; }
+	virtual FStrataOperator* StrataGenerateMaterialTopologyTree(class FMaterialCompiler* Compiler, class UMaterialExpression* Parent, int32 OutputIndex);
 
 	/**
 	 * If true, discards the output index when caching this expression which allows more cases to re-use the output instead of adding a separate instruction
diff --git a/Engine/Source/Runtime/Engine/Private/Materials/HLSLMaterialTranslator.cpp b/Engine/Source/Runtime/Engine/Private/Materials/HLSLMaterialTranslator.cpp
--- a/Engine/Source/Runtime/Engine/Private/Materials/HLSLMaterialTranslator.cpp
+++ b/Engine/Source/Runtime/Engine/Private/Materials/HLSLMaterialTranslator.cpp
@@ -1392,7 +1392,7 @@ bool FHLSLMaterialTranslator::Translate()
 				}
 			}
 
-			if (!bStrataUsesConversionFromLegacy)
+			if (!bStrataUsesConversionFromLegacy && StrataMaterialRootOperator)
 			{
 				// Now implement the functions needed to process the material topology
 
@@ -9660,12 +9660,23 @@ FStrataOperator& FHLSLMaterialTranslator::StrataCompilationRegisterOperator(int3
 FStrataOperator& FHLSLMaterialTranslator::StrataCompilationGetOperator(UMaterialExpression* Expression)
 {
 	auto* OperatorIndex = StrataMaterialExpressionToOperatorIndex.Find(Expression);
-	check(OperatorIndex && *OperatorIndex >= 0 && *OperatorIndex < STRATA_MAX_OPERATOR_COUNT);
+	if (!(OperatorIndex && *OperatorIndex >= 0 && *OperatorIndex < STRATA_MAX_OPERATOR_COUNT))
+	{
+		UE_LOG(LogMaterial, Error, TEXT("Could not find operator for a BSDF in Material %s (asset: %s).\r\n"), *Material->GetDebugName(), *Material->GetAssetPath().ToString());
+		static FStrataOperator DefaultOperator = FStrataOperator();
+		return DefaultOperator;
+	};
 	return StrataMaterialExpressionRegisteredOperators[*OperatorIndex];
 }
 
 bool FHLSLMaterialTranslator::StrataGenerateDerivedMaterialOperatorData()
 {
+	if (StrataMaterialExpressionRegisteredOperators.IsEmpty())
+	{
+		UE_LOG(LogMaterial, Error, TEXT("Could not find any strata operators or BSDFs in Material %s (asset: %s).\r\n"), *Material->GetDebugName(), *Material->GetAssetPath().ToString());
+		return false;
+	}
+
 	//
 	// Evaluate the one and only root node
 	//
@@ -9952,8 +9963,14 @@ bool FHLSLMaterialTranslator::StrataGenerateDerivedMaterialOperatorData()
 
 void FHLSLMaterialTranslator::StrataCompilationInfoRegisterCodeChunk(int32 CodeChunk, FStrataMaterialCompilationInfo& StrataMaterialCompilationInfo)
 {
-	check(CodeChunk != INDEX_NONE);
-	CodeChunkToStrataCompilationInfoMap.Add(CodeChunk, StrataMaterialCompilationInfo);
+	if (CodeChunk != INDEX_NONE)
+	{
+		CodeChunkToStrataCompilationInfoMap.Add(CodeChunk, StrataMaterialCompilationInfo);
+	}
+	else
+	{
+		UE_LOG(LogMaterial, Error, TEXT("Invalid Strata CodeChunk registered in Material %s (asset: %s).\r\n"), *Material->GetDebugName(), *Material->GetAssetPath().ToString());
+	}
 }
 
 bool FHLSLMaterialTranslator::StrataCompilationInfoContainsCodeChunk(int32 CodeChunk)
@@ -9963,8 +9980,12 @@ bool FHLSLMaterialTranslator::StrataCompilationInfoContainsCodeChunk(int32 CodeC
 
 const FStrataMaterialCompilationInfo& FHLSLMaterialTranslator::GetStrataCompilationInfo(int32 CodeChunk)
 {
-	check(CodeChunk != INDEX_NONE);
-	return CodeChunkToStrataCompilationInfoMap[CodeChunk];
+	if (CodeChunk != INDEX_NONE)
+	{
+		return CodeChunkToStrataCompilationInfoMap[CodeChunk];
+	}
+	static FStrataMaterialCompilationInfo InvalidInfo;
+	return InvalidInfo;
 }
 
 FStrataRegisteredSharedLocalBasis FHLSLMaterialTranslator::StrataCompilationInfoRegisterSharedLocalBasis(int32 NormalCodeChunk)
@@ -10098,11 +10119,13 @@ int32 FHLSLMaterialTranslator::StrataSlabBSDF(
 	const FString NormalCode = GetParameterCode(Normal);
 	const FString TangentCode = Tangent != INDEX_NONE ? *GetParameterCode(Tangent) : TEXT("NONE");
 
-
 	if (PromoteToOperator)
 	{
-		check(PromoteToOperator->Index != INDEX_NONE);
-		check(PromoteToOperator->BSDFIndex != INDEX_NONE);
+		if (PromoteToOperator->Index == INDEX_NONE || PromoteToOperator->BSDFIndex == INDEX_NONE)
+		{
+			UE_LOG(LogMaterial, Error, TEXT("Invalid StrataSlabBSDF operator and BSDF indices during promotion in Material %s (asset: %s).\r\n"), *Material->GetDebugName(), *Material->GetAssetPath().ToString());
+			return INDEX_NONE;
+		}
 		return AddCodeChunk(
 			MCT_Strata, TEXT("PromoteParameterBlendedBSDFToOperator(GetStrataSlabBSDF(Parameters.StrataPixelFootprint, %s, %s, %s, %s, %s, %s, %s, %s, %s, %s, %s, %s, %s, %s, %s, %s, %s, %s, %s, Parameters.SharedLocalBases.Types) /* Normal = %s ; Tangent = %s */, Parameters.StrataTree, %u, %u, %u, %u)"),
 			*GetParameterCode(UseMetalness),
diff --git a/Engine/Source/Runtime/Engine/Private/Materials/MaterialExpressions.cpp b/Engine/Source/Runtime/Engine/Private/Materials/MaterialExpressions.cpp
--- a/Engine/Source/Runtime/Engine/Private/Materials/MaterialExpressions.cpp
+++ b/Engine/Source/Runtime/Engine/Private/Materials/MaterialExpressions.cpp
@@ -1580,6 +1580,17 @@ bool UMaterialExpression::IsUsingNewHLSLGenerator() const
 	return false;
 }
 
+FStrataOperator* UMaterialExpression::StrataGenerateMaterialTopologyTree(class FMaterialCompiler* Compiler, class UMaterialExpression* Parent, int32 OutputIndex)
+{
+	Compiler->Errorf(TEXT("Missing StrataGenerateMaterialTopologyTree implementation for node %s."), *GetClass()->GetName());
+	return nullptr;
+}
+
+static void AssignOperatorIndexIfNotNull(int32& NextOperatorPin, FStrataOperator* Operator)
+{
+	NextOperatorPin = Operator ? Operator->Index : INDEX_NONE;
+}
+
 #endif // WITH_EDITOR
 
 UMaterialExpressionTextureBase::UMaterialExpressionTextureBase(const FObjectInitializer& ObjectInitializer)
@@ -6889,11 +6900,11 @@ FStrataOperator* UMaterialExpressionBlendMaterialAttributes::StrataGenerateMater
 	UMaterialExpression* ChildBExpression = B.GetTracedInput().Expression;
 	if (ChildAExpression)
 	{
-		StrataOperator.LeftIndex = ChildAExpression->StrataGenerateMaterialTopologyTree(Compiler, this, A.OutputIndex)->Index;
+		AssignOperatorIndexIfNotNull(StrataOperator.LeftIndex, ChildAExpression->StrataGenerateMaterialTopologyTree(Compiler, this, A.OutputIndex));
 	}
 	if (ChildBExpression)
 	{
-		StrataOperator.RightIndex = ChildBExpression->StrataGenerateMaterialTopologyTree(Compiler, this, B.OutputIndex)->Index;
+		AssignOperatorIndexIfNotNull(StrataOperator.RightIndex, ChildBExpression->StrataGenerateMaterialTopologyTree(Compiler, this, B.OutputIndex));
 	}
 
 	return &StrataOperator;
@@ -20729,8 +20740,9 @@ FStrataOperator* UMaterialExpressionStrataLegacyConversion::StrataGenerateMateri
 	// with fake addresses for BSDF defined in shader code.
 	UMaterialExpression* ChildAExpression = reinterpret_cast<UMaterialExpression*>(0);
 	UMaterialExpression* ChildBExpression = reinterpret_cast<UMaterialExpression*>(1);
-	StrataOperator.LeftIndex = Compiler->StrataCompilationRegisterOperator(STRATA_OPERATOR_BSDF_LEGACY, ChildAExpression, this).Index;
-	StrataOperator.RightIndex = Compiler->StrataCompilationRegisterOperator(STRATA_OPERATOR_BSDF_LEGACY, ChildBExpression, this).Index;
+
+	AssignOperatorIndexIfNotNull(StrataOperator.LeftIndex, &Compiler->StrataCompilationRegisterOperator(STRATA_OPERATOR_BSDF_LEGACY, ChildAExpression, this));
+	AssignOperatorIndexIfNotNull(StrataOperator.RightIndex, &Compiler->StrataCompilationRegisterOperator(STRATA_OPERATOR_BSDF_LEGACY, ChildBExpression, this));
 
 	return &StrataOperator;
 }
@@ -21593,7 +21605,7 @@ FStrataOperator* UMaterialExpressionStrataConvertToDecal::StrataGenerateMaterial
 	UMaterialExpression* ChildDecalMaterialExpression = DecalMaterial.GetTracedInput().Expression;
 	if (ChildDecalMaterialExpression)
 	{
-		StrataOperator.LeftIndex = ChildDecalMaterialExpression->StrataGenerateMaterialTopologyTree(Compiler, this, DecalMaterial.OutputIndex)->Index;
+		AssignOperatorIndexIfNotNull(StrataOperator.LeftIndex, ChildDecalMaterialExpression->StrataGenerateMaterialTopologyTree(Compiler, this, DecalMaterial.OutputIndex));
 	}
 
 	return &StrataOperator;
@@ -22049,11 +22061,11 @@ FStrataOperator* UMaterialExpressionStrataHorizontalMixing::StrataGenerateMateri
 	UMaterialExpression* ChildBExpression = Foreground.GetTracedInput().Expression;
 	if (ChildAExpression)
 	{
-		StrataOperator.LeftIndex = ChildAExpression->StrataGenerateMaterialTopologyTree(Compiler, this, Background.OutputIndex)->Index;
+		AssignOperatorIndexIfNotNull(StrataOperator.LeftIndex, ChildAExpression->StrataGenerateMaterialTopologyTree(Compiler, this, Background.OutputIndex));
 	}
 	if (ChildBExpression)
 	{
-		StrataOperator.RightIndex = ChildBExpression->StrataGenerateMaterialTopologyTree(Compiler, this, Foreground.OutputIndex)->Index;
+		AssignOperatorIndexIfNotNull(StrataOperator.RightIndex, ChildBExpression->StrataGenerateMaterialTopologyTree(Compiler, this, Foreground.OutputIndex));
 	}
 
 	return &StrataOperator;
@@ -22211,11 +22223,11 @@ FStrataOperator* UMaterialExpressionStrataVerticalLayering::StrataGenerateMateri
 	UMaterialExpression* ChildBExpression = Base.GetTracedInput().Expression;
 	if (ChildAExpression)
 	{
-		StrataOperator.LeftIndex = ChildAExpression->StrataGenerateMaterialTopologyTree(Compiler, this, Top.OutputIndex)->Index;
+		AssignOperatorIndexIfNotNull(StrataOperator.LeftIndex, ChildAExpression->StrataGenerateMaterialTopologyTree(Compiler, this, Top.OutputIndex));
 	}
 	if (ChildBExpression)
 	{
-		StrataOperator.RightIndex = ChildBExpression->StrataGenerateMaterialTopologyTree(Compiler, this, Base.OutputIndex)->Index;
+		AssignOperatorIndexIfNotNull(StrataOperator.RightIndex, ChildBExpression->StrataGenerateMaterialTopologyTree(Compiler, this, Base.OutputIndex));
 	}
 
 	return &StrataOperator;
@@ -22372,11 +22384,11 @@ FStrataOperator* UMaterialExpressionStrataAdd::StrataGenerateMaterialTopologyTre
 	UMaterialExpression* ChildBExpression = B.GetTracedInput().Expression;
 	if (ChildAExpression)
 	{
-		StrataOperator.LeftIndex = ChildAExpression->StrataGenerateMaterialTopologyTree(Compiler, this, A.OutputIndex)->Index;
+		AssignOperatorIndexIfNotNull(StrataOperator.LeftIndex, ChildAExpression->StrataGenerateMaterialTopologyTree(Compiler, this, A.OutputIndex));
 	}
 	if (ChildBExpression)
 	{
-		StrataOperator.RightIndex = ChildBExpression->StrataGenerateMaterialTopologyTree(Compiler, this, B.OutputIndex)->Index;
+		AssignOperatorIndexIfNotNull(StrataOperator.RightIndex, ChildBExpression->StrataGenerateMaterialTopologyTree(Compiler, this, B.OutputIndex));
 	}
 
 	return &StrataOperator;
@@ -22476,7 +22488,7 @@ FStrataOperator* UMaterialExpressionStrataWeight::StrataGenerateMaterialTopology
 	UMaterialExpression* ChildAExpression = A.GetTracedInput().Expression;
 	if (ChildAExpression)
 	{
-		StrataOperator.LeftIndex = ChildAExpression->StrataGenerateMaterialTopologyTree(Compiler, this, A.OutputIndex)->Index;
+		AssignOperatorIndexIfNotNull(StrataOperator.LeftIndex, ChildAExpression->StrataGenerateMaterialTopologyTree(Compiler, this, A.OutputIndex));
 	}
 
 	return &StrataOperator;
@@ -22578,7 +22590,7 @@ FStrataOperator* UMaterialExpressionStrataThinFilm::StrataGenerateMaterialTopolo
 	UMaterialExpression* ChildAExpression = A.GetTracedInput().Expression;
 	if (ChildAExpression)
 	{
-		StrataOperator.LeftIndex = ChildAExpression->StrataGenerateMaterialTopologyTree(Compiler, this, A.OutputIndex)->Index;
+		AssignOperatorIndexIfNotNull(StrataOperator.LeftIndex, ChildAExpression->StrataGenerateMaterialTopologyTree(Compiler, this, A.OutputIndex));
 	}
 
 	return &StrataOperator;
