Strata - add support for a few switch pass nodes.

diff --git a/Engine/Source/Runtime/Engine/Classes/Materials/MaterialExpressionStaticSwitch.h b/Engine/Source/Runtime/Engine/Classes/Materials/MaterialExpressionStaticSwitch.h
--- a/Engine/Source/Runtime/Engine/Classes/Materials/MaterialExpressionStaticSwitch.h
+++ b/Engine/Source/Runtime/Engine/Classes/Materials/MaterialExpressionStaticSwitch.h
@@ -37,8 +37,17 @@ class UMaterialExpressionStaticSwitch : public UMaterialExpression
 	virtual uint32 GetOutputType(int32 OutputIndex) override {return MCT_Unknown;}
 
 	virtual bool GenerateHLSLExpression(FMaterialHLSLGenerator& Generator, UE::HLSLTree::FScope& Scope, int32 OutputIndex, UE::HLSLTree::FExpression const*& OutExpression) const override;
+
+	virtual bool IsResultStrataMaterial(int32 OutputIndex) override;
+	virtual void GatherStrataMaterialInfo(FStrataMaterialInfo& StrataMaterialInfo, int32 OutputIndex) override;
+	virtual FStrataOperator* StrataGenerateMaterialTopologyTree(class FMaterialCompiler* Compiler, class UMaterialExpression* Parent, int32 OutputIndex) override;
 #endif // WITH_EDITOR
 	//~ End UMaterialExpression Interface
+
+#if WITH_EDITOR
+protected:
+	FExpressionInput* GetEffectiveInput(class FMaterialCompiler* Compiler);
+#endif
 };
 
 
diff --git a/Engine/Source/Runtime/Engine/Classes/Materials/MaterialExpressionStaticSwitchParameter.h b/Engine/Source/Runtime/Engine/Classes/Materials/MaterialExpressionStaticSwitchParameter.h
--- a/Engine/Source/Runtime/Engine/Classes/Materials/MaterialExpressionStaticSwitchParameter.h
+++ b/Engine/Source/Runtime/Engine/Classes/Materials/MaterialExpressionStaticSwitchParameter.h
@@ -31,7 +31,16 @@ class UMaterialExpressionStaticSwitchParameter : public UMaterialExpressionStati
 	virtual uint32 GetInputType(int32 InputIndex) override {return MCT_Unknown;}
 	virtual uint32 GetOutputType(int32 OutputIndex) override {return MCT_Unknown;}
 	virtual bool GenerateHLSLExpression(FMaterialHLSLGenerator& Generator, UE::HLSLTree::FScope& Scope, int32 OutputIndex, UE::HLSLTree::FExpression const*& OutExpression) const override;
+
+	virtual bool IsResultStrataMaterial(int32 OutputIndex) override;
+	virtual void GatherStrataMaterialInfo(FStrataMaterialInfo& StrataMaterialInfo, int32 OutputIndex) override;
+	virtual FStrataOperator* StrataGenerateMaterialTopologyTree(class FMaterialCompiler* Compiler, class UMaterialExpression* Parent, int32 OutputIndex) override;
 #endif // WITH_EDITOR
 	//~ End UMaterialExpression Interface
+
+#if WITH_EDITOR
+protected:
+	FExpressionInput* GetEffectiveInput(class FMaterialCompiler* Compiler);
+#endif
 };
 
diff --git a/Engine/Source/Runtime/Engine/Private/Materials/HLSLMaterialTranslator.cpp b/Engine/Source/Runtime/Engine/Private/Materials/HLSLMaterialTranslator.cpp
--- a/Engine/Source/Runtime/Engine/Private/Materials/HLSLMaterialTranslator.cpp
+++ b/Engine/Source/Runtime/Engine/Private/Materials/HLSLMaterialTranslator.cpp
@@ -703,7 +703,11 @@ bool FHLSLMaterialTranslator::Translate()
 		UMaterialExpression* FrontMaterialExpr = FrontMaterialInput ? FrontMaterialInput->GetTracedInput().Expression : nullptr;
 		if (bStrataEnabled && FrontMaterialExpr)
 		{
-			FrontMaterialExpr->StrataGenerateMaterialTopologyTree(this, nullptr, 0);
+			// Temp code chunk scope (e.g.needed for the creation of static booleans from static switch parameter node, see UMaterialExpressionStaticSwitch::GetEffectiveInput).
+			TArray<FShaderCodeChunk> TempChunks;
+			AssignTempScope(TempChunks);
+
+			FrontMaterialExpr->StrataGenerateMaterialTopologyTree(this, nullptr, 0); 
 			if (!StrataGenerateDerivedMaterialOperatorData())
 			{
 				Errorf(TEXT("Strata material errors encountered."));
diff --git a/Engine/Source/Runtime/Engine/Private/Materials/MaterialExpressions.cpp b/Engine/Source/Runtime/Engine/Private/Materials/MaterialExpressions.cpp
--- a/Engine/Source/Runtime/Engine/Private/Materials/MaterialExpressions.cpp
+++ b/Engine/Source/Runtime/Engine/Private/Materials/MaterialExpressions.cpp
@@ -8374,18 +8374,18 @@ bool UMaterialExpressionStaticSwitchParameter::IsResultMaterialAttributes(int32
 	}
 }
 
-int32 UMaterialExpressionStaticSwitchParameter::Compile(class FMaterialCompiler* Compiler, int32 OutputIndex)
+FExpressionInput* UMaterialExpressionStaticSwitchParameter::GetEffectiveInput(class FMaterialCompiler* Compiler)
 {
 	bool bSucceeded;
-	const bool bValue = Compiler->GetStaticBoolValue(Compiler->StaticBoolParameter(ParameterName,DefaultValue), bSucceeded);
-	
+	const bool bValue = Compiler->GetStaticBoolValue(Compiler->StaticBoolParameter(ParameterName, DefaultValue), bSucceeded);
+
 	//Both A and B must be connected in a parameter. 
-	if( !A.GetTracedInput().IsConnected() )
+	if (!A.GetTracedInput().IsConnected())
 	{
 		Compiler->Errorf(TEXT("Missing A input"));
 		bSucceeded = false;
 	}
-	if( !B.GetTracedInput().IsConnected() )
+	if (!B.GetTracedInput().IsConnected())
 	{
 		Compiler->Errorf(TEXT("Missing B input"));
 		bSucceeded = false;
@@ -8393,17 +8393,19 @@ int32 UMaterialExpressionStaticSwitchParameter::Compile(class FMaterialCompiler*
 
 	if (!bSucceeded)
 	{
-		return INDEX_NONE;
+		return nullptr;
 	}
+	return bValue ? &A : &B;
+}
 
-	if (bValue)
-	{
-		return A.Compile(Compiler);
-	}
-	else
+int32 UMaterialExpressionStaticSwitchParameter::Compile(class FMaterialCompiler* Compiler, int32 OutputIndex)
+{
+	FExpressionInput* EffectiveInput = GetEffectiveInput(Compiler);
+	if (EffectiveInput)
 	{
-		return B.Compile(Compiler);
+		return EffectiveInput->Compile(Compiler);
 	}
+	return INDEX_NONE;
 }
 
 void UMaterialExpressionStaticSwitchParameter::GetCaption(TArray<FString>& OutCaptions) const
@@ -8423,6 +8425,38 @@ FName UMaterialExpressionStaticSwitchParameter::GetInputName(int32 InputIndex) c
 		return TEXT("False");
 	}
 }
+
+bool UMaterialExpressionStaticSwitchParameter::IsResultStrataMaterial(int32 OutputIndex)
+{
+	if (A.GetTracedInput().Expression && B.GetTracedInput().Expression)
+	{
+		return A.GetTracedInput().Expression->IsResultStrataMaterial(A.OutputIndex) && B.GetTracedInput().Expression->IsResultStrataMaterial(B.OutputIndex);
+	}
+	return false;
+}
+
+void UMaterialExpressionStaticSwitchParameter::GatherStrataMaterialInfo(FStrataMaterialInfo& StrataMaterialInfo, int32 OutputIndex)
+{
+	// STRATA_TODO: this is incorrect because we should only use A or B based on GetEffectiveInput, but we have no compiler at this stage so we just gather both.
+	if (A.GetTracedInput().Expression)
+	{
+		A.GetTracedInput().Expression->GatherStrataMaterialInfo(StrataMaterialInfo, A.OutputIndex);
+	}
+	if (B.GetTracedInput().Expression)
+	{
+		B.GetTracedInput().Expression->GatherStrataMaterialInfo(StrataMaterialInfo, B.OutputIndex);
+	}
+}
+
+FStrataOperator* UMaterialExpressionStaticSwitchParameter::StrataGenerateMaterialTopologyTree(class FMaterialCompiler* Compiler, class UMaterialExpression* Parent, int32 OutputIndex)
+{
+	FExpressionInput* EffectiveInput = GetEffectiveInput(Compiler);
+	if (EffectiveInput && EffectiveInput->Expression)
+	{
+		return EffectiveInput->Expression->StrataGenerateMaterialTopologyTree(Compiler, Parent, 0);
+	}
+	return nullptr;
+}
 #endif // WITH_EDITOR
 
 
@@ -8572,30 +8606,29 @@ bool UMaterialExpressionStaticSwitch::IsResultMaterialAttributes(int32 OutputInd
 	}
 }
 
-int32 UMaterialExpressionStaticSwitch::Compile(class FMaterialCompiler* Compiler, int32 OutputIndex)
+FExpressionInput* UMaterialExpressionStaticSwitch::GetEffectiveInput(class FMaterialCompiler* Compiler)
 {
 	bool bValue = DefaultValue;
-
 	if (Value.GetTracedInput().Expression)
 	{
 		bool bSucceeded;
 		bValue = Compiler->GetStaticBoolValue(Value.Compile(Compiler), bSucceeded);
-
 		if (!bSucceeded)
 		{
-			return INDEX_NONE;
+			return nullptr;
 		}
 	}
-	
-	// We only call Compile on the branch that is taken to avoid compile errors in the disabled branch.
-	if (bValue)
-	{
-		return A.Compile(Compiler);
-	}
-	else
+	return bValue ? &A : &B;
+}
+
+int32 UMaterialExpressionStaticSwitch::Compile(class FMaterialCompiler* Compiler, int32 OutputIndex)
+{
+	FExpressionInput* EffectiveInput = GetEffectiveInput(Compiler);
+	if (EffectiveInput)
 	{
-		return B.Compile(Compiler);
+		return EffectiveInput->Compile(Compiler);
 	}
+	return INDEX_NONE;
 }
 
 void UMaterialExpressionStaticSwitch::GetCaption(TArray<FString>& OutCaptions) const
@@ -8630,6 +8663,38 @@ uint32 UMaterialExpressionStaticSwitch::GetInputType(int32 InputIndex)
 		return MCT_StaticBool;
 	}
 }
+
+bool UMaterialExpressionStaticSwitch::IsResultStrataMaterial(int32 OutputIndex)
+{
+	if (A.GetTracedInput().Expression && B.GetTracedInput().Expression)
+	{
+		return A.GetTracedInput().Expression->IsResultStrataMaterial(A.OutputIndex) && B.GetTracedInput().Expression->IsResultStrataMaterial(B.OutputIndex);
+	}
+	return false;
+}
+
+void UMaterialExpressionStaticSwitch::GatherStrataMaterialInfo(FStrataMaterialInfo& StrataMaterialInfo, int32 OutputIndex)
+{
+	// STRATA_TODO: this is incorrect because we should only use A or B based on GetEffectiveInput, but we have no compiler at this stage so we just gather both.
+	if (A.GetTracedInput().Expression)
+	{
+		A.GetTracedInput().Expression->GatherStrataMaterialInfo(StrataMaterialInfo, A.OutputIndex);
+	}
+	if (B.GetTracedInput().Expression)
+	{
+		B.GetTracedInput().Expression->GatherStrataMaterialInfo(StrataMaterialInfo, B.OutputIndex);
+	}
+}
+
+FStrataOperator* UMaterialExpressionStaticSwitch::StrataGenerateMaterialTopologyTree(class FMaterialCompiler* Compiler, class UMaterialExpression* Parent, int32 OutputIndex)
+{
+	FExpressionInput* EffectiveInput = GetEffectiveInput(Compiler);
+	if (EffectiveInput && EffectiveInput->Expression)
+	{
+		return EffectiveInput->Expression->StrataGenerateMaterialTopologyTree(Compiler, Parent, 0);
+	}
+	return nullptr;
+}
 #endif
 
 //
