World Partition: remove unrequired option to always load a partitioned world in the editor but generate streaming for runtime.

#rnx

diff --git a/Engine/Source/Runtime/Engine/Private/WorldPartition/WorldPartition.cpp b/Engine/Source/Runtime/Engine/Private/WorldPartition/WorldPartition.cpp
--- a/Engine/Source/Runtime/Engine/Private/WorldPartition/WorldPartition.cpp
+++ b/Engine/Source/Runtime/Engine/Private/WorldPartition/WorldPartition.cpp
@@ -256,7 +256,6 @@ UWorldPartition::UWorldPartition(const FObjectInitializer& ObjectInitializer)
 	: Super(ObjectInitializer)
 #if WITH_EDITOR
 	, EditorHash(nullptr)
-	, bIsEditorLevelAlwaysLoaded(false)
 	, WorldPartitionEditor(nullptr)
 	, bEnableStreaming(true)
 	, bStreamingWasEnabled(true)
@@ -359,27 +358,6 @@ FName UWorldPartition::GetWorldPartitionEditorName() const
 	}
 	return NAME_None;
 }
-
-void UWorldPartition::SetIsEditorLevelAlwaysLoaded(bool bIsAlwaysLoaded)
-{
-	if (bIsEditorLevelAlwaysLoaded != bIsAlwaysLoaded)
-	{
-		bIsEditorLevelAlwaysLoaded = bIsAlwaysLoaded;
-		UpdateEditorLevelAlwaysLoaded();
-	}
-}
-
-void UWorldPartition::UpdateEditorLevelAlwaysLoaded()
-{
-	// Can be called if SetIsEditorLevelAlwaysLoaded is called on a freshly created UWorldPartition.
-	if (IsInitialized() && EditorHash && bIsEditorLevelAlwaysLoaded)
-	{
-		EditorHash->ForEachCell([this](UWorldPartitionEditorCell* Cell)
-		{
-			UpdateLoadingEditorCell(Cell, true, false);
-		});
-	}
-}
 #endif
 
 void UWorldPartition::Initialize(UWorld* InWorld, const FTransform& InTransform)
@@ -530,11 +508,7 @@ void UWorldPartition::Initialize(UWorld* InWorld, const FTransform& InTransform)
 	InitState = EWorldPartitionInitState::Initialized;
 
 #if WITH_EDITOR
-	if (bIsEditor)
-	{
-		UpdateEditorLevelAlwaysLoaded();
-	}
-	else
+	if (!bIsEditor)
 	{
 		if (bIsGame || bPIEWorldTravel || bIsDedicatedServer)
 		{
@@ -593,14 +567,14 @@ void UWorldPartition::Uninitialize()
 				// Unload all Editor cells
 				if (EditorHash)
 				{
-				// @todo_ow: Once Metadata is removed from external actor's package, this won't be necessary anymore.
-				EditorHash->ForEachCell([this](UWorldPartitionEditorCell* Cell)
-				{
-					UpdateLoadingEditorCell(Cell, /*bShouldBeLoaded*/false, /*bIsFromUserChange*/false);
-				});
+					// @todo_ow: Once Metadata is removed from external actor's package, this won't be necessary anymore.
+					EditorHash->ForEachCell([this](UWorldPartitionEditorCell* Cell)
+					{
+						UpdateLoadingEditorCell(Cell, /*bShouldBeLoaded*/false, /*bIsFromUserChange*/false);
+					});
+				}
 			}
 		}
-		}
 
 		WorldDataLayersActor = FWorldPartitionReference();
 
@@ -1569,9 +1543,9 @@ FBox UWorldPartition::GetWorldBounds() const
 	for (UActorDescContainer::TConstIterator<> ActorDescIterator(this); ActorDescIterator; ++ActorDescIterator)
 	{
 		if (ActorDescIterator->GetIsSpatiallyLoaded())
-			{
-				WorldBounds += ActorDescIterator->GetBounds();
-			}
+		{
+			WorldBounds += ActorDescIterator->GetBounds();
+		}
 	}
 	return WorldBounds;
 }
diff --git a/Engine/Source/Runtime/Engine/Public/WorldPartition/WorldPartition.h b/Engine/Source/Runtime/Engine/Public/WorldPartition/WorldPartition.h
--- a/Engine/Source/Runtime/Engine/Public/WorldPartition/WorldPartition.h
+++ b/Engine/Source/Runtime/Engine/Public/WorldPartition/WorldPartition.h
@@ -134,7 +134,6 @@ public:
 
 #if WITH_EDITOR
 	FName GetWorldPartitionEditorName() const;
-	void SetIsEditorLevelAlwaysLoaded(bool bIsAlwaysLoaded);
 
 	void LoadEditorCells(const FBox& Box, bool bIsFromUserChange);
 	void LoadEditorCells(const TArray<FName>& CellNames, bool bIsFromUserChange);
@@ -201,9 +200,6 @@ public:
 	UPROPERTY(DuplicateTransient)
 	TObjectPtr<UWorldPartitionEditorHash> EditorHash;
 
-	UPROPERTY()
-	bool bIsEditorLevelAlwaysLoaded;
-
 	IWorldPartitionEditor* WorldPartitionEditor;
 
 	/** Class of WorldPartitionStreamingPolicy to be used to manage world partition streaming. */
@@ -259,7 +255,6 @@ private:
 	void UnregisterDelegates();	
 
 #if WITH_EDITOR
-	void UpdateEditorLevelAlwaysLoaded();
 	void UpdateLoadingEditorCell(UWorldPartitionEditorCell* Cell, bool bShouldBeLoaded, bool bFromUserOperation);
 	void HashActorDesc(FWorldPartitionActorDesc* ActorDesc);
 	void UnhashActorDesc(FWorldPartitionActorDesc* ActorDesc);
