Strata - Implemented more switch type nodes.
Some are not supported and some are dynamic so cannot be suported on strata data. They can only be supported on input to strata data, not on strata data itself.
Tested on local material shaders without and with strata enabled.

diff --git a/Engine/Source/Runtime/Engine/Classes/Materials/MaterialExpressionPathTracingQualitySwitch.h b/Engine/Source/Runtime/Engine/Classes/Materials/MaterialExpressionPathTracingQualitySwitch.h
--- a/Engine/Source/Runtime/Engine/Classes/Materials/MaterialExpressionPathTracingQualitySwitch.h
+++ b/Engine/Source/Runtime/Engine/Classes/Materials/MaterialExpressionPathTracingQualitySwitch.h
@@ -13,13 +13,13 @@ class UMaterialExpressionPathTracingQualitySwitch : public UMaterialExpression
 {
 	GENERATED_UCLASS_BODY()
 
-		/** Used for non-path traced modes */
-		UPROPERTY()
-		FExpressionInput Normal;
+	/** Used for non-path traced modes */
+	UPROPERTY()
+	FExpressionInput Normal;
 
 	/** Used by the path tracer */
 	UPROPERTY()
-		FExpressionInput PathTraced;
+	FExpressionInput PathTraced;
 
 	//~ Begin UMaterialExpression Interface
 #if WITH_EDITOR
@@ -28,6 +28,10 @@ class UMaterialExpressionPathTracingQualitySwitch : public UMaterialExpression
 	virtual bool IsResultMaterialAttributes(int32 OutputIndex) override;
 	virtual uint32 GetInputType(int32 InputIndex) override;
 	virtual uint32 GetOutputType(int32 OutputIndex) override { return MCT_Unknown; }
+
+	virtual bool IsResultStrataMaterial(int32 OutputIndex) override;
+	virtual void GatherStrataMaterialInfo(FStrataMaterialInfo& StrataMaterialInfo, int32 OutputIndex) override;
+	virtual FStrataOperator* StrataGenerateMaterialTopologyTree(class FMaterialCompiler* Compiler, class UMaterialExpression* Parent, int32 OutputIndex) override;
 #endif
 	//~ End UMaterialExpression Interface
 };
diff --git a/Engine/Source/Runtime/Engine/Classes/Materials/MaterialExpressionQualitySwitch.h b/Engine/Source/Runtime/Engine/Classes/Materials/MaterialExpressionQualitySwitch.h
--- a/Engine/Source/Runtime/Engine/Classes/Materials/MaterialExpressionQualitySwitch.h
+++ b/Engine/Source/Runtime/Engine/Classes/Materials/MaterialExpressionQualitySwitch.h
@@ -32,7 +32,16 @@ class UMaterialExpressionQualitySwitch : public UMaterialExpression
 	virtual bool IsInputConnectionRequired(int32 InputIndex) const override;
 	virtual bool IsResultMaterialAttributes(int32 OutputIndex) override;
 	virtual uint32 GetInputType(int32 InputIndex) override {return MCT_Unknown;}
-	virtual uint32 GetOutputType(int32 InputIndex) override {return MCT_Unknown;}
+	virtual uint32 GetOutputType(int32 InputIndex) override { return MCT_Unknown; }
+
+	virtual bool IsResultStrataMaterial(int32 OutputIndex) override;
+	virtual void GatherStrataMaterialInfo(FStrataMaterialInfo& StrataMaterialInfo, int32 OutputIndex) override;
+	virtual FStrataOperator* StrataGenerateMaterialTopologyTree(class FMaterialCompiler* Compiler, class UMaterialExpression* Parent, int32 OutputIndex) override;
 #endif // WITH_EDITOR
 	//~ End UMaterialExpression Interface
+
+#if WITH_EDITOR
+protected:
+	FExpressionInput* GetEffectiveInput(class FMaterialCompiler* Compiler);
+#endif
 };
diff --git a/Engine/Source/Runtime/Engine/Classes/Materials/MaterialExpressionRayTracingQualitySwitch.h b/Engine/Source/Runtime/Engine/Classes/Materials/MaterialExpressionRayTracingQualitySwitch.h
--- a/Engine/Source/Runtime/Engine/Classes/Materials/MaterialExpressionRayTracingQualitySwitch.h
+++ b/Engine/Source/Runtime/Engine/Classes/Materials/MaterialExpressionRayTracingQualitySwitch.h
@@ -13,13 +13,13 @@ class UMaterialExpressionRayTracingQualitySwitch : public UMaterialExpression
 {
 	GENERATED_UCLASS_BODY()
 
-		/** Used for standard rasterization */
-		UPROPERTY()
-		FExpressionInput Normal;
+	/** Used for standard rasterization */
+	UPROPERTY()
+	FExpressionInput Normal;
 
 	/** Used for simplified ray trace eval */
 	UPROPERTY()
-		FExpressionInput RayTraced;
+	FExpressionInput RayTraced;
 
 	//~ Begin UMaterialExpression Interface
 #if WITH_EDITOR
@@ -28,6 +28,10 @@ class UMaterialExpressionRayTracingQualitySwitch : public UMaterialExpression
 	virtual bool IsResultMaterialAttributes(int32 OutputIndex) override;
 	virtual uint32 GetInputType(int32 InputIndex) override;
 	virtual uint32 GetOutputType(int32 OutputIndex) override { return MCT_Unknown; }
+
+	virtual bool IsResultStrataMaterial(int32 OutputIndex) override;
+	virtual void GatherStrataMaterialInfo(FStrataMaterialInfo& StrataMaterialInfo, int32 OutputIndex) override;
+	virtual FStrataOperator* StrataGenerateMaterialTopologyTree(class FMaterialCompiler* Compiler, class UMaterialExpression* Parent, int32 OutputIndex) override;
 #endif
 	//~ End UMaterialExpression Interface
 };
diff --git a/Engine/Source/Runtime/Engine/Classes/Materials/MaterialExpressionReflectionCapturePassSwitch.h b/Engine/Source/Runtime/Engine/Classes/Materials/MaterialExpressionReflectionCapturePassSwitch.h
--- a/Engine/Source/Runtime/Engine/Classes/Materials/MaterialExpressionReflectionCapturePassSwitch.h
+++ b/Engine/Source/Runtime/Engine/Classes/Materials/MaterialExpressionReflectionCapturePassSwitch.h
@@ -25,6 +25,10 @@ class UMaterialExpressionReflectionCapturePassSwitch : public UMaterialExpressio
 	virtual int32 Compile(class FMaterialCompiler* Compiler, int32 OutputIndex) override;
 	virtual void GetCaption(TArray<FString>& OutCaptions) const override;
 	virtual void GetExpressionToolTip(TArray<FString>& OutToolTip) override;
+
+	virtual bool IsResultStrataMaterial(int32 OutputIndex) override;
+	virtual void GatherStrataMaterialInfo(FStrataMaterialInfo& StrataMaterialInfo, int32 OutputIndex) override;
+	virtual FStrataOperator* StrataGenerateMaterialTopologyTree(class FMaterialCompiler* Compiler, class UMaterialExpression* Parent, int32 OutputIndex) override;
 #endif
 	//~ End UMaterialExpression Interface
 };
diff --git a/Engine/Source/Runtime/Engine/Classes/Materials/MaterialExpressionShaderStageSwitch.h b/Engine/Source/Runtime/Engine/Classes/Materials/MaterialExpressionShaderStageSwitch.h
--- a/Engine/Source/Runtime/Engine/Classes/Materials/MaterialExpressionShaderStageSwitch.h
+++ b/Engine/Source/Runtime/Engine/Classes/Materials/MaterialExpressionShaderStageSwitch.h
@@ -32,6 +32,10 @@ class UMaterialExpressionShaderStageSwitch : public UMaterialExpression
 	virtual void GetCaption(TArray<FString>& OutCaptions) const override;
 	virtual void GetExpressionToolTip(TArray<FString>& OutToolTip) override;
 	virtual FText GetKeywords() const override { return FText::FromString(TEXT("vertex pixel shader")); }
+
+	virtual bool IsResultStrataMaterial(int32 OutputIndex) override;
+	virtual void GatherStrataMaterialInfo(FStrataMaterialInfo& StrataMaterialInfo, int32 OutputIndex) override;
+	virtual FStrataOperator* StrataGenerateMaterialTopologyTree(class FMaterialCompiler* Compiler, class UMaterialExpression* Parent, int32 OutputIndex) override;
 #endif
 	//~ End UMaterialExpression Interface
 };
diff --git a/Engine/Source/Runtime/Engine/Classes/Materials/MaterialExpressionShadingPathSwitch.h b/Engine/Source/Runtime/Engine/Classes/Materials/MaterialExpressionShadingPathSwitch.h
--- a/Engine/Source/Runtime/Engine/Classes/Materials/MaterialExpressionShadingPathSwitch.h
+++ b/Engine/Source/Runtime/Engine/Classes/Materials/MaterialExpressionShadingPathSwitch.h
@@ -34,6 +34,15 @@ class UMaterialExpressionShadingPathSwitch : public UMaterialExpression
 	virtual uint32 GetInputType(int32 InputIndex) override {return MCT_Unknown;}
 	virtual uint32 GetOutputType(int32 InputIndex) override {return MCT_Unknown;}
 	virtual bool GenerateHLSLExpression(FMaterialHLSLGenerator& Generator, UE::HLSLTree::FScope& Scope, int32 OutputIndex, UE::HLSLTree::FExpression const*& OutExpression) const override;
+
+	virtual bool IsResultStrataMaterial(int32 OutputIndex) override;
+	virtual void GatherStrataMaterialInfo(FStrataMaterialInfo& StrataMaterialInfo, int32 OutputIndex) override;
+	virtual FStrataOperator* StrataGenerateMaterialTopologyTree(class FMaterialCompiler* Compiler, class UMaterialExpression* Parent, int32 OutputIndex) override;
 #endif // WITH_EDITOR
 	//~ End UMaterialExpression Interface
+
+#if WITH_EDITOR
+protected:
+	FExpressionInput* GetEffectiveInput(class FMaterialCompiler* Compiler);
+#endif
 };
diff --git a/Engine/Source/Runtime/Engine/Private/Materials/MaterialExpressions.cpp b/Engine/Source/Runtime/Engine/Private/Materials/MaterialExpressions.cpp
--- a/Engine/Source/Runtime/Engine/Private/Materials/MaterialExpressions.cpp
+++ b/Engine/Source/Runtime/Engine/Private/Materials/MaterialExpressions.cpp
@@ -8453,7 +8453,7 @@ FStrataOperator* UMaterialExpressionStaticSwitchParameter::StrataGenerateMateria
 	FExpressionInput* EffectiveInput = GetEffectiveInput(Compiler);
 	if (EffectiveInput && EffectiveInput->Expression)
 	{
-		return EffectiveInput->Expression->StrataGenerateMaterialTopologyTree(Compiler, Parent, 0);
+		return EffectiveInput->Expression->StrataGenerateMaterialTopologyTree(Compiler, Parent, EffectiveInput->OutputIndex);
 	}
 	return nullptr;
 }
@@ -8691,7 +8691,7 @@ FStrataOperator* UMaterialExpressionStaticSwitch::StrataGenerateMaterialTopology
 	FExpressionInput* EffectiveInput = GetEffectiveInput(Compiler);
 	if (EffectiveInput && EffectiveInput->Expression)
 	{
-		return EffectiveInput->Expression->StrataGenerateMaterialTopologyTree(Compiler, Parent, 0);
+		return EffectiveInput->Expression->StrataGenerateMaterialTopologyTree(Compiler, Parent, EffectiveInput->OutputIndex);
 	}
 	return nullptr;
 }
@@ -8796,27 +8796,32 @@ UMaterialExpressionQualitySwitch::UMaterialExpressionQualitySwitch(const FObject
 }
 
 #if WITH_EDITOR
-int32 UMaterialExpressionQualitySwitch::Compile(class FMaterialCompiler* Compiler, int32 OutputIndex)
+FExpressionInput* UMaterialExpressionQualitySwitch::GetEffectiveInput(class FMaterialCompiler* Compiler)
 {
 	const EMaterialQualityLevel::Type QualityLevelToCompile = Compiler->GetQualityLevel();
-	FExpressionInput DefaultTraced = Default.GetTracedInput();
-
-	if (!DefaultTraced.Expression)
-	{
-		return Compiler->Errorf(TEXT("Quality switch missing default input"));
-	}
-
 	if (QualityLevelToCompile != EMaterialQualityLevel::Num)
 	{
 		check(QualityLevelToCompile < UE_ARRAY_COUNT(Inputs));
-		FExpressionInput QualityInput = Inputs[QualityLevelToCompile].GetTracedInput();
-		if (QualityInput.Expression)
+		FExpressionInput QualityInputTraced = Inputs[QualityLevelToCompile].GetTracedInput();
+		if (QualityInputTraced.Expression)
 		{
-			return QualityInput.Compile(Compiler);
+			return &Inputs[QualityLevelToCompile];
 		}
 	}
+	return &Default;
+}
 
-	return DefaultTraced.Compile(Compiler);
+int32 UMaterialExpressionQualitySwitch::Compile(class FMaterialCompiler* Compiler, int32 OutputIndex)
+{
+	FExpressionInput DefaultTraced = Default.GetTracedInput();
+	if (!DefaultTraced.Expression)
+	{
+		return Compiler->Errorf(TEXT("Quality switch missing default input"));
+	}
+
+	FExpressionInput* EffectiveInput = GetEffectiveInput(Compiler);
+	check(EffectiveInput);
+	return EffectiveInput->Compile(Compiler);
 }
 
 void UMaterialExpressionQualitySwitch::GetCaption(TArray<FString>& OutCaptions) const
@@ -8879,6 +8884,47 @@ bool UMaterialExpressionQualitySwitch::IsResultMaterialAttributes(int32 OutputIn
 
 	return false;
 }
+
+bool UMaterialExpressionQualitySwitch::IsResultStrataMaterial(int32 OutputIndex)
+{
+	// Return strata only if all inputs are strata
+	bool bResultStrataMaterial = Default.GetTracedInput().Expression && Default.GetTracedInput().Expression->IsResultStrataMaterial(Default.OutputIndex);
+	for (int i = 0; i < EMaterialQualityLevel::Num; ++i)
+	{
+		bResultStrataMaterial = bResultStrataMaterial && Inputs[i].GetTracedInput().Expression && Inputs[i].GetTracedInput().Expression->IsResultStrataMaterial(Inputs[i].OutputIndex);
+	}
+	return bResultStrataMaterial;
+}
+
+void UMaterialExpressionQualitySwitch::GatherStrataMaterialInfo(FStrataMaterialInfo& StrataMaterialInfo, int32 OutputIndex)
+{
+	// STRATA_TODO: this is incorrect because we should only use a single input based on GetEffectiveInput, but we have no compiler at this stage so we just gather all.
+	if (Default.GetTracedInput().Expression)
+	{
+		Default.GetTracedInput().Expression->GatherStrataMaterialInfo(StrataMaterialInfo, Default.OutputIndex);
+	}
+	for (int i = 0; i < EMaterialQualityLevel::Num; ++i)
+	{
+		Inputs[i].GetTracedInput().Expression->GatherStrataMaterialInfo(StrataMaterialInfo, Inputs[i].OutputIndex);
+	}
+}
+
+FStrataOperator* UMaterialExpressionQualitySwitch::StrataGenerateMaterialTopologyTree(class FMaterialCompiler* Compiler, class UMaterialExpression* Parent, int32 OutputIndex)
+{
+	FExpressionInput DefaultTraced = Default.GetTracedInput();
+	if (!DefaultTraced.Expression)
+	{
+		Compiler->Errorf(TEXT("Quality switch missing default input"));
+		return nullptr;
+	}
+
+	FExpressionInput* EffectiveInput = GetEffectiveInput(Compiler);
+	if (EffectiveInput && EffectiveInput->Expression)
+	{
+		return EffectiveInput->Expression->StrataGenerateMaterialTopologyTree(Compiler, Parent, EffectiveInput->OutputIndex);
+	}
+	return nullptr;
+}
 #endif // WITH_EDITOR
 
 //
@@ -9035,7 +9081,7 @@ UMaterialExpressionShadingPathSwitch::UMaterialExpressionShadingPathSwitch(const
 }
 
 #if WITH_EDITOR
-int32 UMaterialExpressionShadingPathSwitch::Compile(class FMaterialCompiler* Compiler, int32 OutputIndex)
+FExpressionInput* UMaterialExpressionShadingPathSwitch::GetEffectiveInput(class FMaterialCompiler* Compiler)
 {
 	const EShaderPlatform ShaderPlatform = Compiler->GetShaderPlatform();
 	ERHIShadingPath::Type ShadingPathToCompile = ERHIShadingPath::Deferred;
@@ -9051,19 +9097,24 @@ int32 UMaterialExpressionShadingPathSwitch::Compile(class FMaterialCompiler* Com
 
 	check(ShadingPathToCompile < UE_ARRAY_COUNT(Inputs));
 	FExpressionInput ShadingPathInput = Inputs[ShadingPathToCompile].GetTracedInput();
-	FExpressionInput DefaultTraced = Default.GetTracedInput();
-
-	if (!DefaultTraced.Expression)
+	if (ShadingPathInput.Expression)
 	{
-		return Compiler->Errorf(TEXT("Shading path switch missing default input"));
+		return &Inputs[ShadingPathToCompile];
 	}
+	return &Default;
+}
 
-	if (ShadingPathInput.Expression)
+int32 UMaterialExpressionShadingPathSwitch::Compile(class FMaterialCompiler* Compiler, int32 OutputIndex)
+{
+	FExpressionInput DefaultTraced = Default.GetTracedInput();
+	if (!DefaultTraced.Expression)
 	{
-		return ShadingPathInput.Compile(Compiler);
+		return Compiler->Errorf(TEXT("Shading path switch missing default input"));
 	}
 
-	return DefaultTraced.Compile(Compiler);
+	FExpressionInput* EffectiveInput = GetEffectiveInput(Compiler);
+	check(EffectiveInput);
+	return EffectiveInput->Compile(Compiler);
 }
 
 void UMaterialExpressionShadingPathSwitch::GetCaption(TArray<FString>& OutCaptions) const
@@ -9128,6 +9179,47 @@ bool UMaterialExpressionShadingPathSwitch::IsResultMaterialAttributes(int32 Outp
 
 	return false;
 }
+
+bool UMaterialExpressionShadingPathSwitch::IsResultStrataMaterial(int32 OutputIndex)
+{
+	// Return strata only if all inputs are strata
+	bool bResultStrataMaterial = Default.GetTracedInput().Expression && Default.GetTracedInput().Expression->IsResultStrataMaterial(Default.OutputIndex);
+	for (int i = 0; i < ERHIShadingPath::Num; ++i)
+	{
+		bResultStrataMaterial = bResultStrataMaterial && Inputs[i].GetTracedInput().Expression && Inputs[i].GetTracedInput().Expression->IsResultStrataMaterial(Inputs[i].OutputIndex);
+	}
+	return bResultStrataMaterial;
+}
+
+void UMaterialExpressionShadingPathSwitch::GatherStrataMaterialInfo(FStrataMaterialInfo& StrataMaterialInfo, int32 OutputIndex)
+{
+	// STRATA_TODO: this is incorrect because we should only use a single input based on GetEffectiveInput, but we have no compiler at this stage so we just gather all.
+	if (Default.GetTracedInput().Expression)
+	{
+		Default.GetTracedInput().Expression->GatherStrataMaterialInfo(StrataMaterialInfo, Default.OutputIndex);
+	}
+	for (int i = 0; i < EMaterialQualityLevel::Num; ++i)
+	{
+		Inputs[i].GetTracedInput().Expression->GatherStrataMaterialInfo(StrataMaterialInfo, Inputs[i].OutputIndex);
+	}
+}
+
+FStrataOperator* UMaterialExpressionShadingPathSwitch::StrataGenerateMaterialTopologyTree(class FMaterialCompiler* Compiler, class UMaterialExpression* Parent, int32 OutputIndex)
+{
+	FExpressionInput DefaultTraced = Default.GetTracedInput();
+	if (!DefaultTraced.Expression)
+	{
+		Compiler->Errorf(TEXT("Shading path switch missing default input"));
+		return nullptr;
+	}
+
+	FExpressionInput* EffectiveInput = GetEffectiveInput(Compiler);
+	if (EffectiveInput && EffectiveInput->Expression)
+	{
+		return EffectiveInput->Expression->StrataGenerateMaterialTopologyTree(Compiler, Parent, EffectiveInput->OutputIndex);
+	}
+	return nullptr;
+}
 #endif // WITH_EDITOR
 
 //
@@ -15886,6 +15978,21 @@ int32 UMaterialExpressionShaderStageSwitch::Compile(class FMaterialCompiler* Com
 	}
 }
 
+bool UMaterialExpressionShaderStageSwitch::IsResultStrataMaterial(int32 OutputIndex)
+{
+	return false; // Not supported
+}
+
+void UMaterialExpressionShaderStageSwitch::GatherStrataMaterialInfo(FStrataMaterialInfo& StrataMaterialInfo, int32 OutputIndex)
+{
+	// Not supported
+}
+
+FStrataOperator* UMaterialExpressionShaderStageSwitch::StrataGenerateMaterialTopologyTree(class FMaterialCompiler* Compiler, class UMaterialExpression* Parent, int32 OutputIndex)
+{
+	Compiler->Errorf(TEXT("Strata materials are only supported in pixel shaders: ShaderStageSwitch thus should not be plugged to convey Strata material informations."));
+	return nullptr;
+}
 
 void UMaterialExpressionShaderStageSwitch::GetCaption(TArray<FString>& OutCaptions) const
 {
@@ -16076,6 +16183,22 @@ uint32 UMaterialExpressionRayTracingQualitySwitch::GetInputType(int32 InputIndex
 {
 	return MCT_Unknown;
 }
+
+bool UMaterialExpressionRayTracingQualitySwitch::IsResultStrataMaterial(int32 OutputIndex)
+{
+	return false; // Not supported
+}
+
+void UMaterialExpressionRayTracingQualitySwitch::GatherStrataMaterialInfo(FStrataMaterialInfo& StrataMaterialInfo, int32 OutputIndex)
+{
+	// Not supported
+}
+
+FStrataOperator* UMaterialExpressionRayTracingQualitySwitch::StrataGenerateMaterialTopologyTree(class FMaterialCompiler* Compiler, class UMaterialExpression* Parent, int32 OutputIndex)
+{
+	Compiler->Errorf(TEXT("Strata material topology must be statically define. We do not support topology update via dynamic evaluation such as `is raytracing or not`. Only input to BSDFs or Operators can be controled this way."));
+	return nullptr;
+}
 #endif // WITH_EDITOR
 
 //
@@ -16152,6 +16275,22 @@ uint32 UMaterialExpressionPathTracingQualitySwitch::GetInputType(int32 InputInde
 {
 	return MCT_Unknown;
 }
+
+bool UMaterialExpressionPathTracingQualitySwitch::IsResultStrataMaterial(int32 OutputIndex)
+{
+	return false; // Not supported
+}
+
+void UMaterialExpressionPathTracingQualitySwitch::GatherStrataMaterialInfo(FStrataMaterialInfo& StrataMaterialInfo, int32 OutputIndex)
+{
+	// Not supported
+}
+
+FStrataOperator* UMaterialExpressionPathTracingQualitySwitch::StrataGenerateMaterialTopologyTree(class FMaterialCompiler* Compiler, class UMaterialExpression* Parent, int32 OutputIndex)
+{
+	Compiler->Errorf(TEXT("Strata material topology must be statically define. We do not support topology update via dynamic evaluation such as `is pathtracing or not`. Only input to BSDFs or Operators can be controled this way."));
+	return nullptr;
+}
 #endif // WITH_EDITOR
 
 
@@ -20070,6 +20209,22 @@ void UMaterialExpressionReflectionCapturePassSwitch::GetExpressionToolTip(TArray
 {
 	ConvertToMultilineToolTip(TEXT("Allows material to define specialized behavior when being rendered into reflection capture views."), 40, OutToolTip);
 }
+
+bool UMaterialExpressionReflectionCapturePassSwitch::IsResultStrataMaterial(int32 OutputIndex)
+{
+	return false; // Not supported
+}
+
+void UMaterialExpressionReflectionCapturePassSwitch::GatherStrataMaterialInfo(FStrataMaterialInfo& StrataMaterialInfo, int32 OutputIndex)
+{
+	// Not supported
+}
+
+FStrataOperator* UMaterialExpressionReflectionCapturePassSwitch::StrataGenerateMaterialTopologyTree(class FMaterialCompiler* Compiler, class UMaterialExpression* Parent, int32 OutputIndex)
+{
+	Compiler->Errorf(TEXT("Strata material topology must be statically define. We do not support topology update via dynamic evaluation such as `is reflection or not`. Only input to BSDFs or Operators can be controled this way."));
+	return nullptr;
+}
 #endif // WITH_EDITOR
 
 
