Add Override properties for the Player Input and Input Component classes on the Pawn and Player controller. This will let you override the input classes at the blueprint level so that you don't need native C++ to override it. This is useful for asset packs because they can't have C++ in them. This is also going to make it easier for people to upgrade their project to Enhanced Input, if they would like to only do part of a project at a time.

diff --git a/Engine/Source/Runtime/Engine/Classes/GameFramework/Pawn.h b/Engine/Source/Runtime/Engine/Classes/GameFramework/Pawn.h
--- a/Engine/Source/Runtime/Engine/Classes/GameFramework/Pawn.h
+++ b/Engine/Source/Runtime/Engine/Classes/GameFramework/Pawn.h
@@ -512,6 +512,9 @@ public:
 	UFUNCTION(BlueprintCallable, Category="Pawn|Input")
 	virtual bool IsMoveInputIgnored() const;
 
+	UFUNCTION(BlueprintCallable, Category = "Pawn|Input")
+	TSubclassOf<UInputComponent> GetOverrideInputComponentClass() const;
+
 protected:
 	/**
 	 * Accumulated control input vector, stored in world space. This is the pending input, which is cleared (zeroed) once consumed.
@@ -527,6 +530,10 @@ protected:
 	UPROPERTY(Transient)
 	FVector LastControlInputVector;
 
+	/** If set, then this InputComponent class will be used instead of the Input Settings' DefaultInputComponentClass */
+	UPROPERTY(EditDefaultsOnly, Category = "Pawn|Input")
+	TSubclassOf<UInputComponent> OverrideInputComponentClass = nullptr;
+
 public:
 	/** Internal function meant for use only within Pawn or by a PawnMovementComponent. Adds movement input if not ignored, or if forced. */
 	void Internal_AddMovementInput(FVector WorldAccel, bool bForce = false);
diff --git a/Engine/Source/Runtime/Engine/Classes/GameFramework/PlayerController.h b/Engine/Source/Runtime/Engine/Classes/GameFramework/PlayerController.h
--- a/Engine/Source/Runtime/Engine/Classes/GameFramework/PlayerController.h
+++ b/Engine/Source/Runtime/Engine/Classes/GameFramework/PlayerController.h
@@ -1613,6 +1613,10 @@ protected:
 	UPROPERTY()
 	TObjectPtr<class UTouchInterface> CurrentTouchInterface;
 
+	/** If set, then this UPlayerInput class will be used instead of the Input Settings' DefaultPlayerInputClass */
+	UPROPERTY(EditDefaultsOnly, Category = Input)
+	TSubclassOf<UPlayerInput> OverridePlayerInputClass;
+
 	/** Handle for efficient management of UnFreeze timer */
 	FTimerHandle TimerHandle_UnFreeze;
 
@@ -1633,6 +1637,9 @@ public:
 	/** Flushes the current key state. */
 	virtual void FlushPressedKeys();
 
+	UFUNCTION(BlueprintCallable, Category = Input)
+	TSubclassOf<UPlayerInput> GetOverridePlayerInputClass() const;
+
 	/** Handles a key press */
 	UE_DEPRECATED(5.0, "This version of InputKey has been deprecated, please use the version that takes FInputKeyParams instead")
 	virtual bool InputKey(FKey Key, EInputEvent EventType, float AmountDepressed, bool bGamepad);
diff --git a/Engine/Source/Runtime/Engine/Private/Pawn.cpp b/Engine/Source/Runtime/Engine/Private/Pawn.cpp
--- a/Engine/Source/Runtime/Engine/Private/Pawn.cpp
+++ b/Engine/Source/Runtime/Engine/Private/Pawn.cpp
@@ -710,7 +710,8 @@ class UPlayer* APawn::GetNetOwningPlayer()
 UInputComponent* APawn::CreatePlayerInputComponent()
 {
 	static const FName InputComponentName(TEXT("PawnInputComponent0"));
-	return NewObject<UInputComponent>(this, UInputSettings::GetDefaultInputComponentClass(), InputComponentName);
+	const UClass* OverrideClass = OverrideInputComponentClass.Get();
+	return NewObject<UInputComponent>(this, OverrideClass ? OverrideClass : UInputSettings::GetDefaultInputComponentClass(), InputComponentName);
 }
 
 void APawn::DestroyPlayerInputComponent()
@@ -728,6 +729,10 @@ bool APawn::IsMoveInputIgnored() const
 	return Controller != nullptr && Controller->IsMoveInputIgnored();
 }
 
+TSubclassOf<UInputComponent> APawn::GetOverrideInputComponentClass() const
+{
+	return OverrideInputComponentClass;	
+}
 
 void APawn::AddMovementInput(FVector WorldDirection, float ScaleValue, bool bForce /*=false*/)
 {
diff --git a/Engine/Source/Runtime/Engine/Private/PlayerController.cpp b/Engine/Source/Runtime/Engine/Private/PlayerController.cpp
--- a/Engine/Source/Runtime/Engine/Private/PlayerController.cpp
+++ b/Engine/Source/Runtime/Engine/Private/PlayerController.cpp
@@ -641,9 +641,11 @@ void APlayerController::SmoothTargetViewRotation(APawn* TargetPawn, float DeltaS
 
 void APlayerController::InitInputSystem()
 {
-	if (PlayerInput == NULL)
+	if (PlayerInput == nullptr)
 	{
-		PlayerInput = NewObject<UPlayerInput>(this, UInputSettings::GetDefaultPlayerInputClass());
+		const UClass* OverrideClass = OverridePlayerInputClass.Get();
+		
+		PlayerInput = NewObject<UPlayerInput>(this, OverrideClass ? OverrideClass : UInputSettings::GetDefaultPlayerInputClass());
 	}
 
 	SetupInputComponent();
@@ -2282,6 +2284,11 @@ void APlayerController::FlushPressedKeys()
 	}
 }
 
+TSubclassOf<UPlayerInput> APlayerController::GetOverridePlayerInputClass() const
+{
+	return OverridePlayerInputClass;
+}
+
 bool APlayerController::InputKey(FKey Key, EInputEvent EventType, float AmountDepressed, bool bGamepad)
 {
 	FInputKeyParams Params;
