Fix for ComputeSprungMasses ensure. Check for slightly smaller epsilon. Check for aligned non-zero axis for single axis computation in addition to aligned zero axis.

diff --git a/Engine/Source/Runtime/Engine/Private/PhysicsEngine/SimpleSuspension.cpp b/Engine/Source/Runtime/Engine/Private/PhysicsEngine/SimpleSuspension.cpp
--- a/Engine/Source/Runtime/Engine/Private/PhysicsEngine/SimpleSuspension.cpp
+++ b/Engine/Source/Runtime/Engine/Private/PhysicsEngine/SimpleSuspension.cpp
@@ -3,6 +3,13 @@
 #include "Physics/SimpleSuspension.h"
 #include "Chaos/Real.h"
 
+
+// Some calculations are expected to exceed the engine's SMALL_NUMBER threshold
+static double SUSPENSION_SMALL_NUMBER = 1.e-10;
+
+// Tolerance for using single axis calculations
+static double SUSPENSION_ALIGNMENT_TOLERANCE = 0.1;
+
 void FSimpleSuspension::Setup(const FSimpleSuspensionParams& InSuspensionParams)
 {
 	const int32 Count = SuspensionParams.SpringParams.Num();
@@ -43,7 +50,7 @@ bool FSimpleSuspensionHelpers::ComputeSingleAxisLambda(const FVector::FReal Axis
 
 	//compute determinant
 	const FReal DetLL = AxisDot * Count - SumAxis * SumAxis;
-	if (!ensureMsgf(DetLL > SMALL_NUMBER || DetLL < -SMALL_NUMBER,
+	if (!ensureMsgf(!FMath::IsNearlyZero(DetLL, SUSPENSION_SMALL_NUMBER),
 		TEXT("Spring configuration is invalid! Please make sure no two springs are at the same location.")))
 	{
 		return false;
@@ -233,6 +240,9 @@ bool FSimpleSuspensionHelpers::ComputeSprungMasses(const TArray<FVector>& MassSp
 	FReal B0 = 0.f;
 	FReal B1 = 0.f;
 	FReal B2 = 0.f;
+	bool bAlignedX = true;
+	bool bAlignedY = true;
+
 	for (uint32 Index = 0; Index < Count; ++Index)
 	{
 		const FReal X = MassSpringPositions[Index].X;
@@ -242,6 +252,9 @@ bool FSimpleSuspensionHelpers::ComputeSprungMasses(const TArray<FVector>& MassSp
 		XDotX += X * X;
 		YDotY += Y * Y;
 		XDotY += X * Y;
+
+		bAlignedX &= FMath::IsNearlyEqual(MassSpringPositions[Index].X, MassSpringPositions[0].X, SUSPENSION_ALIGNMENT_TOLERANCE);
+		bAlignedY &= FMath::IsNearlyEqual(MassSpringPositions[Index].Y, MassSpringPositions[0].Y, SUSPENSION_ALIGNMENT_TOLERANCE);
 	}
 
 	//calculate the lambdas - we approximate the center of mass as zero for each axis
@@ -253,8 +266,8 @@ bool FSimpleSuspensionHelpers::ComputeSprungMasses(const TArray<FVector>& MassSp
 	const FReal LambdaB1 = 2.f * AverageMass * SumY;
 	const FReal LambdaB2 = (2.f * AverageMass * CountN) - (2.f * TotalMass);
 
-	//if one axis is zero, we actually lose a constraint/equation and we need to adjust our calculation
-	if (YDotY < SMALL_NUMBER  && YDotY > -SMALL_NUMBER)
+	//if one axis is aligned, we actually lose a constraint/equation and we need to adjust our calculation
+	if (bAlignedY)
 	{
 		TArray<Chaos::FReal, TFixedAllocator<2>> Lambdas;
 		Lambdas.Add(LambdaB0);
@@ -266,7 +279,7 @@ bool FSimpleSuspensionHelpers::ComputeSprungMasses(const TArray<FVector>& MassSp
 		Lambda2 = Lambdas[1];
 	}
 
-	else if (XDotX < SMALL_NUMBER && XDotX > -SMALL_NUMBER)
+	else if (bAlignedX)
 	{
 		TArray<Chaos::FReal, TFixedAllocator<2>> Lambdas;
 		Lambdas.Add(LambdaB1);
@@ -290,7 +303,7 @@ bool FSimpleSuspensionHelpers::ComputeSprungMasses(const TArray<FVector>& MassSp
 			- (XDotY * XDotY * Count);
 
 		// Make sure the matrix is invertible!
-		if (!ensureMsgf(DetLL > SMALL_NUMBER || DetLL < -SMALL_NUMBER,
+		if (!ensureMsgf(!FMath::IsNearlyZero(DetLL, SUSPENSION_SMALL_NUMBER),
 			TEXT("Spring configuration is invalid! Please make sure no two springs are at the same location.")))
 		{
 			return false;
