Added async physics input component to easily send inputs from GT to async physics. Also allows client and server to synchronize inputs so the execute on the same physics step for both client and server

diff --git a/Engine/Source/Runtime/Engine/Classes/GameFramework/PlayerController.h b/Engine/Source/Runtime/Engine/Classes/GameFramework/PlayerController.h
--- a/Engine/Source/Runtime/Engine/Classes/GameFramework/PlayerController.h
+++ b/Engine/Source/Runtime/Engine/Classes/GameFramework/PlayerController.h
@@ -2110,7 +2110,7 @@ public:
 	{
 		int32 HeadFrame() const { return LastWritten; }
 		int32 TailFrame() const { return FMath::Max(0, LastWritten - Buffer.Num() + 1); }
-		TArray<uint8>& Write(int32 Frame) { LastWritten = Frame; return Buffer[Frame % Buffer.Num()]; }
+		TArray<uint8>& Write(int32 Frame) { LastWritten = FMath::Max(Frame, LastWritten); return Buffer[Frame % Buffer.Num()]; }
 		const TArray<uint8>& Get(int32 Frame) const { return Buffer[Frame % Buffer.Num()]; }
 
 	private:
diff --git a/Engine/Source/Runtime/Engine/Private/PhysicsEngine/AsyncPhysicsInputComponent.cpp b/Engine/Source/Runtime/Engine/Private/PhysicsEngine/AsyncPhysicsInputComponent.cpp
new file mode 100644
--- /dev/null
+++ b/Engine/Source/Runtime/Engine/Private/PhysicsEngine/AsyncPhysicsInputComponent.cpp
@@ -0,0 +1,491 @@
+// Copyright Epic Games, Inc. All Rights Reserved.
+
+#include "Physics/AsyncPhysicsInputComponent.h"
+
+#include "Chaos/SimCallbackObject.h"
+#include "ChaosSolversModule.h"
+#include "Components/PrimitiveComponent.h"
+#include "Engine/World.h"
+#include "EngineUtils.h"
+#include "Experimental/Chaos/Public/RewindData.h"
+#include "GameFramework/Actor.h"
+#include "GameFramework/Pawn.h"
+#include "GameFramework/PlayerController.h"
+#include "HAL/FileManager.h"
+#include "Net/UnrealNetwork.h"
+#include "PBDRigidsSolver.h"
+#include "Physics/Experimental/PhysScene_Chaos.h"
+#include "PhysicsProxy/SingleParticlePhysicsProxy.h"
+#include "RewindData.h"
+#include "UObject/ObjectKey.h"
+#include "UObject/UObjectIterator.h"
+
+namespace InputCmdCVars
+{
+	static int32 ForceFault = 0;
+	static FAutoConsoleVariableRef CVarForceFault(TEXT("p.net.ForceFault"), ForceFault, TEXT("Forces server side input fault"));
+
+	static int32 MaxBufferedCmds = 16;
+	static FAutoConsoleVariableRef CVarMaxBufferedCmds(TEXT("p.net.MaxBufferedCmds"), MaxBufferedCmds, TEXT("MaxNumber of buffered server side commands"));
+
+	static int32 TimeDilationEnabled = 1;
+	static FAutoConsoleVariableRef CVarTimeDilationEnabled(TEXT("p.net.TimeDilationEnabled"), TimeDilationEnabled, TEXT("Enable clientside TimeDilation"));
+
+	static float MaxTargetNumBufferedCmds = 5.0;
+	static FAutoConsoleVariableRef CVarMaxTargetNumBufferedCmds(TEXT("p.net.MaxTargetNumBufferedCmds"), MaxTargetNumBufferedCmds, TEXT("Maximum number of buffered inputs the server will target per client."));
+
+	static float MaxTimeDilationMag = 0.01f;
+	static FAutoConsoleVariableRef CVarMaxTimeDilationMag(TEXT("p.net.MaxTimeDilationMag"), MaxTimeDilationMag, TEXT("Maximum time dilation that client will use to slow down / catch up with server"));
+
+	static float TimeDilationAlpha = 0.1f;
+	static FAutoConsoleVariableRef CVarTimeDilationAlpha(TEXT("p.net.TimeDilationAlpha"), TimeDilationAlpha, TEXT("Lerp strength for sliding client time dilation"));
+
+	static float TargetNumBufferedCmdsDeltaOnFault = 1.0f;
+	static FAutoConsoleVariableRef CVarTargetNumBufferedCmdsDeltaOnFault(TEXT("p.net.TargetNumBufferedCmdsDeltaOnFault"), TargetNumBufferedCmdsDeltaOnFault, TEXT("How much to increase TargetNumBufferedCmds when an input fault occurs"));
+
+	static float TargetNumBufferedCmds = 1.9f;
+	static FAutoConsoleVariableRef CVarTargetNumBufferedCmds(TEXT("p.net.TargetNumBufferedCmds"), TargetNumBufferedCmds, TEXT("How much to increase TargetNumBufferedCmds when an input fault occurs"));
+
+	static float TargetNumBufferedCmdsAlpha = 0.005f;
+	static FAutoConsoleVariableRef CVarTargetNumBufferedCmdsAlpha(TEXT("p.net.TargetNumBufferedCmdsAlpha"), TargetNumBufferedCmdsAlpha, TEXT("Lerp strength for TargetNumBufferedCmds"));
+
+	static int32 LerpTargetNumBufferedCmdsAggresively = 0;
+	static FAutoConsoleVariableRef CVarLerpTargetNumBufferedCmdsAggresively(TEXT("p.net.LerpTargetNumBufferedCmdsAggresively"), LerpTargetNumBufferedCmdsAggresively, TEXT("Aggresively lerp towards TargetNumBufferedCmds. Reduces server side buffering but can cause more artifacts."));
+}
+
+
+// --------------------------------------------------------------------------------------------------------------------------------------------------
+//	Client InputCmd Stream stuff
+// --------------------------------------------------------------------------------------------------------------------------------------------------
+
+
+int8 QuantizeTimeDilation(float F)
+{
+	if (F == 1.f)
+	{
+		return 0;
+	}
+
+	float Normalized = FMath::Clamp<float>((F - 1.f) / InputCmdCVars::MaxTimeDilationMag, -1.f, 1.f);
+	return (int8)(Normalized * 128.f);
+}
+
+float DeQuantizeTimeDilation(int8 i)
+{
+	if (i == 0)
+	{
+		return 1.f;
+	}
+
+	float Normalized = (float)i / 128.f;
+	float Uncompressed = 1.f + (Normalized * InputCmdCVars::MaxTimeDilationMag);
+	return Uncompressed;
+}
+
+
+DECLARE_MULTICAST_DELEGATE_ThreeParams(FOnDispatchPhysicsTick, int32 PhysicsStep, int32 NumSteps, int32 ServerFrame);
+
+struct FAsyncPhysicsInputRewindCallback : public Chaos::IRewindCallback
+{
+	FAsyncPhysicsInputRewindCallback(UWorld* InWorld) : World(InWorld) { }
+	UWorld* World = nullptr;
+
+	FOnDispatchPhysicsTick DispatchPhysicsTick;
+
+	int32 CachedServerFrame = 0;
+
+	virtual void InjectInputs_External(int32 PhysicsStep, int32 NumSteps) override
+	{
+		int32 LocalOffset = 0;
+		if (World->GetNetMode() == NM_Client)
+		{
+			if (APlayerController* PC = World->GetFirstPlayerController())
+			{
+				APlayerController::FClientFrameInfo& ClientFrameInfo = PC->GetClientFrameInfo();
+
+				if (ClientFrameInfo.LastProcessedInputFrame != INDEX_NONE)
+				{
+					LocalOffset = ClientFrameInfo.GetLocalFrameOffset();
+				}
+			}
+		}
+		DispatchPhysicsTick.Broadcast(PhysicsStep, NumSteps, PhysicsStep - LocalOffset);
+	}
+
+	virtual void ProcessInputs_External(int32 PhysicsStep, const TArray<Chaos::FSimCallbackInputAndObject>& SimCallbackInputs)
+	{
+		CachedServerFrame = PhysicsStep;
+
+		if (World->GetNetMode() == NM_Client)
+		{
+			int32 LocalOffset = 0;
+			if (APlayerController* PC = World->GetFirstPlayerController())
+			{
+				// ------------------------------------------------
+				// Send RPC to server telling them what (client/local) physics step we are running
+				//	* Note that SendData is empty because of the existing API, should change this
+				// ------------------------------------------------	
+
+				TArray<uint8> SendData;
+				PC->PushClientInput(PhysicsStep, SendData);
+
+				// -----------------------------------------------------------------
+				// Calculate latest frame offset to map server frame to local frame
+				// -----------------------------------------------------------------
+				APlayerController::FClientFrameInfo& ClientFrameInfo = PC->GetClientFrameInfo();
+
+				if (ClientFrameInfo.LastProcessedInputFrame != INDEX_NONE)
+				{
+					const int32 FrameOffset = ClientFrameInfo.GetLocalFrameOffset();
+					
+					CachedServerFrame = PhysicsStep - FrameOffset; // Local = Server + Offset
+				}
+
+				// -----------------------------------------------------------------
+				// Apply local TIme Dilation based on server's recommendation.
+				// This speeds up or slows down our consumption of real time (by like < 1%)
+				// Ultimately this causes us to send InputCmds at a lower or higher rate in
+				// order to keep server side buffer at optimal capacity. 
+				// Optimal capacity = as small as possible without ever "missing" a frame (e.g, minimal buffer yet always a new fresh cmd to consume server side)
+				// -----------------------------------------------------------------
+				const float RealTimeDilation = DeQuantizeTimeDilation(ClientFrameInfo.QuantizedTimeDilation);
+
+				if (InputCmdCVars::TimeDilationEnabled > 0)
+				{
+					FPhysScene* PhysScene = World->GetPhysicsScene();
+					PhysScene->SetNetworkDeltaTimeScale(RealTimeDilation);
+				}
+			}
+		}
+		else
+		{
+			// -----------------------------------------------
+			// Server: "consume" an InputCmd from each Player Controller
+			// All this means in this context is updating FServerFrameInfo:: LastProcessedInputFrame, LastLocalFrame
+			// (E.g, telling each client what "Input" of theirs we were processing and our local physics frame number.
+			// In cases where the buffer has a fault, we calculate a suggested time dilation to temporarily make client speed up 
+			// or slow down their input cmd production.
+			// -----------------------------------------------
+
+			const bool bForceFault = InputCmdCVars::ForceFault > 0;
+			InputCmdCVars::ForceFault = FMath::Max(0, InputCmdCVars::ForceFault - 1);
+
+			for (FConstPlayerControllerIterator Iterator = World->GetPlayerControllerIterator(); Iterator; ++Iterator)
+			{
+				if (APlayerController* PC = Iterator->Get())
+				{
+					APlayerController::FServerFrameInfo& FrameInfo = PC->GetServerFrameInfo();
+					APlayerController::FInputCmdBuffer& InputBuffer = PC->GetInputBuffer();
+
+					{
+						const int32 NumBufferedInputCmds = bForceFault ? 0 : (InputBuffer.HeadFrame() - FrameInfo.LastProcessedInputFrame);
+
+						// Check Overflow
+						if (NumBufferedInputCmds > InputCmdCVars::MaxBufferedCmds)
+						{
+							UE_LOG(LogPhysics, Warning, TEXT("[Remote.Input] overflow %d %d -> %d"), InputBuffer.HeadFrame(), FrameInfo.LastProcessedInputFrame, NumBufferedInputCmds);
+							FrameInfo.LastProcessedInputFrame = InputBuffer.HeadFrame() - InputCmdCVars::MaxBufferedCmds + 1;
+						}
+
+						// Check fault - we are waiting for Cmds to reach TargetNumBufferedCmds before continuing
+						if (FrameInfo.bFault)
+						{
+							if (NumBufferedInputCmds < (int32)FrameInfo.TargetNumBufferedCmds)
+							{
+								// Skip this because it is in fault. We will use the prev input for this frame.
+								UE_CLOG(FrameInfo.LastProcessedInputFrame != INDEX_NONE, LogPhysics, Warning, TEXT("[Remote.Input] in fault. Reusing Inputcmd. (Client) Input: %d. (Server) Local Frame: %d"), FrameInfo.LastProcessedInputFrame, FrameInfo.LastLocalFrame);
+								return;
+							}
+							FrameInfo.bFault = false;
+						}
+						else if (NumBufferedInputCmds <= 0)
+						{
+							// No Cmds to process, enter fault state. Increment TargetNumBufferedCmds each time this happens.
+							// TODO: We should have something to bring this back down (which means skipping frames) we don't want temporary poor conditions to cause permanent high input buffering
+							FrameInfo.bFault = true;
+							FrameInfo.TargetNumBufferedCmds = FMath::Min(FrameInfo.TargetNumBufferedCmds + InputCmdCVars::TargetNumBufferedCmdsDeltaOnFault, InputCmdCVars::MaxTargetNumBufferedCmds);
+
+							UE_CLOG(FrameInfo.LastProcessedInputFrame != INDEX_NONE, LogPhysics, Warning, TEXT("[Remote.Input] ENTERING fault. New Target: %.2f. (Client) Input: %d. (Server) Local Frame: %d"), FrameInfo.TargetNumBufferedCmds, FrameInfo.LastProcessedInputFrame, FrameInfo.LastLocalFrame);
+							return;
+						}
+
+						float TargetTimeDilation = 1.f;
+						if (NumBufferedInputCmds < (int32)FrameInfo.TargetNumBufferedCmds)
+						{
+							TargetTimeDilation += InputCmdCVars::MaxTimeDilationMag; // Tell client to speed up, we are starved on cmds
+						}
+						/*else if (NumBufferedInputCmds > (int32)FrameInfo.TargetNumBufferedCmds)
+						{
+							TargetTimeDilation -= InputCmdCVars::MaxTimeDilationMag; // Tell client to slow down, we have too many buffered cmds
+
+							if (InputCmdCVars::LerpTargetNumBufferedCmdsAggresively == 0)
+							{
+								// When non aggressive, only lerp when we are above our limit
+								FrameInfo.TargetNumBufferedCmds = FMath::Lerp(FrameInfo.TargetNumBufferedCmds, InputCmdCVars::TargetNumBufferedCmds, InputCmdCVars::TargetNumBufferedCmdsAlpha);
+							}
+						}*/
+
+						FrameInfo.TargetTimeDilation = FMath::Lerp(FrameInfo.TargetTimeDilation, TargetTimeDilation, InputCmdCVars::TimeDilationAlpha);
+						FrameInfo.QuantizedTimeDilation = QuantizeTimeDilation(TargetTimeDilation);
+
+						if (InputCmdCVars::LerpTargetNumBufferedCmdsAggresively != 0)
+						{
+							// When aggressive, always lerp towards target
+							FrameInfo.TargetNumBufferedCmds = FMath::Lerp(FrameInfo.TargetNumBufferedCmds, InputCmdCVars::TargetNumBufferedCmds, InputCmdCVars::TargetNumBufferedCmdsAlpha);
+						}
+
+						FrameInfo.LastProcessedInputFrame++;
+						FrameInfo.LastLocalFrame = PhysicsStep;
+					}
+				}
+			}
+		}
+	}
+
+
+	// Updates the TMap on PhysScene that stores (non interpolated) physics data for replication.
+	// 
+	// Needs to be called from PT context to access fixed tick handle
+	// but also needs to be able to access GT data (actor iterator, actor state)
+	int32 LastUpdatedStep = 0;
+	void UpdateReplicationMap_Internal(int32 PhysicsStep)
+	{
+		if (LastUpdatedStep == PhysicsStep)
+		{
+			return;
+		}
+
+		LastUpdatedStep = PhysicsStep;
+
+		// Go through all "managed" primitive components and update Tmap to have latest physics state
+		// This is a temp hack and we will eventually replace this with something internal to the physics system that can be updated as things move
+		if (FPhysScene_Chaos* Scene = static_cast<FPhysScene_Chaos*>(World->GetPhysicsScene()))
+		{
+			Scene->ReplicationCache.ServerFrame = PhysicsStep;
+			for (TActorIterator<AActor> It(World); It; ++It)
+			{
+				AActor* Actor = *It;
+				if (Actor && Actor->GetIsReplicated() && Actor->IsReplicatingMovement())
+				{
+					UPrimitiveComponent* RootComponent = Cast<UPrimitiveComponent>(Actor->GetRootComponent());
+					if (RootComponent && RootComponent->IsSimulatingPhysics())
+					{
+						if (FBodyInstanceAsyncPhysicsTickHandle Handle = RootComponent->GetBodyInstanceAsyncPhysicsTickHandle())
+						{
+							FObjectKey Key(RootComponent);
+							FRigidBodyState& LatestState = Scene->ReplicationCache.Map.FindOrAdd(Key);
+
+							// This might be wrong... see FBodyInstance::GetRigidBodyState (converts to unreal units?) and FRepMovement::FillFrom
+							LatestState.Position = Handle->X();
+							LatestState.Quaternion = Handle->R();
+							LatestState.LinVel = Handle->V();
+							LatestState.AngVel = Handle->W();
+							LatestState.Flags = Handle->ObjectState() == Chaos::EObjectStateType::Sleeping ? ERigidBodyFlags::Sleeping : 0;
+						}
+					}
+				}
+			}
+		}
+	}
+};
+
+
+UAsyncPhysicsInputComponent::UAsyncPhysicsInputComponent()
+{
+	bAsyncPhysicsTickEnabled = true;
+	bWantsInitializeComponent = true;
+	bAutoActivate = true;
+	SetIsReplicatedByDefault(true);
+}
+
+void UAsyncPhysicsInputComponent::AsyncPhysicsTickComponent(float DeltaTime, float SimTime)
+{
+	Super::AsyncPhysicsTickComponent(DeltaTime, SimTime);
+	ensureMsgf(Pool != nullptr, TEXT("You must call RegisterInputPool in InitializeComponent"));
+
+	UWorld* World = GetWorld();
+	FPhysScene* PhysScene = World->GetPhysicsScene();
+	Chaos::FPhysicsSolver* Solver = PhysScene->GetSolver();
+
+	// Pull server frame out in this unsafe way. Better if we can get piped into AsyncPhysicsTickComponent
+	FAsyncPhysicsInputRewindCallback* Callback = static_cast<FAsyncPhysicsInputRewindCallback*>(Solver->GetRewindCallback());
+	const int32 ServerFrame = Callback->CachedServerFrame;
+
+	Pool->SetCurrentInputToAsyncExecute(nullptr);	//set current input to none in case we don't find it
+
+	for (int32 Idx = BufferedInputs.Num() - 1; Idx >= 0; --Idx)
+	{
+		FAsyncPhysicsInput* Input = BufferedInputs[Idx];
+		if (Input->ServerFrame == ServerFrame)
+		{
+			Pool->SetCurrentInputToAsyncExecute(Input);
+		}
+
+		bool bFreeInput = Input->ServerFrame < ServerFrame;	//for most cases once we're passed this frame we can free
+		if (APlayerController* PC = GetPlayerController())
+		{
+			if (PC->IsLocalController())
+			{
+				FAsyncPhysicsInputWrapper Wrapper;
+				Wrapper.Input = Input;
+				Wrapper.OwnerComponent = this;
+				ServerRPCBufferInput(Wrapper);
+
+				//If we are the local player then we need to keep this input around to send redundant RPCs to deal with potential packet loss
+				bFreeInput = --Input->Replicated == 0;
+			}
+		}
+
+		if (bFreeInput)
+		{
+			Pool->FreeInputToPool(Input);
+			BufferedInputs.RemoveAtSwap(Idx);
+		}
+	}
+}
+
+void UAsyncPhysicsInputComponent::ServerRPCBufferInput_Implementation(FAsyncPhysicsInputWrapper Wrapper)
+{
+	for (FAsyncPhysicsInput* Input : BufferedInputs)
+	{
+		if (Input->ServerFrame == Wrapper.Input->ServerFrame)
+		{
+			//already buffered this is just a redundant send to deal with potential packet loss
+			//question: what if client's offset changes so we get a new instruction for what we think is an existing frame? The case where offset changes is input fault (hopefully rare) so maybe we just accept it
+			return;
+		}
+	}
+
+	BufferedInputs.Add(Wrapper.Input);
+}
+
+
+bool FAsyncPhysicsInputWrapper::NetSerialize(FArchive& Ar, class UPackageMap* Map, bool& bOutSuccess)
+{
+	Ar << OwnerComponent;
+	FAsyncPhysicsInputPool* Pool = OwnerComponent->GetInputPool();
+	Pool->NetSerializeHelper(Input, Ar, Map, bOutSuccess);
+	return bOutSuccess;
+}
+
+void UAsyncPhysicsInputComponent::RegisterInputPool(FAsyncPhysicsInputPool* InPool)
+{
+	ensureMsgf(Pool == nullptr, TEXT("You can only register the input pool once during initialization"));
+	Pool = InPool;
+}
+
+void UAsyncPhysicsInputComponent::OnDispatchPhysicsTick(int32 PhysicsStep, int32 NumSteps, int32 ServerFrame)
+{
+	ensureMsgf(Pool != nullptr, TEXT("You must call RegisterInputPool in InitializeComponent"));
+
+	// It would be better if we only registered while we had a local controller,
+	// but this is simpler to implement
+
+	if (APlayerController* PC = GetPlayerController())
+	{
+		if (PC->IsLocalController())
+		{
+			FAsyncPhysicsInput* Input = Pool->FlushLatestInputToPopulate();
+			for (int32 Step = 0; Step < NumSteps; ++Step)
+			{
+				if(Step > 0)
+				{
+					//Make sure each sub-step gets its own unique data so we don't have to worry about ref count
+					//TODO: should probably given user opportunity to modify this per sub-step. For example a jump instruction should only happen on first sub-step
+					Input = Pool->CloneInput(Input);
+				}
+				Input->ServerFrame = ServerFrame + 1 + Step;
+				BufferedInputs.Add(Input);
+			}
+		}
+	}
+}
+
+void UAsyncPhysicsInputComponent::InitializeComponent()
+{
+	Super::InitializeComponent();
+
+	// Test is component is valid for Network Physics. Needs a valid physics ActorHandle
+	auto ValidComponent = [](UActorComponent* Component)
+	{
+		bool bValid = false;
+		if (UPrimitiveComponent* PrimComp = Cast<UPrimitiveComponent>(Component))
+		{
+			bValid = FPhysicsInterface::IsValid(PrimComp->BodyInstance.ActorHandle);
+		}
+		return bValid;
+	};
+
+	auto SelectComponent = [&ValidComponent](const TArray<UActorComponent*>& Components)
+	{
+		UPrimitiveComponent* Pc = nullptr;
+		for (UActorComponent* Ac : Components)
+		{
+			if (ValidComponent(Ac))
+			{
+				Pc = (UPrimitiveComponent*)Ac;
+				break;
+			}
+
+		}
+		return Pc;
+	};
+
+	if (AActor* MyActor = GetOwner())
+	{
+		TInlineComponentArray<UPrimitiveComponent*> FoundComponents;
+		MyActor->GetComponents<UPrimitiveComponent>(FoundComponents);
+		for (UPrimitiveComponent* FoundComponent : FoundComponents)
+		{
+			if (FPhysicsInterface::IsValid(FoundComponent->BodyInstance.ActorHandle))
+			{
+				UpdateComponent = FoundComponent;
+				break;
+			}
+		}
+	}
+
+	//TODO: set up tick dependency so that async input component runs before the update component
+	//ensure(UpdateComponent);
+
+
+
+	UWorld* World = GetWorld();
+	if (World->WorldType == EWorldType::PIE || World->WorldType == EWorldType::Game)
+	{
+		FPhysScene* PhysScene = World->GetPhysicsScene();
+		if (ensureAlways(PhysScene))
+		{
+			Chaos::FPhysicsSolver* Solver = PhysScene->GetSolver();
+			if (Solver->GetRewindCallback() == nullptr)
+			{
+				// -------------------------------------------------------------------------------------
+				// Hack: enable Rewind capture here. This needs to be moved somewhere else
+				// -------------------------------------------------------------------------------------
+
+				int32 NumFrames = 64;
+				const IConsoleVariable* NumFramesCVar = IConsoleManager::Get().FindConsoleVariable(TEXT("p.RewindCaptureNumFrames"));
+				if (ensure(NumFramesCVar))
+				{
+					// 1 frame is required to enable rewind capture. 
+					NumFrames = FMath::Max<int32>(1, NumFramesCVar->GetInt());
+				}
+				Solver->EnableRewindCapture(NumFrames, false, MakeUnique<FAsyncPhysicsInputRewindCallback>(World));
+			}
+
+			FAsyncPhysicsInputRewindCallback* Callback = static_cast<FAsyncPhysicsInputRewindCallback*>(Solver->GetRewindCallback());
+			Callback->DispatchPhysicsTick.AddUObject(this, &UAsyncPhysicsInputComponent::OnDispatchPhysicsTick);
+		}
+	}
+}
+
+APlayerController* UAsyncPhysicsInputComponent::GetPlayerController()
+{
+	APawn* Pawn = Cast<APawn>(GetOwner());
+	if (!Pawn)
+	{
+		return nullptr;
+	}
+
+	return Pawn->GetController<APlayerController>();
+}
\ No newline at end of file
diff --git a/Engine/Source/Runtime/Engine/Private/PlayerController.cpp b/Engine/Source/Runtime/Engine/Private/PlayerController.cpp
--- a/Engine/Source/Runtime/Engine/Private/PlayerController.cpp
+++ b/Engine/Source/Runtime/Engine/Private/PlayerController.cpp
@@ -1671,7 +1671,7 @@ void APlayerController::ServerRecvClientInputFrame_Implementation(int32 InRecvCl
 	for (int32 DroppedFrame = InputBuffer.HeadFrame()+1; DroppedFrame < InRecvClientInputFrame && DroppedFrame > 0; ++DroppedFrame)
 	{
 		UE_LOG(LogPlayerController, Warning, TEXT( "ClientInput Gap in frames (Dropped). %s [%s]. DroppedFrame: %d. RecvFrame: %d. LastProcessInputFrame: %d" ), *GetName(), PlayerState ? *PlayerState->GetPlayerName() : TEXT("???"), DroppedFrame, InRecvClientInputFrame, ServerFrameInfo.LastProcessedInputFrame);
-		ServerFrameInfo.LastProcessedInputFrame++; // Ehh lets try this for now
+		//ServerFrameInfo.LastProcessedInputFrame++; // Ehh lets try this for now
 		InputBuffer.Write(DroppedFrame) = InputBuffer.Get(DroppedFrame-1);
 	}
 
diff --git a/Engine/Source/Runtime/Engine/Public/Physics/AsyncPhysicsInputComponent.h b/Engine/Source/Runtime/Engine/Public/Physics/AsyncPhysicsInputComponent.h
new file mode 100644
--- /dev/null
+++ b/Engine/Source/Runtime/Engine/Public/Physics/AsyncPhysicsInputComponent.h
@@ -0,0 +1,241 @@
+// Copyright Epic Games, Inc. All Rights Reserved.
+
+#pragma once
+#include "Components/ActorComponent.h"
+#include "Engine/NetSerialization.h"
+#include "Engine/NetConnection.h"
+#include "Engine/PackageMapClient.h"
+#include "Net/RepLayout.h"
+
+#include "AsyncPhysicsInputComponent.generated.h"
+
+class UAsyncPhysicsInputComponent;
+
+/** 
+	The base struct for async physics input. Inherit from this to create custom input data for async physics tick.
+	When no input is available (say due to massive latency or packet loss) we fall back on the default constructed input.
+	This means you should set the default values to something equivalent to no input (for example bPlayerWantsToJump should probably default to false)
+*/
+USTRUCT()
+struct FAsyncPhysicsInput
+{
+	GENERATED_BODY()
+public:
+	virtual ~FAsyncPhysicsInput() = default;
+	int32 GetServerFrame() const { return ServerFrame; }
+private:
+	UPROPERTY()
+	int32 ServerFrame = INDEX_NONE;
+
+	int32 Replicated = 4;	//how many times we want to replicate to server. TODO: use cvar
+
+	friend UAsyncPhysicsInputComponent;
+};
+
+/** Helper struct to replicate polymorphic input data */
+USTRUCT()
+struct FAsyncPhysicsInputWrapper
+{
+	GENERATED_BODY()
+
+	FAsyncPhysicsInput* Input = nullptr;
+
+	bool NetSerialize(FArchive& Ar, class UPackageMap* Map, bool& bOutSuccess);
+
+	UPROPERTY()
+	TObjectPtr<UAsyncPhysicsInputComponent> OwnerComponent = nullptr;
+};
+
+template<>
+struct TStructOpsTypeTraits<FAsyncPhysicsInputWrapper> : public TStructOpsTypeTraitsBase2<FAsyncPhysicsInputWrapper>
+{
+	enum
+	{
+		WithNetSerializer = true
+	};
+};
+
+class FAsyncPhysicsInputPool;
+
+/** Base component used to easily send data from variable tick to async physics.
+	Has networking support so server and client execute AsyncPhysicsTickComponent on same data for same step.
+	Inherit from this class to create your own component that you can act on inputs with.
+
+	Extended class must hold a TAsyncPhysicsInputPool<T> where T is your custom input extended from FAsyncPhysicsInput.
+	Must call RegisterInputPool in InitializeComponent
+*/
+UCLASS(BlueprintType)
+class ENGINE_API UAsyncPhysicsInputComponent : public UActorComponent
+{
+	GENERATED_BODY()
+public:
+	UAsyncPhysicsInputComponent();
+	void InitializeComponent() override;
+
+	/** Must be called in InitializeComponent by extended class */
+	void RegisterInputPool(FAsyncPhysicsInputPool* InPool);
+
+	/** Executes physics logic based on the current input (See TNetworkedInputPool::GetCurrentInput)*/
+	virtual void AsyncPhysicsTickComponent(float DeltaTime, float SimTime) override;
+
+	UFUNCTION(Server, unreliable)
+	void ServerRPCBufferInput(FAsyncPhysicsInputWrapper PlayerInput);
+
+	void OnDispatchPhysicsTick(int32 PhysicsStep, int32 NumSteps, int32 ServerFrame);
+
+	FAsyncPhysicsInputPool* GetInputPool() { return Pool; }
+
+protected:
+	UPrimitiveComponent* UpdateComponent = nullptr;
+
+	APlayerController* GetPlayerController();
+
+private:
+	TArray<FAsyncPhysicsInput*> BufferedInputs;
+	FAsyncPhysicsInputPool* Pool = nullptr;
+};
+
+
+class FAsyncPhysicsInputPool
+{
+public:
+	virtual ~FAsyncPhysicsInputPool() = default;
+
+private:
+	friend class UAsyncPhysicsInputComponent;
+	friend FAsyncPhysicsInputWrapper;
+
+	/** Flushes the current input to populate so that new inputs can be populated. Returns the previous input to be saved off */
+	virtual FAsyncPhysicsInput* FlushLatestInputToPopulate() = 0;
+
+	/** Makes a deep copy of the input so that we can treat it as identical input over multiple physics steps from one slow gamethread step */
+	virtual FAsyncPhysicsInput* CloneInput(const FAsyncPhysicsInput* Input) = 0;
+
+	/** Sets the current input for execution during async physics */
+	virtual void SetCurrentInputToAsyncExecute(FAsyncPhysicsInput* Input) = 0;
+
+	/** Frees the input back into the pool*/
+	virtual void FreeInputToPool(FAsyncPhysicsInput* Input) = 0;
+
+	/** Serialize the underlying data. If needed a new entry in the pool will be created (during deserialization) */
+	virtual void NetSerializeHelper(FAsyncPhysicsInput*& Data, FArchive& Ar, class UPackageMap* Map, bool& bOutSuccess) = 0;
+};
+
+/** Helper class to pool inputs and do various operations on like serialization.
+*   T must have FAsyncPhysicsInput as its base */
+template <typename T>
+class TAsyncPhysicsInputPool final : public FAsyncPhysicsInputPool
+{
+public:
+	TAsyncPhysicsInputPool()
+	{
+		PendingInputToPopulate = NewInput();
+	}
+
+	~TAsyncPhysicsInputPool()
+	{
+		delete PendingInputToPopulate;
+		delete CurrentInputToExecute;
+		while (Pool.Num())
+		{
+			T* Input = Pool.Pop();
+			delete Input;
+		}
+	}
+
+	/** Gets the pending input object for populating. This is what gets sent to the server and the async physics tick to execute logic off of.
+		Should not be used in async tick as it may not be the right physics step*/
+	T& GetPendingInputToPopulate() { return *PendingInputToPopulate; }
+	const T& GetPendingInputToPopulate() const { return *PendingInputToPopulate; }
+
+	/** Gets the current input object to execute logic off of. This should be used during async physics tick. If no input is available we use default constructed values*/
+	const T& GetCurrentInput()
+	{
+		static T NoInput;
+		return CurrentInputToExecute ? *CurrentInputToExecute : NoInput;
+	}
+
+private:
+
+	virtual FAsyncPhysicsInput* FlushLatestInputToPopulate() override
+	{
+		//Assumes CurrentInputToPopulate has been saved into the send to server buffer and will be properly freed back into pool
+		FAsyncPhysicsInput* OldInput = PendingInputToPopulate;
+		PendingInputToPopulate = NewInput();
+		return OldInput;
+	}
+
+	virtual FAsyncPhysicsInput* CloneInput(const FAsyncPhysicsInput* Input) override
+	{
+		T* Input2 = NewInput();
+		*Input2 = *(static_cast<const T*>(Input));
+		return Input2;
+	}
+
+	virtual void SetCurrentInputToAsyncExecute(FAsyncPhysicsInput* Input) override
+	{
+		CurrentInputToExecute = static_cast<T*>(Input);
+	}
+
+	virtual void NetSerializeHelper(FAsyncPhysicsInput*& Data, FArchive& Ar, class UPackageMap* Map, bool& bOutSuccess) override
+	{
+		if (Ar.IsLoading())
+		{
+			Data = NewInput();
+		}
+
+		UScriptStruct* ScriptStruct = T::StaticStruct();
+		if (ScriptStruct->StructFlags & STRUCT_NetSerializeNative)
+		{
+			ScriptStruct->GetCppStructOps()->NetSerialize(Ar, Map, bOutSuccess, Data);
+		}
+		else
+		{
+			UNetConnection* Connection = CastChecked<UPackageMapClient>(Map)->GetConnection();
+			UNetDriver* NetDriver = Connection ? Connection->GetDriver() : nullptr;
+			TSharedPtr<FRepLayout> RepLayout = NetDriver ? NetDriver->GetStructRepLayout(ScriptStruct) : nullptr;
+
+			if (RepLayout.IsValid())
+			{
+				if (FBitArchive* BitAr = static_cast<FBitArchive*>(&Ar))
+				{
+					bool bHasUnmapped = false;
+					RepLayout->SerializePropertiesForStruct(ScriptStruct, *BitAr, Map, Data, bHasUnmapped);
+
+					bOutSuccess = true;
+				}
+			}
+		}
+	}
+
+	T* NewInput()
+	{
+		T* Input;
+		if (Pool.Num())
+		{
+			Input = Pool.Pop();
+			new (Input) T();
+		}
+		else
+		{
+			Input = new T();
+		}
+
+		return Input;
+	}
+
+	void FreeInput(T* Input)
+	{
+		Input->~T();
+		Pool.Add(Input);
+	}
+
+	virtual void FreeInputToPool(FAsyncPhysicsInput* Input) override
+	{
+		FreeInput(static_cast<T*>(Input));
+	}
+
+	T* CurrentInputToExecute = nullptr;
+	T* PendingInputToPopulate = nullptr;
+	TArray<T*> Pool;
+};
\ No newline at end of file
