Prevent static init calling SetDefaults from FQualityLevels CTOR while is undefined as it access other globals and logging that may not be setup at this point

diff --git a/Engine/Source/Runtime/Engine/Private/Scalability.cpp b/Engine/Source/Runtime/Engine/Private/Scalability.cpp
--- a/Engine/Source/Runtime/Engine/Private/Scalability.cpp
+++ b/Engine/Source/Runtime/Engine/Private/Scalability.cpp
@@ -159,8 +159,8 @@ static TAutoConsoleVariable<int32> CVarShadingQuality_NumLevels(
 
 namespace Scalability
 {
-static FQualityLevels GScalabilityBackupQualityLevels;
-static FQualityLevels GScalabilityQualityLevelsOverride;
+static FQualityLevels GScalabilityBackupQualityLevels(false);
+static FQualityLevels GScalabilityQualityLevelsOverride(false);
 static bool GScalabilityUsingTemporaryQualityLevels = false;
 
 // Select a the correct quality level for the given benchmark value and thresholds
diff --git a/Engine/Source/Runtime/Engine/Public/Scalability.h b/Engine/Source/Runtime/Engine/Public/Scalability.h
--- a/Engine/Source/Runtime/Engine/Public/Scalability.h
+++ b/Engine/Source/Runtime/Engine/Public/Scalability.h
@@ -41,11 +41,15 @@ namespace Scalability
 		TArray<float> CPUBenchmarkSteps;
 		TArray<float> GPUBenchmarkSteps;
 
-		FQualityLevels()
+		// Allows us to avoid SetDefaults for static init variables, as SetDefaults is not defined to call during static int
+		FQualityLevels(bool bSetDefaults = true)
 			: CPUBenchmarkResults(-1.0f)
 			, GPUBenchmarkResults(-1.0f)
 		{
-			SetDefaults();
+			if (bSetDefaults)
+			{
+				SetDefaults();
+			}
 		}
 		
 		bool operator==(const FQualityLevels& Other ) const
