Missing include added for FPhysScene.

diff --git a/Engine/Source/Runtime/Engine/Public/PhysicsPublic.h b/Engine/Source/Runtime/Engine/Public/PhysicsPublic.h
--- a/Engine/Source/Runtime/Engine/Public/PhysicsPublic.h
+++ b/Engine/Source/Runtime/Engine/Public/PhysicsPublic.h
@@ -13,6 +13,7 @@
 #include "Misc/CoreMisc.h"
 #include "Misc/App.h"
 #include "EngineDefines.h"
+#include "PhysicsInterfaceDeclaresCore.h"
 #include "PhysicsPublicCore.h"
 //#include "StaticMeshResources.h"
 
