Fixed FTextureSource hash computation, which could lead to the same DDC key being used for 2 identical payloads (e.g. all black pixels) but with different sizes (e.g. 512x256 vs 256x512)

diff --git a/Engine/Source/Runtime/Engine/Classes/Engine/Texture.h b/Engine/Source/Runtime/Engine/Classes/Engine/Texture.h
--- a/Engine/Source/Runtime/Engine/Classes/Engine/Texture.h
+++ b/Engine/Source/Runtime/Engine/Classes/Engine/Texture.h
@@ -340,7 +340,7 @@ struct FTextureSource
 
 	/** Trivial accessors. These will only give values for Block0 so may not be correct for UDIM/multi-block textures, use GetBlock() for this case. */
 	FGuid GetPersistentId() const { return BulkData.GetIdentifier(); }
-	FORCEINLINE FGuid GetId() const { return Id; }
+	ENGINE_API FGuid GetId() const;
 	FORCEINLINE int32 GetSizeX() const { return SizeX; }
 	FORCEINLINE int32 GetSizeY() const { return SizeY; }
 	FORCEINLINE int32 GetNumSlices() const { return NumSlices; }
diff --git a/Engine/Source/Runtime/Engine/Private/Texture.cpp b/Engine/Source/Runtime/Engine/Private/Texture.cpp
--- a/Engine/Source/Runtime/Engine/Private/Texture.cpp
+++ b/Engine/Source/Runtime/Engine/Private/Texture.cpp
@@ -40,6 +40,7 @@
 #include "ImageUtils.h"
 
 #if WITH_EDITOR
+#include "DerivedDataBuildVersion.h"
 #include "TextureCompiler.h"
 #include "Misc/ScopeRWLock.h"
 #endif
@@ -1878,7 +1879,7 @@ FIntPoint FTextureSource::GetSizeInBlocks() const
 
 FString FTextureSource::GetIdString() const
 {
-	FString GuidString = Id.ToString();
+	FString GuidString = GetId().ToString();
 	if (bGuidIsHash)
 	{
 		GuidString += TEXT("X");
@@ -2175,6 +2176,34 @@ void FTextureSource::UseHashAsGuid()
 	}
 }
 
+FGuid FTextureSource::GetId() const
+{
+	if (!bGuidIsHash)
+	{
+		return Id;
+	}
+
+	UE::DerivedData::FBuildVersionBuilder IdBuilder;
+	IdBuilder << BaseBlockX;
+	IdBuilder << BaseBlockX;
+	IdBuilder << BaseBlockY;
+	IdBuilder << SizeX;
+	IdBuilder << SizeY;
+	IdBuilder << NumSlices;
+	IdBuilder << NumMips;
+	IdBuilder << NumLayers;
+	IdBuilder << bPNGCompressed;
+	IdBuilder << bLongLatCubemap;
+	IdBuilder << CompressionFormat;
+	IdBuilder << bGuidIsHash;
+	IdBuilder << static_cast<uint8>(Format.GetValue());
+	IdBuilder.Serialize(const_cast<void*>(static_cast<const void*>(LayerFormat.GetData())), LayerFormat.Num());
+	IdBuilder.Serialize(const_cast<void*>(static_cast<const void*>(Blocks.GetData())), Blocks.Num());
+	IdBuilder.Serialize(const_cast<void*>(static_cast<const void*>(BlockDataOffsets.GetData())), BlockDataOffsets.Num());
+	IdBuilder << const_cast<FGuid&>(Id);
+	return IdBuilder.Build();
+}
+
 void FTextureSource::OperateOnLoadedBulkData(TFunctionRef<void(const FSharedBuffer& BulkDataBuffer)> Operation)
 {
 	checkf(LockState == ELockState::None, TEXT("OperateOnLoadedBulkData shouldn't be called in-between LockMip/UnlockMip"));
