Try to load cooked global shaders even in editor builds

diff --git a/Engine/Source/Runtime/Engine/Private/ShaderCompiler/ShaderCompiler.cpp b/Engine/Source/Runtime/Engine/Private/ShaderCompiler/ShaderCompiler.cpp
--- a/Engine/Source/Runtime/Engine/Private/ShaderCompiler/ShaderCompiler.cpp
+++ b/Engine/Source/Runtime/Engine/Private/ShaderCompiler/ShaderCompiler.cpp
@@ -6817,6 +6817,73 @@ bool IsGlobalShaderMapComplete(const TCHAR* TypeNameSubstring)
 	return true;
 }
 
+static bool TryLoadCookedGlobalShaderMap(EShaderPlatform Platform, FScopedSlowTask& SlowTask)
+{
+	SlowTask.EnterProgressFrame(50);
+
+	bool bLoadedFromCacheFile = false;
+
+	// Load from the override global shaders first, this allows us to hot reload in cooked / pak builds
+	TArray<uint8> GlobalShaderData;
+	const bool bAllowOverrideGlobalShaders = !WITH_EDITOR && !UE_BUILD_SHIPPING;
+	if (bAllowOverrideGlobalShaders)
+	{
+		FString OverrideGlobalShaderCacheFilename = GetGlobalShaderCacheOverrideFilename(Platform);
+		FPaths::MakeStandardFilename(OverrideGlobalShaderCacheFilename);
+
+		bool bFileExist = IFileManager::Get().FileExists(*OverrideGlobalShaderCacheFilename);
+
+		if (!bFileExist)
+		{
+			UE_LOG(LogShaders, Display, TEXT("%s doesn't exists"), *OverrideGlobalShaderCacheFilename);
+		}
+		else
+		{
+			bLoadedFromCacheFile = FFileHelper::LoadFileToArray(GlobalShaderData, *OverrideGlobalShaderCacheFilename, FILEREAD_Silent);
+
+			if (bLoadedFromCacheFile)
+			{
+				UE_LOG(LogShaders, Display, TEXT("%s has been loaded successfully"), *OverrideGlobalShaderCacheFilename);
+			}
+			else
+			{
+				UE_LOG(LogShaders, Error, TEXT("%s failed to load"), *OverrideGlobalShaderCacheFilename);
+			}
+		}
+	}
+
+	// is the data already loaded?
+	int64 PreloadedSize = 0;
+	void* PreloadedData = nullptr;
+	if (!bLoadedFromCacheFile)
+	{
+		PreloadedData = GGlobalShaderPreLoadFile.TakeOwnershipOfLoadedData(&PreloadedSize);
+	}
+
+	if (PreloadedData != nullptr)
+	{
+		FLargeMemoryReader MemoryReader((uint8*)PreloadedData, PreloadedSize, ELargeMemoryReaderFlags::TakeOwnership);
+		GGlobalShaderMap[Platform]->LoadFromGlobalArchive(MemoryReader);
+	}
+	else
+	{
+		FString GlobalShaderCacheFilename = FPaths::GetRelativePathToRoot() / GetGlobalShaderCacheFilename(Platform);
+		FPaths::MakeStandardFilename(GlobalShaderCacheFilename);
+		if (!bLoadedFromCacheFile)
+		{
+			bLoadedFromCacheFile = FFileHelper::LoadFileToArray(GlobalShaderData, *GlobalShaderCacheFilename, FILEREAD_Silent);
+		}
+
+		if (bLoadedFromCacheFile)
+		{
+			FMemoryReader MemoryReader(GlobalShaderData);
+			GGlobalShaderMap[Platform]->LoadFromGlobalArchive(MemoryReader);
+		}
+	}
+
+	return bLoadedFromCacheFile;
+}
+
 void CompileGlobalShaderMap(EShaderPlatform Platform, const ITargetPlatform* TargetPlatform, bool bRefreshShaderMap)
 {
 	// No global shaders needed on dedicated server or clients that use NullRHI. Note that cook commandlet needs to have them, even if it is not allowed to render otherwise.
@@ -6857,94 +6924,22 @@ void CompileGlobalShaderMap(EShaderPlatform Platform, const ITargetPlatform* Tar
 
 		GGlobalShaderMap[Platform] = new FGlobalShaderMap(Platform);
 
-		bool bLoadedFromCacheFile = false;
 		bool bShaderMapIsBeingCompiled = false;
 
 		// Try to load the global shaders from a local cache file if it exists
-		// This method is used exclusively with cooked content, since the DDC is not present
-		// If shader compiling is disabled, go down the cooked path
-		if (FPlatformProperties::RequiresCookedData() || !AllowShaderCompiling())
-		{
-			SlowTask.EnterProgressFrame(50);
-
-			// Load from the override global shaders first, this allows us to hot reload in cooked / pak builds
-			TArray<uint8> GlobalShaderData;
-			const bool bAllowOverrideGlobalShaders = !WITH_EDITOR && !UE_BUILD_SHIPPING;
-			if (bAllowOverrideGlobalShaders)
-			{
-				FString OverrideGlobalShaderCacheFilename = GetGlobalShaderCacheOverrideFilename(Platform);
-				FPaths::MakeStandardFilename(OverrideGlobalShaderCacheFilename);
-
-				bool bFileExist = IFileManager::Get().FileExists(*OverrideGlobalShaderCacheFilename);
-
-				if (!bFileExist)
-				{
-					UE_LOG(LogShaders, Display, TEXT("%s doesn't exists"), *OverrideGlobalShaderCacheFilename);
-				}
-				else
-				{
-					bLoadedFromCacheFile = FFileHelper::LoadFileToArray(GlobalShaderData, *OverrideGlobalShaderCacheFilename, FILEREAD_Silent);
-
-					if (bLoadedFromCacheFile)
-					{
-						UE_LOG(LogShaders, Display, TEXT("%s has been loaded successfully"), *OverrideGlobalShaderCacheFilename);
-					}
-					else
-					{
-						UE_LOG(LogShaders, Error, TEXT("%s failed to load"), *OverrideGlobalShaderCacheFilename);
-					}
-				}
-			}
-
-			// is the data already loaded?
-			int64 PreloadedSize = 0;
-			void* PreloadedData = nullptr;
-			if (!bLoadedFromCacheFile)
-			{
-				PreloadedData = GGlobalShaderPreLoadFile.TakeOwnershipOfLoadedData(&PreloadedSize);
-			}
-
-			if (PreloadedData != nullptr)
-			{
-				FLargeMemoryReader MemoryReader((uint8*)PreloadedData, PreloadedSize, ELargeMemoryReaderFlags::TakeOwnership);
-				GGlobalShaderMap[Platform]->LoadFromGlobalArchive(MemoryReader);
-			}
-			else
-			{
-				FString GlobalShaderCacheFilename = FPaths::GetRelativePathToRoot() / GetGlobalShaderCacheFilename(Platform);
-				FPaths::MakeStandardFilename(GlobalShaderCacheFilename);
-				if (!bLoadedFromCacheFile)
-				{
-					bLoadedFromCacheFile = FFileHelper::LoadFileToArray(GlobalShaderData, *GlobalShaderCacheFilename, FILEREAD_Silent);
-				}
-
-				if (!bLoadedFromCacheFile)
-				{
-					// Handle this gracefully and exit.
-					FString SandboxPath = IFileManager::Get().ConvertToAbsolutePathForExternalAppForWrite(*GlobalShaderCacheFilename);
-					// This can be too early to localize in some situations.
-					const FText Message = FText::Format(NSLOCTEXT("Engine", "GlobalShaderCacheFileMissing", "The global shader cache file '{0}' is missing.\n\nYour application is built to load COOKED content. No COOKED content was found; This usually means you did not cook content for this build.\nIt also may indicate missing cooked data for a shader platform(e.g., OpenGL under Windows): Make sure your platform's packaging settings include this Targeted RHI.\n\nAlternatively build and run the UNCOOKED version instead."), FText::FromString(SandboxPath));
-					if (FPlatformProperties::SupportsWindowedMode())
-					{
-						UE_LOG(LogShaders, Error, TEXT("%s"), *Message.ToString());
-						FMessageDialog::Open(EAppMsgType::Ok, Message);
-						FPlatformMisc::RequestExit(false);
-						return;
-					}
-					else
-					{
-						UE_LOG(LogShaders, Fatal, TEXT("%s"), *Message.ToString());
-					}
-				}
+		// We always try this first, even when running in the editor or if shader compiler is enabled
+		// It's always possible we'll find a cooked local cache
+		const bool bLoadedFromCacheFile = TryLoadCookedGlobalShaderMap(Platform, SlowTask);
+#if WITH_EDITOR
+		const bool bAllowShaderCompiling = !FPlatformProperties::RequiresCookedData() && AllowShaderCompiling();
+#else
+		const bool bAllowShaderCompiling = false;
+#endif
 
-				FMemoryReader MemoryReader(GlobalShaderData);
-				GGlobalShaderMap[Platform]->LoadFromGlobalArchive(MemoryReader);
-			}
-		}
 #if WITH_EDITOR
-		// Uncooked platform
-		else
+		if (!bLoadedFromCacheFile && bAllowShaderCompiling)
 		{
+			// If we didn't find cooked shaders, we can try loading from the DDC or compiling them if supported by the current configuration
 			FGlobalShaderMapId ShaderMapId(Platform, TargetPlatform);
 
 			const int32 ShaderFilenameNum = ShaderMapId.GetShaderFilenameToDependeciesMap().Num();
@@ -7020,6 +7015,27 @@ void CompileGlobalShaderMap(EShaderPlatform Platform, const ITargetPlatform* Tar
 			}
 		}
 #endif // WITH_EDITOR
+		
+		if (!bLoadedFromCacheFile && !bAllowShaderCompiling)
+		{
+			// Failed to load cooked shaders, and no support for compiling
+			// Handle this gracefully and exit.
+			const FString GlobalShaderCacheFilename = FPaths::GetRelativePathToRoot() / GetGlobalShaderCacheFilename(Platform);
+			const FString SandboxPath = IFileManager::Get().ConvertToAbsolutePathForExternalAppForWrite(*GlobalShaderCacheFilename);
+			// This can be too early to localize in some situations.
+			const FText Message = FText::Format(NSLOCTEXT("Engine", "GlobalShaderCacheFileMissing", "The global shader cache file '{0}' is missing.\n\nYour application is built to load COOKED content. No COOKED content was found; This usually means you did not cook content for this build.\nIt also may indicate missing cooked data for a shader platform(e.g., OpenGL under Windows): Make sure your platform's packaging settings include this Targeted RHI.\n\nAlternatively build and run the UNCOOKED version instead."), FText::FromString(SandboxPath));
+			if (FPlatformProperties::SupportsWindowedMode())
+			{
+				UE_LOG(LogShaders, Error, TEXT("%s"), *Message.ToString());
+				FMessageDialog::Open(EAppMsgType::Ok, Message);
+				FPlatformMisc::RequestExit(false);
+				return;
+			}
+			else
+			{
+				UE_LOG(LogShaders, Fatal, TEXT("%s"), *Message.ToString());
+			}
+		}
 
 		// If any shaders weren't loaded, compile them now.
 		VerifyGlobalShaders(Platform, TargetPlatform, bLoadedFromCacheFile);
