Address merge issue caused by the deprecation and renaming of bPNGCompressed.  This configuration is now incorporated into CompressionFormat.

diff --git a/Engine/Source/Runtime/Engine/Private/Texture.cpp b/Engine/Source/Runtime/Engine/Private/Texture.cpp
--- a/Engine/Source/Runtime/Engine/Private/Texture.cpp
+++ b/Engine/Source/Runtime/Engine/Private/Texture.cpp
@@ -2192,7 +2192,6 @@ FGuid FTextureSource::GetId() const
 	IdBuilder << NumSlices;
 	IdBuilder << NumMips;
 	IdBuilder << NumLayers;
-	IdBuilder << bPNGCompressed;
 	IdBuilder << bLongLatCubemap;
 	IdBuilder << CompressionFormat;
 	IdBuilder << bGuidIsHash;
