Remove duplicate hashing of BaseBlockX when composing a texture id.

diff --git a/Engine/Source/Runtime/Engine/Private/Texture.cpp b/Engine/Source/Runtime/Engine/Private/Texture.cpp
--- a/Engine/Source/Runtime/Engine/Private/Texture.cpp
+++ b/Engine/Source/Runtime/Engine/Private/Texture.cpp
@@ -2185,7 +2185,6 @@ FGuid FTextureSource::GetId() const
 
 	UE::DerivedData::FBuildVersionBuilder IdBuilder;
 	IdBuilder << BaseBlockX;
-	IdBuilder << BaseBlockX;
 	IdBuilder << BaseBlockY;
 	IdBuilder << SizeX;
 	IdBuilder << SizeY;
