Make sure source doesn't get transitioned to SRVMask after first call to CopyToResolveTarget when source & dest are the same

diff --git a/Engine/Source/Runtime/Engine/Private/UnrealClient.cpp b/Engine/Source/Runtime/Engine/Private/UnrealClient.cpp
--- a/Engine/Source/Runtime/Engine/Private/UnrealClient.cpp
+++ b/Engine/Source/Runtime/Engine/Private/UnrealClient.cpp
@@ -1719,14 +1719,14 @@ const TArray<FColor>& FViewport::GetRawHitProxyData(FIntRect InRect)
 				// Should not be multisampled, so can safely use copy src as state instead of taking care of ResolveSrc here
 				check(!RenderTargetTexture->IsMultisampled());
 				
-				// Keep in copy source because 2 resolve calls are done on the render target (skip one extra transition call)
+				// Don't touch the copy source state after the first call to CopyToResolveTarget
 				FResolveParams ResolveParams;
-				ResolveParams.SourceAccessFinal = ERHIAccess::CopySrc;
+				ResolveParams.SourceAccessFinal = ERHIAccess::Unknown;
 
 				// Copy (resolve) the rendered thumbnail from the render target to its texture
 				RHICmdList.Transition(FRHITransitionInfo(RenderTargetTexture, ERHIAccess::Unknown, ERHIAccess::CopySrc));
-				RHICmdList.CopyToResolveTarget(RenderTargetTexture, HitProxyMapPtr->GetHitProxyTexture(), ResolveParams);
-				RHICmdList.CopyToResolveTarget(RenderTargetTexture, HitProxyMapPtr->GetHitProxyCPUTexture(), FResolveParams());
+				RHICmdList.CopyToResolveTarget(RenderTargetTexture, HitProxyMapPtr->GetHitProxyCPUTexture(), ResolveParams);
+				RHICmdList.CopyToResolveTarget(RenderTargetTexture, HitProxyMapPtr->GetHitProxyTexture(), FResolveParams());
 			});
 
 		ENQUEUE_RENDER_COMMAND(EndDrawingCommand)(
