Fixed more non unity compile errors

diff --git a/Engine/Source/Runtime/Engine/Private/Materials/MaterialHLSLEmitter.cpp b/Engine/Source/Runtime/Engine/Private/Materials/MaterialHLSLEmitter.cpp
--- a/Engine/Source/Runtime/Engine/Private/Materials/MaterialHLSLEmitter.cpp
+++ b/Engine/Source/Runtime/Engine/Private/Materials/MaterialHLSLEmitter.cpp
@@ -18,6 +18,7 @@
 #include "HLSLTree/HLSLTreeCommon.h"
 #include "HLSLTree/HLSLTreeEmit.h"
 #include "Containers/LazyPrintf.h"
+#include "RenderUtils.h"
 
 static bool SharedPixelProperties[CompiledMP_MAX];
 
