Draw net correction float history in PhysicsReplication.

diff --git a/Engine/Source/Runtime/Engine/Private/PhysicsEngine/PhysicsReplication.cpp b/Engine/Source/Runtime/Engine/Private/PhysicsEngine/PhysicsReplication.cpp
--- a/Engine/Source/Runtime/Engine/Private/PhysicsEngine/PhysicsReplication.cpp
+++ b/Engine/Source/Runtime/Engine/Private/PhysicsEngine/PhysicsReplication.cpp
@@ -445,10 +445,7 @@ bool FPhysicsReplication::ApplyRigidBodyState(float DeltaSeconds, FBodyInstance*
 			{
 				FColor Color = FColor::White;
 				DrawDebugDirectionalArrow(OwningWorld, CurrentState.Position, TargetPos, 5.0f, Color, true, CharacterMovementCVars::NetCorrectionLifetime, 0, 1.5f);
-#if 0
-				//todo: do we show this in async mode?
-				DrawDebugFloatHistory(*OwningWorld, PhysicsTarget.ErrorHistory, NewPos + FVector(0.0f, 0.0f, 100.0f), FVector2D(100.0f, 50.0f), FColor::White);
-#endif
+				DrawDebugFloatHistory(*OwningWorld, PhysicsTarget.ErrorHistory, CurrentState.Position + FVector(0.0f, 0.0f, 100.0f), FVector2D(100.0f, 50.0f), FColor::White, false, 0, -1);
 			}
 		}
 #endif
