Adding `ForEachLoadedAsset` to `FStreamableHandle`.
Allows for operations on the loaded assets without creating an intermediate array.

#tests using `GetAllLoadedAssets` and verifying that there are no regressions via debugger

diff --git a/Engine/Source/Runtime/Engine/Classes/Engine/StreamableManager.h b/Engine/Source/Runtime/Engine/Classes/Engine/StreamableManager.h
--- a/Engine/Source/Runtime/Engine/Classes/Engine/StreamableManager.h
+++ b/Engine/Source/Runtime/Engine/Classes/Engine/StreamableManager.h
@@ -202,6 +202,14 @@ struct ENGINE_API FStreamableHandle : public TSharedFromThis<FStreamableHandle>
 	/** Returns number of assets that have completed loading out of initial list, failed loads will count as loaded */
 	void GetLoadedCount(int32& LoadedCount, int32& RequestedCount) const;
 
+	/**
+	 * Invokes a callable for each loaded asset if load has succeeded. Some entries will be null if loading failed
+	 *
+	 * @param Callable	Callable object
+	 */
+	template <typename CallableT>
+	void ForEachLoadedAsset(CallableT Callable) const;
+
 	/** Returns progress as a value between 0.0 and 1.0. */
 	float GetProgress() const;
 
@@ -599,4 +607,39 @@ private:
 	FString ManagerName;
 };
 
+template <typename CallableT>
+void FStreamableHandle::ForEachLoadedAsset(CallableT Callable) const
+{
+	if (HasLoadCompleted())
+	{
+		for (const FSoftObjectPath& Ref : RequestedAssets)
+		{
+			// Try manager, should be faster and will handle redirects better
+			if (IsActive())
+			{
+				Invoke(Callable, OwningManager->GetStreamed(Ref));
+			}
+			else
+			{
+				Invoke(Callable, Ref.ResolveObject());
+			}
+		}
 
+		// Check child handles
+		for (const TSharedPtr<FStreamableHandle>& ChildHandle : ChildHandles)
+		{
+			for (const FSoftObjectPath& Ref : ChildHandle->RequestedAssets)
+			{
+				// Try manager, should be faster and will handle redirects better
+				if (IsActive())
+				{
+					Invoke(Callable, OwningManager->GetStreamed(Ref));
+				}
+				else
+				{
+					Invoke(Callable, Ref.ResolveObject());
+				}
+			}
+		}
+	}
+}
diff --git a/Engine/Source/Runtime/Engine/Private/StreamableManager.cpp b/Engine/Source/Runtime/Engine/Private/StreamableManager.cpp
--- a/Engine/Source/Runtime/Engine/Private/StreamableManager.cpp
+++ b/Engine/Source/Runtime/Engine/Private/StreamableManager.cpp
@@ -477,38 +477,10 @@ UObject* FStreamableHandle::GetLoadedAsset() const
 
 void FStreamableHandle::GetLoadedAssets(TArray<UObject *>& LoadedAssets) const
 {
-	if (HasLoadCompleted())
+	ForEachLoadedAsset([&LoadedAssets](UObject* LoadedAsset)
 	{
-		for (const FSoftObjectPath& Ref : RequestedAssets)
-		{
-			// Try manager, should be faster and will handle redirects better
-			if (IsActive())
-			{
-				LoadedAssets.Add(OwningManager->GetStreamed(Ref));
-			}
-			else
-			{
-				LoadedAssets.Add(Ref.ResolveObject());
-			}
-		}
-
-		// Check child handles
-		for (const TSharedPtr<FStreamableHandle>& ChildHandle : ChildHandles)
-		{
-			for (const FSoftObjectPath& Ref : ChildHandle->RequestedAssets)
-			{
-				// Try manager, should be faster and will handle redirects better
-				if (IsActive())
-				{
-					LoadedAssets.Add(OwningManager->GetStreamed(Ref));
-				}
-				else
-				{
-					LoadedAssets.Add(Ref.ResolveObject());
-				}
-			}
-		}
-	}
+		LoadedAssets.Add(LoadedAsset);
+	});
 }
 
 void FStreamableHandle::GetLoadedCount(int32& LoadedCount, int32& RequestedCount) const
