Adding in LLM scope for consistent tracking.

diff --git a/Engine/Source/Runtime/Engine/Private/NetConnection.cpp b/Engine/Source/Runtime/Engine/Private/NetConnection.cpp
--- a/Engine/Source/Runtime/Engine/Private/NetConnection.cpp
+++ b/Engine/Source/Runtime/Engine/Private/NetConnection.cpp
@@ -564,6 +564,8 @@ void UNetConnection::InitConnection(UNetDriver* InDriver, EConnectionState InSta
 
 void UNetConnection::InitHandler()
 {
+	LLM_SCOPE_BYTAG(NetConnection);
+
 	check(!Handler.IsValid());
 
 #if !UE_BUILD_SHIPPING
