Physics Replication : Allow per-object error correction settings to propagate to the async callback.

diff --git a/Engine/Source/Runtime/Engine/Private/PhysicsEngine/PhysicsReplication.cpp b/Engine/Source/Runtime/Engine/Private/PhysicsEngine/PhysicsReplication.cpp
--- a/Engine/Source/Runtime/Engine/Private/PhysicsEngine/PhysicsReplication.cpp
+++ b/Engine/Source/Runtime/Engine/Private/PhysicsEngine/PhysicsReplication.cpp
@@ -102,10 +102,7 @@ namespace PhysicsReplicationCVars
 struct FAsyncPhysicsRepCallbackData : public Chaos::FSimCallbackInput
 {
 	TArray<FAsyncPhysicsDesiredState> Buffer;
-	float LinearVelocityCoefficient;
-	float AngularVelocityCoefficient;
-	float PositionLerp;
-	float AngleLerp;
+	ErrorCorrectionData ErrorCorrection;
 
 	void Reset()
 	{
@@ -424,6 +421,7 @@ bool FPhysicsReplication::ApplyRigidBodyState(float DeltaSeconds, FBodyInstance*
 				AsyncDesiredState.AngularVelocity = NewState.AngVel;
 				AsyncDesiredState.Proxy = static_cast<Chaos::FSingleParticlePhysicsProxy*>(BI->GetPhysicsActorHandle());
 				AsyncDesiredState.ObjectState = AsyncDesiredState.Proxy->GetGameThreadAPI().ObjectState();
+				AsyncDesiredState.ErrorCorrection = { ErrorCorrection.LinearVelocityCoefficient, ErrorCorrection.AngularVelocityCoefficient, ErrorCorrection.PositionLerp, ErrorCorrection.AngleLerp };
 				AsyncDesiredState.bShouldSleep = bShouldSleep;
 
 				CurAsyncData->Buffer.Add(AsyncDesiredState);
@@ -645,10 +643,10 @@ void FPhysicsReplication::PrepareAsyncData_External(const FRigidBodyErrorCorrect
 	const float AngularVelocityCoefficient = CharacterMovementCVars::AngularVelocityCoefficient >= 0.0f ? CharacterMovementCVars::AngularVelocityCoefficient : ErrorCorrection.AngularVelocityCoefficient;
 
 	CurAsyncData = AsyncCallback->GetProducerInputData_External();
-	CurAsyncData->PositionLerp = PositionLerp;
-	CurAsyncData->AngleLerp = AngleLerp;
-	CurAsyncData->LinearVelocityCoefficient = LinearVelocityCoefficient;
-	CurAsyncData->AngularVelocityCoefficient = AngularVelocityCoefficient;
+	CurAsyncData->ErrorCorrection.PositionLerp = PositionLerp;
+	CurAsyncData->ErrorCorrection.AngleLerp = AngleLerp;
+	CurAsyncData->ErrorCorrection.LinearVelocityCoefficient = LinearVelocityCoefficient;
+	CurAsyncData->ErrorCorrection.AngularVelocityCoefficient = AngularVelocityCoefficient;
 }
 
 void FPhysicsReplication::ApplyAsyncDesiredState(const float DeltaSeconds, const FAsyncPhysicsRepCallbackData* AsyncData)
@@ -656,13 +654,20 @@ void FPhysicsReplication::ApplyAsyncDesiredState(const float DeltaSeconds, const
 	using namespace Chaos;
 	if(AsyncData)
 	{
-		const float LinearVelocityCoefficient = AsyncData->LinearVelocityCoefficient;
-		const float AngularVelocityCoefficient = AsyncData->AngularVelocityCoefficient;
-		const float PositionLerp = AsyncData->PositionLerp;
-		const float AngleLerp = AsyncData->AngleLerp;
-
 		for (const FAsyncPhysicsDesiredState& State : AsyncData->Buffer)
 		{
+			float LinearVelocityCoefficient = AsyncData->ErrorCorrection.LinearVelocityCoefficient;
+			float AngularVelocityCoefficient = AsyncData->ErrorCorrection.AngularVelocityCoefficient;
+			float PositionLerp = AsyncData->ErrorCorrection.PositionLerp;
+			float AngleLerp = AsyncData->ErrorCorrection.AngleLerp;
+			if (State.ErrorCorrection.IsSet())
+			{
+				ErrorCorrectionData ECData = State.ErrorCorrection.GetValue();
+				LinearVelocityCoefficient = ECData.LinearVelocityCoefficient;
+				AngularVelocityCoefficient = ECData.AngularVelocityCoefficient;
+				PositionLerp = ECData.PositionLerp;
+				AngleLerp = ECData.AngleLerp;
+			}
 			//Proxy should exist because we are using latest and any pending deletes would have been enqueued after
 			Chaos::FSingleParticlePhysicsProxy* Proxy = State.Proxy;
 			auto* Handle = Proxy->GetPhysicsThreadAPI();
diff --git a/Engine/Source/Runtime/Engine/Public/PhysicsReplication.h b/Engine/Source/Runtime/Engine/Public/PhysicsReplication.h
--- a/Engine/Source/Runtime/Engine/Public/PhysicsReplication.h
+++ b/Engine/Source/Runtime/Engine/Public/PhysicsReplication.h
@@ -61,6 +61,14 @@ struct FReplicatedPhysicsTarget
 };
 
 #if WITH_CHAOS
+struct ErrorCorrectionData
+{
+	float LinearVelocityCoefficient;
+	float AngularVelocityCoefficient;
+	float PositionLerp;
+	float AngleLerp;
+};
+
 /** Final computed desired state passed into the physics sim */
 struct FAsyncPhysicsDesiredState
 {
@@ -69,6 +77,7 @@ struct FAsyncPhysicsDesiredState
 	FVector AngularVelocity;
 	Chaos::FSingleParticlePhysicsProxy* Proxy;
 	Chaos::EObjectStateType ObjectState;
+	TOptional<ErrorCorrectionData> ErrorCorrection;
 	bool bShouldSleep;
 	int32 ServerFrame;
 };
