Adding additional recording metadata to the replay header

diff --git a/Engine/Source/Runtime/Engine/Classes/Engine/DemoNetDriver.h b/Engine/Source/Runtime/Engine/Classes/Engine/DemoNetDriver.h
--- a/Engine/Source/Runtime/Engine/Classes/Engine/DemoNetDriver.h
+++ b/Engine/Source/Runtime/Engine/Classes/Engine/DemoNetDriver.h
@@ -385,12 +385,16 @@ public:
 	/** Sets the desired maximum recording time in milliseconds. */
 	void SetMaxDesiredRecordTimeMS(const float InMaxDesiredRecordTimeMS) { MaxDesiredRecordTimeMS = InMaxDesiredRecordTimeMS; }
 
+	float GetMaxDesiredRecordTimeMS() const { return MaxDesiredRecordTimeMS; }
+
 	/** Sets the controller to use as the viewpoint for recording prioritization purposes. */
 	void SetViewerOverride(APlayerController* const InViewerOverride ) { ViewerOverride = InViewerOverride; }
 
 	/** Enable or disable prioritization of actors for recording. */
 	void SetActorPrioritizationEnabled(const bool bInPrioritizeActors) { bPrioritizeActors = bInPrioritizeActors; }
 
+	bool IsActorPrioritizationEnabled() const { return bPrioritizeActors; }
+
 	/** Sets CheckpointSaveMaxMSPerFrame. */
 	void SetCheckpointSaveMaxMSPerFrame(const float InCheckpointSaveMaxMSPerFrame)
 	{ 
diff --git a/Engine/Source/Runtime/Engine/Private/DemoNetDriver.cpp b/Engine/Source/Runtime/Engine/Private/DemoNetDriver.cpp
--- a/Engine/Source/Runtime/Engine/Private/DemoNetDriver.cpp
+++ b/Engine/Source/Runtime/Engine/Private/DemoNetDriver.cpp
@@ -51,8 +51,8 @@ CSV_DEFINE_CATEGORY(Demo, false);
 CSV_DEFINE_CATEGORY(Demo, true);
 #endif
 
-static TAutoConsoleVariable<float> CVarDemoRecordHz( TEXT( "demo.RecordHz" ), 8, TEXT( "Maximum number of demo frames recorded per second" ) );
-static TAutoConsoleVariable<float> CVarDemoMinRecordHz(TEXT("demo.MinRecordHz"), 0, TEXT("Minimum number of demo frames recorded per second (use with care)"));
+TAutoConsoleVariable<float> CVarDemoRecordHz( TEXT( "demo.RecordHz" ), 8, TEXT( "Maximum number of demo frames recorded per second" ) );
+TAutoConsoleVariable<float> CVarDemoMinRecordHz(TEXT("demo.MinRecordHz"), 0, TEXT("Minimum number of demo frames recorded per second (use with care)"));
 static TAutoConsoleVariable<float> CVarDemoTimeDilation( TEXT( "demo.TimeDilation" ), -1.0f, TEXT( "Override time dilation during demo playback (-1 = don't override)" ) );
 static TAutoConsoleVariable<float> CVarDemoSkipTime( TEXT( "demo.SkipTime" ), 0, TEXT( "Skip fixed amount of network replay time (in seconds)" ) );
 TAutoConsoleVariable<int32> CVarEnableCheckpoints( TEXT( "demo.EnableCheckpoints" ), 1, TEXT( "Whether or not checkpoints save on the server" ) );
@@ -65,9 +65,9 @@ static TAutoConsoleVariable<int32> CVarDemoAsyncLoadWorld( TEXT( "demo.AsyncLoad
 TAutoConsoleVariable<float> CVarCheckpointUploadDelayInSeconds( TEXT( "demo.CheckpointUploadDelayInSeconds" ), 30.0f, TEXT( "" ) );
 static TAutoConsoleVariable<int32> CVarDemoLoadCheckpointGarbageCollect( TEXT( "demo.LoadCheckpointGarbageCollect" ), 1, TEXT("If nonzero, CollectGarbage will be called during LoadCheckpoint after the old actors and connection are cleaned up." ) );
 TAutoConsoleVariable<float> CVarCheckpointSaveMaxMSPerFrameOverride( TEXT( "demo.CheckpointSaveMaxMSPerFrameOverride" ), -1.0f, TEXT( "If >= 0, this value will override the CheckpointSaveMaxMSPerFrame member variable, which is the maximum time allowed each frame to spend on saving a checkpoint. If 0, it will save the checkpoint in a single frame, regardless of how long it takes." ) );
-static TAutoConsoleVariable<int32> CVarDemoClientRecordAsyncEndOfFrame( TEXT( "demo.ClientRecordAsyncEndOfFrame" ), 0, TEXT( "If true, TickFlush will be called on a thread in parallel with Slate." ) );
+TAutoConsoleVariable<int32> CVarDemoClientRecordAsyncEndOfFrame( TEXT( "demo.ClientRecordAsyncEndOfFrame" ), 0, TEXT( "If true, TickFlush will be called on a thread in parallel with Slate." ) );
 static TAutoConsoleVariable<int32> CVarForceDisableAsyncPackageMapLoading( TEXT( "demo.ForceDisableAsyncPackageMapLoading" ), 0, TEXT( "If true, async package map loading of network assets will be disabled." ) );
-static TAutoConsoleVariable<int32> CVarDemoUseNetRelevancy( TEXT( "demo.UseNetRelevancy" ), 0, TEXT( "If 1, will enable relevancy checks and distance culling, using all connected clients as reference." ) );
+TAutoConsoleVariable<int32> CVarDemoUseNetRelevancy( TEXT( "demo.UseNetRelevancy" ), 0, TEXT( "If 1, will enable relevancy checks and distance culling, using all connected clients as reference." ) );
 static TAutoConsoleVariable<float> CVarDemoCullDistanceOverride( TEXT( "demo.CullDistanceOverride" ), 0.0f, TEXT( "If > 0, will represent distance from any viewer where actors will stop being recorded." ) );
 static TAutoConsoleVariable<float> CVarDemoRecordHzWhenNotRelevant( TEXT( "demo.RecordHzWhenNotRelevant" ), 2.0f, TEXT( "Record at this frequency when actor is not relevant." ) );
 static TAutoConsoleVariable<int32> CVarLoopDemo(TEXT("demo.Loop"), 0, TEXT("<1> : play replay from beginning once it reaches the end / <0> : stop replay at the end"));
@@ -3171,6 +3171,14 @@ void UDemoNetDriver::ReplayStreamingReady(const FStartStreamingResult& Result)
 			UE_LOG(LogDemo, Log, TEXT("ReplayStreamingReady: playing back replay [%s] %s, which was recorded on engine version %s with flags [%s]"),
 				*ReplayHelper.GetPlaybackGuid().ToString(EGuidFormats::Digits), *ReplayHelper.DemoURL.Map, *ReplayHelper.PlaybackDemoHeader.EngineVersion.ToString(), *HeaderFlags);
 
+			if (ReplayHelper.PlaybackDemoHeader.Version >= ENetworkVersionHistory::HISTORY_RECORDING_METADATA)
+			{
+				UE_LOG(LogDemo, Log, TEXT("ReplayStreamingReady: replay was recorded with: MinHz: %0.2f MaxHz: %0.2f FrameMS: %0.2f CheckpointMS: %0.2f Platform: [%s] Config: [%s] Target: [%s]"),
+					ReplayHelper.PlaybackDemoHeader.MinRecordHz, ReplayHelper.PlaybackDemoHeader.MaxRecordHz,
+					ReplayHelper.PlaybackDemoHeader.FrameLimitInMS, ReplayHelper.PlaybackDemoHeader.CheckpointLimitInMS,
+					*ReplayHelper.PlaybackDemoHeader.Platform, LexToString(ReplayHelper.PlaybackDemoHeader.BuildConfig), LexToString(ReplayHelper.PlaybackDemoHeader.BuildTarget));
+			}
+
 			CSV_METADATA(TEXT("ReplayID"), *ReplayHelper.GetPlaybackGuid().ToString(EGuidFormats::Digits));
 		}
 
diff --git a/Engine/Source/Runtime/Engine/Private/ReplayHelper.cpp b/Engine/Source/Runtime/Engine/Private/ReplayHelper.cpp
--- a/Engine/Source/Runtime/Engine/Private/ReplayHelper.cpp
+++ b/Engine/Source/Runtime/Engine/Private/ReplayHelper.cpp
@@ -18,6 +18,7 @@
 #include "UnrealEngine.h"
 #include "EngineUtils.h"
 #include "ReplayNetConnection.h"
+#include "Engine/DemoNetDriver.h"
 
 extern TAutoConsoleVariable<int32> CVarWithLevelStreamingFixes;
 extern TAutoConsoleVariable<int32> CVarWithDeltaCheckpoints;
@@ -25,6 +26,10 @@ extern TAutoConsoleVariable<int32> CVarWithGameSpecificFrameData;
 extern TAutoConsoleVariable<int32> CVarEnableCheckpoints;
 extern TAutoConsoleVariable<float> CVarCheckpointUploadDelayInSeconds;
 extern TAutoConsoleVariable<float> CVarCheckpointSaveMaxMSPerFrameOverride;
+extern TAutoConsoleVariable<int32> CVarDemoUseNetRelevancy;
+extern TAutoConsoleVariable<int32> CVarDemoClientRecordAsyncEndOfFrame;
+extern TAutoConsoleVariable<float> CVarDemoRecordHz;
+extern TAutoConsoleVariable<float> CVarDemoMinRecordHz;
 
 CSV_DECLARE_CATEGORY_EXTERN(Demo);
 
@@ -198,6 +203,11 @@ void FReplayHelper::WriteNetworkDemoHeader(UNetConnection* Connection)
 			if (LocalWorld->GetNetDriver() != nullptr && !LocalWorld->GetNetDriver()->IsServer())
 			{
 				DemoHeader.HeaderFlags |= EReplayHeaderFlags::ClientRecorded;
+
+				if (CVarDemoClientRecordAsyncEndOfFrame.GetValueOnAnyThread() > 0)
+				{
+					DemoHeader.HeaderFlags |= EReplayHeaderFlags::AsyncRecorded;
+				}
 			}
 		}
 
@@ -221,8 +231,35 @@ void FReplayHelper::WriteNetworkDemoHeader(UNetConnection* Connection)
 			DemoHeader.HeaderFlags |= EReplayHeaderFlags::ReplayConnection;
 		}
 
+		if (Connection)
+		{
+			if (UDemoNetDriver* DemoDriver = Cast<UDemoNetDriver>(Connection->GetDriver()))
+			{
+				if (DemoDriver->IsActorPrioritizationEnabled())
+				{
+					DemoHeader.HeaderFlags |= EReplayHeaderFlags::ActorPrioritizationEnabled;
+				}
+
+				DemoHeader.FrameLimitInMS = DemoDriver->GetMaxDesiredRecordTimeMS();
+			}
+		}
+
+		if (CVarDemoUseNetRelevancy.GetValueOnAnyThread() > 0)
+		{
+			DemoHeader.HeaderFlags |= EReplayHeaderFlags::NetRelevancyEnabled;
+		}
+
 		DemoHeader.Guid = FGuid::NewGuid();
 
+		DemoHeader.CheckpointLimitInMS = CheckpointSaveMaxMSPerFrame;
+
+		DemoHeader.MinRecordHz = CVarDemoMinRecordHz.GetValueOnAnyThread();
+		DemoHeader.MaxRecordHz = CVarDemoRecordHz.GetValueOnAnyThread();
+		
+		DemoHeader.Platform = FPlatformProperties::PlatformName();
+		DemoHeader.BuildConfig = FApp::GetBuildConfiguration();
+		DemoHeader.BuildTarget = FApp::GetBuildTargetType();
+
 		// Write the header
 		(*FileAr) << DemoHeader;
 		FileAr->Flush();
@@ -2246,6 +2283,12 @@ const TCHAR* LexToString(EReplayHeaderFlags Flag)
 		return TEXT("GameSpecificFrameData");
 	case EReplayHeaderFlags::ReplayConnection:
 		return TEXT("ReplayConnection");
+	case EReplayHeaderFlags::ActorPrioritizationEnabled:
+		return TEXT("ActorPrioritizationEnabled");
+	case EReplayHeaderFlags::NetRelevancyEnabled:
+		return TEXT("NetRelevancyEnabled");
+	case EReplayHeaderFlags::AsyncRecorded:
+		return TEXT("AsyncRecorded");
 	default:
 		check(false);
 		return TEXT("Unknown");
diff --git a/Engine/Source/Runtime/Engine/Public/ReplayTypes.h b/Engine/Source/Runtime/Engine/Public/ReplayTypes.h
--- a/Engine/Source/Runtime/Engine/Public/ReplayTypes.h
+++ b/Engine/Source/Runtime/Engine/Public/ReplayTypes.h
@@ -25,6 +25,9 @@ enum class EReplayHeaderFlags : uint32
 	DeltaCheckpoints = (1 << 2),
 	GameSpecificFrameData = (1 << 3),
 	ReplayConnection = (1 << 4),
+	ActorPrioritizationEnabled = (1 << 5),
+	NetRelevancyEnabled = (1 << 6),
+	AsyncRecorded = (1 << 7),
 };
 
 ENUM_CLASS_FLAGS(EReplayHeaderFlags);
@@ -104,6 +107,7 @@ enum ENetworkVersionHistory
 	HISTORY_GUID_NAMETABLE = 15,				// Added a string table for exported guids
 	HISTORY_GUIDCACHE_CHECKSUMS = 16,			// Removing guid export checksums from saved data, they are ignored during playback
 	HISTORY_SAVE_PACKAGE_VERSION_UE = 17,		// Save engine and licensee package version as well, in case serialization functions need them for compatibility
+	HISTORY_RECORDING_METADATA = 18,			// Adding additional record-time information to the header
 
 	// -----<new versions can be added before this line>-------------------------------------------------
 	HISTORY_PLUS_ONE,
@@ -128,6 +132,15 @@ struct FNetworkDemoHeader
 	uint32	GameNetworkProtocolVersion;				// Version of the game internal network format
 	FGuid	Guid;									// Unique identifier
 
+	float MinRecordHz;
+	float MaxRecordHz;
+	float FrameLimitInMS;
+	float CheckpointLimitInMS;
+
+	FString Platform;
+	EBuildConfiguration BuildConfig;
+	EBuildTargetType BuildTarget;
+
 	FEngineVersion EngineVersion;					// Full engine version on which the replay was recorded
 	EReplayHeaderFlags HeaderFlags;					// Replay flags
 	TArray<FLevelNameAndTime> LevelNamesAndTimes;	// Name and time changes of levels loaded for demo
@@ -257,6 +270,19 @@ struct FNetworkDemoHeader
 
 		Ar << Header.GameSpecificData;
 
+		if (Header.Version >= HISTORY_RECORDING_METADATA)
+		{
+			Ar << Header.MinRecordHz;
+			Ar << Header.MaxRecordHz;
+
+			Ar << Header.FrameLimitInMS;
+			Ar << Header.CheckpointLimitInMS;
+
+			Ar << Header.Platform;
+			Ar << Header.BuildConfig;
+			Ar << Header.BuildTarget;
+		}
+
 		return Ar;
 	}
 
