Fix long wait times in FODSCManager::Tick caused by running when cook-on-the-fly is not enabled.
Also only flush in ProcessCookOnTheFlyShaders if some work was actually done.

diff --git a/Engine/Source/Runtime/Engine/Private/ODSC/ODSCManager.cpp b/Engine/Source/Runtime/Engine/Private/ODSC/ODSCManager.cpp
--- a/Engine/Source/Runtime/Engine/Private/ODSC/ODSCManager.cpp
+++ b/Engine/Source/Runtime/Engine/Private/ODSC/ODSCManager.cpp
@@ -15,9 +15,12 @@ FODSCManager* GODSCManager = nullptr;
 
 FODSCManager::FODSCManager()
 	: FTSTickerObjectBase(0.0f, FTSBackgroundableTicker::GetCoreTicker())
-	, Thread(new FODSCThread())
 {
-	Thread->StartThread();
+	if (IsRunningCookOnTheFly())
+	{
+		Thread = new FODSCThread();
+		Thread->StartThread();
+	}
 }
 
 FODSCManager::~FODSCManager()
@@ -55,12 +58,18 @@ bool FODSCManager::Tick(float DeltaSeconds)
 
 void FODSCManager::AddThreadedRequest(const TArray<FString>& MaterialsToCompile, EShaderPlatform ShaderPlatform, ODSCRecompileCommand RecompileCommandType)
 {
-	check(Thread);
-	Thread->AddRequest(MaterialsToCompile, ShaderPlatform, RecompileCommandType);
+	if (IsRunningCookOnTheFly())
+	{
+		check(Thread);
+		Thread->AddRequest(MaterialsToCompile, ShaderPlatform, RecompileCommandType);
+	}
 }
 
 void FODSCManager::AddThreadedShaderPipelineRequest(EShaderPlatform ShaderPlatform, const FString& MaterialName, const FString& VertexFactoryName, const FString& PipelineName, const TArray<FString>& ShaderTypeNames)
 {
-	check(Thread);
-	Thread->AddShaderPipelineRequest(ShaderPlatform, MaterialName, VertexFactoryName, PipelineName, ShaderTypeNames);
+	if (IsRunningCookOnTheFly())
+	{
+		check(Thread);
+		Thread->AddShaderPipelineRequest(ShaderPlatform, MaterialName, VertexFactoryName, PipelineName, ShaderTypeNames);
+	}
 }
diff --git a/Engine/Source/Runtime/Engine/Private/ShaderCompiler/ShaderCompiler.cpp b/Engine/Source/Runtime/Engine/Private/ShaderCompiler/ShaderCompiler.cpp
--- a/Engine/Source/Runtime/Engine/Private/ShaderCompiler/ShaderCompiler.cpp
+++ b/Engine/Source/Runtime/Engine/Private/ShaderCompiler/ShaderCompiler.cpp
@@ -6053,12 +6053,22 @@ void ProcessCookOnTheFlyShaders(bool bReloadGlobalShaders, const TArray<uint8>&
 {
 	check(IsInGameThread());
 
-	// now we need to refresh the RHI resources
-	FlushRenderingCommands();
+	bool bHasFlushed = false;
+
+	auto DoFlushIfNecessary = [&bHasFlushed]() {
+		if (!bHasFlushed )
+		{
+			// now we need to refresh the RHI resources
+			FlushRenderingCommands();
+			bHasFlushed = true;
+		}
+	};
 
 	// reload the global shaders
 	if (bReloadGlobalShaders)
 	{
+		DoFlushIfNecessary();
+
 		// Some platforms rely on global shaders to be created to implement basic RHI functionality
 		TGuardValue<int32> Guard(GCreateShadersOnLoad, 1);
 		CompileGlobalShaderMap(true);
@@ -6067,6 +6077,8 @@ void ProcessCookOnTheFlyShaders(bool bReloadGlobalShaders, const TArray<uint8>&
 	// load all the mesh material shaders if any were sent back
 	if (MeshMaterialMaps.Num() > 0)
 	{
+		DoFlushIfNecessary();
+
 		// parse the shaders
 		FMemoryReader MemoryReader(MeshMaterialMaps, true);
 		FNameAsStringProxyArchive Ar(MemoryReader);
@@ -6092,6 +6104,8 @@ void ProcessCookOnTheFlyShaders(bool bReloadGlobalShaders, const TArray<uint8>&
 	// load all the global shaders if any were sent back
 	if (GlobalShaderMap.Num() > 0)
 	{
+		DoFlushIfNecessary();
+
 		// parse the shaders
 		FMemoryReader MemoryReader(GlobalShaderMap, true);
 		FNameAsStringProxyArchive Ar(MemoryReader);
