Fix MinLod quality level getting strip if there are not specify in the asset data. Serialize all quality level supported by the platform.

diff --git a/Engine/Source/Runtime/Engine/Private/PerQualityLevelProperties.cpp b/Engine/Source/Runtime/Engine/Private/PerQualityLevelProperties.cpp
--- a/Engine/Source/Runtime/Engine/Private/PerQualityLevelProperties.cpp
+++ b/Engine/Source/Runtime/Engine/Private/PerQualityLevelProperties.cpp
@@ -207,26 +207,37 @@ void FPerQualityLevelProperty<_StructType, _ValueType, _BasePropertyName>::Strip
 	if (This->PerQuality.Num() > 0)
 	{
 		FSupportedQualityLevelArray CookQualityLevelInfo = This->GetSupportedQualityLevels(InPlatformName);
-
-		int32 LowestQualityLevel = (int32)QualityLevelProperty::EQualityLevels::Num;
+		CookQualityLevelInfo.Sort([&](const int32& A, const int32& B) { return (A > B); });
 
 		// remove unsupported quality levels 
 		for (TMap<int32, int32>::TIterator It(This->PerQuality); It; ++It)
 		{
-			if(!CookQualityLevelInfo.Contains(It.Key()))
+			if (!CookQualityLevelInfo.Contains(It.Key()))
 			{
 				It.RemoveCurrent();
 			}
-			else
-			{
-				LowestQualityLevel = (It.Key() < LowestQualityLevel) ? It.Key() : LowestQualityLevel;
-			}
 		}
 
-		//if found supported platforms, put the lowest quality level in Default
-		if (LowestQualityLevel != (int32)QualityLevelProperty::EQualityLevels::Num)
+		if (This->PerQuality.Num() > 0)
 		{
-			This->Default = This->GetValue(LowestQualityLevel);
+			int32 PreviousQualityLevel = This->Default;
+
+			for (TSet<int32>::TIterator It(CookQualityLevelInfo); It; ++It)
+			{
+				int32* QualityLevelMinLod = This->PerQuality.Find(*It);
+			
+				// add quality level supported by the platform
+				if (!QualityLevelMinLod)
+				{
+					This->PerQuality.Add(*It, PreviousQualityLevel);
+				}
+				else
+				{
+					PreviousQualityLevel = *QualityLevelMinLod;
+				}
+			}
+
+			This->Default = PreviousQualityLevel;
 		}
 	}
 }
@@ -321,4 +332,4 @@ FString FPerQualityLevelInt::ToString() const
 #endif
 
 	return Result;
-}
\ No newline at end of file
+}
