Fixed deleting of an old DataLayerInstance not using DataLayer assets.

diff --git a/Engine/Source/Runtime/Engine/Private/WorldPartition/DataLayer/WorldDataLayers.cpp b/Engine/Source/Runtime/Engine/Private/WorldPartition/DataLayer/WorldDataLayers.cpp
--- a/Engine/Source/Runtime/Engine/Private/WorldPartition/DataLayer/WorldDataLayers.cpp
+++ b/Engine/Source/Runtime/Engine/Private/WorldPartition/DataLayer/WorldDataLayers.cpp
@@ -536,6 +536,10 @@ bool AWorldDataLayers::RemoveDataLayers(const TArray<UDataLayerInstance*>& InDat
 			Modify();
 			DataLayerInstance->SetChildParent(DataLayerInstance->GetParent());
 			DataLayerInstances.Remove(DataLayerInstance);
+			if (DataLayerInstance->IsA<UDeprecatedDataLayerInstance>())
+			{
+				DeprecatedDataLayerNameToDataLayerInstance.Remove(DataLayerInstance->GetDataLayerFName());
+			}
 			bIsModified = true;
 		}
 	}
@@ -557,6 +561,7 @@ bool AWorldDataLayers::RemoveDataLayer(const UDataLayerInstance* InDataLayerInst
 
 		if (InDataLayerInstance->IsA<UDeprecatedDataLayerInstance>())
 		{
+			DeprecatedDataLayerNameToDataLayerInstance.Remove(InDataLayerInstance->GetDataLayerFName());
 			UpdateContainsDeprecatedDataLayers();
 		}
 
