Cache physics sim results to replicate

diff --git a/Engine/Source/Runtime/Engine/Private/PhysicsEngine/AsyncPhysicsInputComponent.cpp b/Engine/Source/Runtime/Engine/Private/PhysicsEngine/AsyncPhysicsInputComponent.cpp
--- a/Engine/Source/Runtime/Engine/Private/PhysicsEngine/AsyncPhysicsInputComponent.cpp
+++ b/Engine/Source/Runtime/Engine/Private/PhysicsEngine/AsyncPhysicsInputComponent.cpp
@@ -311,6 +311,12 @@ void UAsyncPhysicsInputComponent::AsyncPhysicsTickComponent(float DeltaTime, flo
 	FAsyncPhysicsInputRewindCallback* Callback = static_cast<FAsyncPhysicsInputRewindCallback*>(Solver->GetRewindCallback());
 	const int32 ServerFrame = Callback->CachedServerFrame;
 
+	if(World->IsServer())
+	{
+		//TODO: move this somewhere else - here because we need to run this when GT and PT are both on same core. Function guards against multiple calls per ServerFrame so calling from each instance is ok
+		Callback->UpdateReplicationMap_Internal(ServerFrame);
+	}
+
 	Pool->SetCurrentInputToAsyncExecute(nullptr);	//set current input to none in case we don't find it
 
 	for (int32 Idx = BufferedInputs.Num() - 1; Idx >= 0; --Idx)
