Texture values are no longer generated as HLSL code directly.  Instead, they are replaced with a new 'Object' type system, which allows expressions to generate C++ objects of arbitrary type.  This gives the texture sample expression full visibility into the texture object being sampled, which simplifies some code generation.  More importantly it will allow VT to work, which wasn't really possible with the previous system

diff --git a/Engine/Shaders/Private/MaterialTexture.ush b/Engine/Shaders/Private/MaterialTexture.ush
--- a/Engine/Shaders/Private/MaterialTexture.ush
+++ b/Engine/Shaders/Private/MaterialTexture.ush
@@ -124,6 +124,21 @@ FTextureExternal MakeTextureExternal(TextureExternal Texture, SamplerState Sampl
 	Result.SamplerType = SamplerType;
 	return Result;
 }
+
+#if NUM_VIRTUALTEXTURE_SAMPLES
+struct FTextureVirtual
+{
+	Texture2D<uint4> PageTable0;
+	Texture2D<uint4> PageTable1;
+	Texture2D<uint> PageTableIndirection;
+	Texture2D PhysicalTexture;
+	SamplerState PhysicalSampler;
+	VTPageTableUniform PageTableUniform;
+	VTUniform Uniform;
+	int SamplerType; // EMaterialSamplerType
+};
+#endif // NUM_VIRTUALTEXTURE_SAMPLES
+
 #endif // ENABLE_NEW_HLSL_GENERATOR
 
 MaterialFloat4 ProcessMaterialColorTextureLookup(MaterialFloat4 TextureValue)
diff --git a/Engine/Source/Editor/UnrealEd/Private/MaterialGraphSchema.cpp b/Engine/Source/Editor/UnrealEd/Private/MaterialGraphSchema.cpp
--- a/Engine/Source/Editor/UnrealEd/Private/MaterialGraphSchema.cpp
+++ b/Engine/Source/Editor/UnrealEd/Private/MaterialGraphSchema.cpp
@@ -632,14 +632,14 @@ FLinearColor GetColorForConnectionType(const UGraphEditorSettings* Settings, UE:
 	{
 		return Settings->StructPinTypeColor;
 	}
+	else if (ConnectionType == UE::Shader::EValueType::Object)
+	{
+		return Settings->ObjectPinTypeColor;
+	}
 	else
 	{
 		const UE::Shader::FValueTypeDescription TypeDesc = UE::Shader::GetValueTypeDescription(ConnectionType);
-		if (UE::Shader::IsTextureType(TypeDesc.ComponentType))
-		{
-			return Settings->ObjectPinTypeColor;
-		}
-		else if (TypeDesc.ComponentType == UE::Shader::EValueComponentType::Float)
+		if (TypeDesc.ComponentType == UE::Shader::EValueComponentType::Float)
 		{
 			if (TypeDesc.NumComponents == 1)
 			{
diff --git a/Engine/Source/Runtime/Engine/Classes/Materials/MaterialLayersFunctions.h b/Engine/Source/Runtime/Engine/Classes/Materials/MaterialLayersFunctions.h
--- a/Engine/Source/Runtime/Engine/Classes/Materials/MaterialLayersFunctions.h
+++ b/Engine/Source/Runtime/Engine/Classes/Materials/MaterialLayersFunctions.h
@@ -4,168 +4,13 @@
 
 #include "CoreMinimal.h"
 #include "AssetData.h"
+#include "MaterialTypes.h"
 #include "MaterialLayersFunctions.generated.h"
 
 #define LOCTEXT_NAMESPACE "MaterialLayersFunctions"
 
 class FArchive;
 
-
-UENUM()
-enum EMaterialParameterAssociation
-{
-	LayerParameter,
-	BlendParameter,
-	GlobalParameter,
-};
-
-USTRUCT(BlueprintType)
-struct ENGINE_API FMaterialParameterInfo
-{
-	GENERATED_USTRUCT_BODY()
-
-	UPROPERTY(EditAnywhere, BlueprintReadWrite, Category=ParameterInfo)
-	FName Name;
-
-	/** Whether this is a global parameter, or part of a layer or blend */
-	UPROPERTY(EditAnywhere, BlueprintReadWrite, Category=ParameterInfo)
-	TEnumAsByte<EMaterialParameterAssociation> Association;
-
-	/** Layer or blend index this parameter is part of. INDEX_NONE for global parameters. */
-	UPROPERTY(EditAnywhere, BlueprintReadWrite, Category=ParameterInfo)
-	int32 Index;
-
-	FMaterialParameterInfo(const TCHAR* InName, EMaterialParameterAssociation InAssociation = EMaterialParameterAssociation::GlobalParameter, int32 InIndex = INDEX_NONE)
-		: Name(InName)
-		, Association(InAssociation)
-		, Index(InIndex)
-	{
-	}
-	FMaterialParameterInfo(FName InName = FName(), EMaterialParameterAssociation InAssociation = EMaterialParameterAssociation::GlobalParameter, int32 InIndex = INDEX_NONE)
-	: Name(InName)
-	, Association(InAssociation)
-	, Index(InIndex)
-	{
-	}
-
-	explicit FMaterialParameterInfo(const struct FMemoryImageMaterialParameterInfo& Rhs);
-
-	FString ToString() const
-	{
-		return *Name.ToString() + FString::FromInt(Association) + FString::FromInt(Index);
-	}
-
-	friend FArchive& operator<<(FArchive& Ar, FMaterialParameterInfo& Ref)
-	{
-		Ar << Ref.Name << Ref.Association << Ref.Index;
-		return Ar;
-	}
-
-	bool RemapLayerIndex(TArrayView<const int32> IndexRemap, FMaterialParameterInfo& OutResult) const;
-};
-
-struct FMemoryImageMaterialParameterInfo
-{
-	DECLARE_TYPE_LAYOUT(FMemoryImageMaterialParameterInfo, NonVirtual);
-public:
-	FMemoryImageMaterialParameterInfo(const TCHAR* InName, EMaterialParameterAssociation InAssociation = EMaterialParameterAssociation::GlobalParameter, int32 InIndex = INDEX_NONE)
-		: Name(NameToScriptName(FName(InName)))
-		, Index(InIndex)
-		, Association(InAssociation)
-	{}
-
-	FMemoryImageMaterialParameterInfo(const FName& InName, EMaterialParameterAssociation InAssociation = EMaterialParameterAssociation::GlobalParameter, int32 InIndex = INDEX_NONE)
-		: Name(NameToScriptName(InName))
-		, Index(InIndex)
-		, Association(InAssociation)
-	{}
-
-	FMemoryImageMaterialParameterInfo(const FScriptName& InName = FScriptName(), EMaterialParameterAssociation InAssociation = EMaterialParameterAssociation::GlobalParameter, int32 InIndex = INDEX_NONE)
-		: Name(InName)
-		, Index(InIndex)
-		, Association(InAssociation)
-	{}
-
-	FMemoryImageMaterialParameterInfo(const FMaterialParameterInfo& Rhs)
-		: Name(NameToScriptName(Rhs.Name))
-		, Index(Rhs.Index)
-		, Association(Rhs.Association)
-	{}
-
-	FMemoryImageMaterialParameterInfo(const FMemoryImageMaterialParameterInfo& Rhs) = default;
-
-	FORCEINLINE FName GetName() const { return ScriptNameToName(Name); }
-
-	friend FArchive& operator<<(FArchive& Ar, FMemoryImageMaterialParameterInfo& Ref)
-	{
-		FName RefName = ScriptNameToName(Ref.Name);
-		Ar << RefName << Ref.Association << Ref.Index;
-		Ref.Name = NameToScriptName(RefName);
-		return Ar;
-	}
-
-	bool RemapLayerIndex(TArrayView<const int32> IndexRemap, FMemoryImageMaterialParameterInfo& OutResult) const;
-
-	LAYOUT_FIELD(FScriptName, Name);
-	LAYOUT_FIELD(int32, Index);
-	LAYOUT_FIELD(TEnumAsByte<EMaterialParameterAssociation>, Association);
-};
-
-FORCEINLINE FMaterialParameterInfo::FMaterialParameterInfo(const struct FMemoryImageMaterialParameterInfo& Rhs)
-	: Name(ScriptNameToName(Rhs.Name))
-	, Association(Rhs.Association)
-	, Index(Rhs.Index)
-{
-}
-
-FORCEINLINE bool operator==(const FMaterialParameterInfo& Lhs, const FMaterialParameterInfo& Rhs)
-{
-	return Lhs.Name.IsEqual(Rhs.Name) && Lhs.Association == Rhs.Association && Lhs.Index == Rhs.Index;
-}
-
-FORCEINLINE bool operator!=(const FMaterialParameterInfo& Lhs, const FMaterialParameterInfo& Rhs)
-{
-	return !operator==(Lhs, Rhs);
-}
-
-FORCEINLINE bool operator==(const FMemoryImageMaterialParameterInfo& Lhs, const FMemoryImageMaterialParameterInfo& Rhs)
-{
-	return Lhs.Name == Rhs.Name && Lhs.Association == Rhs.Association && Lhs.Index == Rhs.Index;
-}
-
-FORCEINLINE bool operator!=(const FMemoryImageMaterialParameterInfo& Lhs, const FMemoryImageMaterialParameterInfo& Rhs)
-{
-	return !operator==(Lhs, Rhs);
-}
-
-FORCEINLINE bool operator==(const FMaterialParameterInfo& Lhs, const FMemoryImageMaterialParameterInfo& Rhs)
-{
-	return Lhs.Name == Rhs.Name && Lhs.Index == Rhs.Index && Lhs.Association == Rhs.Association;
-}
-
-FORCEINLINE bool operator!=(const FMaterialParameterInfo& Lhs, const FMemoryImageMaterialParameterInfo& Rhs)
-{
-	return !operator==(Lhs, Rhs);
-}
-
-FORCEINLINE bool operator==(const FMemoryImageMaterialParameterInfo& Lhs, const FMaterialParameterInfo& Rhs)
-{
-	return Lhs.Name == Rhs.Name && Lhs.Index == Rhs.Index && Lhs.Association == Rhs.Association;
-}
-
-FORCEINLINE bool operator!=(const FMemoryImageMaterialParameterInfo& Lhs, const FMaterialParameterInfo& Rhs)
-{
-	return !operator==(Lhs, Rhs);
-}
-
-FORCEINLINE uint32 GetTypeHash(const FMemoryImageMaterialParameterInfo& Value)
-{
-	return HashCombine(HashCombine(GetTypeHash(Value.Name), Value.Index), (uint32)Value.Association);
-}
-
-// Backwards compat
-using FHashedMaterialParameterInfo = FMemoryImageMaterialParameterInfo;
-
 UENUM()
 enum class EMaterialLayerLinkState : uint8
 {
diff --git a/Engine/Source/Runtime/Engine/Private/HLSLTree/HLSLTree.cpp b/Engine/Source/Runtime/Engine/Private/HLSLTree/HLSLTree.cpp
--- a/Engine/Source/Runtime/Engine/Private/HLSLTree/HLSLTree.cpp
+++ b/Engine/Source/Runtime/Engine/Private/HLSLTree/HLSLTree.cpp
@@ -549,34 +549,56 @@ FRequestedType::FRequestedType(const Shader::FType& InType, bool bDefaultRequest
 	}
 }
 
-Shader::FType FRequestedType::GetType() const
+FRequestedType::FRequestedType(FName InObjectType) : ObjectType(InObjectType)
 {
-	if (StructType)
-	{
-		return StructType;
-	}
-	return Shader::MakeValueType(ValueComponentType, GetNumComponents());
+	RequestedComponents.Init(true, 1);
 }
 
 int32 FRequestedType::GetNumComponents() const
 {
+	if (DefaultType == EDefaultRequestedType::Any)
+	{
+		return 4; // TODO?
+	}
+
 	const int32 MaxComponentIndex = RequestedComponents.FindLast(true);
-	if (MaxComponentIndex != INDEX_NONE)
+	if (MaxComponentIndex == INDEX_NONE)
 	{
-		return MaxComponentIndex + 1;
+		return 0;
 	}
-	return 0;
+	return MaxComponentIndex + 1;
 }
 
-void FRequestedType::SetComponentRequest(int32 Index, bool bRequested)
+Shader::FType FRequestedType::GetType() const
 {
-	if (bRequested)
+	if (StructType)
+	{
+		return StructType;
+	}
+	else if (!ObjectType.IsNone())
 	{
-		RequestedComponents.PadToNum(Index + 1, false);
+		return ObjectType;
 	}
-	if (RequestedComponents.IsValidIndex(Index))
+	else if (DefaultType == EDefaultRequestedType::Any)
 	{
-		RequestedComponents[Index] = bRequested;
+		return FName(TEXT("Any"));
+	}
+
+	return Shader::MakeValueType(ValueComponentType, GetNumComponents());
+}
+
+void FRequestedType::SetComponentRequest(int32 Index, bool bRequested)
+{
+	if (DefaultType != EDefaultRequestedType::Any)
+	{
+		if (bRequested)
+		{
+			RequestedComponents.PadToNum(Index + 1, false);
+		}
+		if (RequestedComponents.IsValidIndex(Index))
+		{
+			RequestedComponents[Index] = bRequested;
+		}
 	}
 }
 
@@ -691,6 +713,10 @@ int32 FPreparedType::GetNumComponents() const
 	{
 		return StructType->ComponentTypes.Num();
 	}
+	else if (!ObjectType.IsNone())
+	{
+		return 1;
+	}
 	else if (ValueComponentType != Shader::EValueComponentType::Void)
 	{
 		auto Predicate = [](const FPreparedComponent& InComponent) { return InComponent.IsRequested(); };
@@ -711,6 +737,11 @@ Shader::FType FPreparedType::GetType() const
 	{
 		return StructType;
 	}
+	else if (IsObject())
+	{
+		return ObjectType;
+	}
+
 	return Shader::MakeValueType(ValueComponentType, GetNumComponents());
 }
 
@@ -742,9 +773,8 @@ EExpressionEvaluation FPreparedType::GetEvaluation(const FEmitScope& Scope) cons
 
 EExpressionEvaluation FPreparedType::GetEvaluation(const FEmitScope& Scope, const FRequestedType& RequestedType) const
 {
-	const int32 NumComponents = RequestedType.GetNumComponents();
 	EExpressionEvaluation Result = EExpressionEvaluation::ConstantZero;
-	for (int32 Index = 0; Index < NumComponents; ++Index)
+	for (int32 Index = 0; Index < PreparedComponents.Num(); ++Index)
 	{
 		if (RequestedType.IsComponentRequested(Index))
 		{
@@ -982,9 +1012,12 @@ FPreparedType MakeNonLWCType(const FPreparedType& Type)
 	return Result;
 }
 
-bool FPrepareValueResult::TryMergePreparedType(FEmitContext& Context, const Shader::FStructType* StructType, Shader::EValueComponentType ComponentType)
+bool FPrepareValueResult::TryMergePreparedType(FEmitContext& Context,
+	const Shader::FStructType* StructType,
+	FName ObjectType,
+	Shader::EValueComponentType ComponentType)
 {
-	if (!StructType && ComponentType == Shader::EValueComponentType::Void)
+	if (!StructType && ObjectType.IsNone() && ComponentType == Shader::EValueComponentType::Void)
 	{
 		return false;
 	}
@@ -993,6 +1026,7 @@ bool FPrepareValueResult::TryMergePreparedType(FEmitContext& Context, const Shad
 	{
 		PreparedType.PreparedComponents.Reset();
 		PreparedType.ValueComponentType = ComponentType;
+		PreparedType.ObjectType = ObjectType;
 		PreparedType.StructType = StructType;
 		return true;
 	}
@@ -1000,11 +1034,20 @@ bool FPrepareValueResult::TryMergePreparedType(FEmitContext& Context, const Shad
 	if (StructType)
 	{
 		check(ComponentType == Shader::EValueComponentType::Void);
+		check(ObjectType.IsNone());
 		if (StructType != PreparedType.StructType)
 		{
 			return Context.Error(TEXT("Invalid type"));
 		}
 	}
+	else if (!ObjectType.IsNone())
+	{
+		check(ComponentType == Shader::EValueComponentType::Void);
+		if (ObjectType != PreparedType.ObjectType)
+		{
+			return Context.Error(TEXT("Invalid type"));
+		}
+	}
 	else
 	{
 		PreparedType.ValueComponentType = Shader::CombineComponentTypes(PreparedType.ValueComponentType, ComponentType);
@@ -1017,13 +1060,14 @@ bool FPrepareValueResult::SetTypeVoid()
 {
 	PreparedType.PreparedComponents.Reset();
 	PreparedType.ValueComponentType = Shader::EValueComponentType::Void;
+	PreparedType.ObjectType = FName();
 	PreparedType.StructType = nullptr;
 	return false;
 }
 
 bool FPrepareValueResult::SetType(FEmitContext& Context, const FRequestedType& RequestedType, EExpressionEvaluation Evaluation, const Shader::FType& Type)
 {
-	if (TryMergePreparedType(Context, Type.StructType, Shader::GetValueTypeDescription(Type.ValueType).ComponentType))
+	if (TryMergePreparedType(Context, Type.StructType, Type.ObjectType, Shader::GetValueTypeDescription(Type.ValueType).ComponentType))
 	{
 		if (Evaluation != EExpressionEvaluation::None)
 		{
@@ -1058,7 +1102,7 @@ bool FPrepareValueResult::SetType(FEmitContext& Context, const FRequestedType& R
 		}
 	}
 
-	if (TryMergePreparedType(Context, Type.StructType, ComponentType))
+	if (TryMergePreparedType(Context, Type.StructType, Type.ObjectType, ComponentType))
 	{
 		const int32 NumComponents = Type.GetNumComponents();
 		for (int32 Index = 0; Index < NumComponents; ++Index)
@@ -1109,6 +1153,16 @@ void FExpression::EmitValuePreshader(FEmitContext& Context, FEmitScope& Scope, c
 	check(false);
 }
 
+bool FExpression::EmitValueObject(FEmitContext& Context, FEmitScope& Scope, const FName& ObjectTypeName, void* OutObjectBase) const
+{
+	return false;
+}
+
+bool FExpression::EmitCustomHLSLParameter(FEmitContext& Context, FEmitScope& Scope, const FName& ObjectTypeName, const TCHAR* ParameterName, FEmitCustomHLSLParameterResult& OutResult) const
+{
+	return false;
+}
+
 FEmitShaderExpression* FExpression::GetValueShader(FEmitContext& Context, FEmitScope& Scope, const FRequestedType& RequestedType, const FPreparedType& PreparedType, const Shader::FType& ResultType) const
 {
 	FEmitOwnerScope OwnerScope(Context, this);
@@ -1263,6 +1317,36 @@ Shader::FValue FExpression::GetValueConstant(FEmitContext& Context, FEmitScope&
 	return GetValueConstant(Context, Scope, FRequestedType(ResultType), PreparedType, ResultType);
 }
 
+bool FExpression::GetValueObject(FEmitContext& Context, FEmitScope& Scope, const FName& ObjectTypeName, void* OutObjectBase) const
+{
+	const FPreparedType PreparedType = Context.GetPreparedType(this);
+	check(PreparedType.ObjectType == ObjectTypeName);
+	return EmitValueObject(Context, Scope, ObjectTypeName, OutObjectBase);
+}
+
+bool FExpression::CheckObjectSupportsCustomHLSL(FEmitContext& Context, FEmitScope& Scope, const FName& ObjectTypeName) const
+{
+	const FPreparedType PreparedType = Context.GetPreparedType(this);
+	check(PreparedType.ObjectType == ObjectTypeName);
+
+	FEmitCustomHLSLParameterResult UnusedResult;
+	return EmitCustomHLSLParameter(Context, Scope, ObjectTypeName, nullptr, UnusedResult);
+}
+
+void FExpression::GetObjectCustomHLSLParameter(FEmitContext& Context,
+	FEmitScope& Scope,
+	const FName& ObjectTypeName,
+	const TCHAR* ParameterName,
+	FStringBuilderBase& OutDeclarationCode,
+	FStringBuilderBase& OutForwardCode) const
+{
+	FEmitCustomHLSLParameterResult Result;
+	Result.DeclarationCode = &OutDeclarationCode;
+	Result.ForwardCode = &OutForwardCode;
+	const bool bResult = EmitCustomHLSLParameter(Context, Scope, ObjectTypeName, ParameterName, Result);
+	check(bResult);
+}
+
 bool FScope::HasParentScope(const FScope& InParentScope) const
 {
 	const FScope* CurrentScope = this;
diff --git a/Engine/Source/Runtime/Engine/Private/HLSLTree/HLSLTreeCommon.cpp b/Engine/Source/Runtime/Engine/Private/HLSLTree/HLSLTreeCommon.cpp
--- a/Engine/Source/Runtime/Engine/Private/HLSLTree/HLSLTreeCommon.cpp
+++ b/Engine/Source/Runtime/Engine/Private/HLSLTree/HLSLTreeCommon.cpp
@@ -145,164 +145,6 @@ void FExpressionConstant::EmitValuePreshader(FEmitContext& Context, FEmitScope&
 	OutResult.Preshader.WriteOpcode(Shader::EPreshaderOpcode::Constant).Write(Value);
 }
 
-namespace Private
-{
-Shader::EValueType GetTexCoordType(Shader::EValueType TextureType)
-{
-	switch (TextureType)
-	{
-	case Shader::EValueType::Texture2D:
-	case Shader::EValueType::TextureExternal: return Shader::EValueType::Float2;
-	case Shader::EValueType::Texture2DArray:
-	case Shader::EValueType::TextureCube:
-	case Shader::EValueType::Texture3D: return Shader::EValueType::Float3;
-	case Shader::EValueType::TextureCubeArray: return Shader::EValueType::Float4;
-	default: checkNoEntry(); return Shader::EValueType::Void;
-	}
-}
-}
-
-bool FExpressionTextureSample::PrepareValue(FEmitContext& Context, FEmitScope& Scope, const FRequestedType& RequestedType, FPrepareValueResult& OutResult) const
-{
-	// TODO - need the correct texture type?
-	const FPreparedType& TextureType = Context.PrepareExpression(TextureExpression, Scope, Shader::EValueType::Texture2D);
-	if (!Shader::IsTextureType(TextureType.ValueComponentType))
-	{
-		return Context.Error(TEXT("Expected texture"));
-	}
-
-	const FRequestedType RequestedTexCoordType = Private::GetTexCoordType(TextureType.GetType());
-	const FPreparedType& TexCoordType = Context.PrepareExpression(TexCoordExpression, Scope, RequestedTexCoordType);
-	if (TexCoordType.IsVoid())
-	{
-		return false;
-	}
-
-	if (AutomaticMipBiasExpression)
-	{
-		const FPreparedType& AutomaticMipBiasType = Context.PrepareExpression(AutomaticMipBiasExpression, Scope, Shader::EValueType::Bool1);
-		if (!IsConstantEvaluation(AutomaticMipBiasType.GetEvaluation(Scope, Shader::EValueType::Bool1)))
-		{
-			return Context.Error(TEXT("Automatic Mip Bias input must be constant"));
-		}
-	}
-
-	const bool bUseAnalyticDerivatives = Context.bUseAnalyticDerivatives && (MipValueMode != TMVM_MipLevel) && TexCoordDerivatives.IsValid();
-	if (MipValueMode == TMVM_Derivative || bUseAnalyticDerivatives)
-	{
-		Context.PrepareExpression(TexCoordDerivatives.ExpressionDdx, Scope, RequestedTexCoordType);
-		Context.PrepareExpression(TexCoordDerivatives.ExpressionDdy, Scope, RequestedTexCoordType);
-	}
-	else if (MipValueMode == TMVM_MipLevel || MipValueMode == TMVM_MipBias)
-	{
-		Context.PrepareExpression(MipValueExpression, Scope, Shader::EValueType::Float1);
-	}
-
-	return OutResult.SetType(Context, RequestedType, EExpressionEvaluation::Shader, Shader::EValueType::Float4);
-}
-
-void FExpressionTextureSample::EmitValueShader(FEmitContext& Context, FEmitScope& Scope, const FRequestedType& RequestedType, FEmitValueShaderResult& OutResult) const
-{
-	FEmitShaderExpression* EmitTexture = TextureExpression->GetValueShader(Context, Scope);
-	const Shader::EValueType TextureType = EmitTexture->Type;
-	check(Shader::IsTextureType(TextureType));
-
-	bool bVirtualTexture = false;
-	const TCHAR* SampleFunctionName = nullptr;
-	switch (TextureType)
-	{
-	case Shader::EValueType::Texture2D:
-		SampleFunctionName = TEXT("Texture2DSample");
-		break;
-	case Shader::EValueType::TextureCube:
-		SampleFunctionName = TEXT("TextureCubeSample");
-		break;
-	case Shader::EValueType::Texture2DArray:
-		SampleFunctionName = TEXT("Texture2DArraySample");
-		break;
-	case Shader::EValueType::Texture3D:
-		SampleFunctionName = TEXT("Texture3DSample");
-		break;
-	case Shader::EValueType::TextureExternal:
-		SampleFunctionName = TEXT("TextureExternalSample");
-		break;
-	/*case MCT_TextureVirtual:
-		// TODO
-		SampleFunctionName = TEXT("TextureVirtualSample");
-		bVirtualTexture = true;
-		break;*/
-	default:
-		checkNoEntry();
-		break;
-	}
-
-	const bool AutomaticViewMipBias = AutomaticMipBiasExpression ? AutomaticMipBiasExpression->GetValueConstant(Context, Scope, Shader::EValueType::Bool1).AsBoolScalar() : false;
-	TStringBuilder<256> FormattedSampler;
-	switch (SamplerSource)
-	{
-	case SSM_FromTextureAsset:
-		FormattedSampler.Appendf(TEXT("%s.Sampler"), EmitTexture->Reference);
-		break;
-	case SSM_Wrap_WorldGroupSettings:
-		FormattedSampler.Appendf(TEXT("GetMaterialSharedSampler(%s.Sampler,%s)"),
-			EmitTexture->Reference,
-			AutomaticViewMipBias ? TEXT("View.MaterialTextureBilinearWrapedSampler") : TEXT("Material.Wrap_WorldGroupSettings"));
-		break;
-	case SSM_Clamp_WorldGroupSettings:
-		FormattedSampler.Appendf(TEXT("GetMaterialSharedSampler(%s.Sampler,%s)"),
-			EmitTexture->Reference,
-			AutomaticViewMipBias ? TEXT("View.MaterialTextureBilinearClampedSampler") : TEXT("Material.Clamp_WorldGroupSettings"));
-		break;
-	default:
-		checkNoEntry();
-		break;
-	}
-
-	const Shader::EValueType TexCoordType = Private::GetTexCoordType(TextureType);
-	FEmitShaderExpression* TexCoordValue = TexCoordExpression->GetValueShader(Context, Scope, TexCoordType);
-
-	FEmitShaderExpression* TextureResult = nullptr;
-	if (MipValueMode == TMVM_MipLevel)
-	{
-		TextureResult = Context.EmitExpression(Scope, Shader::EValueType::Float4, TEXT("%Level(%.Texture, %, %, %)"),
-			SampleFunctionName,
-			EmitTexture,
-			FormattedSampler.ToString(),
-			TexCoordValue,
-			MipValueExpression->GetValueShader(Context, Scope, Shader::EValueType::Float1));
-	}
-	else if (MipValueMode == TMVM_Derivative || (Context.bUseAnalyticDerivatives && TexCoordDerivatives.IsValid()))
-	{
-		TextureResult = Context.EmitExpression(Scope, Shader::EValueType::Float4, TEXT("%Grad(%.Texture, %, %, %, %)"),
-			SampleFunctionName,
-			EmitTexture,
-			FormattedSampler.ToString(),
-			TexCoordValue,
-			TexCoordDerivatives.ExpressionDdx->GetValueShader(Context, Scope, TexCoordType),
-			TexCoordDerivatives.ExpressionDdy->GetValueShader(Context, Scope, TexCoordType));
-	}
-	else if (MipValueMode == TMVM_MipBias)
-	{
-		TextureResult = Context.EmitExpression(Scope, Shader::EValueType::Float4, TEXT("%Bias(%.Texture, %, %, %)"),
-			SampleFunctionName,
-			EmitTexture,
-			FormattedSampler.ToString(),
-			TexCoordValue,
-			MipValueExpression->GetValueShader(Context, Scope, Shader::EValueType::Float1));
-	}
-	else
-	{
-		check(MipValueMode == TMVM_None);
-		TextureResult = Context.EmitExpression(Scope, Shader::EValueType::Float4, TEXT("%(%.Texture, %, %)"),
-			SampleFunctionName,
-			EmitTexture,
-			FormattedSampler.ToString(),
-			TexCoordValue);
-	}
-
-	OutResult.Code = Context.EmitExpression(Scope, Shader::EValueType::Float4, TEXT("ApplyMaterialSamplerType(%, %.SamplerType)"), TextureResult, EmitTexture);
-}
-
 void FExpressionGetStructField::ComputeAnalyticDerivatives(FTree& Tree, FExpressionDerivatives& OutResult) const
 {
 	const FExpressionDerivatives StructDerivatives = Tree.GetAnalyticDerivatives(StructExpression);
@@ -1016,12 +858,21 @@ bool FExpressionCustomHLSL::PrepareValue(FEmitContext& Context, FEmitScope& Scop
 {
 	for (const FCustomHLSLInput& Input : Inputs)
 	{
-		// TODO - do we need to support untyped input here?  Add explicit custom input types?
-		const FPreparedType& InputType = Context.PrepareExpression(Input.Expression, Scope, Shader::EValueType::Float4);
+		const FPreparedType& InputType = Context.PrepareExpression(Input.Expression, Scope, EDefaultRequestedType::Any);
 		if (InputType.IsVoid())
 		{
 			return false;
 		}
+
+		if (InputType.IsObject())
+		{
+			if (!Input.Expression->CheckObjectSupportsCustomHLSL(Context, Scope, InputType.ObjectType))
+			{
+				return Context.Error(TEXT("Object type not supported for custom HLSL"));
+			}
+		}
+
+		Context.MarkInputType(Input.Expression, InputType.GetType());
 	}
 
 	return OutResult.SetType(Context, RequestedType, EExpressionEvaluation::Shader, Shader::FType(OutputStructType));
diff --git a/Engine/Source/Runtime/Engine/Private/HLSLTree/HLSLTreeEmit.cpp b/Engine/Source/Runtime/Engine/Private/HLSLTree/HLSLTreeEmit.cpp
--- a/Engine/Source/Runtime/Engine/Private/HLSLTree/HLSLTreeEmit.cpp
+++ b/Engine/Source/Runtime/Engine/Private/HLSLTree/HLSLTreeEmit.cpp
@@ -250,30 +250,16 @@ void EmitCustomHLSL(const FEmitCustomHLSL& EmitCustomHLSL, const TCHAR* Paramete
 
 	// Function that wraps custom HLSL, provides the interface expected by custom HLSL
 	// * LWC inputs have both LWC version and non-LWC version
-	// * Texture inputs are split into a Texture/SamplerState pair
 	// * First output is given through return value, additional outputs use inout function parameters
 	OutCode.Appendf(TEXT("%s CustomExpressionInternal%d(%s Parameters"), OutputStruct->Fields[0].Type.GetName(), EmitCustomHLSL.Index, ParametersTypeName);
 	for (const FEmitCustomHLSLInput& Input : EmitCustomHLSL.Inputs)
 	{
 		Shader::FType InputType = Input.Type;
-		if (InputType.IsTexture())
+		if (InputType.IsObject())
 		{
-			// Texture parameter followed by a SamplerState parameter
-			const TCHAR* TextureTypeName = nullptr;
-			switch (InputType.ValueType)
-			{
-			case Shader::EValueType::Texture2D: TextureTypeName = TEXT("Texture2D"); break;
-			case Shader::EValueType::Texture2DArray: TextureTypeName = TEXT("Texture2DArray"); break;
-			case Shader::EValueType::TextureCube: TextureTypeName = TEXT("TextureCube"); break;
-			case Shader::EValueType::TextureCubeArray: TextureTypeName = TEXT("TextureCubeArray"); break;
-			case Shader::EValueType::Texture3D: TextureTypeName = TEXT("Texture3D"); break;
-			default: checkNoEntry(); break;
-			}
-			OutCode.Appendf(TEXT(", %s "), TextureTypeName);
-			OutCode.Append(Input.Name);
-			OutCode.Appendf(TEXT(", SamplerState "));
-			OutCode.Append(Input.Name);
-			OutCode.Append(TEXT("Sampler"));
+			check(Input.ObjectDeclarationCode.Len() > 0);
+			OutCode.Append(TEXT(", "));
+			OutCode.Append(Input.ObjectDeclarationCode);
 		}
 		else
 		{
@@ -305,8 +291,17 @@ void EmitCustomHLSL(const FEmitCustomHLSL& EmitCustomHLSL, const TCHAR* Paramete
 	OutCode.Appendf(TEXT("%s CustomExpression%d(%s Parameters"), OutputStruct->Name, EmitCustomHLSL.Index, ParametersTypeName);
 	for (const FEmitCustomHLSLInput& Input : EmitCustomHLSL.Inputs)
 	{
-		OutCode.Appendf(TEXT(", %s "), Input.Type.GetName());
-		OutCode.Append(Input.Name);
+		Shader::FType InputType = Input.Type;
+		if (InputType.IsObject())
+		{
+			OutCode.Append(TEXT(", "));
+			OutCode.Append(Input.ObjectDeclarationCode);
+		}
+		else
+		{
+			OutCode.Appendf(TEXT(", %s "), InputType.GetName());
+			OutCode.Append(Input.Name);
+		}
 	}
 	OutCode.Append(TEXT(")\n{\n"));
 	OutCode.Appendf(TEXT("\t%s Result = (%s)0;\n"), OutputStruct->Name, OutputStruct->Name);
@@ -314,19 +309,21 @@ void EmitCustomHLSL(const FEmitCustomHLSL& EmitCustomHLSL, const TCHAR* Paramete
 	for (const FEmitCustomHLSLInput& Input : EmitCustomHLSL.Inputs)
 	{
 		OutCode.Append(TEXT(", "));
-		OutCode.Append(Input.Name);
-		if (Input.Type.IsTexture())
+		if (Input.Type.IsObject() &&
+			Input.ObjectForwardCode.Len() > 0)
 		{
-			// Pass the Texture and SamplerState to the wrapper
-			OutCode.Append(TEXT(".Texture, "));
-			OutCode.Append(Input.Name);
-			OutCode.Append(TEXT(".Sampler"));
+			// CustomHLSL object parameters may have dedicated code to forward parameters from 1 function to another
+			OutCode.Append(Input.ObjectForwardCode);
 		}
-		else if (Input.Type.IsNumericLWC())
+		else
 		{
-			OutCode.Append(TEXT(", LWCToFloat("));
 			OutCode.Append(Input.Name);
-			OutCode.Append(TEXT(")"));
+			if (Input.Type.IsNumericLWC())
+			{
+				OutCode.Append(TEXT(", LWCToFloat("));
+				OutCode.Append(Input.Name);
+				OutCode.Append(TEXT(")"));
+			}
 		}
 	}
 	for (int32 OutputIndex = 1; OutputIndex < OutputStruct->Fields.Num(); ++OutputIndex)
@@ -1265,15 +1262,32 @@ FEmitShaderExpression* FEmitContext::EmitCustomHLSL(FEmitScope& Scope, FStringVi
 {
 	TArray<FEmitCustomHLSLInput, TInlineAllocator<8>> EmitInputs;
 	EmitInputs.Reserve(Inputs.Num());
+
+	FHasher Hasher;
 	for (const FCustomHLSLInput& Input : Inputs)
 	{
-		EmitInputs.Add(FEmitCustomHLSLInput{ Input.Name, GetType(Input.Expression) });
+		const Shader::FType InputType = GetType(Input.Expression);
+		AppendHash(Hasher, Input.Name);
+		AppendHash(Hasher, InputType);
+
+		FStringView ObjectDeclarationCode;
+		FStringView ObjectForwardCode;
+		if (InputType.IsObject())
+		{
+			TStringBuilder<256> FormattedDeclarationCode;
+			TStringBuilder<256> FormattedForwardCode;
+			Input.Expression->GetObjectCustomHLSLParameter(*this, Scope, InputType.ObjectType, Input.Name.GetData(), FormattedDeclarationCode, FormattedForwardCode);
+			ObjectDeclarationCode = MemStack::AllocateStringView(*Allocator, FormattedDeclarationCode.ToView());
+			ObjectForwardCode = MemStack::AllocateStringView(*Allocator, FormattedForwardCode.ToView());
+			AppendHash(Hasher, ObjectDeclarationCode);
+			AppendHash(Hasher, ObjectForwardCode);
+		}
+
+		EmitInputs.Add(FEmitCustomHLSLInput{ Input.Name, ObjectDeclarationCode, ObjectForwardCode, InputType });
 	}
 
-	FHasher Hasher;
 	AppendHash(Hasher, DeclarationCode);
 	AppendHash(Hasher, FunctionCode);
-	AppendHash(Hasher, MakeArrayView(EmitInputs));
 	AppendHash(Hasher, OutputType);
 	const FXxHash64 Hash = Hasher.Finalize();
 
diff --git a/Engine/Source/Runtime/Engine/Private/Materials/MaterialExpressionHLSL.cpp b/Engine/Source/Runtime/Engine/Private/Materials/MaterialExpressionHLSL.cpp
--- a/Engine/Source/Runtime/Engine/Private/Materials/MaterialExpressionHLSL.cpp
+++ b/Engine/Source/Runtime/Engine/Private/Materials/MaterialExpressionHLSL.cpp
@@ -313,7 +313,7 @@ bool UMaterialExpressionParameter::GenerateHLSLExpression(FMaterialHLSLGenerator
 		return Generator.Error(TEXT("Failed to get parameter value"));
 	}
 
-	OutExpression = Generator.GenerateMaterialParameter(ParameterName, ParameterMeta, ParameterMeta.Value.AsShaderValue());
+	OutExpression = Generator.GenerateMaterialParameter(ParameterName, ParameterMeta);
 	return true;
 }
 
@@ -335,23 +335,22 @@ bool UMaterialExpressionCurveAtlasRowParameter::GenerateHLSLExpression(FMaterial
 	{
 		return Generator.Error(TEXT("Failed to get parameter value"));
 	}
-	const FExpression* ExpressionSlot = Generator.GenerateMaterialParameter(ParameterName, ParameterMeta, ParameterMeta.Value.AsShaderValue());
+	const FExpression* ExpressionSlot = Generator.GenerateMaterialParameter(ParameterName, ParameterMeta);
 
 	const FExpression* ExpressionTexture = nullptr;
 	{
 		const FMaterialParameterMetadata TextureParameterMeta(Atlas.Get());
-		const FTextureValue* TextureValue = Generator.AcquireTextureValue(FTextureValue(Atlas.Get(), SAMPLERTYPE_LinearColor));
-		ExpressionTexture = Generator.GenerateMaterialParameter(FName(), TextureParameterMeta, TextureValue);
+		ExpressionTexture = Generator.GenerateMaterialParameter(FName(), TextureParameterMeta, SAMPLERTYPE_LinearColor);
 	}
 
 	FTree& Tree = Generator.GetTree();
-	const FExpression* ExpressionTextureSize = Tree.NewExpression<Material::FExpressionTextureSize>(FMaterialParameterInfo(), Atlas.Get());
+	const FExpression* ExpressionTextureSize = Tree.NewExpression<Material::FExpressionTextureSize>(ExpressionTexture);
 	const FExpression* ExpressionTextureHeight = Tree.NewExpression<FExpressionSwizzle>(MakeSwizzleMask(false, true, false, false), ExpressionTextureSize);
 	const FExpression* ExpressionCoordU = InputTime.AcquireHLSLExpressionOrConstant(Generator, Scope, 0.0f);
 	const FExpression* ExpressionCoordV = Tree.NewMul(Tree.NewAdd(ExpressionSlot, Tree.NewConstant(0.5f)), Tree.NewRcp(ExpressionTextureHeight));
 	const FExpression* ExpressionUV = Tree.NewExpression<FExpressionAppend>(ExpressionCoordU, ExpressionCoordV);
 
-	OutExpression = Generator.GetTree().NewExpression<FExpressionTextureSample>(ExpressionTexture, ExpressionUV, nullptr, nullptr, FExpressionDerivatives(), SSM_Clamp_WorldGroupSettings, TMVM_None);
+	OutExpression = Generator.GetTree().NewExpression<Material::FExpressionTextureSample>(ExpressionTexture, ExpressionUV, nullptr, nullptr, FExpressionDerivatives(), SSM_Clamp_WorldGroupSettings, TMVM_None);
 	return true;
 }
 
@@ -372,7 +371,7 @@ bool UMaterialExpressionStaticSwitchParameter::GenerateHLSLExpression(FMaterialH
 		return false;
 	}
 
-	const FExpression* ExpressionSwitch = Generator.GenerateMaterialParameter(ParameterName, ParameterMeta, ParameterMeta.Value.AsShaderValue());
+	const FExpression* ExpressionSwitch = Generator.GenerateMaterialParameter(ParameterName, ParameterMeta);
 	OutExpression = Generator.GetTree().NewExpression <FExpressionSelect>(ExpressionSwitch, ExpressionA, ExpressionB);
 	return true;
 }
@@ -684,8 +683,7 @@ bool UMaterialExpressionTextureObject::GenerateHLSLExpression(FMaterialHLSLGener
 	using namespace UE::Shader;
 
 	const FMaterialParameterMetadata ParameterMeta(Texture);
-	const FTextureValue* TextureValue = Generator.AcquireTextureValue(FTextureValue(Texture, SamplerType));
-	OutExpression = Generator.GenerateMaterialParameter(FName(), ParameterMeta, TextureValue);
+	OutExpression = Generator.GenerateMaterialParameter(FName(), ParameterMeta, SamplerType);
 	return true;
 }
 
@@ -700,8 +698,7 @@ bool UMaterialExpressionTextureObjectParameter::GenerateHLSLExpression(FMaterial
 		return Generator.Error(TEXT("Failed to get parameter value"));
 	}
 
-	const FTextureValue* TextureValue = Generator.AcquireTextureValue(FTextureValue(Texture, SamplerType));
-	OutExpression = Generator.GenerateMaterialParameter(ParameterName, ParameterMeta, TextureValue);
+	OutExpression = Generator.GenerateMaterialParameter(ParameterName, ParameterMeta, SamplerType);
 	return true;
 }
 
@@ -743,7 +740,7 @@ bool UMaterialExpressionTextureSample::GenerateHLSLExpressionBase(FMaterialHLSLG
 	}
 
 	const FExpression* AutomaticMipBiasExpression = AutomaticViewMipBiasValue.AcquireHLSLExpressionOrConstant(Generator, Scope, (bool)AutomaticViewMipBias);
-	OutExpression = Generator.GetTree().NewExpression<FExpressionTextureSample>(TextureExpression, TexCoordExpression, MipLevelExpression, AutomaticMipBiasExpression, TexCoordDerivatives, SamplerSource, MipValueMode);
+	OutExpression = Generator.GetTree().NewExpression<Material::FExpressionTextureSample>(TextureExpression, TexCoordExpression, MipLevelExpression, AutomaticMipBiasExpression, TexCoordDerivatives, SamplerSource, MipValueMode);
 	return true;
 }
 
@@ -759,8 +756,7 @@ bool UMaterialExpressionTextureSample::GenerateHLSLExpression(FMaterialHLSLGener
 	else if (Texture)
 	{
 		const FMaterialParameterMetadata ParameterMeta(Texture);
-		const FTextureValue* TextureValue = Generator.AcquireTextureValue(FTextureValue(Texture, SamplerType));
-		TextureExpression = Generator.GenerateMaterialParameter(FName(), ParameterMeta, TextureValue);
+		TextureExpression = Generator.GenerateMaterialParameter(FName(), ParameterMeta, SamplerType);
 	}
 
 	return GenerateHLSLExpressionBase(Generator, Scope, TextureExpression, OutExpression);
@@ -782,9 +778,7 @@ bool UMaterialExpressionTextureSampleParameter::GenerateHLSLExpression(FMaterial
 		{
 			return Generator.Error(TEXT("Failed to get parameter value"));
 		}
-
-		const FTextureValue* TextureValue = Generator.AcquireTextureValue(FTextureValue(Texture, SamplerType));
-		TextureExpression = Generator.GenerateMaterialParameter(ParameterName, ParameterMeta, TextureValue);
+		TextureExpression = Generator.GenerateMaterialParameter(ParameterName, ParameterMeta, SamplerType);
 	}
 
 	return GenerateHLSLExpressionBase(Generator, Scope, TextureExpression, OutExpression);
@@ -832,19 +826,17 @@ bool UMaterialExpressionFontSample::GenerateHLSLExpression(FMaterialHLSLGenerato
 		return Generator.Errorf(TEXT("%s"), *SamplerTypeError);
 	}*/
 
-	const FTextureValue* TextureValue = Generator.AcquireTextureValue(FTextureValue(Texture, ExpectedSamplerType));
-
 	FMaterialParameterMetadata ParameterMeta;
 	if (!GetParameterValue(ParameterMeta))
 	{
 		ParameterMeta.Value = Texture;
 	}
 
-	const FExpression* TextureExpression = Generator.GenerateMaterialParameter(GetParameterName(), ParameterMeta, TextureValue);
+	const FExpression* TextureExpression = Generator.GenerateMaterialParameter(GetParameterName(), ParameterMeta, ExpectedSamplerType);
 	const FExpression* TexCoordExpression = Generator.NewTexCoord(0);
 	const FExpressionDerivatives TexCoordDerivatives = Generator.GetTree().GetAnalyticDerivatives(TextureExpression);
 
-	OutExpression = Generator.GetTree().NewExpression<FExpressionTextureSample>(TextureExpression, TexCoordExpression, nullptr, nullptr, TexCoordDerivatives, SSM_FromTextureAsset, TMVM_None);
+	OutExpression = Generator.GetTree().NewExpression<Material::FExpressionTextureSample>(TextureExpression, TexCoordExpression, nullptr, nullptr, TexCoordDerivatives, SSM_FromTextureAsset, TMVM_None);
 	return true;
 }
 
diff --git a/Engine/Source/Runtime/Engine/Private/Materials/MaterialHLSLGenerator.cpp b/Engine/Source/Runtime/Engine/Private/Materials/MaterialHLSLGenerator.cpp
--- a/Engine/Source/Runtime/Engine/Private/Materials/MaterialHLSLGenerator.cpp
+++ b/Engine/Source/Runtime/Engine/Private/Materials/MaterialHLSLGenerator.cpp
@@ -349,37 +349,6 @@ const UE::HLSLTree::FExpression* FMaterialHLSLGenerator::NewSwizzle(const UE::HL
 	return GetTree().NewSwizzle(Params, Input);
 }
 
-const UE::Shader::FTextureValue* FMaterialHLSLGenerator::AcquireTextureValue(const UE::Shader::FTextureValue& InValue)
-{
-	using namespace UE::Shader;
-	
-	// Need to move this to HLSLTreeEmit
-	/*FString SamplerTypeError;
-	if (!UMaterialExpressionTextureBase::VerifySamplerType(CompileTarget.FeatureLevel, CompileTarget.TargetPlatform, InValue.Texture, InValue.SamplerType, SamplerTypeError))
-	{
-		Errors.AddError(SamplerTypeError);
-		return nullptr;
-	}*/
-
-	FXxHash64Builder Hasher;
-	Hasher.Update(&InValue.Texture, sizeof(InValue.Texture));
-	Hasher.Update(&InValue.SamplerType, sizeof(InValue.SamplerType));
-	Hasher.Update(&InValue.ExternalTextureGuid, sizeof(InValue.ExternalTextureGuid));
-	const FXxHash64 Hash = Hasher.Finalize();
-
-	FTextureValue const* const* PrevValue = TextureValueMap.Find(Hash);
-	if (PrevValue)
-	{
-		check(*PrevValue);
-		check(**PrevValue == InValue);
-		return *PrevValue;
-	}
-
-	FTextureValue* Value = new(GetTree().GetAllocator()) FTextureValue(InValue);
-	TextureValueMap.Add(Hash, Value);
-	return Value;
-}
-
 const UE::HLSLTree::FExpression* FMaterialHLSLGenerator::InternalNewErrorExpression(FStringView Error)
 {
 	return GetTree().NewExpression<UE::HLSLTree::FExpressionError>(UE::MemStack::AllocateStringView(GetTree().GetAllocator(), Error));
@@ -586,13 +555,14 @@ bool FMaterialHLSLGenerator::GenerateStatements(UE::HLSLTree::FScope& Scope, UMa
 	return bResult;
 }
 
-const UE::HLSLTree::FExpression* FMaterialHLSLGenerator::GenerateMaterialParameter(FName InParameterName, const FMaterialParameterMetadata& InParameterMeta, const UE::Shader::FValue& InDefaultValue)
+const UE::HLSLTree::FExpression* FMaterialHLSLGenerator::GenerateMaterialParameter(FName InParameterName,
+	const FMaterialParameterMetadata& InParameterMeta,
+	EMaterialSamplerType InSamplerType,
+	const FGuid& InExternalTextureGuid)
 {
 	using namespace UE::Shader;
 
 	FMaterialParameterMetadata ParameterMeta(InParameterMeta);
-	FValue DefaultValue(InDefaultValue);
-
 	FMaterialParameterMetadata OverrideParameterMeta;
 	if (GetParameterOverrideValueForCurrentFunction(InParameterMeta.Value.Type, InParameterName, OverrideParameterMeta))
 	{
@@ -601,20 +571,9 @@ const UE::HLSLTree::FExpression* FMaterialHLSLGenerator::GenerateMaterialParamet
 		ParameterMeta.bUsedAsAtlasPosition = OverrideParameterMeta.bUsedAsAtlasPosition;
 		ParameterMeta.ScalarAtlas = OverrideParameterMeta.ScalarAtlas;
 		ParameterMeta.ScalarCurve = OverrideParameterMeta.ScalarCurve;
-
-		if (DefaultValue.Type.IsTexture())
-		{
-			FTextureValue TextureValue(*DefaultValue.AsTexture());
-			TextureValue.Texture = OverrideParameterMeta.Value.Texture;
-			DefaultValue = AcquireTextureValue(TextureValue);
-		}
-		else
-		{
-			DefaultValue = OverrideParameterMeta.Value.AsShaderValue();
-		}
 	}
 
-	return GetTree().NewExpression<UE::HLSLTree::Material::FExpressionParameter>(GetParameterInfo(InParameterName), ParameterMeta, DefaultValue);
+	return GetTree().NewExpression<UE::HLSLTree::Material::FExpressionParameter>(GetParameterInfo(InParameterName), ParameterMeta, InSamplerType, InExternalTextureGuid);
 }
 
 const UE::HLSLTree::FExpression* FMaterialHLSLGenerator::GenerateFunctionCall(UE::HLSLTree::FScope& Scope,
diff --git a/Engine/Source/Runtime/Engine/Private/Materials/MaterialHLSLTree.cpp b/Engine/Source/Runtime/Engine/Private/Materials/MaterialHLSLTree.cpp
--- a/Engine/Source/Runtime/Engine/Private/Materials/MaterialHLSLTree.cpp
+++ b/Engine/Source/Runtime/Engine/Private/Materials/MaterialHLSLTree.cpp
@@ -362,6 +362,75 @@ void FExpressionShadingModel::EmitValuePreshader(FEmitContext& Context, FEmitSco
 	OutResult.Preshader.WriteOpcode(Shader::EPreshaderOpcode::Constant).Write(Shader::FValue((int32)ShadingModel));
 }
 
+namespace Private
+{
+
+void EmitTextureShader(FEmitContext& Context,
+	const FMaterialTextureValue& TextureValue,
+	FStringBuilderBase& OutCode)
+{
+	const EMaterialValueType TextureType = TextureValue.Texture->GetMaterialType();
+
+	bool bVirtualTexture = false;
+	const TCHAR* TextureTypeName = nullptr;
+	int32 TextureParameterIndex = INDEX_NONE;
+	if (TextureType == MCT_TextureExternal)
+	{
+		check(TextureValue.SamplerType == SAMPLERTYPE_External);
+		TextureTypeName = TEXT("ExternalTexture");
+
+		FMaterialExternalTextureParameterInfo TextureParameterInfo;
+		TextureParameterInfo.ParameterName = NameToScriptName(TextureValue.ParameterInfo.Name);
+		TextureParameterInfo.ExternalTextureGuid = TextureValue.ExternalTextureGuid;
+		if (TextureValue.Texture)
+		{
+			TextureParameterInfo.SourceTextureIndex = Context.Material->GetReferencedTextures().Find(TextureValue.Texture);
+		}
+		TextureParameterIndex = Context.MaterialCompilationOutput->UniformExpressionSet.FindOrAddExternalTextureParameter(TextureParameterInfo);
+	}
+	else
+	{
+		EMaterialTextureParameterType TextureParameterType = EMaterialTextureParameterType::Count;
+		switch (TextureType)
+		{
+		case MCT_Texture2D:
+			TextureParameterType = EMaterialTextureParameterType::Standard2D;
+			TextureTypeName = TEXT("Texture2D");
+			break;
+		case MCT_Texture2DArray:
+			TextureParameterType = EMaterialTextureParameterType::Array2D;
+			TextureTypeName = TEXT("Texture2DArray");
+			break;
+		case MCT_TextureCube:
+			TextureParameterType = EMaterialTextureParameterType::Cube;
+			TextureTypeName = TEXT("TextureCube");
+			break;
+		case MCT_TextureCubeArray:
+			TextureParameterType = EMaterialTextureParameterType::ArrayCube;
+			TextureTypeName = TEXT("TextureCubeArray");
+			break;
+		case MCT_VolumeTexture:
+			TextureParameterType = EMaterialTextureParameterType::Volume;
+			TextureTypeName = TEXT("VolumeTexture");
+			break;
+		default:
+			checkNoEntry();
+			break;
+		}
+
+		FMaterialTextureParameterInfo TextureParameterInfo;
+		TextureParameterInfo.ParameterInfo = TextureValue.ParameterInfo;
+		TextureParameterInfo.TextureIndex = Context.Material->GetReferencedTextures().Find(TextureValue.Texture);
+		TextureParameterInfo.SamplerSource = SSM_FromTextureAsset; // TODO - Is this needed?
+		check(TextureParameterInfo.TextureIndex != INDEX_NONE);
+		TextureParameterIndex = Context.MaterialCompilationOutput->UniformExpressionSet.FindOrAddTextureParameter(TextureParameterType, TextureParameterInfo);
+	}
+
+	OutCode.Appendf(TEXT("Material.%s_%d"), TextureTypeName, TextureParameterIndex);
+}
+
+} // namespace Private
+
 void FExpressionParameter::ComputeAnalyticDerivatives(FTree& Tree, FExpressionDerivatives& OutResult) const
 {
 	const Shader::FType Type = GetShaderValueType(ParameterMeta.Value.Type);
@@ -377,6 +446,8 @@ void FExpressionParameter::ComputeAnalyticDerivatives(FTree& Tree, FExpressionDe
 bool FExpressionParameter::PrepareValue(FEmitContext& Context, FEmitScope& Scope, const FRequestedType& RequestedType, FPrepareValueResult& OutResult) const
 {
 	const EMaterialParameterType ParameterType = ParameterMeta.Value.Type;
+	const Shader::FType ResultType = GetShaderValueType(ParameterType);
+
 	if (Context.bMarkLiveValues)
 	{
 		FEmitData& EmitData = Context.FindData<FEmitData>();
@@ -387,20 +458,7 @@ bool FExpressionParameter::PrepareValue(FEmitContext& Context, FEmitScope& Scope
 				EmitData.CachedExpressionData->Parameters.AddParameter(ParameterInfo, ParameterMeta);
 			}
 
-			UObject* ReferencedTexture = nullptr;
-			switch (ParameterType)
-			{
-			case EMaterialParameterType::Texture: ReferencedTexture = ParameterMeta.Value.Texture; break;
-			case EMaterialParameterType::RuntimeVirtualTexture: ReferencedTexture = ParameterMeta.Value.RuntimeVirtualTexture; break;
-			case EMaterialParameterType::Font:
-				if (ParameterMeta.Value.Font.Value && ParameterMeta.Value.Font.Value->Textures.IsValidIndex(ParameterMeta.Value.Font.Page))
-				{
-					ReferencedTexture = ParameterMeta.Value.Font.Value->Textures[ParameterMeta.Value.Font.Page];
-				}
-				break;
-			default:
-				break;
-			}
+			UObject* ReferencedTexture = ParameterMeta.Value.AsTextureObject();
 			if (ReferencedTexture)
 			{
 				EmitData.CachedExpressionData->ReferencedTextures.AddUnique(ReferencedTexture);
@@ -419,92 +477,26 @@ bool FExpressionParameter::PrepareValue(FEmitContext& Context, FEmitScope& Scope
 	{
 		Evaluation = EExpressionEvaluation::Preshader;
 	}
-	return OutResult.SetType(Context, RequestedType, Evaluation, DefaultValue.Type);
+
+	return OutResult.SetType(Context, RequestedType, Evaluation, ResultType);
 }
 
 void FExpressionParameter::EmitValueShader(FEmitContext& Context, FEmitScope& Scope, const FRequestedType& RequestedType, FEmitValueShaderResult& OutResult) const
 {
-	const EMaterialParameterType ParameterType = ParameterMeta.Value.Type;
-	if (ParameterType == EMaterialParameterType::Texture ||
-		ParameterType == EMaterialParameterType::RuntimeVirtualTexture ||
-		ParameterType == EMaterialParameterType::Font)
+	FMaterialTextureValue TextureValue;
+	if (GetValueObject(Context, Scope, TextureValue))
 	{
-		const Shader::FTextureValue* TextureValue = DefaultValue.AsTexture();
-		check(TextureValue);
-
-		const Shader::EValueType TextureType = TextureValue->GetType();
-		const TCHAR* ConstructorName = nullptr;
-		const TCHAR* TextureTypeName = nullptr;
-		int32 ParameterIndex = INDEX_NONE;
-
-		if (TextureType == Shader::EValueType::TextureExternal)
-		{
-			check(TextureValue->SamplerType == SAMPLERTYPE_External);
-
-			ConstructorName = TEXT("MakeTextureExternal");
-			TextureTypeName = TEXT("ExternalTexture");
-
-			FMaterialExternalTextureParameterInfo TextureParameterInfo;
-			TextureParameterInfo.ParameterName = NameToScriptName(ParameterInfo.Name);
-			TextureParameterInfo.ExternalTextureGuid = TextureValue->ExternalTextureGuid;
-			if (TextureValue->Texture)
-			{
-				TextureParameterInfo.SourceTextureIndex = Context.Material->GetReferencedTextures().Find(TextureValue->Texture);
-			}
-			ParameterIndex = Context.MaterialCompilationOutput->UniformExpressionSet.FindOrAddExternalTextureParameter(TextureParameterInfo);
-		}
-		else
-		{
-			EMaterialTextureParameterType TextureParameterType = EMaterialTextureParameterType::Count;
-			switch (TextureType)
-			{
-			case Shader::EValueType::Texture2D:
-				TextureParameterType = EMaterialTextureParameterType::Standard2D;
-				ConstructorName = TEXT("MakeTexture2D");
-				TextureTypeName = TEXT("Texture2D");
-				break;
-			case Shader::EValueType::Texture2DArray:
-				TextureParameterType = EMaterialTextureParameterType::Array2D;
-				ConstructorName = TEXT("MakeTexture2DArray");
-				TextureTypeName = TEXT("Texture2DArray");
-				break;
-			case Shader::EValueType::TextureCube:
-				TextureParameterType = EMaterialTextureParameterType::Cube;
-				ConstructorName = TEXT("MakeTextureCube");
-				TextureTypeName = TEXT("TextureCube");
-				break;
-			case Shader::EValueType::TextureCubeArray:
-				TextureParameterType = EMaterialTextureParameterType::ArrayCube;
-				ConstructorName = TEXT("MakeTextureCubeArray");
-				TextureTypeName = TEXT("TextureCubeArray");
-				break;
-			case Shader::EValueType::Texture3D:
-				TextureParameterType = EMaterialTextureParameterType::Volume;
-				ConstructorName = TEXT("MakeTexture3D");
-				TextureTypeName = TEXT("VolumeTexture");
-				break;
-			default:
-				checkNoEntry();
-				break;
-			}
-
-			FMaterialTextureParameterInfo TextureParameterInfo;
-			TextureParameterInfo.ParameterInfo = ParameterInfo;
-			TextureParameterInfo.TextureIndex = Context.Material->GetReferencedTextures().Find(TextureValue->Texture);
-			TextureParameterInfo.SamplerSource = SSM_FromTextureAsset; // TODO - Is this needed?
-			check(TextureParameterInfo.TextureIndex != INDEX_NONE);
-			ParameterIndex = Context.MaterialCompilationOutput->UniformExpressionSet.FindOrAddTextureParameter(TextureParameterType, TextureParameterInfo);
-		}
-
-		TStringBuilder<256> FormattedCode;
-		FormattedCode.Appendf(TEXT("%s(Material.%s_%d, Material.%s_%dSampler, %d)"), ConstructorName, TextureTypeName, ParameterIndex, TextureTypeName, ParameterIndex, (int32)TextureValue->SamplerType);
-		OutResult.Code = Context.EmitExpression(Scope, TextureType, FormattedCode.ToView());
+		TStringBuilder<64> FormattedTexture;
+		Private::EmitTextureShader(Context, TextureValue, FormattedTexture);
+		// Emit a texture/sampler pair1
+		OutResult.Code = Context.EmitInlineExpression(Scope, FMaterialTextureValue::GetTypeName(), TEXT("%,%Sampler"), FormattedTexture.ToString(), FormattedTexture.ToString());
 	}
 }
 
 void FExpressionParameter::EmitValuePreshader(FEmitContext& Context, FEmitScope& Scope, const FRequestedType& RequestedType, FEmitValuePreshaderResult& OutResult) const
 {
 	const EMaterialParameterType ParameterType = ParameterMeta.Value.Type;
+	const Shader::FValue DefaultValue = ParameterMeta.Value.AsShaderValue();
 	FEmitData& EmitMaterialData = Context.FindData<FEmitData>();
 
 	Context.PreshaderStackPosition++;
@@ -527,6 +519,7 @@ void FExpressionParameter::EmitValuePreshader(FEmitContext& Context, FEmitScope&
 	}
 	else
 	{
+		check(IsNumericMaterialParameter(ParameterType));
 		const uint32* PrevDefaultOffset = EmitMaterialData.DefaultUniformValues.Find(DefaultValue);
 		uint32 DefaultOffset;
 		if (PrevDefaultOffset)
@@ -544,27 +537,242 @@ void FExpressionParameter::EmitValuePreshader(FEmitContext& Context, FEmitScope&
 	}
 }
 
-bool FExpressionTextureSize::PrepareValue(FEmitContext& Context, FEmitScope& Scope, const FRequestedType& RequestedType, FPrepareValueResult& OutResult) const
+bool FExpressionParameter::EmitValueObject(FEmitContext& Context, FEmitScope& Scope, const FName& ObjectTypeName, void* OutObjectBase) const
 {
-	if (Context.bMarkLiveValues)
+	UTexture* Texture = Cast<UTexture>(ParameterMeta.Value.AsTextureObject());
+	if (Texture && ObjectTypeName == FMaterialTextureValue::GetTypeName())
 	{
-		FEmitData& EmitData = Context.FindData<FEmitData>();
-		if (EmitData.CachedExpressionData)
+		FMaterialTextureValue& OutObject = *static_cast<FMaterialTextureValue*>(OutObjectBase);
+		OutObject.ParameterInfo = ParameterInfo;
+		OutObject.Texture = Texture;
+		OutObject.SamplerType = TextureSamplerType;
+		OutObject.ExternalTextureGuid = ExternalTextureGuid;
+		return true;
+	}
+	return false;
+}
+
+bool FExpressionParameter::EmitCustomHLSLParameter(FEmitContext& Context, FEmitScope& Scope, const FName& ObjectTypeName, const TCHAR* ParameterName, FEmitCustomHLSLParameterResult& OutResult) const
+{
+	UTexture* Texture = Cast<UTexture>(ParameterMeta.Value.AsTextureObject());
+	if (Texture && ObjectTypeName == FMaterialTextureValue::GetTypeName())
+	{
+		const EMaterialValueType TextureMaterialType = Texture->GetMaterialType();
+		const TCHAR* TextureTypeName = nullptr;
+		switch (TextureMaterialType)
 		{
-			EmitData.CachedExpressionData->ReferencedTextures.AddUnique(Texture);
+		case MCT_Texture2D: TextureTypeName = TEXT("Texture2D"); break;
+		case MCT_Texture2DArray: TextureTypeName = TEXT("Texture2DArray"); break;
+		case MCT_TextureCube: TextureTypeName = TEXT("TextureCube"); break;
+		case MCT_TextureCubeArray: TextureTypeName = TEXT("TextureCubeArray"); break;
+		case MCT_VolumeTexture: TextureTypeName = TEXT("Texture3D"); break;
+		default: break;
 		}
+
+		if (TextureTypeName)
+		{
+			if (OutResult.DeclarationCode)
+			{
+				OutResult.DeclarationCode->Appendf(TEXT("%s %s, SamplerState %sSampler"), TextureTypeName, ParameterName, ParameterName);
+			}
+			if (OutResult.ForwardCode)
+			{
+				OutResult.ForwardCode->Appendf(TEXT("%s, %sSampler"), ParameterName, ParameterName);
+			}
+			return true;
+		}
+	}
+	return false;
+}
+
+namespace Private
+{
+Shader::EValueType GetTexCoordType(EMaterialValueType TextureType)
+{
+	switch (TextureType)
+	{
+	case MCT_Texture2D:
+	case MCT_TextureVirtual:
+	case MCT_TextureExternal: return Shader::EValueType::Float2;
+	case MCT_Texture2DArray:
+	case MCT_TextureCube:
+	case MCT_VolumeTexture: return Shader::EValueType::Float3;
+	case MCT_TextureCubeArray: return Shader::EValueType::Float4;
+	default: checkNoEntry(); return Shader::EValueType::Void;
+	}
+}
+} // namespace Private
+
+bool FExpressionTextureSample::PrepareValue(FEmitContext& Context, FEmitScope& Scope, const FRequestedType& RequestedType, FPrepareValueResult& OutResult) const
+{
+	const FPreparedType& TextureType = Context.PrepareExpression(TextureExpression, Scope, FMaterialTextureValue::GetTypeName());
+	if (TextureType.ObjectType != FMaterialTextureValue::GetTypeName())
+	{
+		return Context.Error(TEXT("Expected texture"));
+	}
+
+	FMaterialTextureValue TextureValue;
+	if (!TextureExpression->GetValueObject(Context, Scope, TextureValue))
+	{
+		return Context.Error(TEXT("Expected texture"));
+	}
+
+	const EMaterialValueType TextureMaterialType = TextureValue.Texture->GetMaterialType();
+	const FRequestedType RequestedTexCoordType = Private::GetTexCoordType(TextureMaterialType);
+	const FPreparedType& TexCoordType = Context.PrepareExpression(TexCoordExpression, Scope, RequestedTexCoordType);
+	if (TexCoordType.IsVoid())
+	{
+		return false;
+	}
+
+	if (AutomaticMipBiasExpression)
+	{
+		const FPreparedType& AutomaticMipBiasType = Context.PrepareExpression(AutomaticMipBiasExpression, Scope, Shader::EValueType::Bool1);
+		if (!IsConstantEvaluation(AutomaticMipBiasType.GetEvaluation(Scope, Shader::EValueType::Bool1)))
+		{
+			return Context.Error(TEXT("Automatic Mip Bias input must be constant"));
+		}
+	}
+
+	const bool bUseAnalyticDerivatives = Context.bUseAnalyticDerivatives && (MipValueMode != TMVM_MipLevel) && TexCoordDerivatives.IsValid();
+	if (MipValueMode == TMVM_Derivative || bUseAnalyticDerivatives)
+	{
+		Context.PrepareExpression(TexCoordDerivatives.ExpressionDdx, Scope, RequestedTexCoordType);
+		Context.PrepareExpression(TexCoordDerivatives.ExpressionDdy, Scope, RequestedTexCoordType);
+	}
+	else if (MipValueMode == TMVM_MipLevel || MipValueMode == TMVM_MipBias)
+	{
+		Context.PrepareExpression(MipValueExpression, Scope, Shader::EValueType::Float1);
 	}
+
+	return OutResult.SetType(Context, RequestedType, EExpressionEvaluation::Shader, Shader::EValueType::Float4);
+}
+
+void FExpressionTextureSample::EmitValueShader(FEmitContext& Context, FEmitScope& Scope, const FRequestedType& RequestedType, FEmitValueShaderResult& OutResult) const
+{
+	FMaterialTextureValue TextureValue;
+	TextureExpression->GetValueObject(Context, Scope, TextureValue);
+	const EMaterialValueType TextureType = TextureValue.Texture->GetMaterialType();
+
+	bool bVirtualTexture = false;
+	const TCHAR* SampleFunctionName = nullptr;
+	switch (TextureType)
+	{
+	case MCT_Texture2D:
+		SampleFunctionName = TEXT("Texture2DSample");
+		break;
+	case MCT_Texture2DArray:
+		SampleFunctionName = TEXT("Texture2DArraySample");
+		break;
+	case MCT_TextureCube:
+		SampleFunctionName = TEXT("TextureCubeSample");
+		break;
+	case MCT_TextureCubeArray:
+		SampleFunctionName = TEXT("TextureCubeArraySample");
+		break;
+	case MCT_VolumeTexture:
+		SampleFunctionName = TEXT("Texture3DSample");
+		break;
+	case MCT_TextureExternal:
+		SampleFunctionName = TEXT("TextureExternalSample");
+		break;
+	default:
+		checkNoEntry();
+		break;
+	}
+
+	TStringBuilder<64> FormattedTexture;
+	Private::EmitTextureShader(Context, TextureValue, FormattedTexture);
+
+	const bool AutomaticViewMipBias = AutomaticMipBiasExpression ? AutomaticMipBiasExpression->GetValueConstant(Context, Scope, Shader::EValueType::Bool1).AsBoolScalar() : false;
+	TStringBuilder<256> FormattedSampler;
+	switch (SamplerSource)
+	{
+	case SSM_FromTextureAsset:
+		FormattedSampler.Appendf(TEXT("%sSampler"), FormattedTexture.ToString());
+		break;
+	case SSM_Wrap_WorldGroupSettings:
+		FormattedSampler.Appendf(TEXT("GetMaterialSharedSampler(%sSampler,%s)"),
+			FormattedTexture.ToString(),
+			AutomaticViewMipBias ? TEXT("View.MaterialTextureBilinearWrapedSampler") : TEXT("Material.Wrap_WorldGroupSettings"));
+		break;
+	case SSM_Clamp_WorldGroupSettings:
+		FormattedSampler.Appendf(TEXT("GetMaterialSharedSampler(%sSampler,%s)"),
+			FormattedTexture.ToString(),
+			AutomaticViewMipBias ? TEXT("View.MaterialTextureBilinearClampedSampler") : TEXT("Material.Clamp_WorldGroupSettings"));
+		break;
+	default:
+		checkNoEntry();
+		break;
+	}
+
+	const Shader::EValueType TexCoordType = Private::GetTexCoordType(TextureType);
+	FEmitShaderExpression* TexCoordValue = TexCoordExpression->GetValueShader(Context, Scope, TexCoordType);
+
+	FEmitShaderExpression* TextureResult = nullptr;
+	if (MipValueMode == TMVM_MipLevel)
+	{
+		TextureResult = Context.EmitExpression(Scope, Shader::EValueType::Float4, TEXT("%Level(%, %, %, %)"),
+			SampleFunctionName,
+			FormattedTexture.ToString(),
+			FormattedSampler.ToString(),
+			TexCoordValue,
+			MipValueExpression->GetValueShader(Context, Scope, Shader::EValueType::Float1));
+	}
+	else if (MipValueMode == TMVM_Derivative || (Context.bUseAnalyticDerivatives && TexCoordDerivatives.IsValid()))
+	{
+		TextureResult = Context.EmitExpression(Scope, Shader::EValueType::Float4, TEXT("%Grad(%, %, %, %, %)"),
+			SampleFunctionName,
+			FormattedTexture.ToString(),
+			FormattedSampler.ToString(),
+			TexCoordValue,
+			TexCoordDerivatives.ExpressionDdx->GetValueShader(Context, Scope, TexCoordType),
+			TexCoordDerivatives.ExpressionDdy->GetValueShader(Context, Scope, TexCoordType));
+	}
+	else if (MipValueMode == TMVM_MipBias)
+	{
+		TextureResult = Context.EmitExpression(Scope, Shader::EValueType::Float4, TEXT("%Bias(%, %, %, %)"),
+			SampleFunctionName,
+			FormattedTexture.ToString(),
+			FormattedSampler.ToString(),
+			TexCoordValue,
+			MipValueExpression->GetValueShader(Context, Scope, Shader::EValueType::Float1));
+	}
+	else
+	{
+		check(MipValueMode == TMVM_None);
+		TextureResult = Context.EmitExpression(Scope, Shader::EValueType::Float4, TEXT("%(%, %, %)"),
+			SampleFunctionName,
+			FormattedTexture.ToString(),
+			FormattedSampler.ToString(),
+			TexCoordValue);
+	}
+
+	OutResult.Code = Context.EmitExpression(Scope, Shader::EValueType::Float4, TEXT("ApplyMaterialSamplerType(%, %)"), TextureResult, TextureValue.SamplerType);
+}
+
+bool FExpressionTextureSize::PrepareValue(FEmitContext& Context, FEmitScope& Scope, const FRequestedType& RequestedType, FPrepareValueResult& OutResult) const
+{
+	const FPreparedType& TextureType = Context.PrepareExpression(TextureExpression, Scope, FMaterialTextureValue::GetTypeName());
+	if (TextureType.ObjectType != FMaterialTextureValue::GetTypeName())
+	{
+		return Context.Error(TEXT("Expected texture"));
+	}
+
 	return OutResult.SetType(Context, RequestedType, EExpressionEvaluation::Preshader, Shader::EValueType::Float2);
 }
 
 void FExpressionTextureSize::EmitValuePreshader(FEmitContext& Context, FEmitScope& Scope, const FRequestedType& RequestedType, FEmitValuePreshaderResult& OutResult) const
 {
-	const int32 TextureIndex = Context.Material->GetReferencedTextures().Find(Texture);
+	FMaterialTextureValue TextureValue;
+	TextureExpression->GetValueObject(Context, Scope, TextureValue);
+
+	const int32 TextureIndex = Context.Material->GetReferencedTextures().Find(TextureValue.Texture);
 	check(TextureIndex != INDEX_NONE);
 	
 	Context.PreshaderStackPosition++;
 	OutResult.Type = Shader::EValueType::Float2;
-	OutResult.Preshader.WriteOpcode(Shader::EPreshaderOpcode::TextureSize).Write<FMemoryImageMaterialParameterInfo>(ParameterInfo).Write(TextureIndex);
+	OutResult.Preshader.WriteOpcode(Shader::EPreshaderOpcode::TextureSize).Write<FMemoryImageMaterialParameterInfo>(TextureValue.ParameterInfo).Write(TextureIndex);
 }
 
 bool FExpressionFunctionCall::PrepareValue(FEmitContext& Context, FEmitScope& Scope, const FRequestedType& RequestedType, FPrepareValueResult& OutResult) const
diff --git a/Engine/Source/Runtime/Engine/Private/Materials/MaterialShared.cpp b/Engine/Source/Runtime/Engine/Private/Materials/MaterialShared.cpp
--- a/Engine/Source/Runtime/Engine/Private/Materials/MaterialShared.cpp
+++ b/Engine/Source/Runtime/Engine/Private/Materials/MaterialShared.cpp
@@ -35,6 +35,7 @@
 #include "EngineModule.h"
 #include "Engine/Texture.h"
 #include "Engine/Texture2D.h"
+#include "Engine/Font.h"
 #include "SceneView.h"
 #include "ShaderPlatformQualitySettings.h"
 #include "MaterialShaderQualitySettings.h"
@@ -5302,7 +5303,26 @@ UE::Shader::FValue FMaterialParameterValue::AsShaderValue() const
 	}
 }
 
-UE::Shader::EValueType GetShaderValueType(EMaterialParameterType Type)
+UObject* FMaterialParameterValue::AsTextureObject() const
+{
+	UObject* Result = nullptr;
+	switch (Type)
+	{
+	case EMaterialParameterType::Texture: Result = Texture; break;
+	case EMaterialParameterType::RuntimeVirtualTexture: Result = RuntimeVirtualTexture; break;
+	case EMaterialParameterType::Font:
+		if (Font.Value && Font.Value->Textures.IsValidIndex(Font.Page))
+		{
+			Result = Font.Value->Textures[Font.Page];
+		}
+		break;
+	default:
+		break;
+	}
+	return Result;
+}
+
+UE::Shader::FType GetShaderValueType(EMaterialParameterType Type)
 {
 	switch (Type)
 	{
@@ -5313,9 +5333,9 @@ UE::Shader::EValueType GetShaderValueType(EMaterialParameterType Type)
 	case EMaterialParameterType::StaticComponentMask: return UE::Shader::EValueType::Bool4;
 	case EMaterialParameterType::Texture:
 	case EMaterialParameterType::Font:
+		return FMaterialTextureValue::GetTypeName();
 	case EMaterialParameterType::RuntimeVirtualTexture:
-		// Non-numeric types, can't represent as shader values
-		return UE::Shader::EValueType::Void;
+		return UE::Shader::EValueType::Void; // TODO
 	default:
 		checkNoEntry();
 		return UE::Shader::EValueType::Void;
diff --git a/Engine/Source/Runtime/Engine/Private/Shader/Preshader.cpp b/Engine/Source/Runtime/Engine/Private/Shader/Preshader.cpp
--- a/Engine/Source/Runtime/Engine/Private/Shader/Preshader.cpp
+++ b/Engine/Source/Runtime/Engine/Private/Shader/Preshader.cpp
@@ -168,7 +168,7 @@ void FPreshaderData::WriteValue(const FValue& Value)
 {
 	const int32 NumComponents = Value.Type.GetNumComponents();
 
-	check(!Value.Type.IsTexture());
+	check(!Value.Type.IsObject());
 
 	WriteType(Value.Type);
 	for (int32 Index = 0; Index < NumComponents; ++Index)
diff --git a/Engine/Source/Runtime/Engine/Private/Shader/ShaderValue.cpp b/Engine/Source/Runtime/Engine/Private/Shader/ShaderValue.cpp
--- a/Engine/Source/Runtime/Engine/Private/Shader/ShaderValue.cpp
+++ b/Engine/Source/Runtime/Engine/Private/Shader/ShaderValue.cpp
@@ -13,12 +13,13 @@ namespace Shader
 
 const TCHAR* FType::GetName() const
 {
-	if (StructType)
+	if (IsStruct())
 	{
 		check(ValueType == EValueType::Struct);
 		return StructType->Name;
 	}
 
+	// TODO - do we need specific object names?
 	check(ValueType != EValueType::Struct);
 	const FValueTypeDescription TypeDesc = GetValueTypeDescription(ValueType);
 	return TypeDesc.Name;
@@ -26,11 +27,15 @@ const TCHAR* FType::GetName() const
 
 FType FType::GetDerivativeType() const
 {
-	if (StructType)
+	if (IsStruct())
 	{
 		check(ValueType == EValueType::Struct);
 		return StructType->DerivativeType; // will convert to 'Void' if DerivativeType is nullptr
 	}
+	else if (IsObject())
+	{
+		return *this;
+	}
 
 	check(ValueType != EValueType::Struct);
 	return MakeDerivativeType(ValueType);
@@ -38,11 +43,15 @@ FType FType::GetDerivativeType() const
 
 int32 FType::GetNumComponents() const
 {
-	if (StructType)
+	if (IsStruct())
 	{
 		check(ValueType == EValueType::Struct);
 		return StructType->ComponentTypes.Num();
 	}
+	else if (IsObject())
+	{
+		return 1;
+	}
 
 	check(ValueType != EValueType::Struct);
 	const FValueTypeDescription TypeDesc = GetValueTypeDescription(ValueType);
@@ -51,7 +60,7 @@ int32 FType::GetNumComponents() const
 
 int32 FType::GetNumFlatFields() const
 {
-	if (StructType)
+	if (IsStruct())
 	{
 		check(ValueType == EValueType::Struct);
 		return StructType->FlatFieldTypes.Num();
@@ -63,7 +72,7 @@ int32 FType::GetNumFlatFields() const
 
 EValueComponentType FType::GetComponentType(int32 Index) const
 {
-	if (StructType)
+	if (IsStruct())
 	{
 		check(ValueType == EValueType::Struct);
 		return StructType->ComponentTypes.IsValidIndex(Index) ? StructType->ComponentTypes[Index] : EValueComponentType::Void;
@@ -77,7 +86,7 @@ EValueComponentType FType::GetComponentType(int32 Index) const
 
 EValueType FType::GetFlatFieldType(int32 Index) const
 {
-	if (StructType)
+	if (IsStruct())
 	{
 		check(ValueType == EValueType::Struct);
 		return StructType->FlatFieldTypes.IsValidIndex(Index) ? StructType->FlatFieldTypes[Index] : EValueType::Void;
@@ -90,9 +99,10 @@ EValueType FType::GetFlatFieldType(int32 Index) const
 
 bool FType::Merge(const FType& OtherType)
 {
-	if (ValueType == EValueType::Void)
+	if (IsVoid())
 	{
 		ValueType = OtherType.ValueType;
+		ObjectType = OtherType.ObjectType;
 		StructType = OtherType.StructType;
 		return true;
 	}
@@ -102,6 +112,11 @@ bool FType::Merge(const FType& OtherType)
 		return StructType == OtherType.StructType;
 	}
 
+	if (IsObject() || OtherType.IsObject())
+	{
+		return ObjectType == OtherType.ObjectType;
+	}
+
 	if (ValueType != OtherType.ValueType)
 	{
 		const FValueTypeDescription TypeDesc = GetValueTypeDescription(ValueType);
@@ -277,43 +292,6 @@ void FormatComponent_Double(double Value, int32 NumComponents, EValueStringForma
 
 } // namespace Private
 
-EValueType FTextureValue::GetType() const
-{
-	if (Texture)
-	{
-		const EMaterialValueType MaterialType = Texture->GetMaterialType();
-		switch (MaterialType)
-		{
-		case MCT_Texture2D: return EValueType::Texture2D;
-		case MCT_Texture2DArray: return EValueType::Texture2DArray;
-		case MCT_TextureCube: return EValueType::TextureCube;
-		case MCT_TextureCubeArray: return EValueType::TextureCubeArray;
-		case MCT_VolumeTexture: return EValueType::Texture3D;
-		case MCT_TextureExternal: return EValueType::TextureExternal;
-		case MCT_TextureVirtual: return EValueType::Texture2D; // TODO
-		default: checkNoEntry(); break;
-		}
-	}
-	else if (ExternalTextureGuid.IsValid())
-	{
-		return EValueType::TextureExternal;
-	}
-
-	return EValueType::Void;
-}
-
-FValue::FValue(const FTextureValue* InValue) : Type(EValueType::Void)
-{
-	if (InValue)
-	{
-		Type = InValue->GetType();
-		if (Type != EValueType::Void)
-		{
-			Component.Add(InValue);
-		}
-	}
-}
-
 FValue FValue::FromMemoryImage(EValueType Type, const void* Data, uint32* OutSizeInBytes)
 {
 	check(IsNumericType(Type));
@@ -435,15 +413,6 @@ bool FValue::IsZero() const
 	return bIsZero;
 }
 
-const FTextureValue* FValue::AsTexture() const
-{
-	if (Type.IsTexture() && Component.Num() > 0)
-	{
-		return Component[0].Texture;
-	}
-	return nullptr;
-}
-
 FValueComponentTypeDescription GetValueComponentTypeDescription(EValueComponentType Type)
 {
 	switch (Type)
@@ -453,12 +422,6 @@ FValueComponentTypeDescription GetValueComponentTypeDescription(EValueComponentT
 	case EValueComponentType::Double: return FValueComponentTypeDescription(TEXT("double"), sizeof(double), EComponentBound::NegDoubleMax, EComponentBound::DoubleMax);
 	case EValueComponentType::Int: return FValueComponentTypeDescription(TEXT("int"), sizeof(int32), EComponentBound::IntMin, EComponentBound::IntMax);
 	case EValueComponentType::Bool: return FValueComponentTypeDescription(TEXT("bool"), 1u, EComponentBound::Zero, EComponentBound::One);
-	case EValueComponentType::Texture2D: return FValueComponentTypeDescription(TEXT("Texture2D"), sizeof(void*), EComponentBound::Zero, EComponentBound::Zero);
-	case EValueComponentType::Texture2DArray: return FValueComponentTypeDescription(TEXT("Texture2DArray"), sizeof(void*), EComponentBound::Zero, EComponentBound::Zero);
-	case EValueComponentType::TextureCube: return FValueComponentTypeDescription(TEXT("TextureCube"), sizeof(void*), EComponentBound::Zero, EComponentBound::Zero);
-	case EValueComponentType::TextureCubeArray: return FValueComponentTypeDescription(TEXT("TextureCubeArray"), sizeof(void*), EComponentBound::Zero, EComponentBound::Zero);
-	case EValueComponentType::Texture3D: return FValueComponentTypeDescription(TEXT("Texture3D"), sizeof(void*), EComponentBound::Zero, EComponentBound::Zero);
-	case EValueComponentType::TextureExternal: return FValueComponentTypeDescription(TEXT("TextureExternal"), sizeof(void*), EComponentBound::Zero, EComponentBound::Zero);
 	default: checkNoEntry() return FValueComponentTypeDescription();
 	}
 }
@@ -656,12 +619,7 @@ FValueTypeDescription GetValueTypeDescription(EValueType Type)
 	case EValueType::Double4x4: return FValueTypeDescription(TEXT("FLWCMatrix"), EValueComponentType::Double, 16);
 	case EValueType::DoubleInverse4x4: return FValueTypeDescription(TEXT("FLWCInverseMatrix"), EValueComponentType::Double, 16);
 	case EValueType::Struct: return FValueTypeDescription(TEXT("struct"), EValueComponentType::Void, 0);
-	case EValueType::Texture2D: return FValueTypeDescription(TEXT("FTexture2D"), EValueComponentType::Texture2D, 1);
-	case EValueType::Texture2DArray: return FValueTypeDescription(TEXT("FTexture2DArray"), EValueComponentType::Texture2DArray, 1);
-	case EValueType::TextureCube: return FValueTypeDescription(TEXT("FTextureCube"), EValueComponentType::TextureCube, 1);
-	case EValueType::TextureCubeArray: return FValueTypeDescription(TEXT("FTextureCubeArray"), EValueComponentType::TextureCubeArray, 1);
-	case EValueType::Texture3D: return FValueTypeDescription(TEXT("FTexture3D"), EValueComponentType::Texture3D, 1);
-	case EValueType::TextureExternal: return FValueTypeDescription(TEXT("FTextureExternal"), EValueComponentType::TextureExternal, 1);
+	case EValueType::Object: return FValueTypeDescription(TEXT("object"), EValueComponentType::Void, 0);
 	default: checkNoEntry(); return FValueTypeDescription(TEXT("<INVALID>"), EValueComponentType::Void, 0);
 	}
 }
@@ -731,12 +689,6 @@ EValueType MakeValueType(EValueComponentType ComponentType, int32 NumComponents)
 		default: break;
 		}
 		break;
-	case EValueComponentType::Texture2D: check(NumComponents == 1); return EValueType::Texture2D; break;
-	case EValueComponentType::Texture2DArray: check(NumComponents == 1); return EValueType::Texture2DArray; break;
-	case EValueComponentType::TextureCube: check(NumComponents == 1); return EValueType::TextureCube; break;
-	case EValueComponentType::TextureCubeArray: check(NumComponents == 1); return EValueType::TextureCubeArray; break;
-	case EValueComponentType::Texture3D: check(NumComponents == 1); return EValueType::Texture3D; break;
-	case EValueComponentType::TextureExternal: check(NumComponents == 1); return EValueType::TextureExternal; break;
 	default:
 		break;
 	}
diff --git a/Engine/Source/Runtime/Engine/Public/HLSLTree/HLSLTree.h b/Engine/Source/Runtime/Engine/Public/HLSLTree/HLSLTree.h
--- a/Engine/Source/Runtime/Engine/Public/HLSLTree/HLSLTree.h
+++ b/Engine/Source/Runtime/Engine/Public/HLSLTree/HLSLTree.h
@@ -134,6 +134,12 @@ private:
 	friend class FExpressionLocalPHI;
 };
 
+enum class EDefaultRequestedType : uint8
+{
+	None,
+	Any,
+};
+
 /**
  * Represents a request for an arbitrary set of untyped components
  */
@@ -143,20 +149,26 @@ public:
 	FRequestedType() = default;
 	FRequestedType(Shader::EValueType InType, bool bDefaultRequest = true);
 	FRequestedType(const Shader::FType& InType, bool bDefaultRequest = true);
+	FRequestedType(FName InObjectName);
+	FRequestedType(EDefaultRequestedType InType) : DefaultType(InType) { check(InType != EDefaultRequestedType::None); }
 
 	Shader::FType GetType() const;
 	const TCHAR* GetName() const { return GetType().GetName(); }
 	int32 GetNumComponents() const;
-	bool IsComponentRequested(int32 Index) const { return RequestedComponents.IsValidIndex(Index) ? (bool)RequestedComponents[Index] : false; }
-	bool IsVoid() const { return RequestedComponents.Find(true) == INDEX_NONE; }
 
-	void SetComponentRequest(int32 Index, bool bRequest = true);
-
-	void Reset()
+	bool IsComponentRequested(int32 Index) const
 	{
-		RequestedComponents.Reset();
+		return (DefaultType == EDefaultRequestedType::Any) || (RequestedComponents.IsValidIndex(Index) ? (bool)RequestedComponents[Index] : false);
 	}
 
+	bool IsStruct() const { return !IsVoid() && StructType != nullptr; }
+	bool IsObject() const { return !IsVoid() && !ObjectType.IsNone(); }
+	bool IsNumeric() const { return !IsVoid() && Shader::IsNumericType(ValueComponentType); }
+	bool IsVoid() const { return (DefaultType != EDefaultRequestedType::Any) && RequestedComponents.Find(true) == INDEX_NONE; }
+
+	void SetComponentRequest(int32 Index, bool bRequest = true);
+
+
 	/** Marks the given field as requested (or not) */
 	void SetFieldRequested(const Shader::FStructField* Field, bool bRequest = true);
 	void ClearFieldRequested(const Shader::FStructField* Field)
@@ -171,7 +183,9 @@ public:
 	FRequestedType GetField(const Shader::FStructField* Field) const;
 
 	const Shader::FStructType* StructType = nullptr;
+	FName ObjectType;
 	Shader::EValueComponentType ValueComponentType = Shader::EValueComponentType::Void;
+	EDefaultRequestedType DefaultType = EDefaultRequestedType::None;
 
 	/** 1 bit per component, a value of 'true' means the specified component is requsted */
 	TBitArray<> RequestedComponents;
@@ -179,6 +193,7 @@ public:
 inline bool operator==(const FRequestedType& Lhs, const FRequestedType& Rhs)
 {
 	return Lhs.ValueComponentType == Rhs.ValueComponentType &&
+		Lhs.ObjectType == Rhs.ObjectType &&
 		Lhs.StructType == Rhs.StructType &&
 		Lhs.RequestedComponents == Rhs.RequestedComponents;
 }
@@ -189,6 +204,9 @@ inline bool operator!=(const FRequestedType& Lhs, const FRequestedType& Rhs)
 
 inline void AppendHash(FHasher& Hasher, const FRequestedType& Value)
 {
+	AppendHash(Hasher, Value.StructType);
+	AppendHash(Hasher, Value.ObjectType);
+	AppendHash(Hasher, Value.ValueComponentType);
 	AppendHash(Hasher, Value.RequestedComponents);
 }
 
@@ -243,8 +261,9 @@ public:
 	Shader::FType GetType() const;
 	const TCHAR* GetName() const { return GetType().GetName(); }
 	bool IsStruct() const { return !IsVoid() && StructType != nullptr; }
-	bool IsNumeric() const { return !IsVoid() && ValueComponentType != Shader::EValueComponentType::Void; }
-	bool IsInitialized() const { return StructType != nullptr || ValueComponentType != Shader::EValueComponentType::Void; }
+	bool IsObject() const { return !IsVoid() && !ObjectType.IsNone(); }
+	bool IsNumeric() const { return !IsVoid() && Shader::IsNumericType(ValueComponentType); }
+	bool IsInitialized() const { return StructType != nullptr || !ObjectType.IsNone() || ValueComponentType != Shader::EValueComponentType::Void; }
 	bool IsVoid() const;
 
 	EExpressionEvaluation GetEvaluation(const FEmitScope& Scope) const;
@@ -262,6 +281,7 @@ public:
 	void MergeComponent(int32 Index, const FPreparedComponent& InComponent);
 
 	const Shader::FStructType* StructType = nullptr;
+	FName ObjectType;
 	Shader::EValueComponentType ValueComponentType = Shader::EValueComponentType::Void;
 
 	void EnsureNumComponents(int32 NumComponents);
@@ -273,6 +293,7 @@ inline bool operator==(const FPreparedType& Lhs, const FPreparedType& Rhs)
 {
 	return Lhs.ValueComponentType == Rhs.ValueComponentType &&
 		Lhs.StructType == Rhs.StructType &&
+		Lhs.ObjectType == Rhs.ObjectType &&
 		Lhs.PreparedComponents == Rhs.PreparedComponents;
 }
 inline bool operator!=(const FPreparedType& Lhs, const FPreparedType& Rhs)
@@ -294,7 +315,10 @@ public:
 	bool SetType(FEmitContext& Context, const FRequestedType& RequestedType, const FPreparedType& Type);
 
 private:
-	bool TryMergePreparedType(FEmitContext& Context, const Shader::FStructType* StructType, Shader::EValueComponentType ComponentType);
+	bool TryMergePreparedType(FEmitContext& Context,
+		const Shader::FStructType* StructType,
+		FName ObjectType,
+		Shader::EValueComponentType ComponentType);
 
 	FPreparedType PreparedType;
 	bool bPreparingValue = false;
@@ -308,6 +332,12 @@ struct FEmitValueShaderResult
 	FEmitShaderExpression* Code = nullptr;
 };
 
+struct FEmitCustomHLSLParameterResult
+{
+	FStringBuilderBase* DeclarationCode = nullptr;
+	FStringBuilderBase* ForwardCode = nullptr;
+};
+
 struct FEmitValuePreshaderResult
 {
 	explicit FEmitValuePreshaderResult(Shader::FPreshaderData& InPreshader) : Preshader(InPreshader) {}
@@ -369,12 +399,43 @@ public:
 	Shader::FValue GetValueConstant(FEmitContext& Context, FEmitScope& Scope, const FPreparedType& PreparedType, const Shader::FType& ResultType) const;
 	Shader::FValue GetValueConstant(FEmitContext& Context, FEmitScope& Scope, const FPreparedType& PreparedType, Shader::EValueType ResultType) const;
 
+	/**
+	 * Gets the value as an 'Object' with a certain type
+	 * Individual expressions will potentially support different types of objects
+	 * @param ObjectTypeName should be the object type that matches the type returned by PrepareExpression
+	 * @param OutObjectBase pointer to an object of the correct type
+	 * @return true if given type is supported, otherwise false
+	 */
+	bool GetValueObject(FEmitContext& Context, FEmitScope& Scope, const FName& ObjectTypeName, void* OutObjectBase) const;
+
+	template<typename ObjectType>
+	bool GetValueObject(FEmitContext& Context, FEmitScope& Scope, ObjectType& OutObject) const
+	{
+		return GetValueObject(Context, Scope, ObjectType::GetTypeName(), &OutObject);
+	}
+
+	/**
+	 * Some expressions may support passing object types to custom HLSL expressions
+	 * CheckObjectSupportsCustomHLSL() will return true if this is supported
+	 * If supported, GetObjectCustomHLSLParameter() will generate code to declare the parameters in the custom HLSL function, and forward the parameters
+	 * EmitValueShader() will be called to generate the actual HLSL code for the expression
+	 */
+	bool CheckObjectSupportsCustomHLSL(FEmitContext& Context, FEmitScope& Scope, const FName& ObjectTypeName) const;
+	void GetObjectCustomHLSLParameter(FEmitContext& Context,
+		FEmitScope& Scope,
+		const FName& ObjectTypeName,
+		const TCHAR* ParameterName,
+		FStringBuilderBase& OutDeclarationCode,
+		FStringBuilderBase& OutForwardCode) const;
+
 protected:
 	virtual void ComputeAnalyticDerivatives(FTree& Tree, FExpressionDerivatives& OutResult) const;
 	virtual const FExpression* ComputePreviousFrame(FTree& Tree, const FRequestedType& RequestedType) const;
 	virtual bool PrepareValue(FEmitContext& Context, FEmitScope& Scope, const FRequestedType& RequestedType, FPrepareValueResult& OutResult) const = 0;
 	virtual void EmitValueShader(FEmitContext& Context, FEmitScope& Scope, const FRequestedType& RequestedType, FEmitValueShaderResult& OutResult) const;
 	virtual void EmitValuePreshader(FEmitContext& Context, FEmitScope& Scope, const FRequestedType& RequestedType, FEmitValuePreshaderResult& OutResult) const;
+	virtual bool EmitValueObject(FEmitContext& Context, FEmitScope& Scope, const FName& ObjectTypeName, void* OutObjectBase) const;
+	virtual bool EmitCustomHLSLParameter(FEmitContext& Context, FEmitScope& Scope, const FName& ObjectTypeName, const TCHAR* ParameterName, FEmitCustomHLSLParameterResult& OutResult) const;
 
 private:
 	TArray<UObject*, TInlineAllocator<2>> Owners;
diff --git a/Engine/Source/Runtime/Engine/Public/HLSLTree/HLSLTreeCommon.h b/Engine/Source/Runtime/Engine/Public/HLSLTree/HLSLTreeCommon.h
--- a/Engine/Source/Runtime/Engine/Public/HLSLTree/HLSLTreeCommon.h
+++ b/Engine/Source/Runtime/Engine/Public/HLSLTree/HLSLTreeCommon.h
@@ -53,37 +53,6 @@ public:
 	virtual void EmitValuePreshader(FEmitContext& Context, FEmitScope& Scope, const FRequestedType& RequestedType, FEmitValuePreshaderResult& OutResult) const override;
 };
 
-class FExpressionTextureSample : public FExpression
-{
-public:
-	FExpressionTextureSample(const FExpression* InTextureExpression,
-		const FExpression* InTexCoordExpression,
-		const FExpression* InMipValueExpression,
-		const FExpression* InAutomaticMipBiasExpression,
-		const FExpressionDerivatives& InTexCoordDerivatives,
-		ESamplerSourceMode InSamplerSource,
-		ETextureMipValueMode InMipValueMode)
-		: TextureExpression(InTextureExpression)
-		, TexCoordExpression(InTexCoordExpression)
-		, MipValueExpression(InMipValueExpression)
-		, AutomaticMipBiasExpression(InAutomaticMipBiasExpression)
-		, TexCoordDerivatives(InTexCoordDerivatives)
-		, SamplerSource(InSamplerSource)
-		, MipValueMode(InMipValueMode)
-	{}
-
-	const FExpression* TextureExpression;
-	const FExpression* TexCoordExpression;
-	const FExpression* MipValueExpression;
-	const FExpression* AutomaticMipBiasExpression;
-	FExpressionDerivatives TexCoordDerivatives;
-	ESamplerSourceMode SamplerSource;
-	ETextureMipValueMode MipValueMode;
-
-	virtual bool PrepareValue(FEmitContext& Context, FEmitScope& Scope, const FRequestedType& RequestedType, FPrepareValueResult& OutResult) const override;
-	virtual void EmitValueShader(FEmitContext& Context, FEmitScope& Scope, const FRequestedType& RequestedType, FEmitValueShaderResult& OutResult) const override;
-};
-
 class FExpressionGetStructField : public FExpression
 {
 public:
diff --git a/Engine/Source/Runtime/Engine/Public/HLSLTree/HLSLTreeEmit.h b/Engine/Source/Runtime/Engine/Public/HLSLTree/HLSLTreeEmit.h
--- a/Engine/Source/Runtime/Engine/Public/HLSLTree/HLSLTreeEmit.h
+++ b/Engine/Source/Runtime/Engine/Public/HLSLTree/HLSLTreeEmit.h
@@ -212,15 +212,11 @@ inline bool IsScopeDead(const FEmitScope* Scope)
 struct FEmitCustomHLSLInput
 {
 	FStringView Name;
+	FStringView ObjectDeclarationCode;
+	FStringView ObjectForwardCode;
 	Shader::FType Type;
 };
 
-inline void AppendHash(FHasher& Hasher, const FEmitCustomHLSLInput& Value)
-{
-	AppendHash(Hasher, Value.Name);
-	AppendHash(Hasher, Value.Type);
-}
-
 class FEmitCustomHLSL
 {
 public:
diff --git a/Engine/Source/Runtime/Engine/Public/MaterialHLSLGenerator.h b/Engine/Source/Runtime/Engine/Public/MaterialHLSLGenerator.h
--- a/Engine/Source/Runtime/Engine/Public/MaterialHLSLGenerator.h
+++ b/Engine/Source/Runtime/Engine/Public/MaterialHLSLGenerator.h
@@ -120,8 +120,6 @@ public:
 		return InternalNewErrorExpression(String.ToView());
 	}
 
-	const UE::Shader::FTextureValue* AcquireTextureValue(const UE::Shader::FTextureValue& Value);
-
 	int32 FindInputIndex(const FExpressionInput* Input) const;
 
 	const UE::HLSLTree::FExpression* NewDefaultInputConstant(int32 InputIndex, const UE::Shader::FValue& Value);
@@ -142,7 +140,10 @@ public:
 
 	bool GenerateStatements(UE::HLSLTree::FScope& Scope, UMaterialExpression* MaterialExpression);
 
-	const UE::HLSLTree::FExpression* GenerateMaterialParameter(FName InParameterName, const FMaterialParameterMetadata& InParameterMeta, const UE::Shader::FValue& InDefaultValue);
+	const UE::HLSLTree::FExpression* GenerateMaterialParameter(FName InParameterName,
+		const FMaterialParameterMetadata& InParameterMeta,
+		EMaterialSamplerType InSamplerType = SAMPLERTYPE_Color,
+		const FGuid& InExternalTextureGuid = FGuid());
 
 	const UE::HLSLTree::FExpression* GenerateFunctionCall(UE::HLSLTree::FScope& Scope,
 		UMaterialFunctionInterface* Function,
@@ -252,7 +253,6 @@ private:
 	FFunctionCallEntry RootFunctionCallEntry;
 	TArray<FFunctionCallEntry*, TInlineAllocator<8>> FunctionCallStack;
 	TArray<UE::HLSLTree::FScope*> JoinedScopeStack;
-	TMap<FXxHash64, const UE::Shader::FTextureValue*> TextureValueMap;
 	TMap<FXxHash64, TUniquePtr<FFunctionCallEntry>> FunctionCallMap;
 	TMap<UMaterialFunctionInterface*, UE::HLSLTree::FFunction*> FunctionMap;
 	TMap<UMaterialExpression*, FStatementEntry> StatementMap;
diff --git a/Engine/Source/Runtime/Engine/Public/MaterialHLSLTree.h b/Engine/Source/Runtime/Engine/Public/MaterialHLSLTree.h
--- a/Engine/Source/Runtime/Engine/Public/MaterialHLSLTree.h
+++ b/Engine/Source/Runtime/Engine/Public/MaterialHLSLTree.h
@@ -8,6 +8,7 @@
 #include "Materials/MaterialLayersFunctions.h"
 #include "RHIDefinitions.h"
 
+class UTexture;
 enum class EMaterialParameterType : uint8;
 struct FMaterialCachedExpressionData;
 struct FMaterialLayersFunctions;
@@ -15,6 +16,34 @@ struct FMaterialLayersFunctions;
 namespace UE::HLSLTree::Material
 {
 
+inline void AppendHash(FHasher& Hasher, const FMaterialParameterInfo& Value)
+{
+	AppendHash(Hasher, Value.Name);
+	AppendHash(Hasher, Value.Index);
+	AppendHash(Hasher, Value.Association);
+}
+
+inline void AppendHash(FHasher& Hasher, const FMaterialParameterValue& Value)
+{
+	switch (Value.Type)
+	{
+	case EMaterialParameterType::Scalar: AppendHash(Hasher, Value.Float[0]); break;
+	case EMaterialParameterType::Vector: AppendHash(Hasher, Value.Float); break;
+	case EMaterialParameterType::DoubleVector: AppendHash(Hasher, Value.Double); break;
+	case EMaterialParameterType::Texture: AppendHash(Hasher, Value.Texture); break;
+	case EMaterialParameterType::Font: AppendHash(Hasher, Value.Font); break;
+	case EMaterialParameterType::RuntimeVirtualTexture: AppendHash(Hasher, Value.RuntimeVirtualTexture); break;
+	case EMaterialParameterType::StaticSwitch: AppendHash(Hasher, Value.Bool[0]); break;
+	case EMaterialParameterType::StaticComponentMask: AppendHash(Hasher, Value.Bool); break;
+	default: checkNoEntry(); break;
+	}
+}
+
+inline void AppendHash(FHasher& Hasher, const FMaterialParameterMetadata& Meta)
+{
+	AppendHash(Hasher, Meta.Value);
+}
+
 enum class EExternalInput : uint8
 {
 	None,
@@ -205,61 +234,69 @@ public:
 	virtual void EmitValuePreshader(FEmitContext& Context, FEmitScope& Scope, const FRequestedType& RequestedType, FEmitValuePreshaderResult& OutResult) const override;
 };
 
-inline void AppendHash(FHasher& Hasher, const FMaterialParameterInfo& Value)
-{
-	AppendHash(Hasher, Value.Name);
-	AppendHash(Hasher, Value.Index);
-	AppendHash(Hasher, Value.Association);
-}
-
-inline void AppendHash(FHasher& Hasher, const FMaterialParameterValue& Value)
-{
-	switch (Value.Type)
-	{
-	case EMaterialParameterType::Scalar: AppendHash(Hasher, Value.Float[0]); break;
-	case EMaterialParameterType::Vector: AppendHash(Hasher, Value.Float); break;
-	case EMaterialParameterType::DoubleVector: AppendHash(Hasher, Value.Double); break;
-	case EMaterialParameterType::Texture: AppendHash(Hasher, Value.Texture); break;
-	case EMaterialParameterType::Font: AppendHash(Hasher, Value.Font); break;
-	case EMaterialParameterType::RuntimeVirtualTexture: AppendHash(Hasher, Value.RuntimeVirtualTexture); break;
-	case EMaterialParameterType::StaticSwitch: AppendHash(Hasher, Value.Bool[0]); break;
-	case EMaterialParameterType::StaticComponentMask: AppendHash(Hasher, Value.Bool); break;
-	default: checkNoEntry(); break;
-	}
-}
-
-inline void AppendHash(FHasher& Hasher, const FMaterialParameterMetadata& Meta)
-{
-	AppendHash(Hasher, Meta.Value);
-}
-
 class FExpressionParameter : public FExpression
 {
 public:
-	explicit FExpressionParameter(const FMaterialParameterInfo& InParameterInfo, const FMaterialParameterMetadata& InParameterMeta, const Shader::FValue& InDefaultValue)
-		: ParameterInfo(InParameterInfo), ParameterMeta(InParameterMeta), DefaultValue(InDefaultValue)
+	explicit FExpressionParameter(const FMaterialParameterInfo& InParameterInfo,
+		const FMaterialParameterMetadata& InParameterMeta,
+		EMaterialSamplerType InSamplerType = SAMPLERTYPE_Color,
+		const FGuid& InExternalTextureGuid = FGuid())
+		: ParameterInfo(InParameterInfo), ParameterMeta(InParameterMeta), ExternalTextureGuid(InExternalTextureGuid), TextureSamplerType(InSamplerType)
 	{
 	}
 
 	FMaterialParameterInfo ParameterInfo;
 	FMaterialParameterMetadata ParameterMeta;
-	Shader::FValue DefaultValue;
+	FGuid ExternalTextureGuid;
+	EMaterialSamplerType TextureSamplerType;
 
 	virtual void ComputeAnalyticDerivatives(FTree& Tree, FExpressionDerivatives& OutResult) const override;
 	virtual bool PrepareValue(FEmitContext& Context, FEmitScope& Scope, const FRequestedType& RequestedType, FPrepareValueResult& OutResult) const override;
-	virtual void EmitValueShader(FEmitContext& Context, FEmitScope& Scope, const FRequestedType& RequestedType, FEmitValueShaderResult& OutResult) const override;
 	virtual void EmitValuePreshader(FEmitContext& Context, FEmitScope& Scope, const FRequestedType& RequestedType, FEmitValuePreshaderResult& OutResult) const override;
+	virtual void EmitValueShader(FEmitContext& Context, FEmitScope& Scope, const FRequestedType& RequestedType, FEmitValueShaderResult& OutResult) const override;
+	virtual bool EmitValueObject(FEmitContext& Context, FEmitScope& Scope, const FName& ObjectTypeName, void* OutObjectBase) const override;
+	virtual bool EmitCustomHLSLParameter(FEmitContext& Context, FEmitScope& Scope, const FName& ObjectTypeName, const TCHAR* ParameterName, FEmitCustomHLSLParameterResult& OutResult) const override;
+};
+
+class FExpressionTextureSample : public FExpression
+{
+public:
+	FExpressionTextureSample(const FExpression* InTextureExpression,
+		const FExpression* InTexCoordExpression,
+		const FExpression* InMipValueExpression,
+		const FExpression* InAutomaticMipBiasExpression,
+		const FExpressionDerivatives& InTexCoordDerivatives,
+		ESamplerSourceMode InSamplerSource,
+		ETextureMipValueMode InMipValueMode)
+		: TextureExpression(InTextureExpression)
+		, TexCoordExpression(InTexCoordExpression)
+		, MipValueExpression(InMipValueExpression)
+		, AutomaticMipBiasExpression(InAutomaticMipBiasExpression)
+		, TexCoordDerivatives(InTexCoordDerivatives)
+		, SamplerSource(InSamplerSource)
+		, MipValueMode(InMipValueMode)
+	{}
+
+	const FExpression* TextureExpression;
+	const FExpression* TexCoordExpression;
+	const FExpression* MipValueExpression;
+	const FExpression* AutomaticMipBiasExpression;
+	FExpressionDerivatives TexCoordDerivatives;
+	ESamplerSourceMode SamplerSource;
+	ETextureMipValueMode MipValueMode;
+
+	virtual bool PrepareValue(FEmitContext& Context, FEmitScope& Scope, const FRequestedType& RequestedType, FPrepareValueResult& OutResult) const override;
+	virtual void EmitValueShader(FEmitContext& Context, FEmitScope& Scope, const FRequestedType& RequestedType, FEmitValueShaderResult& OutResult) const override;
 };
 
 class FExpressionTextureSize : public FExpression
 {
 public:
-	explicit FExpressionTextureSize(const FMaterialParameterInfo& InParameterInfo, UTexture* InTexture)
-		: ParameterInfo(InParameterInfo), Texture(InTexture)
+	explicit FExpressionTextureSize(const FExpression* InTextureExpression)
+		: TextureExpression(InTextureExpression)
 	{}
 
-	FMaterialParameterInfo ParameterInfo;
-	UTexture* Texture;
+	const FExpression* TextureExpression;
 
 	virtual bool PrepareValue(FEmitContext& Context, FEmitScope& Scope, const FRequestedType& RequestedType, FPrepareValueResult& OutResult) const override;
 	virtual void EmitValuePreshader(FEmitContext& Context, FEmitScope& Scope, const FRequestedType& RequestedType, FEmitValuePreshaderResult& OutResult) const override;
diff --git a/Engine/Source/Runtime/Engine/Public/MaterialTypes.h b/Engine/Source/Runtime/Engine/Public/MaterialTypes.h
--- a/Engine/Source/Runtime/Engine/Public/MaterialTypes.h
+++ b/Engine/Source/Runtime/Engine/Public/MaterialTypes.h
@@ -16,6 +16,161 @@ class UCurveLinearColorAtlas;
 class UFont;
 class URuntimeVirtualTexture;
 
+UENUM()
+enum EMaterialParameterAssociation
+{
+	LayerParameter,
+	BlendParameter,
+	GlobalParameter,
+};
+
+USTRUCT(BlueprintType)
+struct ENGINE_API FMaterialParameterInfo
+{
+	GENERATED_USTRUCT_BODY()
+
+		UPROPERTY(EditAnywhere, BlueprintReadWrite, Category = ParameterInfo)
+		FName Name;
+
+	/** Whether this is a global parameter, or part of a layer or blend */
+	UPROPERTY(EditAnywhere, BlueprintReadWrite, Category = ParameterInfo)
+		TEnumAsByte<EMaterialParameterAssociation> Association;
+
+	/** Layer or blend index this parameter is part of. INDEX_NONE for global parameters. */
+	UPROPERTY(EditAnywhere, BlueprintReadWrite, Category = ParameterInfo)
+		int32 Index;
+
+	FMaterialParameterInfo(const TCHAR* InName, EMaterialParameterAssociation InAssociation = EMaterialParameterAssociation::GlobalParameter, int32 InIndex = INDEX_NONE)
+		: Name(InName)
+		, Association(InAssociation)
+		, Index(InIndex)
+	{
+	}
+	FMaterialParameterInfo(FName InName = FName(), EMaterialParameterAssociation InAssociation = EMaterialParameterAssociation::GlobalParameter, int32 InIndex = INDEX_NONE)
+		: Name(InName)
+		, Association(InAssociation)
+		, Index(InIndex)
+	{
+	}
+
+	explicit FMaterialParameterInfo(const struct FMemoryImageMaterialParameterInfo& Rhs);
+
+	FString ToString() const
+	{
+		return *Name.ToString() + FString::FromInt(Association) + FString::FromInt(Index);
+	}
+
+	friend FArchive& operator<<(FArchive& Ar, FMaterialParameterInfo& Ref)
+	{
+		Ar << Ref.Name << Ref.Association << Ref.Index;
+		return Ar;
+	}
+
+	bool RemapLayerIndex(TArrayView<const int32> IndexRemap, FMaterialParameterInfo& OutResult) const;
+};
+
+struct FMemoryImageMaterialParameterInfo
+{
+	DECLARE_TYPE_LAYOUT(FMemoryImageMaterialParameterInfo, NonVirtual);
+public:
+	FMemoryImageMaterialParameterInfo(const TCHAR* InName, EMaterialParameterAssociation InAssociation = EMaterialParameterAssociation::GlobalParameter, int32 InIndex = INDEX_NONE)
+		: Name(NameToScriptName(FName(InName)))
+		, Index(InIndex)
+		, Association(InAssociation)
+	{}
+
+	FMemoryImageMaterialParameterInfo(const FName& InName, EMaterialParameterAssociation InAssociation = EMaterialParameterAssociation::GlobalParameter, int32 InIndex = INDEX_NONE)
+		: Name(NameToScriptName(InName))
+		, Index(InIndex)
+		, Association(InAssociation)
+	{}
+
+	FMemoryImageMaterialParameterInfo(const FScriptName& InName = FScriptName(), EMaterialParameterAssociation InAssociation = EMaterialParameterAssociation::GlobalParameter, int32 InIndex = INDEX_NONE)
+		: Name(InName)
+		, Index(InIndex)
+		, Association(InAssociation)
+	{}
+
+	FMemoryImageMaterialParameterInfo(const FMaterialParameterInfo& Rhs)
+		: Name(NameToScriptName(Rhs.Name))
+		, Index(Rhs.Index)
+		, Association(Rhs.Association)
+	{}
+
+	FMemoryImageMaterialParameterInfo(const FMemoryImageMaterialParameterInfo& Rhs) = default;
+
+	FORCEINLINE FName GetName() const { return ScriptNameToName(Name); }
+
+	friend FArchive& operator<<(FArchive& Ar, FMemoryImageMaterialParameterInfo& Ref)
+	{
+		FName RefName = ScriptNameToName(Ref.Name);
+		Ar << RefName << Ref.Association << Ref.Index;
+		Ref.Name = NameToScriptName(RefName);
+		return Ar;
+	}
+
+	bool RemapLayerIndex(TArrayView<const int32> IndexRemap, FMemoryImageMaterialParameterInfo& OutResult) const;
+
+	LAYOUT_FIELD(FScriptName, Name);
+	LAYOUT_FIELD(int32, Index);
+	LAYOUT_FIELD(TEnumAsByte<EMaterialParameterAssociation>, Association);
+};
+
+FORCEINLINE FMaterialParameterInfo::FMaterialParameterInfo(const struct FMemoryImageMaterialParameterInfo& Rhs)
+	: Name(ScriptNameToName(Rhs.Name))
+	, Association(Rhs.Association)
+	, Index(Rhs.Index)
+{
+}
+
+FORCEINLINE bool operator==(const FMaterialParameterInfo& Lhs, const FMaterialParameterInfo& Rhs)
+{
+	return Lhs.Name.IsEqual(Rhs.Name) && Lhs.Association == Rhs.Association && Lhs.Index == Rhs.Index;
+}
+
+FORCEINLINE bool operator!=(const FMaterialParameterInfo& Lhs, const FMaterialParameterInfo& Rhs)
+{
+	return !operator==(Lhs, Rhs);
+}
+
+FORCEINLINE bool operator==(const FMemoryImageMaterialParameterInfo& Lhs, const FMemoryImageMaterialParameterInfo& Rhs)
+{
+	return Lhs.Name == Rhs.Name && Lhs.Association == Rhs.Association && Lhs.Index == Rhs.Index;
+}
+
+FORCEINLINE bool operator!=(const FMemoryImageMaterialParameterInfo& Lhs, const FMemoryImageMaterialParameterInfo& Rhs)
+{
+	return !operator==(Lhs, Rhs);
+}
+
+FORCEINLINE bool operator==(const FMaterialParameterInfo& Lhs, const FMemoryImageMaterialParameterInfo& Rhs)
+{
+	return Lhs.Name == Rhs.Name && Lhs.Index == Rhs.Index && Lhs.Association == Rhs.Association;
+}
+
+FORCEINLINE bool operator!=(const FMaterialParameterInfo& Lhs, const FMemoryImageMaterialParameterInfo& Rhs)
+{
+	return !operator==(Lhs, Rhs);
+}
+
+FORCEINLINE bool operator==(const FMemoryImageMaterialParameterInfo& Lhs, const FMaterialParameterInfo& Rhs)
+{
+	return Lhs.Name == Rhs.Name && Lhs.Index == Rhs.Index && Lhs.Association == Rhs.Association;
+}
+
+FORCEINLINE bool operator!=(const FMemoryImageMaterialParameterInfo& Lhs, const FMaterialParameterInfo& Rhs)
+{
+	return !operator==(Lhs, Rhs);
+}
+
+FORCEINLINE uint32 GetTypeHash(const FMemoryImageMaterialParameterInfo& Value)
+{
+	return HashCombine(HashCombine(GetTypeHash(Value.Name), Value.Index), (uint32)Value.Association);
+}
+
+// Backwards compat
+using FHashedMaterialParameterInfo = FMemoryImageMaterialParameterInfo;
+
 enum class EMaterialParameterType : uint8
 {
 	Scalar = 0u,
@@ -40,6 +195,19 @@ static const int32 NumMaterialParameterTypes = (int32)EMaterialParameterType::Nu
 static const int32 NumMaterialRuntimeParameterTypes = (int32)EMaterialParameterType::NumRuntime;
 static const int32 NumMaterialEditorOnlyParameterTypes = NumMaterialParameterTypes - NumMaterialRuntimeParameterTypes;
 
+static inline bool IsNumericMaterialParameter(EMaterialParameterType InType)
+{
+	switch (InType)
+	{
+	case EMaterialParameterType::Scalar:
+	case EMaterialParameterType::Vector:
+	case EMaterialParameterType::DoubleVector:
+		return true;
+	default:
+		return false;
+	}
+}
+
 static inline bool IsStaticMaterialParameter(EMaterialParameterType InType)
 {
 	switch (InType)
@@ -52,7 +220,7 @@ static inline bool IsStaticMaterialParameter(EMaterialParameterType InType)
 	}
 }
 
-extern ENGINE_API UE::Shader::EValueType GetShaderValueType(EMaterialParameterType Type);
+extern ENGINE_API UE::Shader::FType GetShaderValueType(EMaterialParameterType Type);
 
 enum class EMaterialGetParameterValueFlags : uint32
 {
@@ -114,6 +282,16 @@ struct FStaticComponentMaskValue
 	bool A = false;
 };
 
+struct FMaterialTextureValue
+{
+	static FName GetTypeName() { return TEXT("FMaterialTextureValue"); }
+
+	UTexture* Texture = nullptr;
+	FGuid ExternalTextureGuid;
+	EMaterialSamplerType SamplerType = SAMPLERTYPE_Color;
+	FMaterialParameterInfo ParameterInfo;
+};
+
 struct FMaterialParameterValue
 {
 	FMaterialParameterValue() : Type(EMaterialParameterType::None) {}
@@ -146,6 +324,7 @@ struct FMaterialParameterValue
 	inline bool AsStaticSwitch() const { check(Type == EMaterialParameterType::StaticSwitch); return Bool[0]; }
 	inline FStaticComponentMaskValue AsStaticComponentMask() const { check(Type == EMaterialParameterType::StaticComponentMask); return FStaticComponentMaskValue(Bool[0], Bool[1], Bool[2], Bool[3]); }
 	ENGINE_API UE::Shader::FValue AsShaderValue() const;
+	ENGINE_API UObject* AsTextureObject() const;
 
 	union
 	{
diff --git a/Engine/Source/Runtime/Engine/Public/Shader/ShaderTypes.h b/Engine/Source/Runtime/Engine/Public/Shader/ShaderTypes.h
--- a/Engine/Source/Runtime/Engine/Public/Shader/ShaderTypes.h
+++ b/Engine/Source/Runtime/Engine/Public/Shader/ShaderTypes.h
@@ -62,14 +62,6 @@ enum class EValueComponentType : uint8
 	Int,
 	Bool,
 
-	FirstTexture,
-	Texture2D = FirstTexture,
-	Texture2DArray,
-	TextureCube,
-	TextureCubeArray,
-	Texture3D,
-	TextureExternal,
-
 	Num,
 };
 static constexpr int32 NumValueComponentTypes = (int32)EValueComponentType::Num;
@@ -87,8 +79,7 @@ struct FValueComponentTypeDescription
 ENGINE_API FValueComponentTypeDescription GetValueComponentTypeDescription(EValueComponentType Type);
 inline const TCHAR* GetComponentTypeName(EValueComponentType Type) { return GetValueComponentTypeDescription(Type).Name; }
 inline uint32 GetComponentTypeSizeInBytes(EValueComponentType Type) { return GetValueComponentTypeDescription(Type).SizeInBytes; }
-inline bool IsNumericType(EValueComponentType Type) { return Type != EValueComponentType::Void && (uint8)Type < (uint8)EValueComponentType::FirstTexture; }
-inline bool IsTextureType(EValueComponentType Type) { return (uint8)Type >= (uint8)EValueComponentType::FirstTexture; }
+inline bool IsNumericType(EValueComponentType Type) { return Type != EValueComponentType::Void; }
 inline bool IsComponentTypeWithinBounds(EValueComponentType Type, FComponentBounds Bounds) { return IsWithinBounds(GetValueComponentTypeDescription(Type).Bounds, Bounds); }
 
 ENGINE_API EValueComponentType CombineComponentTypes(EValueComponentType Lhs, EValueComponentType Rhs);
@@ -129,13 +120,7 @@ enum class EValueType : uint8
 	DoubleInverse4x4,
 
 	Struct,
-
-	Texture2D,
-	Texture2DArray,
-	TextureCube,
-	TextureCubeArray,
-	Texture3D,
-	TextureExternal,
+	Object,
 
 	Num,
 };
@@ -163,21 +148,21 @@ ENGINE_API EValueType MakeArithmeticResultType(EValueType Lhs, EValueType Rhs, F
 ENGINE_API EValueType MakeComparisonResultType(EValueType Lhs, EValueType Rhs, FString& OutErrorMessage);
 
 inline bool IsNumericType(EValueType Type) { return IsNumericType(GetValueTypeDescription(Type).ComponentType); }
-inline bool IsTextureType(EValueType Type) { return IsTextureType(GetValueTypeDescription(Type).ComponentType); }
 
 struct FType
 {
-	FType() : StructType(nullptr), ValueType(EValueType::Void) {}
-	FType(EValueType InValueType) : StructType(nullptr), ValueType(InValueType) {}
+	FType() : ValueType(EValueType::Void) {}
+	FType(EValueType InValueType) : ValueType(InValueType) { check(InValueType != EValueType::Struct && InValueType != EValueType::Object); }
 	FType(const FStructType* InStruct) : StructType(InStruct), ValueType(InStruct ? EValueType::Struct : EValueType::Void) {}
+	FType(FName InObjectType) : ObjectType(InObjectType), ValueType(!InObjectType.IsNone() ? EValueType::Object : EValueType::Void) {}
 
 	const TCHAR* GetName() const;
 	FType GetDerivativeType() const;
 	FType GetNonLWCType() const { return IsNumericLWC() ? FType(MakeNonLWCType(ValueType)) : *this; }
 	bool IsVoid() const { return ValueType == EValueType::Void; }
 	bool IsStruct() const { return ValueType == EValueType::Struct; }
-	bool IsNumeric() const { return !IsStruct() && IsNumericType(ValueType); }
-	bool IsTexture() const { return !IsStruct() && IsTextureType(ValueType); }
+	bool IsObject() const { return ValueType == EValueType::Object; }
+	bool IsNumeric() const { return !IsStruct() && !IsObject() && IsNumericType(ValueType); }
 	bool IsNumericLWC() const { return IsNumeric() && IsLWCType(ValueType); }
 	int32 GetNumComponents() const;
 	int32 GetNumFlatFields() const;
@@ -189,7 +174,8 @@ struct FType
 	inline operator bool() const { return !IsVoid(); }
 	inline bool operator!() const { return IsVoid(); }
 
-	const FStructType* StructType;
+	const FStructType* StructType = nullptr;
+	FName ObjectType;
 	EValueType ValueType;
 };
 
@@ -197,6 +183,7 @@ inline bool operator==(const FType& Lhs, const FType& Rhs)
 {
 	if (Lhs.ValueType != Rhs.ValueType) return false;
 	if (Lhs.ValueType == EValueType::Struct && Lhs.StructType != Rhs.StructType) return false;
+	if (Lhs.ValueType == EValueType::Object && Lhs.ObjectType != Rhs.ObjectType) return false;
 	return true;
 }
 inline bool operator!=(const FType& Lhs, const FType& Rhs)
@@ -312,35 +299,6 @@ struct FMemoryImageValue
 	uint32 Size;
 };
 
-struct FTextureValue
-{
-	FTextureValue() = default;
-	FTextureValue(UTexture* InTexture, EMaterialSamplerType InSamplerType = SAMPLERTYPE_Color)
-		: Texture(InTexture)
-		, SamplerType(InSamplerType)
-	{}
-
-	FTextureValue(const FGuid& InGuid)
-		: ExternalTextureGuid(InGuid)
-		, SamplerType(SAMPLERTYPE_External)
-	{}
-
-	EValueType GetType() const;
-
-	UTexture* Texture = nullptr;
-	FGuid ExternalTextureGuid;
-	EMaterialSamplerType SamplerType = SAMPLERTYPE_Color;
-};
-
-inline bool operator==(const FTextureValue& Lhs, const FTextureValue& Rhs)
-{
-	return Lhs.Texture == Rhs.Texture && Lhs.SamplerType == Rhs.SamplerType && Lhs.ExternalTextureGuid == Rhs.ExternalTextureGuid;
-}
-inline bool operator!=(const FTextureValue& Lhs, const FTextureValue& Rhs)
-{
-	return !operator==(Lhs, Rhs);
-}
-
 union FValueComponent
 {
 	FValueComponent() : Packed(0u) {}
@@ -348,7 +306,6 @@ union FValueComponent
 	FValueComponent(double InValue) : Packed(0u) { Double = InValue; }
 	FValueComponent(int32 InValue) : Packed(0u) { Int = InValue; }
 	FValueComponent(bool InValue) : Packed(0u) { Bool = InValue ? 1 : 0; }
-	FValueComponent(const FTextureValue* InValue) : Packed(0u) { Texture = InValue; }
 
 	// 'Bool' is stored as uint8 to avoid changing on different compilers
 	bool AsBool() const { return Bool != 0u; }
@@ -356,7 +313,6 @@ union FValueComponent
 	const TCHAR* ToString(EValueComponentType Type, FStringBuilderBase& OutString) const;
 
 	uint64 Packed;
-	const FTextureValue* Texture;
 	double Double;
 	float Float;
 	int32 Int;
@@ -492,8 +448,6 @@ struct FValue
 		Component.Add(v);
 	}
 
-	FValue(const FTextureValue* InValue);
-
 	inline const FType& GetType() const { return Type; }
 	inline const FValueComponent& GetComponent(int32 i) const
 	{
@@ -514,7 +468,6 @@ struct FValue
 	FDoubleValue AsDouble() const;
 	FIntValue AsInt() const;
 	FBoolValue AsBool() const;
-	const FTextureValue* AsTexture() const;
 
 	FLinearColor AsLinearColor() const;
 	FVector4d AsVector4d() const;
