Fix deprecated API usage

diff --git a/Engine/Source/Runtime/Engine/Private/PhysicsEngine/AsyncPhysicsInputComponent.cpp b/Engine/Source/Runtime/Engine/Private/PhysicsEngine/AsyncPhysicsInputComponent.cpp
--- a/Engine/Source/Runtime/Engine/Private/PhysicsEngine/AsyncPhysicsInputComponent.cpp
+++ b/Engine/Source/Runtime/Engine/Private/PhysicsEngine/AsyncPhysicsInputComponent.cpp
@@ -311,7 +311,7 @@ void UAsyncPhysicsInputComponent::AsyncPhysicsTickComponent(float DeltaTime, flo
 	FAsyncPhysicsInputRewindCallback* Callback = static_cast<FAsyncPhysicsInputRewindCallback*>(Solver->GetRewindCallback());
 	const int32 ServerFrame = Callback->CachedServerFrame;
 
-	if(World->IsServer())
+	if(World->IsNetMode(ENetMode::NM_DedicatedServer) || World->IsNetMode(ENetMode::NM_ListenServer))
 	{
 		//TODO: move this somewhere else - here because we need to run this when GT and PT are both on same core. Function guards against multiple calls per ServerFrame so calling from each instance is ok
 		Callback->UpdateReplicationMap_Internal(ServerFrame);
