Strata - fix name typo reported by [REDACTED_PERSON], thanks!

diff --git a/Engine/Source/Runtime/Engine/Classes/Materials/MaterialExpressionStrata.h b/Engine/Source/Runtime/Engine/Classes/Materials/MaterialExpressionStrata.h
--- a/Engine/Source/Runtime/Engine/Classes/Materials/MaterialExpressionStrata.h
+++ b/Engine/Source/Runtime/Engine/Classes/Materials/MaterialExpressionStrata.h
@@ -736,7 +736,7 @@ class UMaterialExpressionStrataAdd : public UMaterialExpressionStrataBSDF
 	//~ End UMaterialExpression Interface
 };
 
-UCLASS(MinimalAPI, collapsecategories, hidecategories = Object, DisplayName = "Strata Coverate Weight")
+UCLASS(MinimalAPI, collapsecategories, hidecategories = Object, DisplayName = "Strata Coverage Weight")
 class UMaterialExpressionStrataWeight : public UMaterialExpressionStrataBSDF
 {
 	GENERATED_UCLASS_BODY()
