Fix for USoundWave manually cleaning up data it no longer owns in OnFinishDestroy()

diff --git a/Engine/Source/Runtime/Engine/Private/SoundWave.cpp b/Engine/Source/Runtime/Engine/Private/SoundWave.cpp
--- a/Engine/Source/Runtime/Engine/Private/SoundWave.cpp
+++ b/Engine/Source/Runtime/Engine/Private/SoundWave.cpp
@@ -2510,7 +2510,6 @@ void USoundWave::FinishDestroy()
 	check(GetPrecacheState() != ESoundWavePrecacheState::InProgress);
 	check(AudioDecompressor == nullptr);
 
-	CleanupCachedRunningPlatformData();
 #if WITH_EDITOR
 	if (!GExitPurge)
 	{
