WorldPartition : Added WorldDataLayersActorDesc which will allow changelist validation of unloaded worlds and also allow knowledge of data layers used by level instances.

diff --git a/Engine/Source/Runtime/Core/Public/UObject/FortniteNCBranchObjectVersion.h b/Engine/Source/Runtime/Core/Public/UObject/FortniteNCBranchObjectVersion.h
--- a/Engine/Source/Runtime/Core/Public/UObject/FortniteNCBranchObjectVersion.h
+++ b/Engine/Source/Runtime/Core/Public/UObject/FortniteNCBranchObjectVersion.h
@@ -12,6 +12,9 @@ struct CORE_API FFortniteNCBranchObjectVersion
 		// Before any version changes were made
 		BeforeCustomVersionWasAdded = 0,
 
+		// Added FWorldDataLayersActorDesc
+		AddedWorldDataLayersActorDesc,
+
 		// -----<new versions can be added above this line>-------------------------------------------------
 		VersionPlusOne,
 		LatestVersion = VersionPlusOne - 1
diff --git a/Engine/Source/Runtime/Engine/Private/WorldPartition/DataLayer/WorldDataLayers.cpp b/Engine/Source/Runtime/Engine/Private/WorldPartition/DataLayer/WorldDataLayers.cpp
--- a/Engine/Source/Runtime/Engine/Private/WorldPartition/DataLayer/WorldDataLayers.cpp
+++ b/Engine/Source/Runtime/Engine/Private/WorldPartition/DataLayer/WorldDataLayers.cpp
@@ -14,6 +14,7 @@
 #include "WorldPartition/WorldPartition.h"
 #if WITH_EDITOR
 #include "WorldPartition/WorldPartitionEditorPerProjectUserSettings.h"
+#include "WorldPartition/DataLayer/WorldDataLayersActorDesc.h"
 #endif
 
 #define LOCTEXT_NAMESPACE "WorldDataLayers"
@@ -407,6 +408,11 @@ void AWorldDataLayers::DumpDataLayers(FOutputDevice& OutputDevice) const
 
 #if WITH_EDITOR
 
+TUniquePtr<class FWorldPartitionActorDesc> AWorldDataLayers::CreateClassActorDesc() const
+{
+	return TUniquePtr<FWorldPartitionActorDesc>(new FWorldDataLayersActorDesc());
+}
+
 void AWorldDataLayers::GetUserLoadedInEditorStates(TArray<FName>& OutDataLayersLoadedInEditor, TArray<FName>& OutDataLayersNotLoadedInEditor) const
 {
 	const TArray<FName>& SettingsDataLayersNotLoadedInEditor = GetMutableDefault<UWorldPartitionEditorPerProjectUserSettings>()->GetWorldDataLayersNotLoadedInEditor(GetWorld());
diff --git a/Engine/Source/Runtime/Engine/Private/WorldPartition/DataLayer/WorldDataLayersActorDesc.cpp b/Engine/Source/Runtime/Engine/Private/WorldPartition/DataLayer/WorldDataLayersActorDesc.cpp
new file mode 100644
--- /dev/null
+++ b/Engine/Source/Runtime/Engine/Private/WorldPartition/DataLayer/WorldDataLayersActorDesc.cpp
@@ -0,0 +1,187 @@
+// Copyright Epic Games, Inc. All Rights Reserved.
+
+#include "WorldPartition/DataLayer/WorldDataLayersActorDesc.h"
+
+#if WITH_EDITOR
+
+#include "WorldPartition/DataLayer/WorldDataLayers.h"
+#include "WorldPartition/DataLayer/DataLayerInstanceWithAsset.h"
+#include "WorldPartition/DataLayer/DataLayerAsset.h"
+#include "UObject/FortniteNCBranchObjectVersion.h"
+
+FArchive& operator<<(FArchive& Ar, FDataLayerInstanceDesc& Desc)
+{
+	Ar << Desc.Name << Desc.ParentName << Desc.bIsUsingAsset << Desc.ShortName;
+	if (Desc.bIsUsingAsset)
+	{
+		Ar << Desc.AssetPath << Desc.bIsRuntime;
+	}
+	return Ar;
+}
+
+bool operator == (const FDataLayerInstanceDesc& Lhs, const FDataLayerInstanceDesc& Rhs)
+{
+	return (Lhs.Name == Rhs.Name) &&
+		(Lhs.ParentName == Rhs.ParentName) &&
+		(Lhs.bIsUsingAsset == Rhs.bIsUsingAsset) &&
+		(Lhs.AssetPath == Rhs.AssetPath) &&
+		(Lhs.bIsRuntime == Rhs.bIsRuntime) &&
+		(Lhs.ShortName == Rhs.ShortName);
+}
+
+bool operator < (const FDataLayerInstanceDesc& Lhs, const FDataLayerInstanceDesc& Rhs)
+{
+	if (Lhs.Name == Rhs.Name)
+	{
+		if (Lhs.ParentName == Rhs.ParentName)
+		{
+			if (Lhs.bIsUsingAsset == Rhs.bIsUsingAsset)
+			{
+				if (Lhs.AssetPath == Rhs.AssetPath)
+				{
+					if (Lhs.bIsRuntime == Rhs.bIsRuntime)
+					{
+						return Lhs.ShortName < Rhs.ShortName;
+					}
+					return (int)Lhs.bIsRuntime < (int)Rhs.bIsRuntime;
+				}
+				return Lhs.AssetPath.LexicalLess(Rhs.AssetPath);
+			}
+			return (int)Lhs.bIsUsingAsset < (int)Rhs.bIsUsingAsset;
+		}
+		return Lhs.ParentName.LexicalLess(Rhs.ParentName);
+	}
+	return Lhs.Name.LexicalLess(Rhs.Name);
+}
+
+UDataLayerAsset* FDataLayerInstanceDesc::GetAsset() const
+{
+	if (IsUsingAsset() && !AssetPath.IsNone())
+	{
+		UDataLayerAsset* Asset = LoadObject<UDataLayerAsset>(nullptr, *AssetPath.ToString());
+		return Asset;
+	}
+	return nullptr;
+}
+
+EDataLayerType FDataLayerInstanceDesc::GetDataLayerType() const
+{
+	if (IsUsingAsset())
+	{
+		if (UDataLayerAsset* DataLayerAsset = GetAsset())
+		{
+			return DataLayerAsset->IsRuntime() ? EDataLayerType::Runtime : EDataLayerType::Editor;
+		}
+		return EDataLayerType::Unknown;
+	}
+	return bIsRuntime ? EDataLayerType::Runtime : EDataLayerType::Editor;
+}
+
+FString FDataLayerInstanceDesc::GetShortName() const
+{
+	if (IsUsingAsset())
+	{
+		static FString UnknownString(TEXT("Unknown"));
+		UDataLayerAsset* DataLayerAsset = GetAsset();
+		return DataLayerAsset ? DataLayerAsset->GetName() : UnknownString;
+	}
+	return ShortName;
+}
+
+FDataLayerInstanceDesc::FDataLayerInstanceDesc()
+: bIsUsingAsset(false)
+, bIsRuntime(false)
+{
+}
+
+void FDataLayerInstanceDesc::Init(UDataLayerInstance* InDataLayerInstance)
+{
+	Name = InDataLayerInstance->GetDataLayerFName();
+	ParentName = InDataLayerInstance->GetParent() ? InDataLayerInstance->GetParent()->GetDataLayerFName() : NAME_None;
+	const UDataLayerInstanceWithAsset* DataLayerWithAsset = Cast<UDataLayerInstanceWithAsset>(InDataLayerInstance);
+	bIsUsingAsset = (DataLayerWithAsset != nullptr);
+	if (bIsUsingAsset)
+	{
+		if (const UDataLayerAsset* DataLayerAsset = DataLayerWithAsset->GetAsset())
+		{
+			AssetPath = FName(DataLayerAsset->GetPathName());
+		}
+	}
+	ShortName = InDataLayerInstance->GetDataLayerShortName();
+	bIsRuntime = InDataLayerInstance->IsRuntime();
+}
+
+FWorldDataLayersActorDesc::FWorldDataLayersActorDesc()
+: bIsValid(false)
+{}
+
+void FWorldDataLayersActorDesc::Init(const AActor* InActor)
+{
+	FWorldPartitionActorDesc::Init(InActor);
+
+	const AWorldDataLayers* WorldDataLayers = CastChecked<AWorldDataLayers>(InActor);
+	WorldDataLayers->ForEachDataLayer([this](UDataLayerInstance* DataLayerInstance)
+	{
+		FDataLayerInstanceDesc& DataLayerInstanceDesc = DataLayerInstances.Emplace_GetRef();
+		DataLayerInstanceDesc.Init(DataLayerInstance);
+		return true;
+	});
+
+	bIsValid = true;
+}
+
+bool FWorldDataLayersActorDesc::Equals(const FWorldPartitionActorDesc* Other) const
+{
+	if (FWorldPartitionActorDesc::Equals(Other))
+	{
+		const FWorldDataLayersActorDesc* OtherDesc = (FWorldDataLayersActorDesc*)Other;
+		if (DataLayerInstances.Num() == OtherDesc->DataLayerInstances.Num())
+		{
+			TArray<FDataLayerInstanceDesc> SortedDataLayerInstances(DataLayerInstances);
+			TArray<FDataLayerInstanceDesc> SortedDataLayerInstancesOther(OtherDesc->DataLayerInstances);
+			SortedDataLayerInstances.Sort([](const FDataLayerInstanceDesc& LHS, const FDataLayerInstanceDesc& RHS) { return LHS < RHS; });
+			SortedDataLayerInstancesOther.Sort([](const FDataLayerInstanceDesc& LHS, const FDataLayerInstanceDesc& RHS) { return LHS < RHS; });
+			return SortedDataLayerInstances == SortedDataLayerInstancesOther;
+		}
+	}
+	return false;
+}
+
+void FWorldDataLayersActorDesc::Serialize(FArchive& Ar)
+{
+	Ar.UsingCustomVersion(FFortniteNCBranchObjectVersion::GUID);
+
+	FWorldPartitionActorDesc::Serialize(Ar);
+
+	if (Ar.CustomVer(FFortniteNCBranchObjectVersion::GUID) >= FFortniteNCBranchObjectVersion::AddedWorldDataLayersActorDesc)
+	{
+		Ar << DataLayerInstances;
+		bIsValid = true;
+	}
+}
+
+const FDataLayerInstanceDesc* FWorldDataLayersActorDesc::GetDataLayerInstanceFromInstanceName(FName InDataLayerInstanceName) const
+{
+	for (const FDataLayerInstanceDesc& DataLayerInstance : DataLayerInstances)
+	{
+		if (DataLayerInstance.GetName().IsEqual(InDataLayerInstanceName, ENameCase::CaseSensitive))
+		{
+			return &DataLayerInstance;
+		}
+	}
+	return nullptr;
+}
+
+const FDataLayerInstanceDesc* FWorldDataLayersActorDesc::GetDataLayerInstanceFromAssetPath(FName InDataLayerAssetPath) const
+{
+	for (const FDataLayerInstanceDesc& DataLayerInstance : DataLayerInstances)
+	{
+		if (FName(DataLayerInstance.GetAssetPath()) == InDataLayerAssetPath)
+		{
+			return &DataLayerInstance;
+		}
+	}
+	return nullptr;
+}
+
+#endif
diff --git a/Engine/Source/Runtime/Engine/Public/WorldPartition/DataLayer/WorldDataLayers.h b/Engine/Source/Runtime/Engine/Public/WorldPartition/DataLayer/WorldDataLayers.h
--- a/Engine/Source/Runtime/Engine/Public/WorldPartition/DataLayer/WorldDataLayers.h
+++ b/Engine/Source/Runtime/Engine/Public/WorldPartition/DataLayer/WorldDataLayers.h
@@ -50,6 +50,7 @@ public:
 	virtual bool ShouldImport(FString* ActorPropString, bool IsMovingLevel) override { return false; }
 	virtual bool IsLockLocation() const { return true; }
 	virtual bool IsUserManaged() const override { return false; }
+	virtual TUniquePtr<class FWorldPartitionActorDesc> CreateClassActorDesc() const override;
 
 	static AWorldDataLayers* Create(UWorld* World);
 
@@ -233,6 +234,8 @@ public:
 	FDataLayersFilterDelegate DataLayersFilterDelegate;
 };
 
+DEFINE_ACTORDESC_TYPE(AWorldDataLayers, FWorldDataLayersActorDesc);
+
 #if WITH_EDITOR
 
 template<class DataLayerInstanceType, typename ...CreationsArgs>
diff --git a/Engine/Source/Runtime/Engine/Public/WorldPartition/DataLayer/WorldDataLayersActorDesc.h b/Engine/Source/Runtime/Engine/Public/WorldPartition/DataLayer/WorldDataLayersActorDesc.h
new file mode 100644
--- /dev/null
+++ b/Engine/Source/Runtime/Engine/Public/WorldPartition/DataLayer/WorldDataLayersActorDesc.h
@@ -0,0 +1,72 @@
+// Copyright Epic Games, Inc. All Rights Reserved.
+
+#pragma once
+
+#include "WorldPartition/WorldPartitionActorDesc.h"
+#include "WorldPartition/DataLayer/DataLayerUtils.h"
+
+#if WITH_EDITOR
+
+class ENGINE_API FDataLayerInstanceDesc
+{
+public:
+	FDataLayerInstanceDesc();
+	void Init(class UDataLayerInstance* DataLayerInstance);
+	friend FArchive& operator<<(FArchive& Ar, FDataLayerInstanceDesc& DataLayerInstanceDesc);
+	friend bool operator == (const FDataLayerInstanceDesc& Lhs, const FDataLayerInstanceDesc& Rhs);
+	friend bool operator < (const FDataLayerInstanceDesc& Lhs, const FDataLayerInstanceDesc& Rhs);
+
+	FName GetName() const { return Name; }
+	FName GetParentName() const { return ParentName; }
+	bool IsUsingAsset() const { return bIsUsingAsset; }
+	FName GetAssetPath() const { return AssetPath; }
+	class UDataLayerAsset* GetAsset() const;
+	EDataLayerType GetDataLayerType() const;
+	FString GetShortName() const;
+
+private:
+	// DataLayerInstance Name
+	FName Name;
+	// Parent DataLayerInstance Name
+	FName ParentName;
+	// We can't rely on a valid AssetPath to determine if DataLayerInstance is UDataLayerInstanceWithAsset
+	bool bIsUsingAsset;
+
+	//~ Begin UDataLayerInstanceWithAsset
+	// DataLayer Asset Path 
+	FName AssetPath;
+	//~ End UDataLayerInstanceWithAsset
+
+	//~ Begin UDeprecatedDataLayerInstance
+	// Runtime or Editor
+	bool bIsRuntime;
+	// Label
+	FString ShortName;
+	//~ End UDeprecatedDataLayerInstance
+};
+
+#endif
+
+/**
+ * ActorDesc for AWorldDataLayers actors.
+ */
+class ENGINE_API FWorldDataLayersActorDesc : public FWorldPartitionActorDesc
+{
+#if WITH_EDITOR
+public:
+	FWorldDataLayersActorDesc();
+	bool IsValid() const { return bIsValid; }
+	const TArray<FDataLayerInstanceDesc>& GetDataLayerInstances() const { return DataLayerInstances; }
+	const FDataLayerInstanceDesc* GetDataLayerInstanceFromInstanceName(FName InDataLayerInstanceName) const;
+	const FDataLayerInstanceDesc* GetDataLayerInstanceFromAssetPath(FName InDataLayerAssetPath) const;
+
+protected:
+	virtual void Init(const AActor* InActor) override;
+	virtual bool Equals(const FWorldPartitionActorDesc* Other) const override;
+	virtual void Serialize(FArchive& Ar) override;
+
+private:
+	TArray<FDataLayerInstanceDesc> DataLayerInstances;
+	bool bIsValid;
+#endif
+};
