If we don't have a controller, use our owner's controller to send rotation information to server
Set control rotation to all types of controllers instead of just player controllers when receiving movement information from client

diff --git a/Engine/Source/Runtime/Engine/Private/Components/CharacterMovementComponent.cpp b/Engine/Source/Runtime/Engine/Private/Components/CharacterMovementComponent.cpp
--- a/Engine/Source/Runtime/Engine/Private/Components/CharacterMovementComponent.cpp
+++ b/Engine/Source/Runtime/Engine/Private/Components/CharacterMovementComponent.cpp
@@ -9308,9 +9308,9 @@ void UCharacterMovementComponent::ServerMove_PerformMovement(const FCharacterNet
 		ServerData->ServerTimeStamp = MyWorld->GetTimeSeconds();
 		ServerData->ServerTimeStampLastServerMove = ServerData->ServerTimeStamp;
 
-		if (PC)
+		if (AController* CharacterController = Cast<AController>(CharacterOwner->GetController()))
 		{
-			PC->SetControlRotation(ClientControlRotation);
+			CharacterController->SetControlRotation(ClientControlRotation);
 		}
 
 		if (!bServerReadyForClient)
@@ -11854,6 +11854,13 @@ void FSavedMove_Character::PostUpdate(ACharacter* Character, FSavedMove_Characte
 		}
 
 		SavedControlRotation = Character->GetControlRotation().Clamp();
+		if (!Character->GetController())
+		{
+			if (AController* ControllerOwner = Cast<AController>(Character->GetOwner()))
+			{
+				SavedControlRotation = ControllerOwner->GetControlRotation().Clamp();
+			}
+		}
 	}
 
 	// Only save RootMotion params when initially recording
