Fix nullptr dereference in AActor::PreEditChange()

diff --git a/Engine/Source/Runtime/Engine/Private/ActorEditor.cpp b/Engine/Source/Runtime/Engine/Private/ActorEditor.cpp
--- a/Engine/Source/Runtime/Engine/Private/ActorEditor.cpp
+++ b/Engine/Source/Runtime/Engine/Private/ActorEditor.cpp
@@ -54,7 +54,7 @@ void AActor::PreEditChange(FProperty* PropertyThatWillChange)
 	}
 
 	PreEditChangeDataLayers.Reset();
-	if (PropertyThatWillChange->GetFName() == GET_MEMBER_NAME_CHECKED(AActor, DataLayerAssets))
+	if (PropertyThatWillChange != nullptr && PropertyThatWillChange->GetFName() == GET_MEMBER_NAME_CHECKED(AActor, DataLayerAssets))
 	{
 		PreEditChangeDataLayers = DataLayerAssets;
 	}
