UHLODProxy - Improve dirtyness detection
* SetHLODBakingTransform() now compares the transform using GetTransformCRC()
* Fixed rotation hashing in GetTransformCRC() - could reach 360 when rounded to integer, should be [0, 360[

diff --git a/Engine/Source/Runtime/Engine/Private/HLODProxy.cpp b/Engine/Source/Runtime/Engine/Private/HLODProxy.cpp
--- a/Engine/Source/Runtime/Engine/Private/HLODProxy.cpp
+++ b/Engine/Source/Runtime/Engine/Private/HLODProxy.cpp
@@ -409,10 +409,11 @@ static void AppendRoundedTransform(const FRotator& ComponentRotation, const FVec
 	Ar << Location;
 
 	FRotator Rotator(ComponentRotation.GetDenormalized());
-	FIntVector Rotation(FMath::RoundToInt(Rotator.Pitch), FMath::RoundToInt(Rotator.Yaw), FMath::RoundToInt(Rotator.Roll));
+	const int32 MAX_DEGREES = 360;
+	FIntVector Rotation(FMath::RoundToInt(Rotator.Pitch) % MAX_DEGREES, FMath::RoundToInt(Rotator.Yaw) % MAX_DEGREES, FMath::RoundToInt(Rotator.Roll) % MAX_DEGREES);
 	Ar << Rotation;
 
-	float SCALE_FACTOR = 100;
+	const float SCALE_FACTOR = 100;
 	FIntVector Scale(FMath::RoundToInt(ComponentScale.X * SCALE_FACTOR), FMath::RoundToInt(ComponentScale.Y * SCALE_FACTOR), FMath::RoundToInt(ComponentScale.Z * SCALE_FACTOR));
 	Ar << Scale;
 }
@@ -745,10 +746,14 @@ bool UHLODProxy::SetHLODBakingTransform(const FTransform& InTransform)
 {
 	bool bChanged = false;
 
+	int32 NewTransformCRC = GetTransformCRC(InTransform);
+
 	for (auto ItHLODActor = HLODActors.CreateIterator(); ItHLODActor; ++ItHLODActor)
 	{
 		UHLODProxyDesc* HLODProxyDesc = ItHLODActor.Key();
-		if (!HLODProxyDesc->HLODBakingTransform.Equals(InTransform))
+
+		int32 OldTransformCRC = GetTransformCRC(HLODProxyDesc->HLODBakingTransform);
+		if (OldTransformCRC != NewTransformCRC)
 		{
 			HLODProxyDesc->HLODBakingTransform = InTransform;
 			bChanged = true;
@@ -805,4 +810,4 @@ bool UHLODProxy::ContainsDataForActor(const ALODActor* InLODActor) const
 	return false;
 }
 
-#endif
\ No newline at end of file
+#endif
