PR #8779: Change FCachedKeyToActionInfo::PlayerInput pointer to a WeakObjectPtr. This needs to be weak because otherwise the cache can prevent the player input from being GC'd on level transitions.  (Contributed by [REDACTED_PERSON])

Also add a null check on the PlayerInput in OnInputOwnerEndPlayed because it is a weak pointer now, not guaranteed to be there.

diff --git a/Engine/Source/Runtime/Engine/Classes/Components/InputComponent.h b/Engine/Source/Runtime/Engine/Classes/Components/InputComponent.h
--- a/Engine/Source/Runtime/Engine/Classes/Components/InputComponent.h
+++ b/Engine/Source/Runtime/Engine/Classes/Components/InputComponent.h
@@ -696,7 +696,7 @@ struct FCachedKeyToActionInfo
 
 	/** Which PlayerInput object this has been built for */
 	UPROPERTY()
-	TObjectPtr<UPlayerInput> PlayerInput;
+	TWeakObjectPtr<UPlayerInput> PlayerInput;
 
 	/** What index of the player input's key mappings was the map built for. */
 	uint32 KeyMapBuiltForIndex;
diff --git a/Engine/Source/Runtime/Engine/Private/Components/InputComponent.cpp b/Engine/Source/Runtime/Engine/Private/Components/InputComponent.cpp
--- a/Engine/Source/Runtime/Engine/Private/Components/InputComponent.cpp
+++ b/Engine/Source/Runtime/Engine/Private/Components/InputComponent.cpp
@@ -95,7 +95,8 @@ void UInputComponent::OnInputOwnerEndPlayed(AActor* InOwner, EEndPlayReason::Typ
 	for (int32 Index = CachedKeyToActionInfo.Num() - 1; Index >= 0; --Index)
 	{
 		FCachedKeyToActionInfo& CachedInfo = CachedKeyToActionInfo[Index];
-		if (CachedInfo.PlayerInput->GetTypedOuter<AActor>() == InOwner)
+		const UPlayerInput* CachedInput = CachedInfo.PlayerInput.Get();
+		if (CachedInput && CachedInput->GetTypedOuter<AActor>() == InOwner)
 		{
 			CachedKeyToActionInfo.RemoveAtSwap(Index, 1, false);
 		}
