Fixed bad serialization of FDataLayerInstanceDesc

diff --git a/Engine/Source/Runtime/Core/Public/UObject/FortniteNCBranchObjectVersion.h b/Engine/Source/Runtime/Core/Public/UObject/FortniteNCBranchObjectVersion.h
--- a/Engine/Source/Runtime/Core/Public/UObject/FortniteNCBranchObjectVersion.h
+++ b/Engine/Source/Runtime/Core/Public/UObject/FortniteNCBranchObjectVersion.h
@@ -15,6 +15,9 @@ struct CORE_API FFortniteNCBranchObjectVersion
 		// Added FWorldDataLayersActorDesc
 		AddedWorldDataLayersActorDesc,
 
+		// Fixed FDataLayerInstanceDesc
+		FixedDataLayerInstanceDesc,
+
 		// -----<new versions can be added above this line>-------------------------------------------------
 		VersionPlusOne,
 		LatestVersion = VersionPlusOne - 1
diff --git a/Engine/Source/Runtime/Engine/Private/WorldPartition/DataLayer/WorldDataLayersActorDesc.cpp b/Engine/Source/Runtime/Engine/Private/WorldPartition/DataLayer/WorldDataLayersActorDesc.cpp
--- a/Engine/Source/Runtime/Engine/Private/WorldPartition/DataLayer/WorldDataLayersActorDesc.cpp
+++ b/Engine/Source/Runtime/Engine/Private/WorldPartition/DataLayer/WorldDataLayersActorDesc.cpp
@@ -11,10 +11,26 @@
 
 FArchive& operator<<(FArchive& Ar, FDataLayerInstanceDesc& Desc)
 {
-	Ar << Desc.Name << Desc.ParentName << Desc.bIsUsingAsset << Desc.ShortName;
-	if (Desc.bIsUsingAsset)
+	Ar << Desc.Name << Desc.ParentName << Desc.bIsUsingAsset;
+
+	if (Ar.CustomVer(FFortniteNCBranchObjectVersion::GUID) < FFortniteNCBranchObjectVersion::FixedDataLayerInstanceDesc)
+	{
+		Ar << Desc.ShortName;
+		if (Desc.bIsUsingAsset)
+		{
+			Ar << Desc.AssetPath << Desc.bIsRuntime;
+		}
+	}
+	else
 	{
-		Ar << Desc.AssetPath << Desc.bIsRuntime;
+		if (Desc.bIsUsingAsset)
+		{
+			Ar << Desc.AssetPath;
+		}
+		else
+		{
+			Ar << Desc.bIsRuntime << Desc.ShortName;
+		}
 	}
 	return Ar;
 }
@@ -184,4 +200,4 @@ const FDataLayerInstanceDesc* FWorldDataLayersActorDesc::GetDataLayerInstanceFro
 	return nullptr;
 }
 
-#endif
+#endif
\ No newline at end of file
