New HLSL translator supports VT

diff --git a/Engine/Source/Runtime/Engine/Private/Materials/MaterialCachedHLSLTree.cpp b/Engine/Source/Runtime/Engine/Private/Materials/MaterialCachedHLSLTree.cpp
--- a/Engine/Source/Runtime/Engine/Private/Materials/MaterialCachedHLSLTree.cpp
+++ b/Engine/Source/Runtime/Engine/Private/Materials/MaterialCachedHLSLTree.cpp
@@ -122,6 +122,8 @@ bool FMaterialCachedHLSLTree::GenerateTree(UMaterial* Material, const FMaterialL
 	check(MaterialAttributesDefaultValue.Component.Num() == MaterialAttributesType->ComponentTypes.Num());
 	MaterialAttributesDefaultValue.Type = MaterialAttributesType;
 
+	VTPageTableResultType = TypeRegistry.NewExternalType(TEXT("VTPageTableResult"));
+
 	HLSLTree = UE::HLSLTree::FTree::Create(Allocator);
 
 	FMaterialHLSLGenerator Generator(Material, LayerOverrides, PreviewExpression, *this);
diff --git a/Engine/Source/Runtime/Engine/Private/Materials/MaterialHLSLEmitter.cpp b/Engine/Source/Runtime/Engine/Private/Materials/MaterialHLSLEmitter.cpp
--- a/Engine/Source/Runtime/Engine/Private/Materials/MaterialHLSLEmitter.cpp
+++ b/Engine/Source/Runtime/Engine/Private/Materials/MaterialHLSLEmitter.cpp
@@ -495,23 +495,29 @@ static void GetMaterialEnvironment(EShaderPlatform InPlatform,
 	OutEnvironment.SetDefine(TEXT("MATERIAL_USES_ANISOTROPY"), false);// bUsesAnisotropy);
 
 	// Count the number of VTStacks (each stack will allocate a feedback slot)
-	OutEnvironment.SetDefine(TEXT("NUM_VIRTUALTEXTURE_SAMPLES"), 0);// VTStacks.Num());
+	OutEnvironment.SetDefine(TEXT("NUM_VIRTUALTEXTURE_SAMPLES"), EmitMaterialData.VTStacks.Num());
 
 	// Setup defines to map each VT stack to either 1 or 2 page table textures, depending on how many layers it uses
-	/*for (int i = 0; i < VTStacks.Num(); ++i)
+	for (int i = 0; i < EmitMaterialData.VTStacks.Num(); ++i)
 	{
-		const FMaterialVirtualTextureStack& Stack = MaterialCompilationOutput.UniformExpressionSet.VTStacks[i];
-		FString PageTableValue = FString::Printf(TEXT("Material.VirtualTexturePageTable0_%d"), i);
+		const FMaterialVirtualTextureStack& Stack = MaterialCompilationOutput.UniformExpressionSet.GetVTStacks()[i];
+
+		TStringBuilder<256> PageTableName;
+		PageTableName.Appendf(TEXT("VIRTUALTEXTURE_PAGETABLE_%d"), i);
+
+		TStringBuilder<1024> PageTableValue;
+		PageTableValue.Appendf(TEXT("Material.VirtualTexturePageTable0_%d"), i);
+
 		if (Stack.GetNumLayers() > 4u)
 		{
-			PageTableValue += FString::Printf(TEXT(", Material.VirtualTexturePageTable1_%d"), i);
+			PageTableValue.Appendf(TEXT(", Material.VirtualTexturePageTable1_%d"), i);
 		}
-		if (VTStacks[i].bAdaptive)
+		if (EmitMaterialData.VTStacks[i].bAdaptive)
 		{
-			PageTableValue += FString::Printf(TEXT(", Material.VirtualTexturePageTableIndirection_%d"), i);
+			PageTableValue.Appendf(TEXT(", Material.VirtualTexturePageTableIndirection_%d"), i);
 		}
-		OutEnvironment.SetDefine(*FString::Printf(TEXT("VIRTUALTEXTURE_PAGETABLE_%d"), i), *PageTableValue);
-	}*/
+		OutEnvironment.SetDefine(PageTableName.ToString(), PageTableValue.ToString());
+	}
 
 	/*for (int32 CollectionIndex = 0; CollectionIndex < ParameterCollections.Num(); CollectionIndex++)
 	{
@@ -793,6 +799,7 @@ bool MaterialEmitHLSL(const FMaterialCompileTargetParameters& InCompilerTarget,
 	EmitContext.MaterialCompilationOutput = &OutCompilationOutput;
 
 	Material::FEmitData& EmitMaterialData = EmitContext.AcquireData<Material::FEmitData>();
+	EmitMaterialData.VTPageTableResultType = CachedTree->GetVTPageTableResultType();
 	EmitMaterialData.StaticParameters = &InStaticParameters;
 
 	const FStructField* NormalField = CachedTree->GetMaterialAttributesType()->FindFieldByName(*FMaterialAttributeDefinitionMap::GetAttributeName(MP_Normal));
diff --git a/Engine/Source/Runtime/Engine/Private/Materials/MaterialHLSLTree.cpp b/Engine/Source/Runtime/Engine/Private/Materials/MaterialHLSLTree.cpp
--- a/Engine/Source/Runtime/Engine/Private/Materials/MaterialHLSLTree.cpp
+++ b/Engine/Source/Runtime/Engine/Private/Materials/MaterialHLSLTree.cpp
@@ -11,6 +11,7 @@
 #include "Materials/MaterialFunctionInterface.h"
 #include "Engine/BlendableInterface.h" // BL_AfterTonemapping
 #include "VT/RuntimeVirtualTexture.h"
+#include "VT/VirtualTextureScalability.h"
 #include "Engine/Texture2D.h"
 #include "Engine/Font.h"
 
@@ -587,6 +588,7 @@ bool FExpressionParameter::EmitCustomHLSLParameter(FEmitContext& Context, FEmitS
 
 namespace Private
 {
+
 Shader::EValueType GetTexCoordType(EMaterialValueType TextureType)
 {
 	switch (TextureType)
@@ -601,6 +603,23 @@ Shader::EValueType GetTexCoordType(EMaterialValueType TextureType)
 	default: checkNoEntry(); return Shader::EValueType::Void;
 	}
 }
+
+ETextureMipValueMode GetMipValueMode(FEmitContext& Context, const FExpressionTextureSample* Expression)
+{
+	const ETextureMipValueMode MipValueMode = Expression->MipValueMode;
+	const bool bUseAnalyticDerivatives = Context.bUseAnalyticDerivatives && (MipValueMode != TMVM_MipLevel) && Expression->TexCoordDerivatives.IsValid();
+	if (Context.ShaderFrequency != SF_Pixel)
+	{
+		// TODO - should we allow TMVM_Derivative in non-PS?
+		return TMVM_MipLevel;
+	}
+	else if (bUseAnalyticDerivatives || MipValueMode == TMVM_Derivative)
+	{
+		return TMVM_Derivative;
+	}
+	return MipValueMode;
+}
+
 } // namespace Private
 
 bool FExpressionTextureSample::PrepareValue(FEmitContext& Context, FEmitScope& Scope, const FRequestedType& RequestedType, FPrepareValueResult& OutResult) const
@@ -634,13 +653,13 @@ bool FExpressionTextureSample::PrepareValue(FEmitContext& Context, FEmitScope& S
 		}
 	}
 
-	const bool bUseAnalyticDerivatives = Context.bUseAnalyticDerivatives && (MipValueMode != TMVM_MipLevel) && TexCoordDerivatives.IsValid();
-	if (MipValueMode == TMVM_Derivative || bUseAnalyticDerivatives)
+	const ETextureMipValueMode LocalMipValueMode = Private::GetMipValueMode(Context, this);
+	if (LocalMipValueMode == TMVM_Derivative)
 	{
 		Context.PrepareExpression(TexCoordDerivatives.ExpressionDdx, Scope, RequestedTexCoordType);
 		Context.PrepareExpression(TexCoordDerivatives.ExpressionDdy, Scope, RequestedTexCoordType);
 	}
-	else if (MipValueMode == TMVM_MipLevel || MipValueMode == TMVM_MipBias)
+	else if (LocalMipValueMode == TMVM_MipLevel || LocalMipValueMode == TMVM_MipBias)
 	{
 		Context.PrepareExpression(MipValueExpression, Scope, Shader::EValueType::Float1);
 	}
@@ -648,107 +667,389 @@ bool FExpressionTextureSample::PrepareValue(FEmitContext& Context, FEmitScope& S
 	return OutResult.SetType(Context, RequestedType, EExpressionEvaluation::Shader, Shader::EValueType::Float4);
 }
 
-void FExpressionTextureSample::EmitValueShader(FEmitContext& Context, FEmitScope& Scope, const FRequestedType& RequestedType, FEmitValueShaderResult& OutResult) const
+namespace Private
 {
-	FMaterialTextureValue TextureValue;
-	TextureExpression->GetValueObject(Context, Scope, TextureValue);
-	const EMaterialValueType TextureType = TextureValue.Texture->GetMaterialType();
 
-	bool bVirtualTexture = false;
-	const TCHAR* SampleFunctionName = nullptr;
-	switch (TextureType)
+const TCHAR* GetVTAddressMode(TextureAddress Address)
+{
+	switch (Address)
 	{
-	case MCT_Texture2D:
-		SampleFunctionName = TEXT("Texture2DSample");
-		break;
-	case MCT_Texture2DArray:
-		SampleFunctionName = TEXT("Texture2DArraySample");
-		break;
-	case MCT_TextureCube:
-		SampleFunctionName = TEXT("TextureCubeSample");
-		break;
-	case MCT_TextureCubeArray:
-		SampleFunctionName = TEXT("TextureCubeArraySample");
-		break;
-	case MCT_VolumeTexture:
-		SampleFunctionName = TEXT("Texture3DSample");
-		break;
-	case MCT_TextureExternal:
-		SampleFunctionName = TEXT("TextureExternalSample");
-		break;
-	default:
-		checkNoEntry();
-		break;
+	case TA_Wrap: return TEXT("VTADDRESSMODE_WRAP");
+	case TA_Clamp: return TEXT("VTADDRESSMODE_CLAMP");
+	case TA_Mirror: return TEXT("VTADDRESSMODE_MIRROR");
+	default: checkNoEntry(); return nullptr;
+	}
+}
+
+uint32 AcquireVTStackIndex(
+	FEmitContext& Context,
+	FEmitScope& Scope,
+	FEmitData& EmitMaterialData,
+	ETextureMipValueMode MipValueMode,
+	TextureAddress AddressU,
+	TextureAddress AddressV,
+	float AspectRatio,
+	FEmitShaderExpression* EmitTexCoordValue,
+	FEmitShaderExpression* EmitTexCoordValueDdx,
+	FEmitShaderExpression* EmitTexCoordValueDdy,
+	FEmitShaderExpression* EmitMipValue,
+	int32 PreallocatedStackTextureIndex,
+	bool bAdaptive,
+	bool bGenerateFeedback)
+{
+	FHasher Hasher;
+	AppendHashes(Hasher, MipValueMode, AddressU, AddressV, AspectRatio, EmitTexCoordValue, EmitTexCoordValueDdx, EmitTexCoordValueDdy, EmitMipValue, PreallocatedStackTextureIndex, bAdaptive, bGenerateFeedback);
+	const FXxHash64 Hash = Hasher.Finalize();
+
+	// First check to see if we have an existing VTStack that matches this key, that can still fit another layer
+	FUniformExpressionSet& UniformExpressionSet = Context.MaterialCompilationOutput->UniformExpressionSet;
+	for (int32 Index = EmitMaterialData.VTStackHash.First(Hash.Hash); EmitMaterialData.VTStackHash.IsValid(Index); Index = EmitMaterialData.VTStackHash.Next(Index))
+	{
+		const FVTStackEntry& Entry = EmitMaterialData.VTStacks[Index];
+		if (!UniformExpressionSet.GetVTStacks()[Index].AreLayersFull() &&
+			Entry.EmitTexCoordValue == EmitTexCoordValue &&
+			Entry.EmitTexCoordValueDdx == EmitTexCoordValueDdx &&
+			Entry.EmitTexCoordValueDdy == EmitTexCoordValueDdy &&
+			Entry.EmitMipValue == EmitMipValue &&
+			Entry.MipValueMode == MipValueMode &&
+			Entry.AddressU == AddressU &&
+			Entry.AddressV == AddressV &&
+			Entry.AspectRatio == AspectRatio &&
+			Entry.PreallocatedStackTextureIndex == PreallocatedStackTextureIndex &&
+			Entry.bAdaptive == bAdaptive &&
+			Entry.bGenerateFeedback == bGenerateFeedback)
+		{
+			return Index;
+		}
 	}
 
-	TStringBuilder<64> FormattedTexture;
-	Private::EmitTextureShader(Context, TextureValue, FormattedTexture);
+	// Need to allocate a new VTStack
+	const int32 StackIndex = EmitMaterialData.VTStacks.AddDefaulted();
+	EmitMaterialData.VTStackHash.Add(Hash.Hash, StackIndex);
+	FVTStackEntry& Entry = EmitMaterialData.VTStacks[StackIndex];
+	Entry.EmitTexCoordValue = EmitTexCoordValue;
+	Entry.EmitTexCoordValueDdx = EmitTexCoordValueDdx;
+	Entry.EmitTexCoordValueDdy = EmitTexCoordValueDdy;
+	Entry.MipValueMode = MipValueMode;
+	Entry.AddressU = AddressU;
+	Entry.AddressV = AddressV;
+	Entry.AspectRatio = AspectRatio;
+	Entry.PreallocatedStackTextureIndex = PreallocatedStackTextureIndex;
+	Entry.bAdaptive = bAdaptive;
+	Entry.bGenerateFeedback = bGenerateFeedback;
+
+	const int32 MaterialStackIndex = UniformExpressionSet.AddVTStack(PreallocatedStackTextureIndex);
+	// These two arrays need to stay in sync
+	check(StackIndex == MaterialStackIndex);
+
+	// Select LoadVirtualPageTable function name for this context
+	const TCHAR* BaseFunctionName = bAdaptive ? TEXT("TextureLoadVirtualPageTableAdaptive") : TEXT("TextureLoadVirtualPageTable");
+
+	// Optionally sample without virtual texture feedback but only for miplevel mode
+	check(bGenerateFeedback || MipValueMode == TMVM_MipLevel);
+
+	TStringBuilder<256> FormattedFeedback;
+	if (bGenerateFeedback)
+	{
+		FormattedFeedback.Appendf(TEXT(", %dU + LIGHTMAP_VT_ENABLED, Parameters.VirtualTextureFeedback"), StackIndex);
+	}
 
-	const bool AutomaticViewMipBias = AutomaticMipBiasExpression ? AutomaticMipBiasExpression->GetValueConstant(Context, Scope, Shader::EValueType::Bool1).AsBoolScalar() : false;
-	TStringBuilder<256> FormattedSampler;
-	switch (SamplerSource)
+	// Code to load the VT page table...this will execute the first time a given VT stack is accessed
+	// Additional stack layers will simply reuse these results
+	switch (MipValueMode)
 	{
-	case SSM_FromTextureAsset:
-		FormattedSampler.Appendf(TEXT("%sSampler"), FormattedTexture.ToString());
+	case TMVM_None:
+		Entry.EmitResult = Context.EmitExpression(Scope, EmitMaterialData.VTPageTableResultType, TEXT(
+			"%("
+			"VIRTUALTEXTURE_PAGETABLE_%, "
+			"VTPageTableUniform_Unpack(Material.VTPackedPageTableUniform[%*2], Material.VTPackedPageTableUniform[%*2+1]), "
+			"%, %, %, "
+			"0, Parameters.SvPosition.xy, "
+			"%U + LIGHTMAP_VT_ENABLED, Parameters.VirtualTextureFeedback)"),
+			BaseFunctionName,
+			StackIndex, StackIndex, StackIndex,
+			EmitTexCoordValue, Private::GetVTAddressMode(AddressU), Private::GetVTAddressMode(AddressV),
+			StackIndex);
 		break;
-	case SSM_Wrap_WorldGroupSettings:
-		FormattedSampler.Appendf(TEXT("GetMaterialSharedSampler(%sSampler,%s)"),
-			FormattedTexture.ToString(),
-			AutomaticViewMipBias ? TEXT("View.MaterialTextureBilinearWrapedSampler") : TEXT("Material.Wrap_WorldGroupSettings"));
+	case TMVM_MipBias:
+		Entry.EmitResult = Context.EmitExpression(Scope, EmitMaterialData.VTPageTableResultType, TEXT(
+			"%("
+			"VIRTUALTEXTURE_PAGETABLE_%, "
+			"VTPageTableUniform_Unpack(Material.VTPackedPageTableUniform[%*2], Material.VTPackedPageTableUniform[%*2+1]), "
+			"%, %, %, "
+			"%, Parameters.SvPosition.xy, "
+			"%U + LIGHTMAP_VT_ENABLED, Parameters.VirtualTextureFeedback)"),
+			BaseFunctionName,
+			StackIndex, StackIndex, StackIndex,
+			EmitTexCoordValue, GetVTAddressMode(AddressU), GetVTAddressMode(AddressV),
+			EmitMipValue,
+			StackIndex);
 		break;
-	case SSM_Clamp_WorldGroupSettings:
-		FormattedSampler.Appendf(TEXT("GetMaterialSharedSampler(%sSampler,%s)"),
-			FormattedTexture.ToString(),
-			AutomaticViewMipBias ? TEXT("View.MaterialTextureBilinearClampedSampler") : TEXT("Material.Clamp_WorldGroupSettings"));
+	case TMVM_MipLevel:
+		Entry.EmitResult = Context.EmitExpression(Scope, EmitMaterialData.VTPageTableResultType, TEXT(
+			"%Level("
+			"VIRTUALTEXTURE_PAGETABLE_%, "
+			"VTPageTableUniform_Unpack(Material.VTPackedPageTableUniform[%*2], Material.VTPackedPageTableUniform[%*2+1]), "
+			"%, %, %, "
+			"%"
+			"%)"),
+			BaseFunctionName,
+			StackIndex, StackIndex, StackIndex,
+			EmitTexCoordValue, GetVTAddressMode(AddressU), GetVTAddressMode(AddressV),
+			EmitMipValue,
+			FormattedFeedback.ToString());
+		break;
+	case TMVM_Derivative:
+		Entry.EmitResult = Context.EmitExpression(Scope, EmitMaterialData.VTPageTableResultType, TEXT(
+			"%Grad("
+			"VIRTUALTEXTURE_PAGETABLE_%, "
+			"VTPageTableUniform_Unpack(Material.VTPackedPageTableUniform[%*2], Material.VTPackedPageTableUniform[%*2+1]), "
+			"%, %, %, "
+			"%, %, Parameters.SvPosition.xy, "
+			"%U + LIGHTMAP_VT_ENABLED, Parameters.VirtualTextureFeedback)"),
+			BaseFunctionName,
+			StackIndex, StackIndex, StackIndex,
+			EmitTexCoordValue, GetVTAddressMode(AddressU), GetVTAddressMode(AddressV),
+			EmitTexCoordValueDdx, EmitTexCoordValueDdy,
+			StackIndex);
 		break;
 	default:
 		checkNoEntry();
 		break;
 	}
 
+	return StackIndex;
+}
+
+} // namespace Private
+
+void FExpressionTextureSample::EmitValueShader(FEmitContext& Context, FEmitScope& Scope, const FRequestedType& RequestedType, FEmitValueShaderResult& OutResult) const
+{
+	FMaterialTextureValue TextureValue;
+	TextureExpression->GetValueObject(Context, Scope, TextureValue);
+	UTexture* Texture = TextureValue.Texture;
+	const EMaterialValueType TextureType = Texture->GetMaterialType();
 	const Shader::EValueType TexCoordType = Private::GetTexCoordType(TextureType);
-	FEmitShaderExpression* TexCoordValue = TexCoordExpression->GetValueShader(Context, Scope, TexCoordType);
 
-	FEmitShaderExpression* TextureResult = nullptr;
-	if (MipValueMode == TMVM_MipLevel)
+	TextureAddress StaticAddressX = TA_Wrap;
+	TextureAddress StaticAddressY = TA_Wrap;
+	TextureAddress StaticAddressZ = TA_Wrap;
+	if (Texture && Texture->Source.GetNumBlocks() > 1)
 	{
-		TextureResult = Context.EmitExpression(Scope, Shader::EValueType::Float4, TEXT("%Level(%, %, %, %)"),
-			SampleFunctionName,
-			FormattedTexture.ToString(),
-			FormattedSampler.ToString(),
-			TexCoordValue,
-			MipValueExpression->GetValueShader(Context, Scope, Shader::EValueType::Float1));
+		// UDIM (multi-block) texture are forced to use wrap address mode
+		// This is important for supporting VT stacks made from UDIMs with differing number of blocks, as this requires wrapping vAddress for certain layers
+		StaticAddressX = TA_Wrap;
+		StaticAddressY = TA_Wrap;
+		StaticAddressZ = TA_Wrap;
 	}
-	else if (MipValueMode == TMVM_Derivative || (Context.bUseAnalyticDerivatives && TexCoordDerivatives.IsValid()))
+	else
 	{
-		TextureResult = Context.EmitExpression(Scope, Shader::EValueType::Float4, TEXT("%Grad(%, %, %, %, %)"),
-			SampleFunctionName,
-			FormattedTexture.ToString(),
-			FormattedSampler.ToString(),
-			TexCoordValue,
-			TexCoordDerivatives.ExpressionDdx->GetValueShader(Context, Scope, TexCoordType),
-			TexCoordDerivatives.ExpressionDdy->GetValueShader(Context, Scope, TexCoordType));
+		switch (SamplerSource)
+		{
+		case SSM_FromTextureAsset:
+			if (Texture)
+			{
+				StaticAddressX = Texture->GetTextureAddressX();
+				StaticAddressY = Texture->GetTextureAddressY();
+				StaticAddressZ = Texture->GetTextureAddressZ();
+			}
+			break;
+		case SSM_Wrap_WorldGroupSettings:
+			StaticAddressX = TA_Wrap;
+			StaticAddressY = TA_Wrap;
+			StaticAddressZ = TA_Wrap;
+			break;
+		case SSM_Clamp_WorldGroupSettings:
+			StaticAddressX = TA_Clamp;
+			StaticAddressY = TA_Clamp;
+			StaticAddressZ = TA_Clamp;
+			break;
+		default:
+			checkNoEntry();
+			break;
+		}
 	}
-	else if (MipValueMode == TMVM_MipBias)
+
+	FEmitShaderExpression* EmitTexCoordValue = TexCoordExpression->GetValueShader(Context, Scope, TexCoordType);
+	FEmitShaderExpression* EmitMipValue = nullptr;
+	FEmitShaderExpression* EmitTexCoordValueDdx = nullptr;
+	FEmitShaderExpression* EmitTexCoordValueDdy = nullptr;
+	const ETextureMipValueMode LocalMipValueMode = Private::GetMipValueMode(Context, this);
+	if (LocalMipValueMode == TMVM_Derivative)
+	{
+		EmitTexCoordValueDdx = TexCoordDerivatives.ExpressionDdx->GetValueShader(Context, Scope, TexCoordType);
+		EmitTexCoordValueDdy = TexCoordDerivatives.ExpressionDdy->GetValueShader(Context, Scope, TexCoordType);
+	}
+	else if (LocalMipValueMode == TMVM_MipLevel || LocalMipValueMode == TMVM_MipBias)
 	{
-		TextureResult = Context.EmitExpression(Scope, Shader::EValueType::Float4, TEXT("%Bias(%, %, %, %)"),
+		if (MipValueExpression)
+		{
+			EmitMipValue = MipValueExpression->GetValueShader(Context, Scope, Shader::EValueType::Float1);
+		}
+		else
+		{
+			EmitMipValue = Context.EmitConstantZero(Scope, Shader::EValueType::Float1);
+		}
+	}
+
+	FEmitShaderExpression* EmitTextureResult = nullptr;
+	if (TextureType == MCT_TextureVirtual)
+	{
+		FEmitData& EmitMaterialData = Context.FindData<FEmitData>();
+
+		FMaterialTextureParameterInfo TextureParameterInfo;
+		TextureParameterInfo.ParameterInfo = TextureValue.ParameterInfo;
+		TextureParameterInfo.TextureIndex = Context.Material->GetReferencedTextures().Find(TextureValue.Texture);
+		TextureParameterInfo.SamplerSource = SamplerSource;
+		check(TextureParameterInfo.TextureIndex != INDEX_NONE);
+		const int32 TextureParameterIndex = Context.MaterialCompilationOutput->UniformExpressionSet.FindOrAddTextureParameter(EMaterialTextureParameterType::Virtual, TextureParameterInfo);
+
+		// Using Source size because we care about the aspect ratio of each block (each block of multi-block texture must have same aspect ratio)
+		// We can still combine multi-block textures of different block aspect ratios, as long as each block has the same ratio
+		// This is because we only need to overlay VT pages from within a given block
+		const float TextureAspectRatio = (float)Texture->Source.GetSizeX() / (float)Texture->Source.GetSizeY();
+
+		const bool AdaptiveVirtualTexture = false;
+		const bool bGenerateFeedback = (Context.ShaderFrequency == SF_Pixel);
+		int32 VTStackIndex = Private::AcquireVTStackIndex(Context, Scope, EmitMaterialData,
+			LocalMipValueMode,
+			StaticAddressX,
+			StaticAddressY,
+			TextureAspectRatio,
+			EmitTexCoordValue,
+			EmitTexCoordValueDdx,
+			EmitTexCoordValueDdy,
+			EmitMipValue,
+			INDEX_NONE,
+			AdaptiveVirtualTexture,
+			bGenerateFeedback);
+
+		// Allocate a layer in the virtual texture stack for this physical sample
+		int32 VTLayerIndex = Context.MaterialCompilationOutput->UniformExpressionSet.AddVTLayer(VTStackIndex, TextureParameterIndex);
+		int32 VTPageTableIndex = VTLayerIndex;
+
+		TStringBuilder<64> FormattedTexture;
+		FormattedTexture.Appendf(TEXT("Material.VirtualTexturePhysical_%d"), TextureParameterIndex);
+
+		TStringBuilder<256> FormattedSampler;
+		if (SamplerSource != SSM_FromTextureAsset)
+		{
+			// VT doesn't care if the shared sampler is wrap or clamp. It only cares if it is aniso or not.
+			// The wrap/clamp/mirror operation is handled in the shader explicitly.
+			const bool bUseAnisoSampler = VirtualTextureScalability::IsAnisotropicFilteringEnabled() && LocalMipValueMode != TMVM_MipLevel;
+			const TCHAR* SharedSamplerName = bUseAnisoSampler ? TEXT("View.SharedBilinearAnisoClampedSampler") : TEXT("View.SharedBilinearClampedSampler");
+			FormattedSampler.Appendf(TEXT("GetMaterialSharedSampler(Material.VirtualTexturePhysical_%dSampler, %s)"), TextureParameterIndex, SharedSamplerName);
+		}
+		else
+		{
+			FormattedSampler.Appendf(TEXT("Material.VirtualTexturePhysical_%dSampler"), TextureParameterIndex);
+		}
+
+		const TCHAR* SampleFunctionName = (LocalMipValueMode == TMVM_MipLevel) ? TEXT("TextureVirtualSampleLevel") : TEXT("TextureVirtualSample");
+		EmitTextureResult = Context.EmitExpression(Scope, Shader::EValueType::Float4, TEXT("%(%, %, %, %, VTUniform_Unpack(Material.VTPackedUniform[%]))"),
 			SampleFunctionName,
 			FormattedTexture.ToString(),
 			FormattedSampler.ToString(),
-			TexCoordValue,
-			MipValueExpression->GetValueShader(Context, Scope, Shader::EValueType::Float1));
+			EmitMaterialData.VTStacks[VTStackIndex].EmitResult,
+			VTPageTableIndex,
+			TextureParameterIndex);
 	}
 	else
 	{
-		check(MipValueMode == TMVM_None);
-		TextureResult = Context.EmitExpression(Scope, Shader::EValueType::Float4, TEXT("%(%, %, %)"),
-			SampleFunctionName,
-			FormattedTexture.ToString(),
-			FormattedSampler.ToString(),
-			TexCoordValue);
+		// Non virtual texture
+		const TCHAR* SampleFunctionName = nullptr;
+		switch (TextureType)
+		{
+		case MCT_Texture2D:
+			SampleFunctionName = TEXT("Texture2DSample");
+			break;
+		case MCT_Texture2DArray:
+			SampleFunctionName = TEXT("Texture2DArraySample");
+			break;
+		case MCT_TextureCube:
+			SampleFunctionName = TEXT("TextureCubeSample");
+			break;
+		case MCT_TextureCubeArray:
+			SampleFunctionName = TEXT("TextureCubeArraySample");
+			break;
+		case MCT_VolumeTexture:
+			SampleFunctionName = TEXT("Texture3DSample");
+			break;
+		case MCT_TextureExternal:
+			SampleFunctionName = TEXT("TextureExternalSample");
+			break;
+		default:
+			checkNoEntry();
+			break;
+		}
+
+		TStringBuilder<64> FormattedTexture;
+		Private::EmitTextureShader(Context, TextureValue, FormattedTexture);
+
+		const bool AutomaticViewMipBias = AutomaticMipBiasExpression ? AutomaticMipBiasExpression->GetValueConstant(Context, Scope, Shader::EValueType::Bool1).AsBoolScalar() : false;
+		TStringBuilder<256> FormattedSampler;
+		switch (SamplerSource)
+		{
+		case SSM_FromTextureAsset:
+			FormattedSampler.Appendf(TEXT("%sSampler"), FormattedTexture.ToString());
+			break;
+		case SSM_Wrap_WorldGroupSettings:
+			FormattedSampler.Appendf(TEXT("GetMaterialSharedSampler(%sSampler,%s)"),
+				FormattedTexture.ToString(),
+				AutomaticViewMipBias ? TEXT("View.MaterialTextureBilinearWrapedSampler") : TEXT("Material.Wrap_WorldGroupSettings"));
+			break;
+		case SSM_Clamp_WorldGroupSettings:
+			FormattedSampler.Appendf(TEXT("GetMaterialSharedSampler(%sSampler,%s)"),
+				FormattedTexture.ToString(),
+				AutomaticViewMipBias ? TEXT("View.MaterialTextureBilinearClampedSampler") : TEXT("Material.Clamp_WorldGroupSettings"));
+			break;
+		default:
+			checkNoEntry();
+			break;
+		}
+
+		switch (LocalMipValueMode)
+		{
+		case TMVM_Derivative:
+			EmitTextureResult = Context.EmitExpression(Scope, Shader::EValueType::Float4, TEXT("%Grad(%, %, %, %, %)"),
+				SampleFunctionName,
+				FormattedTexture.ToString(),
+				FormattedSampler.ToString(),
+				EmitTexCoordValue,
+				EmitTexCoordValueDdx,
+				EmitTexCoordValueDdy);
+			break;
+		case TMVM_MipLevel:
+			EmitTextureResult = Context.EmitExpression(Scope, Shader::EValueType::Float4, TEXT("%Level(%, %, %, %)"),
+				SampleFunctionName,
+				FormattedTexture.ToString(),
+				FormattedSampler.ToString(),
+				EmitTexCoordValue,
+				EmitMipValue);
+			break;
+		case TMVM_MipBias:
+			EmitTextureResult = Context.EmitExpression(Scope, Shader::EValueType::Float4, TEXT("%Bias(%, %, %, %)"),
+				SampleFunctionName,
+				FormattedTexture.ToString(),
+				FormattedSampler.ToString(),
+				EmitTexCoordValue,
+				EmitMipValue);
+			break;
+		case TMVM_None:
+			EmitTextureResult = Context.EmitExpression(Scope, Shader::EValueType::Float4, TEXT("%(%, %, %)"),
+				SampleFunctionName,
+				FormattedTexture.ToString(),
+				FormattedSampler.ToString(),
+				EmitTexCoordValue);
+			break;
+		default:
+			checkNoEntry();
+			break;
+		}
 	}
 
-	OutResult.Code = Context.EmitExpression(Scope, Shader::EValueType::Float4, TEXT("ApplyMaterialSamplerType(%, %)"), TextureResult, TextureValue.SamplerType);
+	check(EmitTextureResult);
+	OutResult.Code = Context.EmitExpression(Scope, Shader::EValueType::Float4, TEXT("ApplyMaterialSamplerType(%, %)"), EmitTextureResult, TextureValue.SamplerType);
 }
 
 bool FExpressionTextureSize::PrepareValue(FEmitContext& Context, FEmitScope& Scope, const FRequestedType& RequestedType, FPrepareValueResult& OutResult) const
diff --git a/Engine/Source/Runtime/Engine/Private/Materials/MaterialUniformExpressions.cpp b/Engine/Source/Runtime/Engine/Private/Materials/MaterialUniformExpressions.cpp
--- a/Engine/Source/Runtime/Engine/Private/Materials/MaterialUniformExpressions.cpp
+++ b/Engine/Source/Runtime/Engine/Private/Materials/MaterialUniformExpressions.cpp
@@ -603,6 +603,18 @@ uint32 FUniformExpressionSet::AddDefaultParameterValue(const UE::Shader::FValue&
 	return Offset;
 }
 
+int32 FUniformExpressionSet::AddVTStack(int32 InPreallocatedStackTextureIndex)
+{
+	return VTStacks.Add(FMaterialVirtualTextureStack(InPreallocatedStackTextureIndex));
+}
+
+int32 FUniformExpressionSet::AddVTLayer(int32 StackIndex, int32 TextureIndex)
+{
+	const int32 VTLayerIndex = VTStacks[StackIndex].AddLayer();
+	VTStacks[StackIndex].SetLayer(VTLayerIndex, TextureIndex);
+	return VTLayerIndex;
+}
+
 void FUniformExpressionSet::GetGameThreadTextureValue(EMaterialTextureParameterType Type, int32 Index, const UMaterialInterface* MaterialInterface, const FMaterial& Material, UTexture*& OutValue, bool bAllowOverride) const
 {
 	check(IsInGameThread());
diff --git a/Engine/Source/Runtime/Engine/Private/Shader/ShaderValue.cpp b/Engine/Source/Runtime/Engine/Private/Shader/ShaderValue.cpp
--- a/Engine/Source/Runtime/Engine/Private/Shader/ShaderValue.cpp
+++ b/Engine/Source/Runtime/Engine/Private/Shader/ShaderValue.cpp
@@ -811,21 +811,24 @@ void FStructTypeRegistry::EmitDeclarationsCode(FStringBuilderBase& OutCode) cons
 	for (const auto& It : Types)
 	{
 		const FStructType* StructType = It.Value;
-
-		OutCode.Appendf(TEXT("struct %s\n"), StructType->Name);
-		OutCode.Append(TEXT("{\n"));
-		for (const FStructField& Field : StructType->Fields)
+		// Don't need to emit declaration for external types
+		if (!StructType->IsExternal())
 		{
-			OutCode.Appendf(TEXT("\t%s %s;\n"), Field.Type.GetName(), Field.Name);
-		}
-		OutCode.Append(TEXT("};\n"));
+			OutCode.Appendf(TEXT("struct %s\n"), StructType->Name);
+			OutCode.Append(TEXT("{\n"));
+			for (const FStructField& Field : StructType->Fields)
+			{
+				OutCode.Appendf(TEXT("\t%s %s;\n"), Field.Type.GetName(), Field.Name);
+			}
+			OutCode.Append(TEXT("};\n"));
 
-		for (const FStructField& Field : StructType->Fields)
-		{
-			OutCode.Appendf(TEXT("%s %s_Set%s(%s Self, %s Value) { Self.%s = Value; return Self; }\n"),
-				StructType->Name, StructType->Name, Field.Name, StructType->Name, Field.Type.GetName(), Field.Name);
+			for (const FStructField& Field : StructType->Fields)
+			{
+				OutCode.Appendf(TEXT("%s %s_Set%s(%s Self, %s Value) { Self.%s = Value; return Self; }\n"),
+					StructType->Name, StructType->Name, Field.Name, StructType->Name, Field.Type.GetName(), Field.Name);
+			}
+			OutCode.Append(TEXT("\n"));
 		}
-		OutCode.Append(TEXT("\n"));
 	}
 }
 
@@ -936,6 +939,22 @@ const FStructType* FStructTypeRegistry::NewType(const FStructTypeInitializer& In
 	return StructType;
 }
 
+const FStructType* FStructTypeRegistry::NewExternalType(FStringView Name)
+{
+	uint64 Hash = 0u;
+	{
+		FXxHash64Builder Hasher;
+		Hasher.Update(Name.GetData(), Name.Len() * sizeof(TCHAR));
+		Hash = Hasher.Finalize().Hash;
+	}
+
+	FStructType* StructType = new(*Allocator) FStructType();
+	StructType->Name = MemStack::AllocateString(*Allocator, Name);
+	StructType->Hash = Hash;
+	Types.Add(Hash, StructType);
+	return StructType;
+}
+
 const FStructType* FStructTypeRegistry::FindType(uint64 Hash) const
 {
 	FStructType const* const* PrevType = Types.Find(Hash);
diff --git a/Engine/Source/Runtime/Engine/Public/MaterialCachedHLSLTree.h b/Engine/Source/Runtime/Engine/Public/MaterialCachedHLSLTree.h
--- a/Engine/Source/Runtime/Engine/Public/MaterialCachedHLSLTree.h
+++ b/Engine/Source/Runtime/Engine/Public/MaterialCachedHLSLTree.h
@@ -59,6 +59,7 @@ public:
 	UE::HLSLTree::FScope* GetResultScope() const { return ResultScope; }
 
 	const UE::Shader::FStructType* GetMaterialAttributesType() const { return MaterialAttributesType; }
+	const UE::Shader::FStructType* GetVTPageTableResultType() const { return VTPageTableResultType; }
 	const UE::Shader::FValue& GetMaterialAttributesDefaultValue() const { return MaterialAttributesDefaultValue; }
 
 	const TMap<FMaterialConnectionKey, const UE::HLSLTree::FExpression*>& GetConnections() const { return ConnectionMap; }
@@ -81,6 +82,7 @@ private:
 	TArray<UMaterialExpressionCustomOutput*> MaterialCustomOutputs;
 	TMap<FMaterialConnectionKey, const UE::HLSLTree::FExpression*> ConnectionMap;
 	const UE::Shader::FStructType* MaterialAttributesType = nullptr;
+	const UE::Shader::FStructType* VTPageTableResultType = nullptr;
 	UE::Shader::FValue MaterialAttributesDefaultValue;
 
 	friend class FMaterialHLSLGenerator;
diff --git a/Engine/Source/Runtime/Engine/Public/MaterialHLSLTree.h b/Engine/Source/Runtime/Engine/Public/MaterialHLSLTree.h
--- a/Engine/Source/Runtime/Engine/Public/MaterialHLSLTree.h
+++ b/Engine/Source/Runtime/Engine/Public/MaterialHLSLTree.h
@@ -7,15 +7,18 @@
 #include "MaterialTypes.h"
 #include "Materials/MaterialLayersFunctions.h"
 #include "RHIDefinitions.h"
+#include "Engine/Texture.h" // enum TextureAddress for VTStackEntry
 
 class UTexture;
 enum class EMaterialParameterType : uint8;
 struct FMaterialCachedExpressionData;
 struct FMaterialLayersFunctions;
 
-namespace UE::HLSLTree::Material
+namespace UE::HLSLTree
 {
 
+class FEmitShaderExpression;
+
 inline void AppendHash(FHasher& Hasher, const FMaterialParameterInfo& Value)
 {
 	AppendHash(Hasher, Value.Name);
@@ -44,6 +47,11 @@ inline void AppendHash(FHasher& Hasher, const FMaterialParameterMetadata& Meta)
 	AppendHash(Hasher, Meta.Value);
 }
 
+} // namespace UE::HLSLTree
+
+namespace UE::HLSLTree::Material
+{
+
 enum class EExternalInput : uint8
 {
 	None,
@@ -406,6 +414,25 @@ struct FVertexInterpolator
 	FPreparedType PreparedType;
 };
 
+struct FVTStackEntry
+{
+	FEmitShaderExpression* EmitTexCoordValue;
+	FEmitShaderExpression* EmitTexCoordValueDdx;
+	FEmitShaderExpression* EmitTexCoordValueDdy;
+	FEmitShaderExpression* EmitMipValue;
+	FEmitShaderExpression* EmitResult;
+	ETextureMipValueMode MipValueMode;
+	TextureAddress AddressU;
+	TextureAddress AddressV;
+	int32 DebugCoordinateIndex;
+	int32 DebugMipValue0Index;
+	int32 DebugMipValue1Index;
+	int32 PreallocatedStackTextureIndex;
+	bool bAdaptive;
+	bool bGenerateFeedback;
+	float AspectRatio;
+};
+
 class FEmitData
 {
 public:
@@ -417,10 +444,13 @@ public:
 		}
 	}
 
+	Shader::FType VTPageTableResultType;
 	const FStaticParameterSet* StaticParameters = nullptr;
 	FMaterialCachedExpressionData* CachedExpressionData = nullptr;
 	TMap<Shader::FValue, uint32> DefaultUniformValues;
 	TArray<FVertexInterpolator, TInlineAllocator<8>> VertexInterpolators;
+	TArray<FVTStackEntry, TInlineAllocator<8>> VTStacks;
+	FHashTable VTStackHash;
 	TBitArray<> ExternalInputMask[SF_NumFrequencies];
 	FMaterialShadingModelField ShadingModelsFromCompilation;
 	int32 NumInterpolatorComponents = 0;
diff --git a/Engine/Source/Runtime/Engine/Public/MaterialShared.h b/Engine/Source/Runtime/Engine/Public/MaterialShared.h
--- a/Engine/Source/Runtime/Engine/Public/MaterialShared.h
+++ b/Engine/Source/Runtime/Engine/Public/MaterialShared.h
@@ -578,6 +578,10 @@ public:
 	int32 FindOrAddNumericParameter(EMaterialParameterType Type, const FMaterialParameterInfo& ParameterInfo, uint32 DefaultValueOffset);
 	uint32 AddDefaultParameterValue(const UE::Shader::FValue& Value);
 
+	TConstArrayView<FMaterialVirtualTextureStack> GetVTStacks() const { return VTStacks; }
+	int32 AddVTStack(int32 InPreallocatedStackTextureIndex);
+	int32 AddVTLayer(int32 StackIndex, int32 TextureIndex);
+
 protected:
 	union FVTPackedStackAndLayerIndex
 	{
diff --git a/Engine/Source/Runtime/Engine/Public/Shader/ShaderTypes.h b/Engine/Source/Runtime/Engine/Public/Shader/ShaderTypes.h
--- a/Engine/Source/Runtime/Engine/Public/Shader/ShaderTypes.h
+++ b/Engine/Source/Runtime/Engine/Public/Shader/ShaderTypes.h
@@ -223,7 +223,7 @@ struct FStructType
 {
 	uint64 Hash;
 	const TCHAR* Name;
-	const FStructType* DerivativeType;
+	const FStructType* DerivativeType = nullptr;
 	TArrayView<const FStructField> Fields;
 
 	/**
@@ -237,6 +237,7 @@ struct FStructType
 	 */
 	TArrayView<const EValueType> FlatFieldTypes;
 
+	bool IsExternal() const { return Fields.Num() == 0; }
 	int32 GetNumComponents() const { return ComponentTypes.Num(); }
 
 	const FStructField* FindFieldByName(const TCHAR* InName) const;
@@ -266,6 +267,7 @@ public:
 	void EmitDeclarationsCode(FStringBuilderBase& OutCode) const;
 
 	const FStructType* NewType(const FStructTypeInitializer& Initializer);
+	const FStructType* NewExternalType(FStringView Name);
 	const FStructType* FindType(uint64 Hash) const;
 
 private:
