Fixed more non unity compile errors

diff --git a/Engine/Source/Runtime/Engine/Private/VT/RuntimeVirtualTextureVolume.cpp b/Engine/Source/Runtime/Engine/Private/VT/RuntimeVirtualTextureVolume.cpp
--- a/Engine/Source/Runtime/Engine/Private/VT/RuntimeVirtualTextureVolume.cpp
+++ b/Engine/Source/Runtime/Engine/Private/VT/RuntimeVirtualTextureVolume.cpp
@@ -4,6 +4,7 @@
 
 #include "Components/BoxComponent.h"
 #include "Components/RuntimeVirtualTextureComponent.h"
+#include "UObject/FortniteMainBranchObjectVersion.h"
 
 ARuntimeVirtualTextureVolume::ARuntimeVirtualTextureVolume(const FObjectInitializer& ObjectInitializer)
 	: Super(ObjectInitializer)
