Fix compile error in TraceScreenshot.cpp

diff --git a/Engine/Source/Runtime/Core/Public/ProfilingDebugging/MiscTrace.h b/Engine/Source/Runtime/Core/Public/ProfilingDebugging/MiscTrace.h
--- a/Engine/Source/Runtime/Core/Public/ProfilingDebugging/MiscTrace.h
+++ b/Engine/Source/Runtime/Core/Public/ProfilingDebugging/MiscTrace.h
@@ -82,10 +82,14 @@ private:
 #define TRACE_END_FRAME(FrameType) \
 	FMiscTrace::OutputEndFrame(FrameType);
 
+#define TRACE_SCREENSHOT(Name, Cycle, Width, Height, Data) \
+	FMiscTrace::OutputScreenshot(Name, Cycle, Width, Height, Data);
+
 #else
 
 #define TRACE_BOOKMARK(...)
 #define TRACE_BEGIN_FRAME(...)
 #define TRACE_END_FRAME(...)
+#define TRACE_SCREENSHOT(...)
 
 #endif
diff --git a/Engine/Source/Runtime/Engine/Private/ProfilingDebugging/TraceScreenshot.cpp b/Engine/Source/Runtime/Engine/Private/ProfilingDebugging/TraceScreenshot.cpp
--- a/Engine/Source/Runtime/Engine/Private/ProfilingDebugging/TraceScreenshot.cpp
+++ b/Engine/Source/Runtime/Engine/Private/ProfilingDebugging/TraceScreenshot.cpp
@@ -87,7 +87,7 @@ void FTraceScreenshot::HandleScreenshotData(int32 InSizeX, int32 InSizeY, const
 	TArray64<uint8> CompressedBitmap;
 	FImageUtils::PNGCompressImageArray(ResizedX, ResizedY, TArrayView64<const FColor>(ResizedImage.GetData(), ResizedImage.Num()), CompressedBitmap);
 
-	FMiscTrace::OutputScreenshot(*ScreenshotName, Cycles, ResizedX, ResizedY, CompressedBitmap);
+	TRACE_SCREENSHOT(*ScreenshotName, Cycles, ResizedX, ResizedY, CompressedBitmap);
 	ScreenshotName.Empty();
 }
 
