Fixed crash caused by missing push of ActorEditorContext when creating a level instance from selection

diff --git a/Engine/Source/Runtime/Engine/Private/LevelInstance/LevelInstanceSubsystem.cpp b/Engine/Source/Runtime/Engine/Private/LevelInstance/LevelInstanceSubsystem.cpp
--- a/Engine/Source/Runtime/Engine/Private/LevelInstance/LevelInstanceSubsystem.cpp
+++ b/Engine/Source/Runtime/Engine/Private/LevelInstance/LevelInstanceSubsystem.cpp
@@ -974,7 +974,9 @@ ALevelInstance* ULevelInstanceSubsystem::CreateLevelInstanceFrom(const TArray<AA
 	{
 		DirtyPackages.Add(NewLevelInstanceActor->GetPackage()->GetFName());
 	}
-			
+	
+	// Commit will always pop the actor editor context, make sure to push one here
+	UActorEditorContextSubsystem::Get()->PushContext();
 	bool bCommitted = CommitLevelInstanceInternal(TempLevelInstanceEdit, /*bDiscardEdits=*/false, /*bDiscardOnFailure=*/true, &DirtyPackages);
 	check(bCommitted);
 	check(!TempLevelInstanceEdit);
