VirtualTexture: Initial changes required for tiled media texture integration.

diff --git a/Engine/Source/Runtime/Engine/Private/Materials/MaterialShared.cpp b/Engine/Source/Runtime/Engine/Private/Materials/MaterialShared.cpp
--- a/Engine/Source/Runtime/Engine/Private/Materials/MaterialShared.cpp
+++ b/Engine/Source/Runtime/Engine/Private/Materials/MaterialShared.cpp
@@ -3147,9 +3147,9 @@ int32 FMaterialVirtualTextureStack::FindLayer(int32 UniformExpressionIndex) cons
 	return -1;
 }
 
-void FMaterialVirtualTextureStack::GetTextureValues(const FMaterialRenderContext& Context, const FUniformExpressionSet& UniformExpressionSet, UTexture2D const** OutValues) const
+void FMaterialVirtualTextureStack::GetTextureValues(const FMaterialRenderContext& Context, const FUniformExpressionSet& UniformExpressionSet, UTexture const** OutValues) const
 {
-	FMemory::Memzero(OutValues, sizeof(FVirtualTexture2DResource*) * VIRTUALTEXTURE_SPACE_MAXLAYERS);
+	FMemory::Memzero(OutValues, sizeof(UTexture*) * VIRTUALTEXTURE_SPACE_MAXLAYERS);
 	
 	for (uint32 LayerIndex = 0u; LayerIndex < NumLayers; ++LayerIndex)
 	{
@@ -3158,7 +3158,7 @@ void FMaterialVirtualTextureStack::GetTextureValues(const FMaterialRenderContext
 		{
 			const UTexture* Texture = nullptr;
 			UniformExpressionSet.GetTextureValue(EMaterialTextureParameterType::Virtual, ParameterIndex, Context, Context.Material, Texture);
-			OutValues[LayerIndex] = Cast<UTexture2D>(Texture);
+			OutValues[LayerIndex] = Texture;
 		}
 	}
 }
@@ -3284,7 +3284,7 @@ IAllocatedVirtualTexture* FMaterialRenderProxy::AllocateVTStack(const FMaterialR
 		return nullptr;
 	}
 
-	const UTexture2D* LayerTextures[VIRTUALTEXTURE_SPACE_MAXLAYERS] = { nullptr };
+	const UTexture* LayerTextures[VIRTUALTEXTURE_SPACE_MAXLAYERS] = { nullptr };
 	VTStack.GetTextureValues(Context, UniformExpressionSet, LayerTextures);
 
 	const UMaterialInterface* MaterialInterface = GetMaterialInterface();
@@ -3299,7 +3299,7 @@ IAllocatedVirtualTexture* FMaterialRenderProxy::AllocateVTStack(const FMaterialR
 	bool bFoundValidLayer = false;
 	for (uint32 LayerIndex = 0u; LayerIndex < NumLayers; ++LayerIndex)
 	{
-		const UTexture2D* Texture = LayerTextures[LayerIndex];
+		const UTexture* Texture = LayerTextures[LayerIndex];
 		if (!Texture)
 		{
 			continue;
diff --git a/Engine/Source/Runtime/Engine/Private/Texture2D.cpp b/Engine/Source/Runtime/Engine/Private/Texture2D.cpp
--- a/Engine/Source/Runtime/Engine/Private/Texture2D.cpp
+++ b/Engine/Source/Runtime/Engine/Private/Texture2D.cpp
@@ -1226,6 +1226,19 @@ bool UTexture2D::IsCurrentlyVirtualTextured() const
 	return false;
 }
 
+FVirtualTexture2DResource::FVirtualTexture2DResource()
+	: AllocatedVT(nullptr)
+	, VTData(nullptr)
+	, TextureOwner(nullptr)
+	, ProducerHandle()
+	, FirstMipToUse(0)
+{
+	//NOTE: Empty constructor for use with media textures (which do not derive from UTexture2D).
+
+	// Initialize this resource FeatureLevel, so it gets re-created on FeatureLevel changes
+	SetFeatureLevel(GMaxRHIFeatureLevel);
+}
+
 FVirtualTexture2DResource::FVirtualTexture2DResource(const UTexture2D* InOwner, FVirtualTextureBuiltData* InVTData, int32 InFirstMipToUse)
 	: AllocatedVT(nullptr)
 	, VTData(InVTData)
diff --git a/Engine/Source/Runtime/Engine/Public/MaterialShared.h b/Engine/Source/Runtime/Engine/Public/MaterialShared.h
--- a/Engine/Source/Runtime/Engine/Public/MaterialShared.h
+++ b/Engine/Source/Runtime/Engine/Public/MaterialShared.h
@@ -345,8 +345,8 @@ public:
 
 	/** Returns true if this is a stack that with a preallocated layout of layers (for example when we are using a URuntimeVirtualTexture). */
 	inline bool IsPreallocatedStack() const { return PreallocatedStackTextureIndex != INDEX_NONE; }
-	/** Get the array of UTexture2D objects for the expressions that in the layers of this stack. Can return nullptr objects for layers that don't hold UTexture2D references. */
-	void GetTextureValues(const FMaterialRenderContext& Context, const FUniformExpressionSet& UniformExpressionSet, UTexture2D const** OutValues) const;
+	/** Get the array of UTexture objects for the expressions that in the layers of this stack. Can return nullptr objects for layers that don't hold UTexture2D references. */
+	void GetTextureValues(const FMaterialRenderContext& Context, const FUniformExpressionSet& UniformExpressionSet, UTexture const** OutValues) const;
 	/** Get the URuntimeVirtualTexture object if one was used to initialize this stack. */
 	void GetTextureValue(const FMaterialRenderContext& Context, const FUniformExpressionSet& UniformExpressionSet, const URuntimeVirtualTexture*& OutValue) const;
 
diff --git a/Engine/Source/Runtime/Engine/Public/TextureResource.h b/Engine/Source/Runtime/Engine/Public/TextureResource.h
--- a/Engine/Source/Runtime/Engine/Public/TextureResource.h
+++ b/Engine/Source/Runtime/Engine/Public/TextureResource.h
@@ -155,11 +155,12 @@ protected :
 class FVirtualTexture2DResource : public FTextureResource
 {
 public:
-	FVirtualTexture2DResource(const UTexture2D* InOwner, struct FVirtualTextureBuiltData* InVTData, int32 FirstMipToUse);
-	virtual ~FVirtualTexture2DResource();
+	ENGINE_API FVirtualTexture2DResource();
+	ENGINE_API FVirtualTexture2DResource(const UTexture2D* InOwner, struct FVirtualTextureBuiltData* InVTData, int32 FirstMipToUse);
+	ENGINE_API virtual ~FVirtualTexture2DResource();
 
-	virtual void InitRHI() override;
-	virtual void ReleaseRHI() override;
+	ENGINE_API virtual void InitRHI() override;
+	ENGINE_API virtual void ReleaseRHI() override;
 	
 	// Dynamic cast methods.
 	ENGINE_API virtual FVirtualTexture2DResource* GetVirtualTexture2DResource() { return this; }
@@ -167,11 +168,11 @@ public:
 	ENGINE_API virtual const FVirtualTexture2DResource* GetVirtualTexture2DResource() const { return this; }
 
 #if WITH_EDITOR
-	void InitializeEditorResources(class IVirtualTexture* InVirtualTexture);
+	ENGINE_API virtual void InitializeEditorResources(class IVirtualTexture* InVirtualTexture);
 #endif
 
-	virtual uint32 GetSizeX() const override;
-	virtual uint32 GetSizeY() const override;
+	ENGINE_API virtual uint32 GetSizeX() const override;
+	ENGINE_API virtual uint32 GetSizeY() const override;
 
 	const FVirtualTextureProducerHandle& GetProducerHandle() const { return ProducerHandle; }
 
@@ -187,19 +188,19 @@ public:
 	ENGINE_API class IAllocatedVirtualTexture* AcquireAllocatedVT();
 	ENGINE_API void ReleaseAllocatedVT();
 
-	ENGINE_API EPixelFormat GetFormat(uint32 LayerIndex) const;
-	ENGINE_API FIntPoint GetSizeInBlocks() const;
-	ENGINE_API uint32 GetNumTilesX() const;
-	ENGINE_API uint32 GetNumTilesY() const;
-	ENGINE_API uint32 GetNumMips() const;
-	ENGINE_API uint32 GetNumLayers() const;
-	ENGINE_API uint32 GetTileSize() const; //no borders
-	ENGINE_API uint32 GetBorderSize() const;
+	ENGINE_API virtual EPixelFormat GetFormat(uint32 LayerIndex) const;
+	ENGINE_API virtual FIntPoint GetSizeInBlocks() const;
+	ENGINE_API virtual uint32 GetNumTilesX() const;
+	ENGINE_API virtual uint32 GetNumTilesY() const;
+	ENGINE_API virtual uint32 GetNumMips() const;
+	ENGINE_API virtual uint32 GetNumLayers() const;
+	ENGINE_API virtual uint32 GetTileSize() const; //no borders
+	ENGINE_API virtual uint32 GetBorderSize() const;
 	uint32 GetAllocatedvAddress() const;
 
-	ENGINE_API FIntPoint GetPhysicalTextureSize(uint32 LayerIndex) const;
+	ENGINE_API virtual FIntPoint GetPhysicalTextureSize(uint32 LayerIndex) const;
 
-private:
+protected:
 	class IAllocatedVirtualTexture* AllocatedVT;
 	struct FVirtualTextureBuiltData* VTData;
 	const UTexture2D* TextureOwner;
diff --git a/Engine/Source/Runtime/RenderCore/Public/VirtualTexturing.h b/Engine/Source/Runtime/RenderCore/Public/VirtualTexturing.h
--- a/Engine/Source/Runtime/RenderCore/Public/VirtualTexturing.h
+++ b/Engine/Source/Runtime/RenderCore/Public/VirtualTexturing.h
@@ -133,6 +133,7 @@ struct FVTProducerDescription
 	
 	bool bPersistentHighestMip = true;
 	bool bContinuousUpdate = false;
+	bool bNotifyCompleted = false; /** Producer will receive OnRequestsCompleted() callbacks every frame when enabled. */
 	
 	uint32 TileSize = 0u;
 	uint32 TileBorderSize = 0u;
@@ -307,6 +308,9 @@ public:
 
 	/** Dump any type specific debug info. */
 	virtual void DumpToConsole(bool verbose) {}
+
+	/** Called on every virtual texture system update once all requests are completed, if bNotifyCompleted is enabled. */
+	virtual void OnRequestsCompleted() {}
 };
 
 enum class EVTPageTableFormat : uint8
diff --git a/Engine/Source/Runtime/Renderer/Private/VT/VirtualTextureProducer.cpp b/Engine/Source/Runtime/Renderer/Private/VT/VirtualTextureProducer.cpp
--- a/Engine/Source/Runtime/Renderer/Private/VT/VirtualTextureProducer.cpp
+++ b/Engine/Source/Runtime/Renderer/Private/VT/VirtualTextureProducer.cpp
@@ -225,6 +225,17 @@ uint32 FVirtualTextureProducerCollection::RemoveAllCallbacks(const void* Baton)
 	return NumRemoved;
 }
 
+void FVirtualTextureProducerCollection::NotifyRequestsCompleted()
+{
+	for (FProducerEntry& Entry : Producers)
+	{
+		if (Entry.Producer.Description.bNotifyCompleted)
+		{
+			Entry.Producer.GetVirtualTexture()->OnRequestsCompleted();
+		}
+	}
+}
+
 FVirtualTextureProducer* FVirtualTextureProducerCollection::FindProducer(const FVirtualTextureProducerHandle& Handle)
 {
 	FProducerEntry* Entry = GetEntry(Handle);
diff --git a/Engine/Source/Runtime/Renderer/Private/VT/VirtualTextureProducer.h b/Engine/Source/Runtime/Renderer/Private/VT/VirtualTextureProducer.h
--- a/Engine/Source/Runtime/Renderer/Private/VT/VirtualTextureProducer.h
+++ b/Engine/Source/Runtime/Renderer/Private/VT/VirtualTextureProducer.h
@@ -63,6 +63,9 @@ public:
 	uint32 RemoveAllCallbacks(const void* Baton);
 	void CallPendingCallbacks();
 
+	/** Notify producers marked as "continous notify" that all requests have been completed. */
+	void NotifyRequestsCompleted();
+
 	/**
 	 * Gets the producer associated with the given handle, or nullptr if handle is invalid
 	 * Returned pointer is only valid until the next call to RegisterProducer, so should not be stored beyond scope of a function
diff --git a/Engine/Source/Runtime/Renderer/Private/VT/VirtualTextureSystem.cpp b/Engine/Source/Runtime/Renderer/Private/VT/VirtualTextureSystem.cpp
--- a/Engine/Source/Runtime/Renderer/Private/VT/VirtualTextureSystem.cpp
+++ b/Engine/Source/Runtime/Renderer/Private/VT/VirtualTextureSystem.cpp
@@ -1618,6 +1618,8 @@ void FVirtualTextureSystem::Update(FRDGBuilder& GraphBuilder, ERHIFeatureLevel::
 	// Submit the merged requests
 	SubmitRequests(GraphBuilder, FeatureLevel, MemStack, MergedRequestList, true);
 
+	Producers.NotifyRequestsCompleted();
+
 	UpdateResidencyTracking();
 
 #if !UE_BUILD_SHIPPING
@@ -2150,16 +2152,29 @@ void FVirtualTextureSystem::GatherRequestsTask(const FGatherRequestsParameters&
 void FVirtualTextureSystem::GetContinuousUpdatesToProduce(FUniqueRequestList const* RequestList, int32 MaxTilesToProduce)
 {
 	const int32 NumContinuousUpdateRequests = (int32)RequestList->GetNumContinuousUpdateRequests();
-	const int32 MaxContinousUpdates = FMath::Min(VirtualTextureScalability::GetMaxContinuousUpdatesPerFrame(), NumContinuousUpdateRequests);
-
-	int32 NumContinuousUpdates = 0;
-	while (NumContinuousUpdates < MaxContinousUpdates && ContinuousUpdateTilesToProduce.Num() < MaxTilesToProduce)
+	const int32 MaxContinuousUpdatesPerFrame = VirtualTextureScalability::GetMaxContinuousUpdatesPerFrame();
+	
+	// Negative maximum continous updates allows for uncapped requests
+	if (MaxContinuousUpdatesPerFrame < 0)
 	{
-		// Note it's possible that we add a duplicate value to the TSet here, and so MappedTilesToProduce doesn't grow.
-		// But ending up with fewer continuous updates then the maximum is OK.
-		int32 RandomIndex = FMath::Rand() % NumContinuousUpdateRequests;
-		ContinuousUpdateTilesToProduce.Add(RequestList->GetContinuousUpdateRequest(RandomIndex));
-		NumContinuousUpdates++;
+		for (int32 i = 0; i < NumContinuousUpdateRequests; ++i)
+		{
+			ContinuousUpdateTilesToProduce.Add(RequestList->GetContinuousUpdateRequest(i));
+		}
+	}
+	else
+	{
+		const int32 MaxContinousUpdates = FMath::Min(MaxContinuousUpdatesPerFrame, NumContinuousUpdateRequests);
+
+		int32 NumContinuousUpdates = 0;
+		while (NumContinuousUpdates < MaxContinousUpdates && ContinuousUpdateTilesToProduce.Num() < MaxTilesToProduce)
+		{
+			// Note it's possible that we add a duplicate value to the TSet here, and so MappedTilesToProduce doesn't grow.
+			// But ending up with fewer continuous updates then the maximum is OK.
+			int32 RandomIndex = FMath::Rand() % NumContinuousUpdateRequests;
+			ContinuousUpdateTilesToProduce.Add(RequestList->GetContinuousUpdateRequest(RandomIndex));
+			NumContinuousUpdates++;
+		}
 	}
 }
 
