Revert unintended change to LLM tracker tags, move audio spatialization plugins tracking to use new tag with appropriate stat group

#rnx

diff --git a/Engine/Source/Runtime/Core/Private/HAL/LowLevelMemTracker.cpp b/Engine/Source/Runtime/Core/Private/HAL/LowLevelMemTracker.cpp
--- a/Engine/Source/Runtime/Core/Private/HAL/LowLevelMemTracker.cpp
+++ b/Engine/Source/Runtime/Core/Private/HAL/LowLevelMemTracker.cpp
@@ -135,6 +135,7 @@ DECLARE_LLM_MEMORY_STAT(TEXT("LLM Overhead"), STAT_LLMOverheadTotal, STATGROUP_L
 DEFINE_STAT(STAT_EngineSummaryLLM);
 DEFINE_STAT(STAT_ProjectSummaryLLM);
 DEFINE_STAT(STAT_NetworkingSummaryLLM);
+DEFINE_STAT(STAT_AudioSummaryLLM);
 
 /*
  * LLM stats referenced by ELLMTagNames
@@ -168,7 +169,6 @@ DECLARE_LLM_MEMORY_STAT(TEXT("Audio"), STAT_AudioLLM, STATGROUP_LLMFULL);
 DECLARE_LLM_MEMORY_STAT(TEXT("AudioMisc"), STAT_AudioMiscLLM, STATGROUP_LLMFULL);
 DECLARE_LLM_MEMORY_STAT(TEXT("AudioSoundWaves"), STAT_AudioSoundWavesLLM, STATGROUP_LLMFULL);
 DECLARE_LLM_MEMORY_STAT(TEXT("AudioSoundWaveProxies"), STAT_AudioSoundWaveProxiesLLM, STATGROUP_LLMFULL);
-DECLARE_LLM_MEMORY_STAT(TEXT("AudioSpatializationPlugins"), STAT_AudioSpatializationPluginsLLM, STATGROUP_LLMFULL);
 DECLARE_LLM_MEMORY_STAT(TEXT("AudioMixer"), STAT_AudioMixerLLM, STATGROUP_LLMFULL);
 DECLARE_LLM_MEMORY_STAT(TEXT("AudioMixerPlugins"), STAT_AudioMixerPluginsLLM, STATGROUP_LLMFULL);
 DECLARE_LLM_MEMORY_STAT(TEXT("AudioPrecache"), STAT_AudioPrecacheLLM, STATGROUP_LLMFULL);
@@ -253,7 +253,6 @@ DECLARE_LLM_MEMORY_STAT(TEXT("CustomName"), STAT_CustomName, STATGROUP_LLMFULL);
 * LLM Summary stats referenced by ELLMTagNames
 */
 DECLARE_LLM_MEMORY_STAT(TEXT("Total"), STAT_TrackedTotalSummaryLLM, STATGROUP_LLM);
-DECLARE_LLM_MEMORY_STAT(TEXT("Audio"), STAT_AudioSummaryLLM, STATGROUP_LLM);
 DECLARE_LLM_MEMORY_STAT(TEXT("Meshes"), STAT_MeshesSummaryLLM, STATGROUP_LLM);
 DECLARE_LLM_MEMORY_STAT(TEXT("Physics"), STAT_PhysicsSummaryLLM, STATGROUP_LLM);
 DECLARE_LLM_MEMORY_STAT(TEXT("PhysX"), STAT_PhysXSummaryLLM, STATGROUP_LLM);
diff --git a/Engine/Source/Runtime/Core/Public/HAL/LowLevelMemStats.h b/Engine/Source/Runtime/Core/Public/HAL/LowLevelMemStats.h
--- a/Engine/Source/Runtime/Core/Public/HAL/LowLevelMemStats.h
+++ b/Engine/Source/Runtime/Core/Public/HAL/LowLevelMemStats.h
@@ -25,6 +25,7 @@
 	DECLARE_LLM_MEMORY_STAT_EXTERN(TEXT("Engine"), STAT_EngineSummaryLLM, STATGROUP_LLM, CORE_API);
 	DECLARE_LLM_MEMORY_STAT_EXTERN(TEXT("Project"), STAT_ProjectSummaryLLM, STATGROUP_LLM, CORE_API);
 	DECLARE_LLM_MEMORY_STAT_EXTERN(TEXT("Networking"), STAT_NetworkingSummaryLLM, STATGROUP_LLM, CORE_API);
+	DECLARE_LLM_MEMORY_STAT_EXTERN(TEXT("Audio"), STAT_AudioSummaryLLM, STATGROUP_LLM, CORE_API);
 
 #else
 	#define DECLARE_LLM_MEMORY_STAT(CounterName,StatId,GroupId)
diff --git a/Engine/Source/Runtime/Core/Public/HAL/LowLevelMemTracker.h b/Engine/Source/Runtime/Core/Public/HAL/LowLevelMemTracker.h
--- a/Engine/Source/Runtime/Core/Public/HAL/LowLevelMemTracker.h
+++ b/Engine/Source/Runtime/Core/Public/HAL/LowLevelMemTracker.h
@@ -153,7 +153,6 @@ enum class ELLMTagSet : uint8
 	macro(AudioMisc,							"AudioMisc",					GET_STATFNAME(STAT_AudioMiscLLM),							GET_STATFNAME(STAT_AudioSummaryLLM),			ELLMTag::Audio)\
 	macro(AudioSoundWaves,						"AudioSoundWaves",				GET_STATFNAME(STAT_AudioSoundWavesLLM),						GET_STATFNAME(STAT_AudioSummaryLLM),			ELLMTag::Audio)\
 	macro(AudioSoundWaveProxies,				"AudioSoundWaveProxies",		GET_STATFNAME(STAT_AudioSoundWaveProxiesLLM),				GET_STATFNAME(STAT_AudioSummaryLLM),			ELLMTag::Audio)\
-	macro(AudioSpatializationPlugins,			"AudioSpatializationPlugins",	GET_STATFNAME(STAT_AudioSpatializationPluginsLLM),			GET_STATFNAME(STAT_AudioSummaryLLM),			ELLMTag::Audio)\
 	macro(AudioMixer,							"AudioMixer",					GET_STATFNAME(STAT_AudioMixerLLM),							GET_STATFNAME(STAT_AudioSummaryLLM),			ELLMTag::Audio)\
 	macro(AudioMixerPlugins,					"AudioMixerPlugins",			GET_STATFNAME(STAT_AudioMixerPluginsLLM),					GET_STATFNAME(STAT_AudioSummaryLLM),			ELLMTag::Audio)\
 	macro(AudioPrecache,						"AudioPrecache",				GET_STATFNAME(STAT_AudioPrecacheLLM),						GET_STATFNAME(STAT_AudioSummaryLLM),			ELLMTag::Audio)\
@@ -250,7 +249,7 @@ enum class ELLMTag : LLM_TAG_TYPE
 
 	//------------------------------
 	// Platform tags
-	PlatformTagStart = 109,
+	PlatformTagStart = 108,
 	PlatformTagEnd = 149,
 
 	//------------------------------
diff --git a/Engine/Source/Runtime/Engine/Private/AudioDevice.cpp b/Engine/Source/Runtime/Engine/Private/AudioDevice.cpp
--- a/Engine/Source/Runtime/Engine/Private/AudioDevice.cpp
+++ b/Engine/Source/Runtime/Engine/Private/AudioDevice.cpp
@@ -20,7 +20,6 @@
 #include "GeneralProjectSettings.h"
 #include "Engine/World.h"
 #include "HAL/FileManager.h"
-#include "HAL/LowLevelMemTracker.h"
 #include "IAudioExtensionPlugin.h"
 #include "IAudioParameterTransmitter.h"
 #include "Misc/ConfigCacheIni.h"
@@ -54,6 +53,10 @@
 #include "Editor/EditorEngine.h"
 #endif // WITH_EDITOR
 
+DECLARE_LLM_MEMORY_STAT(TEXT("AudioSpatializationPlugins"), STAT_AudioSpatializationPluginsLLM, STATGROUP_LLMFULL);
+LLM_DEFINE_TAG(AudioSpatializationPlugins, NAME_None, TEXT("Audio"), GET_STATFNAME(STAT_AudioSpatializationPluginsLLM), GET_STATFNAME(STAT_AudioSummaryLLM));
+
+#define AUDIO_SPATIALIZATION_PLUGIN_LLM_SCOPE LLM_SCOPE_BYTAG(AudioSpatializationPlugins); 
 
 static int32 AudioChannelCountCVar = 0;
 FAutoConsoleVariableRef CVarSetAudioChannelCount(
@@ -783,7 +786,7 @@ void FAudioDevice::Teardown()
 
 	for (TAudioPluginListenerPtr PluginListener : PluginListeners)
 	{
-		LLM_SCOPE(ELLMTag::AudioSpatializationPlugins);
+		AUDIO_SPATIALIZATION_PLUGIN_LLM_SCOPE
 
 		PluginListener->OnListenerShutdown(this);
 	}
@@ -3224,7 +3227,7 @@ void FAudioDevice::SetListener(UWorld* World, const int32 InViewportIndex, const
 
 	if (World)
 	{
-		LLM_SCOPE(ELLMTag::AudioSpatializationPlugins);
+		AUDIO_SPATIALIZATION_PLUGIN_LLM_SCOPE
 
 		for (TAudioPluginListenerPtr PluginManager : PluginListeners)
 		{
@@ -3237,7 +3240,7 @@ void FAudioDevice::SetListener(UWorld* World, const int32 InViewportIndex, const
 		// Broadcast to a 3rd party plugin listener observer if enabled
 		for (TAudioPluginListenerPtr PluginManager : PluginListeners)
 		{
-			LLM_SCOPE(ELLMTag::AudioSpatializationPlugins);
+			AUDIO_SPATIALIZATION_PLUGIN_LLM_SCOPE
 
 			PluginManager->OnListenerUpdated(this, InViewportIndex, ListenerTransformCopy, InDeltaSeconds);
 		}
@@ -4790,7 +4793,7 @@ void FAudioDevice::InitializePluginListeners(UWorld* World)
 	check(IsInGameThread());
 	check(!bPluginListenersInitialized);
 
-	LLM_SCOPE(ELLMTag::AudioSpatializationPlugins);
+	AUDIO_SPATIALIZATION_PLUGIN_LLM_SCOPE
 
 	for (TAudioPluginListenerPtr PluginListener : PluginListeners)
 	{
@@ -4802,7 +4805,7 @@ void FAudioDevice::NotifyPluginListenersWorldChanged(UWorld* World)
 {
 	check(IsInGameThread());
 
-	LLM_SCOPE(ELLMTag::AudioSpatializationPlugins);
+	AUDIO_SPATIALIZATION_PLUGIN_LLM_SCOPE
 
 	for (TAudioPluginListenerPtr PluginListener : PluginListeners)
 	{
diff --git a/Engine/Source/Runtime/Engine/Public/AudioDevice.h b/Engine/Source/Runtime/Engine/Public/AudioDevice.h
--- a/Engine/Source/Runtime/Engine/Public/AudioDevice.h
+++ b/Engine/Source/Runtime/Engine/Public/AudioDevice.h
@@ -9,6 +9,7 @@
 #include "DSP/SpectrumAnalyzer.h"
 #include "Engine/Engine.h"
 #include "EngineGlobals.h"
+#include "HAL/LowLevelMemStats.h"
 #include "IAudioExtensionPlugin.h"
 #include "AudioDynamicParameter.h"
 #include "Sound/AudioSettings.h"
@@ -62,6 +63,8 @@ struct FAudioQualitySettings;
 UE_DECLARE_LWC_TYPE(Rotator, 3);
 struct FWaveInstance;
 
+LLM_DECLARE_TAG_API(AudioSpatializationPlugins, ENGINE_API);
+
 namespace Audio
 {
 	class FAudioDebugger;
