Added missing ModularFeatureList lock/unlock

diff --git a/Engine/Source/Developer/AudioSettingsEditor/Private/TargetPlatformAudioCustomization.cpp b/Engine/Source/Developer/AudioSettingsEditor/Private/TargetPlatformAudioCustomization.cpp
--- a/Engine/Source/Developer/AudioSettingsEditor/Private/TargetPlatformAudioCustomization.cpp
+++ b/Engine/Source/Developer/AudioSettingsEditor/Private/TargetPlatformAudioCustomization.cpp
@@ -90,12 +90,14 @@ TSharedRef<SWidget> FAudioPluginWidgetManager::MakeAudioPluginSelectorWidget(con
 
 #if WITH_ENGINE
 	// Scan through all currently enabled audio plugins of this specific type:
-
 	switch (AudioPluginType)
 	{
 		case EAudioPlugin::SPATIALIZATION:
 		{
+			IModularFeatures::Get().LockModularFeatureList();
 			TArray<IAudioSpatializationFactory*> AvailableSpatializationPlugins = IModularFeatures::Get().GetModularFeatureImplementations<IAudioSpatializationFactory>(IAudioSpatializationFactory::GetModularFeatureName());
+			IModularFeatures::Get().UnlockModularFeatureList();
+
 			for (IAudioSpatializationFactory* Plugin : AvailableSpatializationPlugins)
 			{
 				if (Plugin->SupportsPlatform(PlatformName))
@@ -108,7 +110,10 @@ TSharedRef<SWidget> FAudioPluginWidgetManager::MakeAudioPluginSelectorWidget(con
 
 		case EAudioPlugin::REVERB:
 		{
+			IModularFeatures::Get().LockModularFeatureList();
 			TArray<IAudioReverbFactory*> AvailableReverbPlugins = IModularFeatures::Get().GetModularFeatureImplementations<IAudioReverbFactory>(IAudioReverbFactory::GetModularFeatureName());
+			IModularFeatures::Get().UnlockModularFeatureList();
+
 			for (IAudioReverbFactory* Plugin : AvailableReverbPlugins)
 			{
 				if (Plugin->SupportsPlatform(PlatformName))
@@ -121,7 +126,10 @@ TSharedRef<SWidget> FAudioPluginWidgetManager::MakeAudioPluginSelectorWidget(con
 
 		case EAudioPlugin::OCCLUSION:
 		{
+			IModularFeatures::Get().LockModularFeatureList();
 			TArray<IAudioOcclusionFactory*> AvailableOcclusionPlugins = IModularFeatures::Get().GetModularFeatureImplementations<IAudioOcclusionFactory>(IAudioOcclusionFactory::GetModularFeatureName());
+			IModularFeatures::Get().UnlockModularFeatureList();
+
 			for (IAudioOcclusionFactory* Plugin : AvailableOcclusionPlugins)
 			{
 				if (Plugin->SupportsPlatform(PlatformName))
@@ -134,7 +142,10 @@ TSharedRef<SWidget> FAudioPluginWidgetManager::MakeAudioPluginSelectorWidget(con
 
 		case EAudioPlugin::SOURCEDATAOVERRIDE:
 		{
+			IModularFeatures::Get().LockModularFeatureList();
 			TArray<IAudioSourceDataOverrideFactory*> AvailableSourceDataOverridePlugins = IModularFeatures::Get().GetModularFeatureImplementations<IAudioSourceDataOverrideFactory>(IAudioSourceDataOverrideFactory::GetModularFeatureName());
+			IModularFeatures::Get().UnlockModularFeatureList();
+
 			for (IAudioSourceDataOverrideFactory* Plugin : AvailableSourceDataOverridePlugins)
 			{
 				if (Plugin->SupportsPlatform(PlatformName))
@@ -145,7 +156,11 @@ TSharedRef<SWidget> FAudioPluginWidgetManager::MakeAudioPluginSelectorWidget(con
 		}
 		break;
 
+		// note: if adding cases here, make sure to lock the Modular Feature List after retrieving the implementations
+		//		IModularFeatures::Get().LockModularFeatureList(); & IModularFeatures::Get().UnlockModularFeatureList();
+
 		default:
+			IModularFeatures::Get().UnlockModularFeatureList();
 			break;
 	}
 #endif // #if WITH_ENGINE
diff --git a/Engine/Source/Runtime/AudioAnalyzer/Private/AudioAnalyzerFacade.cpp b/Engine/Source/Runtime/AudioAnalyzer/Private/AudioAnalyzerFacade.cpp
--- a/Engine/Source/Runtime/AudioAnalyzer/Private/AudioAnalyzerFacade.cpp
+++ b/Engine/Source/Runtime/AudioAnalyzer/Private/AudioAnalyzerFacade.cpp
@@ -11,7 +11,9 @@ namespace Audio
 	IAnalyzerFactory* GetAnalyzerFactory(FName InFactoryName)
 	{
 		// Get all analyzer nrt factories implementations.
+		IModularFeatures::Get().LockModularFeatureList();
 		TArray<IAnalyzerFactory*> RegisteredFactories = IModularFeatures::Get().GetModularFeatureImplementations<IAnalyzerFactory>(IAnalyzerFactory::GetModularFeatureName());
+		IModularFeatures::Get().UnlockModularFeatureList();
 
 		// Get the factory of interest by matching the name.
 		TArray<IAnalyzerFactory*> MatchingFactories = RegisteredFactories.FilterByPredicate([InFactoryName](IAnalyzerFactory* Factory) { check(nullptr != Factory); return Factory->GetName() == InFactoryName; });
diff --git a/Engine/Source/Runtime/AudioAnalyzer/Private/AudioAnalyzerNRTFacade.cpp b/Engine/Source/Runtime/AudioAnalyzer/Private/AudioAnalyzerNRTFacade.cpp
--- a/Engine/Source/Runtime/AudioAnalyzer/Private/AudioAnalyzerNRTFacade.cpp
+++ b/Engine/Source/Runtime/AudioAnalyzer/Private/AudioAnalyzerNRTFacade.cpp
@@ -12,7 +12,9 @@ namespace Audio
 	IAnalyzerNRTFactory* GetAnalyzerNRTFactory(FName InFactoryName)
 	{
 		// Get all analyzer nrt factories implementations.
+		IModularFeatures::Get().LockModularFeatureList();
 		TArray<IAnalyzerNRTFactory*> RegisteredFactories = IModularFeatures::Get().GetModularFeatureImplementations<IAnalyzerNRTFactory>(IAnalyzerNRTFactory::GetModularFeatureName());
+		IModularFeatures::Get().UnlockModularFeatureList();
 
 		// Get the factory of interest by matching the name.
 		TArray<IAnalyzerNRTFactory*> MatchingFactories = RegisteredFactories.FilterByPredicate([InFactoryName](IAnalyzerNRTFactory* Factory) { check(nullptr != Factory); return Factory->GetName() == InFactoryName; });
diff --git a/Engine/Source/Runtime/AudioCaptureCore/Private/AudioCaptureInternal.h b/Engine/Source/Runtime/AudioCaptureCore/Private/AudioCaptureInternal.h
--- a/Engine/Source/Runtime/AudioCaptureCore/Private/AudioCaptureInternal.h
+++ b/Engine/Source/Runtime/AudioCaptureCore/Private/AudioCaptureInternal.h
@@ -32,7 +32,9 @@ namespace Audio
 
 	FORCEINLINE TUniquePtr<IAudioCaptureStream> FAudioCapture::CreateImpl()
 	{
+		IModularFeatures::Get().LockModularFeatureList();
 		TArray<IAudioCaptureFactory*> AudioCaptureStreamFactories = IModularFeatures::Get().GetModularFeatureImplementations<IAudioCaptureFactory>(IAudioCaptureFactory::GetModularFeatureName());
+		IModularFeatures::Get().UnlockModularFeatureList();
 
 		// For now, just return the first audio capture stream implemented. We can make this configurable at a later point.
 		if (AudioCaptureStreamFactories.Num() > 0 && AudioCaptureStreamFactories[0] != nullptr)
diff --git a/Engine/Source/Runtime/AudioExtensions/Private/IAudioEndpoint.cpp b/Engine/Source/Runtime/AudioExtensions/Private/IAudioEndpoint.cpp
--- a/Engine/Source/Runtime/AudioExtensions/Private/IAudioEndpoint.cpp
+++ b/Engine/Source/Runtime/AudioExtensions/Private/IAudioEndpoint.cpp
@@ -159,7 +159,9 @@ IAudioEndpointFactory* IAudioEndpointFactory::Get(const FName& InName)
 		return nullptr;
 	}
 
+	IModularFeatures::Get().LockModularFeatureList();
 	TArray<IAudioEndpointFactory*> Factories = IModularFeatures::Get().GetModularFeatureImplementations<IAudioEndpointFactory>(GetModularFeatureName());
+	IModularFeatures::Get().UnlockModularFeatureList();
 
 	for (IAudioEndpointFactory* Factory : Factories)
 	{
@@ -180,7 +182,10 @@ TArray<FName> IAudioEndpointFactory::GetAvailableEndpointTypes()
 
 	EndpointNames.Add(GetTypeNameForDefaultEndpoint());
 
+	IModularFeatures::Get().LockModularFeatureList();
 	TArray<IAudioEndpointFactory*> Factories = IModularFeatures::Get().GetModularFeatureImplementations<IAudioEndpointFactory>(GetModularFeatureName());
+	IModularFeatures::Get().UnlockModularFeatureList();
+
 	for (IAudioEndpointFactory* Factory : Factories)
 	{
 		EndpointNames.AddUnique(Factory->GetEndpointTypeName());
diff --git a/Engine/Source/Runtime/AudioExtensions/Private/ISoundfieldEndpoint.cpp b/Engine/Source/Runtime/AudioExtensions/Private/ISoundfieldEndpoint.cpp
--- a/Engine/Source/Runtime/AudioExtensions/Private/ISoundfieldEndpoint.cpp
+++ b/Engine/Source/Runtime/AudioExtensions/Private/ISoundfieldEndpoint.cpp
@@ -116,8 +116,9 @@ ISoundfieldEndpointFactory* ISoundfieldEndpointFactory::Get(const FName& InName)
 	{
 		return nullptr;
 	}
-
+	IModularFeatures::Get().LockModularFeatureList();
 	TArray<ISoundfieldEndpointFactory*> Factories = IModularFeatures::Get().GetModularFeatureImplementations<ISoundfieldEndpointFactory>(GetModularFeatureName());
+	IModularFeatures::Get().UnlockModularFeatureList();
 
 	for (ISoundfieldEndpointFactory* Factory : Factories)
 	{
@@ -137,7 +138,10 @@ TArray<FName> ISoundfieldEndpointFactory::GetAllSoundfieldEndpointTypes()
 
 	SoundfieldFormatNames.Add(DefaultSoundfieldEndpointName());
 
+	IModularFeatures::Get().LockModularFeatureList();
 	TArray<ISoundfieldEndpointFactory*> Factories = IModularFeatures::Get().GetModularFeatureImplementations<ISoundfieldEndpointFactory>(GetModularFeatureName());
+	IModularFeatures::Get().UnlockModularFeatureList();
+
 	for (ISoundfieldEndpointFactory* Factory : Factories)
 	{
 		SoundfieldFormatNames.Add(Factory->GetEndpointTypeName());
diff --git a/Engine/Source/Runtime/AudioExtensions/Private/ISoundfieldFormat.cpp b/Engine/Source/Runtime/AudioExtensions/Private/ISoundfieldFormat.cpp
--- a/Engine/Source/Runtime/AudioExtensions/Private/ISoundfieldFormat.cpp
+++ b/Engine/Source/Runtime/AudioExtensions/Private/ISoundfieldFormat.cpp
@@ -40,7 +40,9 @@ ISoundfieldFactory* ISoundfieldFactory::Get(const FName& InName)
 		return nullptr;
 	}
 
+	IModularFeatures::Get().LockModularFeatureList();
 	TArray<ISoundfieldFactory*> Factories = IModularFeatures::Get().GetModularFeatureImplementations<ISoundfieldFactory>(GetModularFeatureName());
+	IModularFeatures::Get().UnlockModularFeatureList();
 
 	for (ISoundfieldFactory* Factory : Factories)
 	{
@@ -66,7 +68,10 @@ TArray<FName> ISoundfieldFactory::GetAvailableSoundfieldFormats()
 	SoundfieldFormatNames.Add(GetFormatNameForInheritedEncoding());
 	SoundfieldFormatNames.Add(GetFormatNameForNoEncoding());
 
+	IModularFeatures::Get().LockModularFeatureList();
 	TArray<ISoundfieldFactory*> Factories = IModularFeatures::Get().GetModularFeatureImplementations<ISoundfieldFactory>(GetModularFeatureName());
+	IModularFeatures::Get().UnlockModularFeatureList();
+
 	for (ISoundfieldFactory* Factory : Factories)
 	{
 		SoundfieldFormatNames.Add(Factory->GetSoundfieldFormatName());
diff --git a/Engine/Source/Runtime/AudioLink/AudioLinkEngine/Private/IAudioLinkFactory.cpp b/Engine/Source/Runtime/AudioLink/AudioLinkEngine/Private/IAudioLinkFactory.cpp
--- a/Engine/Source/Runtime/AudioLink/AudioLinkEngine/Private/IAudioLinkFactory.cpp
+++ b/Engine/Source/Runtime/AudioLink/AudioLinkEngine/Private/IAudioLinkFactory.cpp
@@ -38,6 +38,7 @@ FName IAudioLinkFactory::GetModularFeatureName()
 
 TArray<IAudioLinkFactory*> IAudioLinkFactory::GetAllRegisteredFactories()
 {
+	IModularFeatures::FScopedLockModularFeatureList ScopedLockModularFeatureList;
 	return IModularFeatures::Get().GetModularFeatureImplementations<IAudioLinkFactory>(GetModularFeatureName());
 }
 
diff --git a/Engine/Source/Runtime/Engine/Private/AudioPluginUtilities.cpp b/Engine/Source/Runtime/Engine/Private/AudioPluginUtilities.cpp
--- a/Engine/Source/Runtime/Engine/Private/AudioPluginUtilities.cpp
+++ b/Engine/Source/Runtime/Engine/Private/AudioPluginUtilities.cpp
@@ -50,7 +50,9 @@ IAudioSpatializationFactory* AudioPluginUtilities::GetDesiredSpatializationPlugi
 {
 	FString DesiredSpatializationPlugin = GetDesiredPluginName(EAudioPlugin::SPATIALIZATION);
 
+	IModularFeatures::Get().LockModularFeatureList();
 	TArray<IAudioSpatializationFactory *> SpatializationPluginFactories = IModularFeatures::Get().GetModularFeatureImplementations<IAudioSpatializationFactory>(IAudioSpatializationFactory::GetModularFeatureName());
+	IModularFeatures::Get().UnlockModularFeatureList();
 
 	//Iterate through all of the plugins we've discovered:
 	for (IAudioSpatializationFactory* PluginFactory : SpatializationPluginFactories)
@@ -68,7 +70,10 @@ IAudioSpatializationFactory* AudioPluginUtilities::GetDesiredSpatializationPlugi
 IAudioSourceDataOverrideFactory* AudioPluginUtilities::GetDesiredSourceDataOverridePlugin()
 {
 	FString DesiredSourceDataOverridePlugin = GetDesiredPluginName(EAudioPlugin::SOURCEDATAOVERRIDE);
+
+	IModularFeatures::Get().LockModularFeatureList();
 	TArray<IAudioSourceDataOverrideFactory*> SourceDataOverridePluginFactories = IModularFeatures::Get().GetModularFeatureImplementations<IAudioSourceDataOverrideFactory>(IAudioSourceDataOverrideFactory::GetModularFeatureName());
+	IModularFeatures::Get().UnlockModularFeatureList();
 
 	//Iterate through all of the plugins we've discovered:
 	for (IAudioSourceDataOverrideFactory* PluginFactory : SourceDataOverridePluginFactories)
@@ -89,7 +94,9 @@ IAudioReverbFactory* AudioPluginUtilities::GetDesiredReverbPlugin()
 	//Get the name of the desired Reverb plugin:
 	FString DesiredReverbPlugin = GetDesiredPluginName(EAudioPlugin::REVERB);
 
+	IModularFeatures::Get().LockModularFeatureList();
 	TArray<IAudioReverbFactory *> ReverbPluginFactories = IModularFeatures::Get().GetModularFeatureImplementations<IAudioReverbFactory>(IAudioReverbFactory::GetModularFeatureName());
+	IModularFeatures::Get().UnlockModularFeatureList();
 
 	//Iterate through all of the plugins we've discovered:
 	for (IAudioReverbFactory* PluginFactory : ReverbPluginFactories)
@@ -108,7 +115,9 @@ IAudioOcclusionFactory* AudioPluginUtilities::GetDesiredOcclusionPlugin()
 {
 	FString DesiredOcclusionPlugin = GetDesiredPluginName(EAudioPlugin::OCCLUSION);
 
+	IModularFeatures::Get().LockModularFeatureList();
 	TArray<IAudioOcclusionFactory *> OcclusionPluginFactories = IModularFeatures::Get().GetModularFeatureImplementations<IAudioOcclusionFactory>(IAudioOcclusionFactory::GetModularFeatureName());
+	IModularFeatures::Get().UnlockModularFeatureList();
 
 	//Iterate through all of the plugins we've discovered:
 	for (IAudioOcclusionFactory* PluginFactory : OcclusionPluginFactories)
@@ -129,7 +138,10 @@ IAudioModulationFactory* AudioPluginUtilities::GetDesiredModulationPlugin()
 	const FName& PluginName = (PlatformPluginName == NAME_None) ? GetDefaultModulationPluginName() : PlatformPluginName;
 	const FName& FeatureName = IAudioModulationFactory::GetModularFeatureName();
 
+	IModularFeatures::Get().LockModularFeatureList();
 	TArray<IAudioModulationFactory*> Factories = IModularFeatures::Get().GetModularFeatureImplementations<IAudioModulationFactory>(FeatureName);
+	IModularFeatures::Get().UnlockModularFeatureList();
+
 	for (IAudioModulationFactory* Factory : Factories)
 	{
 		//if this plugin's name matches the name found in the platform settings, use it:
diff --git a/Engine/Source/Runtime/SignalProcessing/Private/ConvolutionAlgorithm.cpp b/Engine/Source/Runtime/SignalProcessing/Private/ConvolutionAlgorithm.cpp
--- a/Engine/Source/Runtime/SignalProcessing/Private/ConvolutionAlgorithm.cpp
+++ b/Engine/Source/Runtime/SignalProcessing/Private/ConvolutionAlgorithm.cpp
@@ -22,7 +22,9 @@ namespace Audio
 	TUniquePtr<IConvolutionAlgorithm> FConvolutionFactory::NewConvolutionAlgorithm(const FConvolutionSettings& InSettings, const FName& InAlgorithmFactoryName)
 	{
 		// Get all IConvolutionAlgorithm factories. 
+		IModularFeatures::Get().LockModularFeatureList();
 		TArray<IConvolutionAlgorithmFactory*> Factories = IModularFeatures::Get().GetModularFeatureImplementations<IConvolutionAlgorithmFactory>(IConvolutionAlgorithmFactory::GetModularFeatureName());
+		IModularFeatures::Get().UnlockModularFeatureList();
 
 		// Remove null factories
 		Factories = Factories.FilterByPredicate([InAlgorithmFactoryName](const IConvolutionAlgorithmFactory* Factory) 
