Load DP on demand when using dp.override

diff --git a/Engine/Source/Runtime/Engine/Private/DeviceProfiles/DeviceProfileManager.cpp b/Engine/Source/Runtime/Engine/Private/DeviceProfiles/DeviceProfileManager.cpp
--- a/Engine/Source/Runtime/Engine/Private/DeviceProfiles/DeviceProfileManager.cpp
+++ b/Engine/Source/Runtime/Engine/Private/DeviceProfiles/DeviceProfileManager.cpp
@@ -1091,7 +1091,27 @@ void UDeviceProfileManager::HandleDeviceProfileOverrideChange()
 	// only handle when the value is different
 	if (CVarValue.Len() > 0 && CVarValue != GetActiveProfile()->GetName())
 	{
-		UDeviceProfile* NewActiveProfile = FindProfile(CVarValue, false);
+		FString PlatformName = ANSI_TO_TCHAR(FPlatformProperties::IniPlatformName());
+		
+		TArray<FString> DeviceProfileNameAndTypes;
+		FConfigFile PlatformConfigFile;
+		FConfigCacheIni::LoadLocalIniFile(PlatformConfigFile, TEXT("DeviceProfiles"), true, *PlatformName);
+		PlatformConfigFile.GetArray(TEXT("DeviceProfiles"), TEXT("DeviceProfileNameAndTypes"), DeviceProfileNameAndTypes);
+			
+		bool bCreateIfMissing = false;
+		for (const FString& Desc: DeviceProfileNameAndTypes)
+		{
+			FString Name, DeviceType;
+			Desc.Split(TEXT(","), &Name, &DeviceType);
+			if ((DeviceType == PlatformName) && (Name == CVarValue))
+			{
+				bCreateIfMissing = true;
+				break;
+			}
+
+		}
+
+		UDeviceProfile* NewActiveProfile = FindProfile(CVarValue, bCreateIfMissing);
 		if (NewActiveProfile)
 		{
 			SetOverrideDeviceProfile(NewActiveProfile);
