Fixed static assert due to wrong enum used in a loop.

diff --git a/Engine/Source/Runtime/Engine/Private/Materials/MaterialExpressions.cpp b/Engine/Source/Runtime/Engine/Private/Materials/MaterialExpressions.cpp
--- a/Engine/Source/Runtime/Engine/Private/Materials/MaterialExpressions.cpp
+++ b/Engine/Source/Runtime/Engine/Private/Materials/MaterialExpressions.cpp
@@ -9198,7 +9198,7 @@ void UMaterialExpressionShadingPathSwitch::GatherStrataMaterialInfo(FStrataMater
 	{
 		Default.GetTracedInput().Expression->GatherStrataMaterialInfo(StrataMaterialInfo, Default.OutputIndex);
 	}
-	for (int i = 0; i < EMaterialQualityLevel::Num; ++i)
+	for (int i = 0; i < ERHIShadingPath::Num; ++i)
 	{
 		Inputs[i].GetTracedInput().Expression->GatherStrataMaterialInfo(StrataMaterialInfo, Inputs[i].OutputIndex);
 	}
