Strata - UMaterialExpressionStrataTransmittanceToMFP node fixup.

diff --git a/Engine/Source/Runtime/Engine/Classes/Materials/MaterialExpressionStrata.h b/Engine/Source/Runtime/Engine/Classes/Materials/MaterialExpressionStrata.h
--- a/Engine/Source/Runtime/Engine/Classes/Materials/MaterialExpressionStrata.h
+++ b/Engine/Source/Runtime/Engine/Classes/Materials/MaterialExpressionStrata.h
@@ -812,14 +812,14 @@ class UMaterialExpressionStrataTransmittanceToMFP : public UMaterialExpressionSt
 	GENERATED_UCLASS_BODY()
 
 	/**
-	* The colored transmittance for a view perpendicular to the surface. The transmittance or other view orientation will automatically deduced according to surface thickness.
+	* The colored transmittance for a view perpendicular to the surface. The transmittance for other view orientations will automatically be deduced according to surface thickness.
 	*/
 	UPROPERTY()
 	FExpressionInput TransmittanceColor;
 
 	/**
-	* The desired thickness in centimers. This can be set lower than 0.1mm (= 0.01cm) to enable the Thin lighting model on the slab node for instance.
-	* Another use case exemple: the thickenss ouput can be modulated to have simple scattering/transmittance variation of the same material before it is plugged in a slab node.
+	* The desired thickness in centimeter. This can be set lower than 0.1mm (= 0.01cm) to enable the Thin lighting model on the slab node for instance.
+	* Another use case example: this node output called thickness can be modulated before it is plugged in a slab node, this can be used to achieve simple scattering/transmittance variation of the same material.
 	*/
 	UPROPERTY()
 	FExpressionInput Thickness;
diff --git a/Engine/Source/Runtime/Engine/Private/Materials/MaterialExpressions.cpp b/Engine/Source/Runtime/Engine/Private/Materials/MaterialExpressions.cpp
--- a/Engine/Source/Runtime/Engine/Private/Materials/MaterialExpressions.cpp
+++ b/Engine/Source/Runtime/Engine/Private/Materials/MaterialExpressions.cpp
@@ -22871,19 +22871,45 @@ uint32 UMaterialExpressionStrataTransmittanceToMFP::GetOutputType(int32 OutputIn
 
 uint32 UMaterialExpressionStrataTransmittanceToMFP::GetInputType(int32 InputIndex)
 {
-	return MCT_Float3;
-}
-void UMaterialExpressionStrataTransmittanceToMFP::GetConnectorToolTip(int32 InputIndex, int32 OutputIndex, TArray<FString>& OutToolTip)
-{
-	switch (OutputIndex)
+	switch (InputIndex)
 	{
 	case 0:
-		ConvertToMultilineToolTip(TEXT("The Mean Free Path defining the participating media constituting the slab of material (unit = centimeters)."), 80, OutToolTip);
+		return MCT_Float3; // Transmittance
 		break;
 	case 1:
-		ConvertToMultilineToolTip(TEXT("The thickness of the slab of material (unit = centimeters)."), 80, OutToolTip);
+		return MCT_Float1; // Thickness
 		break;
 	}
+
+	check(false);
+	return MCT_Float1;
+}
+void UMaterialExpressionStrataTransmittanceToMFP::GetConnectorToolTip(int32 InputIndex, int32 OutputIndex, TArray<FString>& OutToolTip)
+{
+	if (InputIndex != INDEX_NONE)
+	{
+		switch (InputIndex)
+		{
+		case 0:
+			ConvertToMultilineToolTip(TEXT("The colored transmittance for a view perpendicular to the surface. The transmittance for other view orientations will automatically be deduced according to surface thickness."), 80, OutToolTip);
+			break;
+		case 1:
+			ConvertToMultilineToolTip(TEXT("The desired thickness in centimeter. This can be set lower than 0.1mm(= 0.01cm) to enable the Thin lighting model on the slab node for instance. Another use case example: this node output called thickness can be modulated before it is plugged in a slab node.this can be used to achieve simple scattering/transmittance variation of the same material."), 80, OutToolTip);
+			break;
+		}
+	}
+	else if (OutputIndex != INDEX_NONE)
+	{
+		switch (OutputIndex)
+		{
+		case 0:
+			ConvertToMultilineToolTip(TEXT("The Mean Free Path defining the participating media constituting the slab of material (unit = centimeters)."), 80, OutToolTip);
+			break;
+		case 1:
+			ConvertToMultilineToolTip(TEXT("The thickness of the slab of material (unit = centimeters)."), 80, OutToolTip);
+			break;
+		}
+	}
 }
 
 void UMaterialExpressionStrataTransmittanceToMFP::GetExpressionToolTip(TArray<FString>& OutToolTip)
