Added Conv_StringToInt64() to Kismet string library.

PR #8987: add support to convert between string and int64 in BP (Contributed by [REDACTED_PERSON])

diff --git a/Engine/Source/Runtime/Engine/Classes/Kismet/KismetStringLibrary.h b/Engine/Source/Runtime/Engine/Classes/Kismet/KismetStringLibrary.h
--- a/Engine/Source/Runtime/Engine/Classes/Kismet/KismetStringLibrary.h
+++ b/Engine/Source/Runtime/Engine/Classes/Kismet/KismetStringLibrary.h
@@ -88,6 +88,10 @@ class ENGINE_API UKismetStringLibrary : public UBlueprintFunctionLibrary
 	UFUNCTION(BlueprintPure, meta=(DisplayName = "String To Integer", CompactNodeTitle = "->", BlueprintAutocast), Category="Utilities|String")
 	static int32 Conv_StringToInt(const FString& InString);
 
+	/** Converts a string to a int value */
+	UFUNCTION(BlueprintPure, meta = (DisplayName = "String To Integer64", CompactNodeTitle = "->", BlueprintAutocast), Category = "Utilities|String")
+	static int64 Conv_StringToInt64(const FString& InString);
+
 	/** Converts a string to a float value */
 	UE_DEPRECATED(5.0, "This method has been deprecated and will be removed. Please use LexFromString instead.")
 	static float Conv_StringToFloat(const FString& InString);
diff --git a/Engine/Source/Runtime/Engine/Private/KismetStringLibrary.cpp b/Engine/Source/Runtime/Engine/Private/KismetStringLibrary.cpp
--- a/Engine/Source/Runtime/Engine/Private/KismetStringLibrary.cpp
+++ b/Engine/Source/Runtime/Engine/Private/KismetStringLibrary.cpp
@@ -146,6 +146,11 @@ int32 UKismetStringLibrary::Conv_StringToInt(const FString& InString)
 	return FCString::Atoi(*InString);
 }
 
+int64 UKismetStringLibrary::Conv_StringToInt64(const FString& InString)
+{
+	return FCString::Atoi64(*InString);
+}
+
 float UKismetStringLibrary::Conv_StringToFloat(const FString& InString)
 {
 	return FCString::Atof(*InString);
