Dump DDC keys for global shaders.

- The cvar `r.Material.DumpDDCKeys` has been renamed to `r.ShaderCompiler.DumpDDCKeys`.
- The directory they output to changed from "MaterialDDCKeys" to "ShaderDDCKeys"
- Break getting the global shader map DDC key string out into a function to make it easier to call when we dump DDC keys.
- Global Shader map section files are named "GlobalShaderMap-%s" where %s is the name of the section (usually a global shader type).
- Editor and Game shaders are dumped to separated directories.
- Expose ShouldDumpShaderDDCKeys() function so material code can see if dumping is enabled.

diff --git a/Engine/Source/Runtime/Engine/Private/Materials/MaterialShader.cpp b/Engine/Source/Runtime/Engine/Private/Materials/MaterialShader.cpp
--- a/Engine/Source/Runtime/Engine/Private/Materials/MaterialShader.cpp
+++ b/Engine/Source/Runtime/Engine/Private/Materials/MaterialShader.cpp
@@ -29,14 +29,6 @@ static FAutoConsoleVariableRef CVarMaterialExcludeNonPipelinedShaders(
 	ECVF_ReadOnly
 );
 
-int32 GMaterialDumpDDCKeys = 0;
-static FAutoConsoleVariableRef CVarMaterialDumpDDCKeys(
-	TEXT("r.Material.DumpDDCKeys"),
-	GMaterialDumpDDCKeys,
-	TEXT("if != 0, DDC keys for each material shadermap will be dumped into project's Saved directory (MaterialDDCKeys subdirectory)"),
-	ECVF_Default
-);
-
 #if ENABLE_COOK_STATS
 namespace MaterialShaderCookStats
 {
@@ -1353,17 +1345,10 @@ void FMaterialShaderMap::LoadFromDerivedDataCache(const FMaterial* Material, con
 			const FString DataKey = GetMaterialShaderMapKeyString(ShaderMapId, InPlatform);
 			OutDDCKeyDesc = FSHA1_HashString(DataKey);
 
-			if (UNLIKELY(GMaterialDumpDDCKeys))
+			if (UNLIKELY(ShouldDumpShaderDDCKeys()))
 			{
-				FString TempPath = FPaths::ProjectSavedDir() / TEXT("MaterialDDCKeys") / LexToString(InPlatform);
-				IFileManager::Get().MakeDirectory(*TempPath, true);
-
-				FString FileName = FString::Printf(TEXT("%s-%s-%s.txt"), *Material->GetAssetName(), *LexToString(ShaderMapId.FeatureLevel), *LexToString(ShaderMapId.QualityLevel));
-				FString TempFile = TempPath / FileName;
-
-				TUniquePtr<FArchive> DumpAr(IFileManager::Get().CreateFileWriter(*TempFile));
-				// serializing the string via << produces a non-textual file because it saves string's length, too
-				DumpAr->Serialize(const_cast<TCHAR*>(*DataKey), DataKey.Len() * sizeof(TCHAR));
+				const FString FileName = FString::Printf(TEXT("%s-%s-%s.txt"), *Material->GetAssetName(), *LexToString(ShaderMapId.FeatureLevel), *LexToString(ShaderMapId.QualityLevel));
+				DumpShaderDDCKeyToFile(InPlatform, ShaderMapId.LayoutParams.WithEditorOnly(), FileName, DataKey);
 			}
 
 			bool CheckCache = true;
diff --git a/Engine/Source/Runtime/Engine/Private/ShaderCompiler/ShaderCompiler.cpp b/Engine/Source/Runtime/Engine/Private/ShaderCompiler/ShaderCompiler.cpp
--- a/Engine/Source/Runtime/Engine/Private/ShaderCompiler/ShaderCompiler.cpp
+++ b/Engine/Source/Runtime/Engine/Private/ShaderCompiler/ShaderCompiler.cpp
@@ -188,7 +188,12 @@ static FAutoConsoleVariableRef CVarShaderCompilerTooLongIOThresholdSeconds(
 	ECVF_Default
 );
 
-
+static TAutoConsoleVariable<bool> CVarShaderCompilerDumpDDCKeys(
+	TEXT("r.ShaderCompiler.DumpDDCKeys"),
+	false,
+	TEXT("if != 0, DDC keys for each shadermap will be dumped into project's Saved directory (ShaderDDCKeys subdirectory)"),
+	ECVF_Default
+);
 
 /** Helper functions for logging more debug info */
 namespace ShaderCompiler
@@ -239,6 +244,24 @@ const FString& GetMaterialShaderMapDDCKey()
 	return MaterialShaderMapDDCKey;
 }
 
+bool ShouldDumpShaderDDCKeys()
+{
+	return CVarShaderCompilerDumpDDCKeys.GetValueOnAnyThread();
+}
+
+void DumpShaderDDCKeyToFile(const EShaderPlatform InPlatform, bool bWithEditor, const FString& FileName, const FString& DDCKey)
+{
+	const FString SubDirectory = bWithEditor ? TEXT("Editor") : TEXT("Game");
+	const FString TempPath = FPaths::ProjectSavedDir() / TEXT("ShaderDDCKeys") / SubDirectory / LexToString(InPlatform);
+	IFileManager::Get().MakeDirectory(*TempPath, true);
+
+	const FString TempFile = TempPath / FileName;
+
+	TUniquePtr<FArchive> DumpAr(IFileManager::Get().CreateFileWriter(*TempFile));
+	// serializing the string via << produces a non-textual file because it saves string's length, too
+	DumpAr->Serialize(const_cast<TCHAR*>(*DDCKey), DDCKey.Len() * sizeof(TCHAR));
+}
+
 namespace ShaderCompiler
 {
 	bool IsJobCacheEnabled()
@@ -6661,15 +6684,20 @@ static FString GetGlobalShaderCacheFilename(EShaderPlatform Platform)
 }
 
 #if WITH_EDITOR
-/** Creates a string key for the derived data cache entry for the global shader map. */
-static UE::DerivedData::FCacheKey GetGlobalShaderMapKey(const FGlobalShaderMapId& ShaderMapId, EShaderPlatform Platform, const ITargetPlatform* TargetPlatform, TArray<FShaderTypeDependency> const& Dependencies)
+
+static FString GetGlobalShaderMapKeyString(const FGlobalShaderMapId& ShaderMapId, EShaderPlatform Platform, TArray<FShaderTypeDependency> const& Dependencies)
 {
 	FName Format = LegacyShaderPlatformToShaderFormat(Platform);
 	FString ShaderMapKeyString = Format.ToString() + TEXT("_") + FString(FString::FromInt(GetTargetPlatformManagerRef().ShaderFormatVersion(Format))) + TEXT("_");
 	ShaderMapAppendKeyString(Platform, ShaderMapKeyString);
 	ShaderMapId.AppendKeyString(ShaderMapKeyString, Dependencies);
-	const FString DataKey = FString::Printf(TEXT("%s_%s_%s"), TEXT("GSM"), *GetGlobalShaderMapDDCKey(), *ShaderMapKeyString);
+	return FString::Printf(TEXT("%s_%s_%s"), TEXT("GSM"), *GetGlobalShaderMapDDCKey(), *ShaderMapKeyString);
+}
 
+/** Creates a string key for the derived data cache entry for the global shader map. */
+static UE::DerivedData::FCacheKey GetGlobalShaderMapKey(const FGlobalShaderMapId& ShaderMapId, EShaderPlatform Platform, const ITargetPlatform* TargetPlatform, TArray<FShaderTypeDependency> const& Dependencies)
+{
+	const FString DataKey = GetGlobalShaderMapKeyString(ShaderMapId, Platform, Dependencies);
 	static const UE::DerivedData::FCacheBucket Bucket("GlobalShaderMap");
 	return {Bucket, FIoHash::HashBuffer(MakeMemoryView(FTCHARToUTF8(DataKey)))};
 }
@@ -6992,6 +7020,14 @@ void CompileGlobalShaderMap(EShaderPlatform Platform, const ITargetPlatform* Tar
 					Request.Key = GetGlobalShaderMapKey(ShaderMapId, Platform, TargetPlatform, ShaderFilenameDependencies.Value);
 					Request.UserData = uint64(BufferIndex);
 					++BufferIndex;
+
+					if (UNLIKELY(ShouldDumpShaderDDCKeys()))
+					{
+						const FString ShaderName = ShaderFilenameDependencies.Key.Replace(TEXT("/"), TEXT(".")).Replace(TEXT(".usf"), TEXT(""));
+						const FString FileName = FString::Printf(TEXT("GlobalShaderMap%s.txt"), *ShaderName);
+						const FString DataKey = GetGlobalShaderMapKeyString(ShaderMapId, Platform, ShaderFilenameDependencies.Value);
+						DumpShaderDDCKeyToFile(Platform, ShaderMapId.WithEditorOnly(), FileName, DataKey);
+					}
 				}
 
 				int32 DDCHits = 0;
diff --git a/Engine/Source/Runtime/Engine/Public/ShaderCompiler.h b/Engine/Source/Runtime/Engine/Public/ShaderCompiler.h
--- a/Engine/Source/Runtime/Engine/Public/ShaderCompiler.h
+++ b/Engine/Source/Runtime/Engine/Public/ShaderCompiler.h
@@ -1342,6 +1342,9 @@ extern ENGINE_API const FString& GetGlobalShaderMapDDCKey();
 
 extern ENGINE_API const FString& GetMaterialShaderMapDDCKey();
 
+extern ENGINE_API bool ShouldDumpShaderDDCKeys();
+extern ENGINE_API void DumpShaderDDCKeyToFile(const EShaderPlatform InPlatform, bool bWithEditor, const FString& FileName, const FString& DDCKey);
+
 /**
 * Handles serializing in MeshMaterialMaps or GlobalShaderMap from a CookOnTheFly command and applying them to the in-memory shadermaps.
 *
diff --git a/Engine/Source/Runtime/RenderCore/Private/GlobalShader.cpp b/Engine/Source/Runtime/RenderCore/Private/GlobalShader.cpp
--- a/Engine/Source/Runtime/RenderCore/Private/GlobalShader.cpp
+++ b/Engine/Source/Runtime/RenderCore/Private/GlobalShader.cpp
@@ -535,6 +535,11 @@ void FGlobalShaderMapId::AppendKeyString(FString& KeyString, const TArray<FShade
 #endif // WITH_EDITOR
 }
 
+bool FGlobalShaderMapId::WithEditorOnly() const
+{
+	return LayoutParams.WithEditorOnly();
+}
+
 bool FGlobalShaderType::ShouldCompilePipeline(const FShaderPipelineType* ShaderPipelineType, EShaderPlatform Platform, EShaderPermutationFlags Flags)
 {
 	for (const FShaderType* ShaderType : ShaderPipelineType->GetStages())
diff --git a/Engine/Source/Runtime/RenderCore/Public/GlobalShader.h b/Engine/Source/Runtime/RenderCore/Public/GlobalShader.h
--- a/Engine/Source/Runtime/RenderCore/Public/GlobalShader.h
+++ b/Engine/Source/Runtime/RenderCore/Public/GlobalShader.h
@@ -33,6 +33,8 @@ public:
 
 	RENDERCORE_API const TMap<FString, TArray<FShaderTypeDependency>>& GetShaderFilenameToDependeciesMap() const { return ShaderFilenameToDependenciesMap; }
 
+	RENDERCORE_API bool WithEditorOnly() const;
+
 private:
 	FPlatformTypeLayoutParameters LayoutParams;
 	EShaderPlatform ShaderPlatform;
