Do not allow r.SkinCache.SkipCompilingGPUSkinVF=1 on mobile platforms (FeatureLevel == ERHIFeatureLevel::ES3_1).

- Some mobile platforms have a limit on how big buffers can be for skin cache positions.  If we ever exceed this amount on a specific mesh, we skip the skin cache operation and need to fallback to a GPU Skin Shader.
- Since we unfortunately don't know which shaders we will need on a per-mesh basis we must disable this feature and assume we need GPU Skin Shaders for _all_ materials.
- Fix an issue where I was trying to get a bool cvar as an int and it was failing.

diff --git a/Engine/Source/Runtime/Engine/Private/GPUSkinCache.cpp b/Engine/Source/Runtime/Engine/Private/GPUSkinCache.cpp
--- a/Engine/Source/Runtime/Engine/Private/GPUSkinCache.cpp
+++ b/Engine/Source/Runtime/Engine/Private/GPUSkinCache.cpp
@@ -211,7 +211,7 @@ static inline bool IsGPUSkinCacheInclusive(EShaderPlatform Platform)
 	return (PerPlatformCVar.Get(Platform) != 0);
 }
 
-bool ShouldWeCompileGPUSkinVFShaders(EShaderPlatform Platform)
+bool ShouldWeCompileGPUSkinVFShaders(EShaderPlatform Platform, ERHIFeatureLevel::Type FeatureLevel)
 {
 	// If the skin cache is not available on this platform we need to compile GPU Skin VF shaders.
 	if (IsGPUSkinCacheAvailable(Platform) == false)
@@ -231,6 +231,15 @@ bool ShouldWeCompileGPUSkinVFShaders(EShaderPlatform Platform)
 		return true;
 	}
 
+	// Some mobile GPUs (MALI) has a 64K elements limitation on texel buffers
+	// This results in meshes with more than 64k vertices having their skin cache entries disabled at runtime.
+	// We don't have a reliable way of checking this at cook time, so for mobile we must always cache skin cache
+	// shaders so we have something to fall back to.
+	if (FeatureLevel == ERHIFeatureLevel::ES3_1)
+	{
+		return true;
+	}
+
 	// If the skin cache is enabled and we've been asked to skip GPU Skin VF shaders.
 	static FShaderPlatformCachedIniValue<bool> PerPlatformCVar(TEXT("r.SkinCache.SkipCompilingGPUSkinVF"));
 	return (PerPlatformCVar.Get(Platform) == false);
diff --git a/Engine/Source/Runtime/Engine/Private/GPUSkinVertexFactory.cpp b/Engine/Source/Runtime/Engine/Private/GPUSkinVertexFactory.cpp
--- a/Engine/Source/Runtime/Engine/Private/GPUSkinVertexFactory.cpp
+++ b/Engine/Source/Runtime/Engine/Private/GPUSkinVertexFactory.cpp
@@ -500,7 +500,7 @@ template <GPUSkinBoneInfluenceType BoneInfluenceType>
 bool TGPUSkinVertexFactory<BoneInfluenceType>::ShouldCompilePermutation(const FVertexFactoryShaderPermutationParameters& Parameters)
 {
 	bool bUnlimitedBoneInfluences = (BoneInfluenceType == UnlimitedBoneInfluence && GCVarUnlimitedBoneInfluences);
-	return ShouldWeCompileGPUSkinVFShaders(Parameters.Platform) && 
+	return ShouldWeCompileGPUSkinVFShaders(Parameters.Platform, Parameters.MaterialParameters.FeatureLevel) &&
 		  ((Parameters.MaterialParameters.bIsUsedWithSkeletalMesh && (BoneInfluenceType != UnlimitedBoneInfluence || bUnlimitedBoneInfluences)) || Parameters.MaterialParameters.bIsSpecialEngineMaterial);
 }
 
diff --git a/Engine/Source/Runtime/Engine/Private/SkeletalRenderGPUSkin.cpp b/Engine/Source/Runtime/Engine/Private/SkeletalRenderGPUSkin.cpp
--- a/Engine/Source/Runtime/Engine/Private/SkeletalRenderGPUSkin.cpp
+++ b/Engine/Source/Runtime/Engine/Private/SkeletalRenderGPUSkin.cpp
@@ -1382,8 +1382,8 @@ const FVertexFactory* FSkeletalMeshObjectGPUSkin::GetSkinVertexFactory(const FSc
 	}
 
 	// If we have not compiled GPU Skin vertex factory variants
-	const auto CVar = IConsoleManager::Get().FindTConsoleVariableDataInt(TEXT("r.SkinCache.SkipCompilingGPUSkinVF"));
-	if (CVar && CVar->GetValueOnAnyThread() != 0)
+	const auto CVar = IConsoleManager::Get().FindConsoleVariable(TEXT("r.SkinCache.SkipCompilingGPUSkinVF"));
+	if (FeatureLevel != ERHIFeatureLevel::ES3_1 && CVar && CVar->GetBool() == true)
 	{
 		ensureMsgf(false, TEXT("We are attempting to render with a GPU Skin Vertex Factory, but r.SkinCache.SkipCompilingGPUSkinVF=1 so we don't have shaders.  Skeletal meshes will draw in ref pose.  Either disable r.SkinCache.SkipCompilingGPUSkinVF or increase the r.SkinCache.SceneMemoryLimitInMB size."));
 		return LOD.GPUSkinVertexFactories.PassthroughVertexFactories[ChunkIdx].Get();
diff --git a/Engine/Source/Runtime/Engine/Public/GPUSkinCache.h b/Engine/Source/Runtime/Engine/Public/GPUSkinCache.h
--- a/Engine/Source/Runtime/Engine/Public/GPUSkinCache.h
+++ b/Engine/Source/Runtime/Engine/Public/GPUSkinCache.h
@@ -58,7 +58,7 @@ struct FRayTracingGeometrySegment;
 // Can the skin cache be used (ie shaders added, etc)
 extern ENGINE_API bool IsGPUSkinCacheAvailable(EShaderPlatform Platform);
 
-extern bool ShouldWeCompileGPUSkinVFShaders(EShaderPlatform Platform);
+extern bool ShouldWeCompileGPUSkinVFShaders(EShaderPlatform Platform, ERHIFeatureLevel::Type FeatureLevel);
 
 extern ENGINE_API bool GPUSkinCacheNeedsDuplicatedVertices();
 
