Improved the error message when the texture derived task failed to load some part of the source data.
It now output the fullname of the texture.

Also fixed a compilation issue when using the debug SAVE_TILES define.

diff --git a/Engine/Source/Runtime/Engine/Private/TextureDerivedDataTask.cpp b/Engine/Source/Runtime/Engine/Private/TextureDerivedDataTask.cpp
--- a/Engine/Source/Runtime/Engine/Private/TextureDerivedDataTask.cpp
+++ b/Engine/Source/Runtime/Engine/Private/TextureDerivedDataTask.cpp
@@ -228,7 +228,7 @@ void FTextureSourceData::Init(UTexture& InTexture, TextureMipGenSettings InMipGe
 		Block.MipBias = MipBiasX;
 	}
 
-	TextureName = InTexture.GetFName();
+	TextureFullName = InTexture.GetFullName();
 
 	if (bAllowAsyncLoading && !InTexture.Source.IsBulkDataLoaded())
 	{
@@ -245,7 +245,7 @@ void FTextureSourceData::GetSourceMips(FTextureSource& Source, IImageWrapperModu
 	{
 		if (Source.HasHadBulkDataCleared())
 		{	// don't do any work we can't reload this
-			UE_LOG(LogTexture, Error, TEXT("Unable to get texture source mips because its bulk data was released. %s"), *TextureName.ToString())
+			UE_LOG(LogTexture, Error, TEXT("Unable to get texture source mips because its bulk data was released. %s"), *TextureFullName)
 				return;
 		}
 
@@ -275,7 +275,7 @@ void FTextureSourceData::GetSourceMips(FTextureSource& Source, IImageWrapperModu
 
 						if (!ScopedMipData.GetMipData(SourceMip->RawData, BlockIndex, LayerIndex, MipIndex))
 						{
-							UE_LOG(LogTexture, Warning, TEXT("Cannot retrieve source data for mip %d of texture %s"), MipIndex, *TextureName.ToString());
+							UE_LOG(LogTexture, Warning, TEXT("Cannot retrieve source data for mip %d of %s"), MipIndex, *TextureFullName);
 							ReleaseMemory();
 							bValid = false;
 							break;
diff --git a/Engine/Source/Runtime/Engine/Private/TextureDerivedDataTask.h b/Engine/Source/Runtime/Engine/Private/TextureDerivedDataTask.h
--- a/Engine/Source/Runtime/Engine/Private/TextureDerivedDataTask.h
+++ b/Engine/Source/Runtime/Engine/Private/TextureDerivedDataTask.h
@@ -95,7 +95,7 @@ struct FTextureSourceData
 		Blocks.Empty();
 	}
 
-	FName TextureName;
+	FString TextureFullName;
 	FTextureSource AsyncSource;
 	TArray<FTextureSourceLayerData> Layers;
 	TArray<FTextureSourceBlockData> Blocks;
diff --git a/Engine/Source/Runtime/Engine/Private/VT/VirtualTextureDataBuilder.cpp b/Engine/Source/Runtime/Engine/Private/VT/VirtualTextureDataBuilder.cpp
--- a/Engine/Source/Runtime/Engine/Private/VT/VirtualTextureDataBuilder.cpp
+++ b/Engine/Source/Runtime/Engine/Private/VT/VirtualTextureDataBuilder.cpp
@@ -16,6 +16,7 @@
 #define SAVE_TILES 0
 
 #if SAVE_TILES
+#include "ImageCoreUtils.h"
 #include "ImageUtils.h"
 #endif
 
@@ -200,7 +201,7 @@ struct FPixelDataRectangle
 #if SAVE_TILES
 	FImageView GetImageView() const 
 	{
-		ERawImageFormat::Type RawFormat = FImageUtils::ConvertToRawImageFormat(Format);
+		ERawImageFormat::Type RawFormat = FImageCoreUtils::ConvertToRawImageFormat(Format);
 		return FImageView(Data,Width,Height,RawFormat);
 	}
 
@@ -1088,7 +1089,10 @@ void FVirtualTextureDataBuilder::BuildSourcePixels(const FTextureSourceData& Sou
 #if SAVE_TILES
 			{
 				const FString BasePath = FPaths::ProjectUserDir();
-				const FString MipFileName = BasePath / FString::Format(TEXT("{0}_{1}"), TArray<FStringFormatArg>({ SourceData.TextureName.ToString(), LayerIndex }));
+				int32 ObjectNameStart;
+				SourceData.TextureFullName.FindLastChar('.', ObjectNameStart);
+				FString ObjectName = SourceData.TextureFullName.Right(SourceData.TextureFullName.Len() - ObjectNameStart - 1);
+				const FString MipFileName = BasePath / FString::Format(TEXT("{0}_{1}"), TArray<FStringFormatArg>({MoveTemp(ObjectName), LayerIndex}));
 				DstPixelData.Save(MipFileName, ImageWrapper);
 			}
 #endif // SAVE_TILES
