Only create a new default subobject when destroying a component if the owner's class is not native. Native owner classes do not have to have a DefaultSceneRoot when you drag in the C++ class directly from the content browser into the level.

diff --git a/Engine/Source/Editor/SubobjectDataInterface/Private/SubobjectData.cpp b/Engine/Source/Editor/SubobjectDataInterface/Private/SubobjectData.cpp
--- a/Engine/Source/Editor/SubobjectDataInterface/Private/SubobjectData.cpp
+++ b/Engine/Source/Editor/SubobjectDataInterface/Private/SubobjectData.cpp
@@ -57,10 +57,26 @@ bool FSubobjectData::CanEdit() const
 
 bool FSubobjectData::CanDelete() const
 {
-	// Components can be deleted if they are not inherited or the scene root
+	// Components can be deleted if they are not inherited
 	if (const UActorComponent* ComponentTemplate = GetComponentTemplate())
 	{
-		return !IsInheritedComponent() && !IsDefaultSceneRoot() && !IsChildActor();
+		// Inherited components cannot be deleted
+		if (IsInheritedComponent() || IsChildActor())
+		{
+			return false;
+		}
+		// You can delete the default scene root on instances of Native C++ actors in the level
+		// but not inside of a blueprint or on BP created actors
+		else if (IsDefaultSceneRoot())
+		{
+			if (const AActor* Actor = ComponentTemplate->GetOwner())
+			{
+				return Actor->GetClass()->IsNative();
+			}
+
+			// Otherwise you can't delete the default scene root
+			return false;
+		}
 	}
 	
 	// Otherwise, it can't be deleted
diff --git a/Engine/Source/Editor/SubobjectEditor/Private/SSubobjectInstanceEditor.cpp b/Engine/Source/Editor/SubobjectEditor/Private/SSubobjectInstanceEditor.cpp
--- a/Engine/Source/Editor/SubobjectEditor/Private/SSubobjectInstanceEditor.cpp
+++ b/Engine/Source/Editor/SubobjectEditor/Private/SSubobjectInstanceEditor.cpp
@@ -229,6 +229,12 @@ void SSubobjectInstanceEditor::OnDeleteNodes()
 			{
 				TreeWidget->SetSelection(NodeToSelect);
 			}
+			// If there are no components left, then fall back to the generic root node.
+			// This may be the case if you have deleted all components on a native instance actor
+			else if(!RootNodes.IsEmpty())
+			{
+				TreeWidget->SetSelection(RootNodes[0]);
+			}
 			
 			UpdateTree();
 
diff --git a/Engine/Source/Runtime/Engine/Private/Components/SceneComponent.cpp b/Engine/Source/Runtime/Engine/Private/Components/SceneComponent.cpp
--- a/Engine/Source/Runtime/Engine/Private/Components/SceneComponent.cpp
+++ b/Engine/Source/Runtime/Engine/Private/Components/SceneComponent.cpp
@@ -941,11 +941,15 @@ void USceneComponent::DestroyComponent(bool bPromoteChildren/*= false*/)
 				USceneComponent* const * FindResult =
 					AttachedChildren.FindByPredicate([Owner](USceneComponent* Child){ return Child != nullptr && !Child->IsEditorOnly() && Child->GetOwner() == Owner; });
 
+				const bool bIsNativeOwnerClass = Owner->GetClass()->IsNative();
+
 				if (FindResult != nullptr)
 				{
 					ChildToPromote = *FindResult;
 				}
-				else
+				// Native C++ classes do not need to always have a DefaultSceneRoot, it can be empty on
+				// instances placed in the level directly from the C++ class
+				else if (!bIsNativeOwnerClass)
 				{
 					// Didn't find a suitable component to promote so create a new default component
 
@@ -968,7 +972,6 @@ void USceneComponent::DestroyComponent(bool bPromoteChildren/*= false*/)
 				Owner->Modify();
 
 				// Set the selected child node as the new root
-				check(ChildToPromote != nullptr);
 				Owner->SetRootComponent(ChildToPromote);
 			}
 			else    // ...not the root node, so we'll promote the selected child node to this position in its AttachParent's child array.
