Fix log spam with the line "DMFP has already been upgraded to MFP".

diff --git a/Engine/Source/Runtime/Engine/Private/Rendering/SubsurfaceProfile.cpp b/Engine/Source/Runtime/Engine/Private/Rendering/SubsurfaceProfile.cpp
--- a/Engine/Source/Runtime/Engine/Private/Rendering/SubsurfaceProfile.cpp
+++ b/Engine/Source/Runtime/Engine/Private/Rendering/SubsurfaceProfile.cpp
@@ -107,6 +107,19 @@ void UpgradeDiffuseMeanFreePathToMeanFreePath(FSubsurfaceProfileStruct& Settings
 	Settings.WorldUnitScale /= (SUBSURFACE_KERNEL_SIZE / CmToMm);
 }
 
+void UpgradeSubsurfaceProfileParameters(FSubsurfaceProfileStruct& Settings)
+{
+	if (!Settings.bEnableBurley)
+	{
+		UpgradeSeparableToBurley(Settings);
+	}
+
+	if (!Settings.bEnableMeanFreePath)
+	{
+		UpgradeDiffuseMeanFreePathToMeanFreePath(Settings);
+	}
+}
+
 FSubsurfaceProfileTexture::FSubsurfaceProfileTexture()
 {
 	check(IsInGameThread());
@@ -114,13 +127,7 @@ FSubsurfaceProfileTexture::FSubsurfaceProfileTexture()
 	FSubsurfaceProfileStruct DefaultSkin;
 
 	//The default burley in slot 0 behaves the same as Separable previously
-	UpgradeSeparableToBurley(DefaultSkin);
-
-	//Upgrade DMFP in slot 0 to MFP
-	if (!DefaultSkin.bEnableMeanFreePath)
-	{
-		UpgradeDiffuseMeanFreePathToMeanFreePath(DefaultSkin);
-	}
+	UpgradeSubsurfaceProfileParameters(DefaultSkin);
 
 	// add element 0, it is used as default profile
 	SubsurfaceProfileEntries.Add(FSubsurfaceProfileEntry(DefaultSkin, 0));
@@ -441,6 +448,9 @@ void FSubsurfaceProfileTexture::CreateTexture(FRHICommandListImmediate& RHICmdLi
 	{
 		FSubsurfaceProfileStruct Data = SubsurfaceProfileEntries[y].Settings;
 
+		// Fix for postload() not yet called.
+		UpgradeSubsurfaceProfileParameters(Data);
+
 		Data.Tint = Data.Tint.GetClamped();
 		Data.FalloffColor = Data.FalloffColor.GetClamped(Bias);
 		Data.MeanFreePathColor = Data.MeanFreePathColor.GetClamped(Bias); // In Cm
@@ -755,24 +765,5 @@ void USubsurfaceProfile::PostLoad()
 {
 	Super::PostLoad();
 
-	const auto CVar = IConsoleManager::Get().
-		FindTConsoleVariableDataInt(TEXT("r.SSS.Burley.AlwaysUpdateParametersFromSeparable"));
-	if (CVar)
-	{
-		const bool bUpdateBurleyParametersFromSeparable = CVar->GetValueOnAnyThread() == 1;
-
-		if (!Settings.bEnableBurley
-#if !(UE_BUILD_SHIPPING || UE_BUILD_TEST)
-			|| bUpdateBurleyParametersFromSeparable
-#endif
-			)
-		{
-			UpgradeSeparableToBurley(Settings);
-		}
-	}
-
-	if (!Settings.bEnableMeanFreePath)
-	{
-		UpgradeDiffuseMeanFreePathToMeanFreePath(Settings);
-	}
+	UpgradeSubsurfaceProfileParameters(this->Settings);
 }
