Fixed data race and potential double delete in FCanvas deletion path.

diff --git a/Engine/Source/Developer/MaterialUtilities/Private/MeshRendering.cpp b/Engine/Source/Developer/MaterialUtilities/Private/MeshRendering.cpp
--- a/Engine/Source/Developer/MaterialUtilities/Private/MeshRendering.cpp
+++ b/Engine/Source/Developer/MaterialUtilities/Private/MeshRendering.cpp
@@ -548,12 +548,8 @@ public:
 		RenderMaterial(RenderContext, DrawRenderState, *View, *Data);
 
 		RenderContext.DeferredDelete(View);
-
-		if (Canvas->GetAllowedModes() & FCanvas::Allow_DeleteOnRender)
-		{
-			RenderContext.DeferredDelete(Data);
-			Data = nullptr;
-		}
+		RenderContext.DeferredDelete(Data);
+		Data = nullptr;
 		return true;
 	}
 
@@ -603,15 +599,10 @@ public:
 			RenderMaterial(RenderContext, DrawRenderState, *Parameters.View, *Parameters.RenderData);
 
 			RenderContext.DeferredDelete(Parameters.View);
-			if (Parameters.AllowedCanvasModes & FCanvas::Allow_DeleteOnRender)
-			{
-				RenderContext.DeferredDelete(Parameters.RenderData);
-			}
+			RenderContext.DeferredDelete(Parameters.RenderData);
 		});
-		if (Canvas->GetAllowedModes() & FCanvas::Allow_DeleteOnRender)
-		{
-			Data = nullptr;
-		}
+
+		Data = nullptr;
 		return true;
 	}
 };
diff --git a/Engine/Source/Runtime/Engine/Private/Slate/DebugCanvas.cpp b/Engine/Source/Runtime/Engine/Private/Slate/DebugCanvas.cpp
--- a/Engine/Source/Runtime/Engine/Private/Slate/DebugCanvas.cpp
+++ b/Engine/Source/Runtime/Engine/Private/Slate/DebugCanvas.cpp
@@ -204,7 +204,7 @@ void FDebugCanvasDrawer::InitDebugCanvas(FViewportClient* ViewportClient, UWorld
 		GameThreadCanvas = MakeShared<FCanvas, ESPMode::ThreadSafe>(RenderTarget, nullptr, InWorld, InWorld ? InWorld->FeatureLevel.GetValue() : GMaxRHIFeatureLevel, FCanvas::CDM_DeferDrawing, DPIScale);
 
 		// Do not allow the canvas to be flushed outside of our debug rendering path
-		GameThreadCanvas->SetAllowedModes(FCanvas::Allow_DeleteOnRender);
+		GameThreadCanvas->SetAllowedModes(FCanvas::Allow_None);
 	}
 
 	if (GameThreadCanvas.IsValid())
diff --git a/Engine/Source/Runtime/Engine/Private/TileRendering.cpp b/Engine/Source/Runtime/Engine/Private/TileRendering.cpp
--- a/Engine/Source/Runtime/Engine/Private/TileRendering.cpp
+++ b/Engine/Source/Runtime/Engine/Private/TileRendering.cpp
@@ -264,11 +264,8 @@ bool FCanvasTileRendererItem::Render_RenderThread(FCanvasRenderContext& RenderCo
 
 	Data->RenderTiles(RenderContext, DrawRenderState, View, Canvas->IsHitTesting(), bNeedsToSwitchVerticalAxis);
 
-	if (Canvas->GetAllowedModes() & FCanvas::Allow_DeleteOnRender)
-	{
-		RenderContext.DeferredRelease(MoveTemp(Data));
-		Data = nullptr;
-	}
+	RenderContext.DeferredRelease(MoveTemp(Data));
+	Data = nullptr;
 
 	return true;
 }
@@ -309,7 +306,6 @@ bool FCanvasTileRendererItem::Render_GameThread(const FCanvas* Canvas, FCanvasRe
 
 		const bool bNeedsToSwitchVerticalAxis = RHINeedsToSwitchVerticalAxis(Canvas->GetShaderPlatform()) && Canvas->GetAllowSwitchVerticalAxis();
 		const bool bIsHitTesting = Canvas->IsHitTesting();
-		const bool bDeleteOnRender = Canvas->GetAllowedModes() & FCanvas::Allow_DeleteOnRender;
 
 		bool bRequiresExplicit128bitRT = false;
 
@@ -335,10 +331,7 @@ bool FCanvasTileRendererItem::Render_GameThread(const FCanvas* Canvas, FCanvasRe
 			RenderContext.DeferredDelete(View);
 		});
 
-		if (bDeleteOnRender)
-		{
-			Data = nullptr;
-		}
+		Data = nullptr;
 
 		return true;
 	}
diff --git a/Engine/Source/Runtime/Engine/Private/TriangleRendering.cpp b/Engine/Source/Runtime/Engine/Private/TriangleRendering.cpp
--- a/Engine/Source/Runtime/Engine/Private/TriangleRendering.cpp
+++ b/Engine/Source/Runtime/Engine/Private/TriangleRendering.cpp
@@ -235,11 +235,8 @@ bool FCanvasTriangleRendererItem::Render_RenderThread(FCanvasRenderContext& Rend
 
 	Data->RenderTriangles(RenderContext, DrawRenderState, View, Canvas->IsHitTesting(), bNeedsToSwitchVerticalAxis);
 
-	if (Canvas->GetAllowedModes() & FCanvas::Allow_DeleteOnRender)
-	{
-		RenderContext.DeferredRelease(MoveTemp(Data));
-		Data = nullptr;
-	}
+	RenderContext.DeferredRelease(MoveTemp(Data));
+	Data = nullptr;
 
 	return true;
 }
@@ -279,7 +276,6 @@ bool FCanvasTriangleRendererItem::Render_GameThread(const FCanvas* Canvas, FCanv
 
 	const bool bNeedsToSwitchVerticalAxis = RHINeedsToSwitchVerticalAxis(Canvas->GetShaderPlatform()) && Canvas->GetAllowSwitchVerticalAxis();
 	const bool bIsHitTesting = Canvas->IsHitTesting();
-	const bool bDeleteOnRender = Canvas->GetAllowedModes() & FCanvas::Allow_DeleteOnRender;
 
 	RenderScope.EnqueueRenderCommand(
 		[LocalData = Data, View, bIsHitTesting, bNeedsToSwitchVerticalAxis](FCanvasRenderContext& RenderContext) mutable
@@ -296,10 +292,7 @@ bool FCanvasTriangleRendererItem::Render_GameThread(const FCanvas* Canvas, FCanv
 		RenderContext.DeferredDelete(View);
 	});
 
-	if (Canvas->GetAllowedModes() & FCanvas::Allow_DeleteOnRender)
-	{
-		Data = nullptr;
-	}
+	Data = nullptr;
 
 	return true;
 }
diff --git a/Engine/Source/Runtime/Engine/Private/UserInterface/Canvas.cpp b/Engine/Source/Runtime/Engine/Private/UserInterface/Canvas.cpp
--- a/Engine/Source/Runtime/Engine/Private/UserInterface/Canvas.cpp
+++ b/Engine/Source/Runtime/Engine/Private/UserInterface/Canvas.cpp
@@ -274,7 +274,7 @@ FCanvas::FCanvas(FRenderTarget* InRenderTarget, FHitProxyConsumer* InHitProxyCon
 ,	RenderTarget(InRenderTarget)
 ,	HitProxyConsumer(InHitProxyConsumer)
 ,	Scene(InWorld ? InWorld->Scene : NULL)
-,	AllowedModes(0xFFFFFFFF)
+,	AllowedModes(Allow_Flush)
 ,	bRenderTargetDirty(false)
 ,	FeatureLevel(InFeatureLevel)
 ,	bUseInternalTexture(false)
@@ -295,7 +295,7 @@ FCanvas::FCanvas(FRenderTarget* InRenderTarget,FHitProxyConsumer* InHitProxyCons
 ,	RenderTarget(InRenderTarget)
 ,	HitProxyConsumer(InHitProxyConsumer)
 ,	Scene(NULL)
-,	AllowedModes(0xFFFFFFFF)
+,	AllowedModes(Allow_Flush)
 ,	bRenderTargetDirty(false)
 ,	Time(InTime)
 ,	bAllowsToSwitchVerticalAxis(false)
@@ -455,12 +455,9 @@ bool FCanvasBatchedElementRenderItem::Render_RenderThread(FCanvasRenderContext&
 		});
 	}
 
-	if (Canvas->GetAllowedModes() & FCanvas::Allow_DeleteOnRender)
-	{
-		// delete data since we're done rendering it
-		RenderContext.DeferredDelete(Data);
-		Data = nullptr;
-	}
+	// delete data since we're done rendering it
+	RenderContext.DeferredDelete(Data);
+	Data = nullptr;
 
 	return bDirty;
 }
@@ -531,16 +528,12 @@ bool FCanvasBatchedElementRenderItem::Render_GameThread(const FCanvas* Canvas, F
 				DrawParameters.bHitTesting,
 				DrawParameters.DisplayGamma);
 
-			if(DrawParameters.AllowedCanvasModes & FCanvas::Allow_DeleteOnRender )
-			{
-				delete DrawParameters.RenderData;
-			}
+			delete DrawParameters.RenderData;
 		});
 	}
-	if( Canvas->GetAllowedModes() & FCanvas::Allow_DeleteOnRender )
-	{
-		Data = NULL;
-	}
+
+	Data = nullptr;
+
 	return bDirty;
 }
 
@@ -573,13 +566,20 @@ FCanvasRenderThreadScope::~FCanvasRenderThreadScope()
 {
 	RenderCommandFunctionArray* RenderCommandArray = RenderCommands;
 
+	FIntRect ViewportRect = Canvas.ViewRect;
+
+	if (ViewportRect.Area() <= 0 || Canvas.bScaledToRenderTarget)
+	{
+		ViewportRect = FIntRect(FIntPoint::ZeroValue, Canvas.RenderTarget->GetSizeXY());
+	}
+
 	ENQUEUE_RENDER_COMMAND(DispatchCanvasRenderCommands)(
-		[RenderCommandArray, LocalCanvas = Canvas](FRHICommandListImmediate& RHICmdList)
+		[RenderCommandArray, ViewportRect, ScissorRect = Canvas.ScissorRect, RenderTarget = Canvas.RenderTarget](FRHICommandListImmediate& RHICmdList)
 	{
 		GetRendererModule().InitializeSystemTextures(RHICmdList);
 		FMemMark MemMark(FMemStack::Get());
 		FRDGBuilder GraphBuilder(RHICmdList, RDG_EVENT_NAME("CanvasRenderThreadScope"));
-		FCanvasRenderContext RenderContext(GraphBuilder, LocalCanvas);
+		FCanvasRenderContext RenderContext(GraphBuilder, RenderTarget->GetRenderTargetTexture(GraphBuilder), ViewportRect, ScissorRect);
 
 		for (uint32 Index = 0, Count = RenderCommandArray->Num(); Index < Count; ++Index)
 		{
@@ -810,24 +810,18 @@ void FCanvas::Flush_RenderThread(FRDGBuilder& GraphBuilder, bool bForce)
 			{
 				// mark current render target as dirty since we are drawing to it
 				bRenderTargetDirty |= RenderItem->Render_RenderThread(RenderContext, DrawRenderState, this);
-				if (AllowedModes & Allow_DeleteOnRender)
-				{
-					RenderContext.DeferredDelete(RenderItem);
-				}
+
+				RenderContext.DeferredDelete(RenderItem);
 			}
 		}
-		if (AllowedModes & Allow_DeleteOnRender)
-		{
-			SortElement.RenderBatchArray.Empty();
-		}
-	}
-	if (AllowedModes & Allow_DeleteOnRender)
-	{
-		// empty the array of FCanvasSortElement entries after finished with rendering	
-		SortedElements.Empty();
-		SortedElementLookupMap.Empty();
-		LastElementIndex = INDEX_NONE;
+
+		SortElement.RenderBatchArray.Empty();
 	}
+
+	// empty the array of FCanvasSortElement entries after finished with rendering	
+	SortedElements.Empty();
+	SortedElementLookupMap.Empty();
+	LastElementIndex = INDEX_NONE;
 }
 
 void FCanvas::Flush_GameThread(bool bForce)
@@ -883,16 +877,11 @@ void FCanvas::Flush_GameThread(bool bForce)
 				{
 					// mark current render target as dirty since we are drawing to it
 					bRenderTargetDirty |= RenderItem->Render_GameThread(this, RenderThreadScope);
-					if( AllowedModes & Allow_DeleteOnRender )
-					{
-						RenderThreadScope.DeferredDelete(RenderItem);
-					}
+
+					RenderThreadScope.DeferredDelete(RenderItem);
 				}
 			}
-			if( AllowedModes & Allow_DeleteOnRender )
-			{
-				SortElement.RenderBatchArray.Empty();
-			}
+			SortElement.RenderBatchArray.Empty();
 		}
 	}
 	else
@@ -913,13 +902,10 @@ void FCanvas::Flush_GameThread(bool bForce)
 		}
 	}
 	
-	if( AllowedModes & Allow_DeleteOnRender )
-	{
-		// empty the array of FCanvasSortElement entries after finished with rendering
-		SortedElements.Empty();
-		SortedElementLookupMap.Empty();
-		LastElementIndex = INDEX_NONE;
-	}
+	// empty the array of FCanvasSortElement entries after finished with rendering
+	SortedElements.Empty();
+	SortedElementLookupMap.Empty();
+	LastElementIndex = INDEX_NONE;
 }
 
 void FCanvas::PushRelativeTransform(const FMatrix& Transform)
diff --git a/Engine/Source/Runtime/Engine/Public/CanvasTypes.h b/Engine/Source/Runtime/Engine/Public/CanvasTypes.h
--- a/Engine/Source/Runtime/Engine/Public/CanvasTypes.h
+++ b/Engine/Source/Runtime/Engine/Public/CanvasTypes.h
@@ -217,10 +217,11 @@ public:
 	 **/
 	enum ECanvasAllowModes
 	{
+		Allow_None = 0,
 		// flushing and rendering
 		Allow_Flush = 1 << 0,
 		// delete the render batches when rendering
-		Allow_DeleteOnRender = 1 << 1
+		Allow_DeleteOnRender UE_DEPRECATED(5.1, "Allow_DeleteOnRender is now always enabled.") = 1 << 1
 	};
 
 	enum ECanvasDrawMode
@@ -792,6 +793,7 @@ public:
 	ENGINE_API FCanvasSortElement& GetSortElement(int32 DepthSortKey);
 
 	friend class FCanvasRenderContext;
+	friend class FCanvasRenderThreadScope;
 };
 
 /**
