- WorldPartition ActorDesc full support for actors referencing Data Layer Assets
- Fixed changelist validation on DataLayers using new WorldDataLayersActorDesc

diff --git a/Engine/Plugins/Editor/DataValidation/Source/DataValidation/Private/WorldPartitionChangelistValidator.cpp b/Engine/Plugins/Editor/DataValidation/Source/DataValidation/Private/WorldPartitionChangelistValidator.cpp
--- a/Engine/Plugins/Editor/DataValidation/Source/DataValidation/Private/WorldPartitionChangelistValidator.cpp
+++ b/Engine/Plugins/Editor/DataValidation/Source/DataValidation/Private/WorldPartitionChangelistValidator.cpp
@@ -346,5 +346,9 @@ void UWorldPartitionChangelistValidator::OnDataLayerAssetConflict(const UDataLay
 	}
 }
 
+void UWorldPartitionChangelistValidator::OnActorNeedsResave(const FWorldPartitionActorDescView& ActorDescView)
+{
+	// Changelist validation already ensures that dirty actors must be part of the changelist
+}
 
 #undef LOCTEXT_NAMESPACE
\ No newline at end of file
diff --git a/Engine/Plugins/Editor/DataValidation/Source/DataValidation/Private/WorldPartitionChangelistValidator.h b/Engine/Plugins/Editor/DataValidation/Source/DataValidation/Private/WorldPartitionChangelistValidator.h
--- a/Engine/Plugins/Editor/DataValidation/Source/DataValidation/Private/WorldPartitionChangelistValidator.h
+++ b/Engine/Plugins/Editor/DataValidation/Source/DataValidation/Private/WorldPartitionChangelistValidator.h
@@ -43,5 +43,6 @@ protected:
 	virtual void OnInvalidReferenceDataLayerAsset(const UDataLayerInstanceWithAsset* DataLayerInstance) override;
 	virtual void OnDataLayerHierarchyTypeMismatch(const UDataLayerInstance* DataLayerInstance, const UDataLayerInstance* Parent) override;
 	virtual void OnDataLayerAssetConflict(const UDataLayerInstanceWithAsset* DataLayerInstance, const UDataLayerInstanceWithAsset* ConflictingDataLayerInstance) override;
+	virtual void OnActorNeedsResave(const FWorldPartitionActorDescView& ActorDescView) override;
 };
 
diff --git a/Engine/Source/Editor/DataLayerEditor/Private/DataLayer/DataLayerEditorSubsystem.cpp b/Engine/Source/Editor/DataLayerEditor/Private/DataLayer/DataLayerEditorSubsystem.cpp
--- a/Engine/Source/Editor/DataLayerEditor/Private/DataLayer/DataLayerEditorSubsystem.cpp
+++ b/Engine/Source/Editor/DataLayerEditor/Private/DataLayer/DataLayerEditorSubsystem.cpp
@@ -6,9 +6,12 @@
 #include "WorldPartition/DataLayer/DeprecatedDataLayerInstance.h"
 #include "WorldPartition/DataLayer/WorldDataLayers.h"
 #include "WorldPartition/DataLayer/DataLayerSubsystem.h"
+#include "WorldPartition/DataLayer/DataLayerUtils.h"
 #include "WorldPartition/WorldPartitionSubsystem.h"
 #include "WorldPartition/WorldPartition.h"
 #include "WorldPartition/WorldPartitionEditorPerProjectUserSettings.h"
+#include "WorldPartition/WorldPartitionActorDesc.h"
+#include "WorldPartition/ActorDescContainer.h"
 #include "Modules/ModuleManager.h"
 #include "Misc/MessageDialog.h"
 #include "DataLayerEditorModule.h"
@@ -24,9 +27,12 @@
 #include "Subsystems/ActorEditorContextSubsystem.h"
 #include "IContentBrowserSingleton.h"
 #include "ContentBrowserModule.h"
+#include "ProfilingDebugging/ScopedTimers.h"
 
 #define LOCTEXT_NAMESPACE "DataLayer"
 
+DEFINE_LOG_CATEGORY_STATIC(LogDataLayerEditorSubsystem, All, All);
+
 //////////////////////////////////////////////////////////////////////////
 // FDataLayersBroadcast
 
@@ -44,6 +50,7 @@ private:
 	void OnObjectPostEditChange(UObject* Object, FPropertyChangedEvent& PropertyChangedEvent);
 	void OnLevelActorsAdded(AActor* InActor) { DataLayerEditorSubsystem->InitializeNewActorDataLayers(InActor); }
 	void OnLevelSelectionChanged(UObject* InObject) { DataLayerEditorSubsystem->OnSelectionChanged(); }
+	void OnActorDescContainerInitialized(UActorDescContainer* InActorDescContainer);
 
 	UDataLayerEditorSubsystem* DataLayerEditorSubsystem;
 	bool bIsInitialized;
@@ -72,6 +79,7 @@ void FDataLayersBroadcast::Deinitialize()
 		GEngine->OnLevelActorAdded().RemoveAll(this);
 		USelection::SelectionChangedEvent.RemoveAll(this);
 		USelection::SelectObjectEvent.RemoveAll(this);
+		UActorDescContainer::OnActorDescContainerInitialized.RemoveAll(this);
 	}
 }
 
@@ -86,9 +94,15 @@ void FDataLayersBroadcast::Initialize()
 		GEngine->OnLevelActorAdded().AddRaw(this, &FDataLayersBroadcast::OnLevelActorsAdded);
 		USelection::SelectionChangedEvent.AddRaw(this, &FDataLayersBroadcast::OnLevelSelectionChanged);
 		USelection::SelectObjectEvent.AddRaw(this, &FDataLayersBroadcast::OnLevelSelectionChanged);
+		UActorDescContainer::OnActorDescContainerInitialized.AddRaw(this, &FDataLayersBroadcast::OnActorDescContainerInitialized);
 	}
 }
 
+void FDataLayersBroadcast::OnActorDescContainerInitialized(UActorDescContainer* InActorDescContainer)
+{
+	UDataLayerEditorSubsystem::OnActorDescContainerInitialized(InActorDescContainer);
+}
+
 void FDataLayersBroadcast::OnObjectPostEditChange(UObject* Object, FPropertyChangedEvent& PropertyChangedEvent)
 {
 	if (Object && (PropertyChangedEvent.ChangeType != EPropertyChangeType::Interactive))
@@ -151,6 +165,21 @@ void UDataLayerEditorSubsystem::Deinitialize()
 	DataLayersBroadcast->Deinitialize();
 }
 
+void UDataLayerEditorSubsystem::OnActorDescContainerInitialized(UActorDescContainer* InActorDescContainer)
+{
+	check(InActorDescContainer);
+
+	UE_SCOPED_TIMER(*FString::Printf(TEXT("Resolving Data Layer Instance Names for %s"), *InActorDescContainer->GetContainerPackage().ToString()), LogDataLayerEditorSubsystem, Display);
+
+	const FWorldDataLayersActorDesc* WorldDataLayersActorDesc = FDataLayerUtils::GetWorldDataLayersActorDesc(InActorDescContainer);
+	for (FActorDescList::TIterator<> Iterator(InActorDescContainer); Iterator; ++Iterator)
+	{
+		FWorldPartitionActorDesc* ActorDesc = *Iterator;
+		check(ActorDesc->GetContainer() == InActorDescContainer);
+		ActorDesc->DataLayerInstanceNames = FDataLayerUtils::ResolvedDataLayerInstanceNames(ActorDesc, /*WorldDataLayers*/nullptr, WorldDataLayersActorDesc);
+	}
+}
+
 void UDataLayerEditorSubsystem::OnExecuteActorEditorContextAction(UWorld* InWorld, const EActorEditorContextAction& InType, AActor* InActor)
 {
 	AWorldDataLayers* WorldDataLayer = InWorld->GetWorldDataLayers();
diff --git a/Engine/Source/Editor/DataLayerEditor/Public/DataLayer/DataLayerEditorSubsystem.h b/Engine/Source/Editor/DataLayerEditor/Public/DataLayer/DataLayerEditorSubsystem.h
--- a/Engine/Source/Editor/DataLayerEditor/Public/DataLayer/DataLayerEditorSubsystem.h
+++ b/Engine/Source/Editor/DataLayerEditor/Public/DataLayer/DataLayerEditorSubsystem.h
@@ -665,5 +665,7 @@ private:
 	/** The path at which the "Pick A Data Layer Asset" will be opened*/
 	mutable FString PickDataLayerDialogPath;
 
+	static void OnActorDescContainerInitialized(class UActorDescContainer* InActorDescContainer);
+
 	friend class FDataLayersBroadcast;
 };
diff --git a/Engine/Source/Runtime/Core/Public/UObject/FortniteNCBranchObjectVersion.h b/Engine/Source/Runtime/Core/Public/UObject/FortniteNCBranchObjectVersion.h
--- a/Engine/Source/Runtime/Core/Public/UObject/FortniteNCBranchObjectVersion.h
+++ b/Engine/Source/Runtime/Core/Public/UObject/FortniteNCBranchObjectVersion.h
@@ -18,6 +18,9 @@ struct CORE_API FFortniteNCBranchObjectVersion
 		// Fixed FDataLayerInstanceDesc
 		FixedDataLayerInstanceDesc,
 
+		// Serialize DataLayerAssets in WorldPartitionActorDesc
+		WorldPartitionActorDescSerializeDataLayerAssets,
+
 		// -----<new versions can be added above this line>-------------------------------------------------
 		VersionPlusOne,
 		LatestVersion = VersionPlusOne - 1
diff --git a/Engine/Source/Runtime/Engine/Private/WorldPartition/ActorDescContainer.cpp b/Engine/Source/Runtime/Engine/Private/WorldPartition/ActorDescContainer.cpp
--- a/Engine/Source/Runtime/Engine/Private/WorldPartition/ActorDescContainer.cpp
+++ b/Engine/Source/Runtime/Engine/Private/WorldPartition/ActorDescContainer.cpp
@@ -11,6 +11,8 @@
 #include "UObject/ObjectSaveContext.h"
 #include "UObject/CoreRedirects.h"
 #include "Engine/World.h"
+
+UActorDescContainer::FActorDescContainerInitializeDelegate UActorDescContainer::OnActorDescContainerInitialized;
 #endif
 
 UActorDescContainer::UActorDescContainer(const FObjectInitializer& ObjectInitializer)
@@ -110,6 +112,8 @@ void UActorDescContainer::Initialize(UWorld* InWorld, FName InPackageName)
 		}
 	}
 
+	OnActorDescContainerInitialized.Broadcast(this);
+
 	RegisterEditorDelegates();
 
 	bContainerInitialized = true;
diff --git a/Engine/Source/Runtime/Engine/Private/WorldPartition/DataLayer/DataLayer.cpp b/Engine/Source/Runtime/Engine/Private/WorldPartition/DataLayer/DataLayer.cpp
--- a/Engine/Source/Runtime/Engine/Private/WorldPartition/DataLayer/DataLayer.cpp
+++ b/Engine/Source/Runtime/Engine/Private/WorldPartition/DataLayer/DataLayer.cpp
@@ -5,6 +5,7 @@
 =============================================================================*/
 
 #include "WorldPartition/DataLayer/DataLayer.h"
+#include "WorldPartition/DataLayer/DataLayerUtils.h"
 #include "WorldPartition/DataLayer/WorldDataLayers.h"
 
 #define LOCTEXT_NAMESPACE "DataLayer"
@@ -42,7 +43,7 @@ void UDEPRECATED_DataLayer::PostLoad()
 	bIsVisible = bIsInitiallyVisible;
 
 	// Sanitize Label
-	DataLayerLabel = DataLayerUtils::GetSanitizedDataLayerLabel(DataLayerLabel);
+	DataLayerLabel = FDataLayerUtils::GetSanitizedDataLayerLabel(DataLayerLabel);
 
 	if (DebugColor == FColor::Black)
 	{
diff --git a/Engine/Source/Runtime/Engine/Private/WorldPartition/DataLayer/DataLayerInstance.cpp b/Engine/Source/Runtime/Engine/Private/WorldPartition/DataLayer/DataLayerInstance.cpp
--- a/Engine/Source/Runtime/Engine/Private/WorldPartition/DataLayer/DataLayerInstance.cpp
+++ b/Engine/Source/Runtime/Engine/Private/WorldPartition/DataLayer/DataLayerInstance.cpp
@@ -6,6 +6,7 @@
 #include "Logging/TokenizedMessage.h"
 #include "Misc/UObjectToken.h"
 #include "Misc/MapErrors.h"
+#include "WorldPartition/DataLayer/DataLayerUtils.h"
 #include "WorldPartition/DataLayer/WorldDataLayers.h"
 #include "WorldPartition/ErrorHandling/WorldPartitionStreamingGenerationErrorHandler.h"
 
@@ -93,6 +94,11 @@ bool UDataLayerInstance::IsLocked() const
 	return IsRuntime() && !GetOuterAWorldDataLayers()->GetAllowRuntimeDataLayerEditing();
 }
 
+const TCHAR* UDataLayerInstance::GetDataLayerIconName() const
+{
+	return FDataLayerUtils::GetDataLayerIconName(GetType());
+}
+
 bool UDataLayerInstance::CanParent(const UDataLayerInstance* InParent) const
 {
 	return (this != InParent)
diff --git a/Engine/Source/Runtime/Engine/Private/WorldPartition/DataLayer/DataLayerUtils.cpp b/Engine/Source/Runtime/Engine/Private/WorldPartition/DataLayer/DataLayerUtils.cpp
new file mode 100644
--- /dev/null
+++ b/Engine/Source/Runtime/Engine/Private/WorldPartition/DataLayer/DataLayerUtils.cpp
@@ -0,0 +1,163 @@
+// Copyright Epic Games, Inc. All Rights Reserved.
+
+#include "WorldPartition/DataLayer/DataLayerUtils.h"
+
+#if WITH_EDITOR
+#include "WorldPartition/DataLayer/WorldDataLayersActorDesc.h"
+#include "WorldPartition/DataLayer/WorldDataLayers.h"
+#include "WorldPartition/DataLayer/DataLayerInstanceWithAsset.h"
+#include "WorldPartition/DataLayer/DataLayerAsset.h"
+#include "WorldPartition/DataLayer/DataLayerType.h"
+#include "WorldPartition/WorldPartitionActorDesc.h"
+#include "WorldPartition/ActorDescContainer.h"
+
+FWorldDataLayersActorDesc* FDataLayerUtils::GetWorldDataLayersActorDesc(const UActorDescContainer* InContainer, bool bInCheckValid)
+{
+	if (InContainer)
+	{
+		for (FActorDescList::TIterator<AWorldDataLayers> Iterator(const_cast<UActorDescContainer*>(InContainer)); Iterator; ++Iterator)
+		{
+			if (!bInCheckValid || Iterator->IsValid())
+			{
+				FWorldDataLayersActorDesc* WorldDataLayersActorDesc = *Iterator;
+				return WorldDataLayersActorDesc;
+			}
+			// No need to iterate (we assume that there's only 1 AWorldDataLayer for now)
+			return nullptr;
+		}
+	}
+	return nullptr;
+}
+
+TArray<FName> FDataLayerUtils::ResolvedDataLayerInstanceNames(const FWorldPartitionActorDesc* InActorDesc, const AWorldDataLayers* InWorldDataLayers, const FWorldDataLayersActorDesc* InWorldDataLayersActorDesc, bool* bOutIsResultValid)
+{
+	bool bLocalIsSuccess;
+	bool& bIsSuccess = bOutIsResultValid ? *bOutIsResultValid : bLocalIsSuccess;
+	bIsSuccess = true;
+
+	// Prioritize in-memory AWorldDataLayers
+	const AWorldDataLayers* WorldDataLayers = InWorldDataLayers;
+	if (!WorldDataLayers)
+	{
+		UWorld* World = InActorDesc->GetContainer() ? InActorDesc->GetContainer()->GetWorld() : nullptr;
+		WorldDataLayers = World ? World->GetWorldDataLayers() : nullptr;
+	}
+
+	// DataLayers not using DataLayer Assets represent DataLayerInstanceNames
+	if (!InActorDesc->IsUsingDataLayerAsset())
+	{
+		if (WorldDataLayers)
+		{
+			TArray<FName> Result;
+			for (const FName& DataLayerInstanceName : InActorDesc->DataLayers)
+			{
+				if (const UDataLayerInstance* DataLayerInstance = WorldDataLayers->GetDataLayerInstance(DataLayerInstanceName))
+				{
+					Result.Add(DataLayerInstanceName);
+				}
+			}
+			return Result;
+		}
+		// Fallback on FWorldDataLayersActorDesc
+		else if (InWorldDataLayersActorDesc)
+		{
+			TArray<FName> Result;
+			for (const FName& DataLayerInstanceName : InActorDesc->DataLayers)
+			{
+				if (const FDataLayerInstanceDesc* DataLayerInstanceDesc = InWorldDataLayersActorDesc->GetDataLayerInstanceFromInstanceName(DataLayerInstanceName))
+				{
+					Result.Add(DataLayerInstanceName);
+				}
+			}
+			return Result;
+		}
+	}
+	// ActorDesc DataLayers represents DataLayer Asset Paths
+	else
+	{
+		if (WorldDataLayers)
+		{
+			TArray<FName> Result;
+			for (const FName& DataLayerAssetPath : InActorDesc->DataLayers)
+			{
+				WorldDataLayers->ForEachDataLayer([DataLayerAssetPath, &Result](UDataLayerInstance* DataLayerInstance)
+				{
+					const UDataLayerInstanceWithAsset* DataLayerInstanceWithAsset = Cast<UDataLayerInstanceWithAsset>(DataLayerInstance);
+					const UDataLayerAsset* DataLayerAsset = DataLayerInstanceWithAsset ? DataLayerInstanceWithAsset->GetAsset() : nullptr;
+					if (DataLayerAsset && (FName(DataLayerAsset->GetPathName()) == DataLayerAssetPath))
+					{
+						Result.Add(DataLayerInstance->GetDataLayerFName());
+						return false;
+					}
+					return true;
+				});
+			}
+			return Result;
+		}
+		// Fallback on FWorldDataLayersActorDesc
+		else if (InWorldDataLayersActorDesc)
+		{
+			TArray<FName> Result;
+			for (const FName& DataLayerAssetPath : InActorDesc->DataLayers)
+			{
+				if (const FDataLayerInstanceDesc* DataLayerInstanceDesc = InWorldDataLayersActorDesc->GetDataLayerInstanceFromAssetPath(DataLayerAssetPath))
+				{
+					Result.Add(DataLayerInstanceDesc->GetName());
+				}
+			}
+			return Result;
+		}
+	}
+
+	bIsSuccess = false;
+	return InActorDesc->DataLayers;
+}
+
+// For performance reasons, this function assumes that InActorDesc's DataLayerInstanceNames was already resolved.
+TArray<FName> FDataLayerUtils::ResolveRuntimeDataLayerInstanceNames(const FWorldPartitionActorDesc* InActorDesc, const UActorDescContainer* InContainer, bool* bOutIsResultValid)
+{
+	bool bLocalIsSuccess;
+	bool& bIsSuccess = bOutIsResultValid ? *bOutIsResultValid : bLocalIsSuccess;
+	bIsSuccess = false;
+
+	TArray<FName> Result;
+	bIsSuccess = InActorDesc->GetDataLayerInstanceNames().IsEmpty();
+	if (bIsSuccess)
+	{
+		return Result;
+	}
+
+	const UActorDescContainer* Container = InContainer ? InContainer : InActorDesc->GetContainer();
+	if (Container)
+	{
+		UWorld* World = Container->GetWorld();
+		if (AWorldDataLayers* WorldDataLayers = World ? World->GetWorldDataLayers() : nullptr)
+		{
+			bIsSuccess = true;
+			for (FName DataLayerInstanceName : InActorDesc->GetDataLayerInstanceNames())
+			{
+				const UDataLayerInstance* DataLayerInstance = WorldDataLayers->GetDataLayerInstance(DataLayerInstanceName);
+				if (DataLayerInstance && DataLayerInstance->IsRuntime())
+				{
+					Result.Add(DataLayerInstanceName);
+				}
+			}
+		}
+		// Fallback on FWorldDataLayersActorDesc
+		else if (FWorldDataLayersActorDesc* WorldDataLayersActorDesc = FDataLayerUtils::GetWorldDataLayersActorDesc(Container))
+		{
+			bIsSuccess = true;
+			for (FName DataLayerInstanceName : InActorDesc->GetDataLayerInstanceNames())
+			{
+				const FDataLayerInstanceDesc* DataLayerInstanceDesc = WorldDataLayersActorDesc->GetDataLayerInstanceFromInstanceName(DataLayerInstanceName);
+				if (DataLayerInstanceDesc && (DataLayerInstanceDesc->GetDataLayerType() == EDataLayerType::Runtime))
+				{
+					Result.Add(DataLayerInstanceName);
+				}
+			}
+		}
+	}
+	return Result;
+}
+
+#endif
\ No newline at end of file
diff --git a/Engine/Source/Runtime/Engine/Private/WorldPartition/DataLayer/DeprecatedDataLayerInstance.cpp b/Engine/Source/Runtime/Engine/Private/WorldPartition/DataLayer/DeprecatedDataLayerInstance.cpp
--- a/Engine/Source/Runtime/Engine/Private/WorldPartition/DataLayer/DeprecatedDataLayerInstance.cpp
+++ b/Engine/Source/Runtime/Engine/Private/WorldPartition/DataLayer/DeprecatedDataLayerInstance.cpp
@@ -1,7 +1,7 @@
 // Copyright Epic Games, Inc. All Rights Reserved.
 
 #include "WorldPartition/DataLayer/DeprecatedDataLayerInstance.h"
-
+#include "WorldPartition/DataLayer/DataLayerUtils.h"
 #include "WorldPartition/DataLayer/WorldDataLayers.h"
 
 UDeprecatedDataLayerInstance::UDeprecatedDataLayerInstance(const FObjectInitializer& ObjectInitializer)
@@ -79,7 +79,7 @@ bool UDeprecatedDataLayerInstance::ContainsActor(const AActor* Actor) const
 
 bool UDeprecatedDataLayerInstance::RelabelDataLayer(FName InDataLayerLabel)
 {
-	FName SanitizedLabel = DataLayerUtils::GetSanitizedDataLayerLabel(InDataLayerLabel);
+	FName SanitizedLabel = FDataLayerUtils::GetSanitizedDataLayerLabel(InDataLayerLabel);
 	FName UniqueNewDataLayerLabel = GetOuterAWorldDataLayers()->GenerateUniqueDataLayerLabel(SanitizedLabel);
 	if (Label != UniqueNewDataLayerLabel)
 	{
diff --git a/Engine/Source/Runtime/Engine/Private/WorldPartition/DataLayer/WorldDataLayers.cpp b/Engine/Source/Runtime/Engine/Private/WorldPartition/DataLayer/WorldDataLayers.cpp
--- a/Engine/Source/Runtime/Engine/Private/WorldPartition/DataLayer/WorldDataLayers.cpp
+++ b/Engine/Source/Runtime/Engine/Private/WorldPartition/DataLayer/WorldDataLayers.cpp
@@ -8,6 +8,7 @@
 #include "WorldPartition/DataLayer/DataLayerSubsystem.h"
 #include "WorldPartition/DataLayer/DataLayerInstanceWithAsset.h"
 #include "WorldPartition/DataLayer/DeprecatedDataLayerInstance.h"
+#include "WorldPartition/DataLayer/DataLayerUtils.h"
 #include "EngineUtils.h"
 #include "Engine/CoreSettings.h"
 #include "Net/UnrealNetwork.h"
@@ -903,7 +904,7 @@ bool AWorldDataLayers::RemoveDataLayer(const UDEPRECATED_DataLayer* InDataLayer)
 FName AWorldDataLayers::GenerateUniqueDataLayerLabel(const FName& InDataLayerLabel) const
 {
 	int32 DataLayerIndex = 0;
-	const FName DataLayerLabelSanitized = DataLayerUtils::GetSanitizedDataLayerLabel(InDataLayerLabel);
+	const FName DataLayerLabelSanitized = FDataLayerUtils::GetSanitizedDataLayerLabel(InDataLayerLabel);
 	FName UniqueNewDataLayerLabel = DataLayerLabelSanitized;
 	while (GetDataLayerFromLabel(UniqueNewDataLayerLabel))
 	{
@@ -921,7 +922,7 @@ bool AWorldDataLayers::ContainsDataLayer(const UDEPRECATED_DataLayer* InDataLaye
 
 const UDataLayerInstance* AWorldDataLayers::GetDataLayerFromLabel(const FName& InDataLayerLabel) const
 {
-	const FName DataLayerLabelSanitized = DataLayerUtils::GetSanitizedDataLayerLabel(InDataLayerLabel);
+	const FName DataLayerLabelSanitized = FDataLayerUtils::GetSanitizedDataLayerLabel(InDataLayerLabel);
 
 	for (const UDataLayerInstance* DataLayer : DataLayerInstances)
 	{
diff --git a/Engine/Source/Runtime/Engine/Private/WorldPartition/ErrorHandling/WorldPartitionStreamingGenerationLogErrorHandler.cpp b/Engine/Source/Runtime/Engine/Private/WorldPartition/ErrorHandling/WorldPartitionStreamingGenerationLogErrorHandler.cpp
--- a/Engine/Source/Runtime/Engine/Private/WorldPartition/ErrorHandling/WorldPartitionStreamingGenerationLogErrorHandler.cpp
+++ b/Engine/Source/Runtime/Engine/Private/WorldPartition/ErrorHandling/WorldPartitionStreamingGenerationLogErrorHandler.cpp
@@ -53,4 +53,9 @@ void FStreamingGenerationLogErrorHandler::OnDataLayerAssetConflict(const UDataLa
 {
 	UE_LOG(LogWorldPartition, Log, TEXT("Data Layer Instance %s and Data Layer Instance %s are both referencing Data Layer Asset %s"), *DataLayerInstance->GetDataLayerFName().ToString(), *ConflictingDataLayerInstance->GetDataLayerFName().ToString(), *DataLayerInstance->GetAsset()->GetFullName());
 }
+
+void FStreamingGenerationLogErrorHandler::OnActorNeedsResave(const FWorldPartitionActorDescView& ActorDescView)
+{
+	UE_LOG(LogWorldPartition, Log, TEXT("Actor %s needs to be resaved"), *ActorDescView.GetActorLabelOrName().ToString());
+}
 #endif
diff --git a/Engine/Source/Runtime/Engine/Private/WorldPartition/ErrorHandling/WorldPartitionStreamingGenerationTokenizedMessageErrorHandler.cpp b/Engine/Source/Runtime/Engine/Private/WorldPartition/ErrorHandling/WorldPartitionStreamingGenerationTokenizedMessageErrorHandler.cpp
--- a/Engine/Source/Runtime/Engine/Private/WorldPartition/ErrorHandling/WorldPartitionStreamingGenerationTokenizedMessageErrorHandler.cpp
+++ b/Engine/Source/Runtime/Engine/Private/WorldPartition/ErrorHandling/WorldPartitionStreamingGenerationTokenizedMessageErrorHandler.cpp
@@ -131,6 +131,17 @@ void ITokenizedMessageErrorHandler::OnDataLayerAssetConflict(const UDataLayerIns
 	HandleTokenizedMessage(MoveTemp(Message));
 }
 
+void ITokenizedMessageErrorHandler::OnActorNeedsResave(const FWorldPartitionActorDescView& ActorDescView)
+{
+	TSharedRef<FTokenizedMessage> Message = FTokenizedMessage::Create(EMessageSeverity::Error);
+	Message->AddToken(FTextToken::Create(LOCTEXT("TokenMessage_WorldPartition_ActorNeedsResave", "Actor needs resave")))
+		->AddToken(FActorToken::Create(ActorDescView.GetActorPath().ToString(), ActorDescView.GetGuid(), FText::FromName(ActorDescView.GetActorLabelOrName())));
+
+	AddAdditionalNameToken(Message, FName(TEXT("WorldPartition_ActorNeedsResave_CheckForErrors")));
+
+	HandleTokenizedMessage(MoveTemp(Message));
+}
+
 #undef LOCTEXT_NAMESPACE
 
 #endif
diff --git a/Engine/Source/Runtime/Engine/Private/WorldPartition/WorldPartitionActorCluster.cpp b/Engine/Source/Runtime/Engine/Private/WorldPartition/WorldPartitionActorCluster.cpp
--- a/Engine/Source/Runtime/Engine/Private/WorldPartition/WorldPartitionActorCluster.cpp
+++ b/Engine/Source/Runtime/Engine/Private/WorldPartition/WorldPartitionActorCluster.cpp
@@ -13,7 +13,7 @@
 #include "WorldPartition/DataLayer/DataLayerInstance.h"
 
 template<class LayerNameContainer>
-TSet<const UDataLayerInstance*> GetDataLayers(UWorld* InWorld, const LayerNameContainer& DataLayerInstanceNames)
+TSet<const UDataLayerInstance*> GetDataLayerInstances(UWorld* InWorld, const LayerNameContainer& DataLayerInstanceNames)
 {
 	TSet<const UDataLayerInstance*> DataLayerInstances;
 	if (const AWorldDataLayers* WorldDataLayers = InWorld->GetWorldDataLayers())
@@ -22,7 +22,7 @@ TSet<const UDataLayerInstance*> GetDataLayers(UWorld* InWorld, const LayerNameCo
 		{
 			if (const UDataLayerInstance* DataLayerInstance = WorldDataLayers->GetDataLayerInstance(DataLayerInstanceName))
 			{
-				if (DataLayerInstance->IsRuntime())
+				if (ensure(DataLayerInstance->IsRuntime()))
 				{
 					DataLayerInstances.Add(DataLayerInstance);
 				}
@@ -38,7 +38,7 @@ FActorCluster::FActorCluster(UWorld* InWorld, const FWorldPartitionActorDescView
 	, Bounds(InActorDescView.GetBounds())
 {
 	Actors.Add(InActorDescView.GetGuid());
-	DataLayers = GetDataLayers(InWorld, InActorDescView.GetDataLayers());
+	DataLayers = GetDataLayerInstances(InWorld, InActorDescView.GetRuntimeDataLayers());
 	DataLayersID = FDataLayersID(DataLayers.Array());
 }
 
@@ -83,7 +83,7 @@ FActorClusterContext::FActorClusterContext(TArray<FActorContainerInstance>&& InC
 	}
 }
 
-FActorContainerInstance::FActorContainerInstance(const FActorContainerID& InID, const FTransform& InTransform, const FBox& InBounds, const TSet<FName>& InDataLayers, EContainerClusterMode InClusterMode, const UActorDescContainer* InContainer, TMap<FGuid, FWorldPartitionActorDescView> InActorDescViewMap)
+FActorContainerInstance::FActorContainerInstance(const FActorContainerID& InID, const FTransform& InTransform, const FBox& InBounds, const TSet<FName>& InRuntimeDataLayers, EContainerClusterMode InClusterMode, const UActorDescContainer* InContainer, TMap<FGuid, FWorldPartitionActorDescView> InActorDescViewMap)
 	: ID(InID)
 	, Transform(InTransform)
 	, Bounds(InBounds)
@@ -91,7 +91,7 @@ FActorContainerInstance::FActorContainerInstance(const FActorContainerID& InID,
 	, Container(InContainer)
 	, ActorDescViewMap(InActorDescViewMap)
 {
-	DataLayers = GetDataLayers(InContainer->GetWorld(), InDataLayers);
+	DataLayers = GetDataLayerInstances(InContainer->GetWorld(), InRuntimeDataLayers);
 }
 
 const FWorldPartitionActorDescView& FActorContainerInstance::GetActorDescView(const FGuid& InGuid) const
diff --git a/Engine/Source/Runtime/Engine/Private/WorldPartition/WorldPartitionActorDesc.cpp b/Engine/Source/Runtime/Engine/Private/WorldPartition/WorldPartitionActorDesc.cpp
--- a/Engine/Source/Runtime/Engine/Private/WorldPartition/WorldPartitionActorDesc.cpp
+++ b/Engine/Source/Runtime/Engine/Private/WorldPartition/WorldPartitionActorDesc.cpp
@@ -12,9 +12,13 @@
 #include "Engine/Level.h"
 #include "UObject/UE5MainStreamObjectVersion.h"
 #include "UObject/UE5ReleaseStreamObjectVersion.h"
+#include "UObject/FortniteNCBranchObjectVersion.h"
 #include "WorldPartition/WorldPartitionLog.h"
 #include "WorldPartition/ActorDescContainer.h"
 #include "WorldPartition/HLOD/HLODLayer.h"
+#include "WorldPartition/DataLayer/WorldDataLayers.h"
+#include "WorldPartition/DataLayer/DataLayerAsset.h"
+#include "WorldPartition/DataLayer/DataLayerUtils.h"
 #include "Engine/Public/ActorReferencesUtils.h"
 #endif
 
@@ -23,7 +27,8 @@ uint32 FWorldPartitionActorDesc::GlobalTag = 0;
 bool FWorldPartitionActorDesc::bForceAlwaysLoaded = false;
 
 FWorldPartitionActorDesc::FWorldPartitionActorDesc()
-	: SoftRefCount(0)
+	: bIsUsingDataLayerAsset(false)
+	, SoftRefCount(0)
 	, HardRefCount(0)
 	, Container(nullptr)
 	, Tag(0)
@@ -49,7 +54,40 @@ void FWorldPartitionActorDesc::Init(const AActor* InActor)
 	bLevelBoundsRelevant = InActor->IsLevelBoundsRelevant();
 	bActorIsHLODRelevant = InActor->IsHLODRelevant();
 	HLODLayer = InActor->GetHLODLayer() ? FName(InActor->GetHLODLayer()->GetPathName()) : FName();
-	DataLayers = InActor->GetDataLayerInstanceNames();
+	
+	// DataLayers
+	{
+		TArray<FName> LocalDataLayerAssetPaths;
+		TArray<FName> LocalDataLayerInstanceNames;
+		const AWorldDataLayers* WorldDataLayers = InActor->GetWorld()->GetWorldDataLayers();
+		if (WorldDataLayers)
+		{
+			LocalDataLayerAssetPaths.Reserve(InActor->GetDataLayerAssets().Num());
+			for (const TObjectPtr<const UDataLayerAsset>& DataLayerAsset : InActor->GetDataLayerAssets())
+			{
+				if (DataLayerAsset && WorldDataLayers->GetDataLayerInstance(DataLayerAsset))
+				{
+					LocalDataLayerAssetPaths.Add(*DataLayerAsset->GetPathName());
+				}
+			}
+
+			PRAGMA_DISABLE_DEPRECATION_WARNINGS
+			LocalDataLayerInstanceNames = WorldDataLayers->GetDataLayerInstanceNames(InActor->GetActorDataLayers());
+			PRAGMA_ENABLE_DEPRECATION_WARNINGS
+		}
+		// Validation
+		const bool bHasDataLayerAssets = LocalDataLayerAssetPaths.Num() > 0;
+		const bool bHasDeprecatedDataLayers = LocalDataLayerInstanceNames.Num() > 0;
+		check((!bHasDataLayerAssets && !bHasDeprecatedDataLayers) || (bHasDataLayerAssets != bHasDeprecatedDataLayers));
+
+		// Init DataLayers persistent info
+		bIsUsingDataLayerAsset = bHasDataLayerAssets;
+		DataLayers = bIsUsingDataLayerAsset ? MoveTemp(LocalDataLayerAssetPaths) : MoveTemp(LocalDataLayerInstanceNames);
+
+		// Init DataLayers transient info
+		DataLayerInstanceNames = FDataLayerUtils::ResolvedDataLayerInstanceNames(this, WorldDataLayers);
+	}
+
 	ActorPackage = InActor->GetPackage()->GetFName();
 	ActorPath = *InActor->GetPathName();
 	FolderPath = InActor->GetFolderPath();
@@ -112,6 +150,7 @@ bool FWorldPartitionActorDesc::Equals(const FWorldPartitionActorDesc* Other) con
 		bActorIsEditorOnly == Other->bActorIsEditorOnly && 
 		bLevelBoundsRelevant == Other->bLevelBoundsRelevant && 
 		bActorIsHLODRelevant == Other->bActorIsHLODRelevant && 
+		bIsUsingDataLayerAsset == Other->bIsUsingDataLayerAsset &&
 		HLODLayer == Other->HLODLayer && 
 		FolderPath == Other->FolderPath &&
 		FolderGuid == Other->FolderGuid &&
@@ -124,12 +163,14 @@ bool FWorldPartitionActorDesc::Equals(const FWorldPartitionActorDesc* Other) con
 		SortedDataLayers.Sort([](const FName& LHS, const FName& RHS) { return LHS.LexicalLess(RHS); });
 		SortedDataLayersOther.Sort([](const FName& LHS, const FName& RHS) { return LHS.LexicalLess(RHS); });
 
-		TArray<FGuid> SortedReferences(References);
-		TArray<FGuid> SortedReferencesOther(Other->References);
-		SortedReferences.Sort();
-		SortedReferencesOther.Sort();
-
-		return SortedDataLayers == SortedDataLayersOther && SortedReferences == SortedReferencesOther;
+		if (SortedDataLayers == SortedDataLayersOther)
+		{
+			TArray<FGuid> SortedReferences(References);
+			TArray<FGuid> SortedReferencesOther(Other->References);
+			SortedReferences.Sort();
+			SortedReferencesOther.Sort();
+			return (SortedReferences == SortedReferencesOther);
+		}
 	}
 
 	return false;
@@ -195,6 +236,7 @@ void FWorldPartitionActorDesc::Serialize(FArchive& Ar)
 
 	Ar.UsingCustomVersion(FUE5MainStreamObjectVersion::GUID);
 	Ar.UsingCustomVersion(FUE5ReleaseStreamObjectVersion::GUID);
+	Ar.UsingCustomVersion(FFortniteNCBranchObjectVersion::GUID);
 
 	Ar << Class << Guid;
 
@@ -243,6 +285,11 @@ void FWorldPartitionActorDesc::Serialize(FArchive& Ar)
 		Ar << DataLayers;
 	}
 
+	if (Ar.CustomVer(FFortniteNCBranchObjectVersion::GUID) >= FFortniteNCBranchObjectVersion::WorldPartitionActorDescSerializeDataLayerAssets)
+	{
+		Ar << bIsUsingDataLayerAsset;
+	}
+
 	if (Ar.CustomVer(FUE5MainStreamObjectVersion::GUID) >= FUE5MainStreamObjectVersion::WorldPartitionActorDescSerializeActorLabel)
 	{
 		Ar << ActorLabel;
diff --git a/Engine/Source/Runtime/Engine/Private/WorldPartition/WorldPartitionActorDescView.cpp b/Engine/Source/Runtime/Engine/Private/WorldPartition/WorldPartitionActorDescView.cpp
--- a/Engine/Source/Runtime/Engine/Private/WorldPartition/WorldPartitionActorDescView.cpp
+++ b/Engine/Source/Runtime/Engine/Private/WorldPartition/WorldPartitionActorDescView.cpp
@@ -5,6 +5,8 @@
 #if WITH_EDITOR
 #include "WorldPartition/WorldPartitionLog.h"
 #include "WorldPartition/WorldPartitionActorDesc.h"
+#include "WorldPartition/ActorDescContainer.h"
+#include "WorldPartition/DataLayer/DataLayerUtils.h"
 
 FWorldPartitionActorDescView::FWorldPartitionActorDescView()
 	: FWorldPartitionActorDescView(nullptr)
@@ -15,7 +17,20 @@ FWorldPartitionActorDescView::FWorldPartitionActorDescView(const FWorldPartition
 	, bIsForcedNonSpatiallyLoaded(false)
 	, bInvalidDataLayers(false)
 	, bInvalidRuntimeGrid(false)
-{}
+{
+	ResolveRuntimeDataLayers();
+}
+
+//The only case where we need to call ResolveRuntimeDataLayers and pass a Container is for the "unsaved actors" case of FWorldPartitionStreamingGenerator
+void FWorldPartitionActorDescView::ResolveRuntimeDataLayers(const UActorDescContainer* InContainer)
+{
+	bool bSuccess = true;
+	RuntimeDataLayers = FDataLayerUtils::ResolveRuntimeDataLayerInstanceNames(ActorDesc, InContainer, &bSuccess);
+	if (!bSuccess)
+	{
+		RuntimeDataLayers.Reset();
+	}
+}
 
 const FGuid& FWorldPartitionActorDescView::GetGuid() const
 {
@@ -78,6 +93,12 @@ const TArray<FName>& FWorldPartitionActorDescView::GetDataLayers() const
 	return bInvalidDataLayers ? EmptyDataLayers : ActorDesc->GetDataLayerInstanceNames();
 }
 
+const TArray<FName>& FWorldPartitionActorDescView::GetRuntimeDataLayers() const
+{
+	static TArray<FName> EmptyDataLayers;
+	return (bInvalidDataLayers || !RuntimeDataLayers.IsSet()) ? EmptyDataLayers : RuntimeDataLayers.GetValue();
+}
+
 FName FWorldPartitionActorDescView::GetActorPackage() const
 {
 	return ActorDesc->GetActorPackage();
@@ -161,4 +182,10 @@ void FWorldPartitionActorDescView::SetInvalidDataLayers()
 		UE_LOG(LogWorldPartition, Verbose, TEXT("Actor '%s' data layers invalidated"), *GetActorLabel().ToString());
 	}
 }
+
+bool FWorldPartitionActorDescView::IsResaveNeeded() const
+{
+	return ActorDesc->IsResaveNeeded();
+}
+
 #endif
\ No newline at end of file
diff --git a/Engine/Source/Runtime/Engine/Private/WorldPartition/WorldPartitionStreamingGeneration.cpp b/Engine/Source/Runtime/Engine/Private/WorldPartition/WorldPartitionStreamingGeneration.cpp
--- a/Engine/Source/Runtime/Engine/Private/WorldPartition/WorldPartitionStreamingGeneration.cpp
+++ b/Engine/Source/Runtime/Engine/Private/WorldPartition/WorldPartitionStreamingGeneration.cpp
@@ -42,6 +42,15 @@ class FWorldPartitionStreamingGenerator
 
 		// Consider all actors of a /Temp/ container package as Unsaved because loading them from disk will fail (Outer world name mismatch)
 		const bool bIsTempContainerPackage = FPackageName::IsTempPackage(InContainer->GetPackage()->GetName());
+
+		auto GetModifiedActorDesc = [this, InContainer](AActor* InActor)
+		{
+			FWorldPartitionActorDesc* ModifiedActorDesc = ModifiedActorsDescList->AddActor(InActor);
+			FWorldPartitionActorDescView ModifiedActorDescView(ModifiedActorDesc);
+			// ModifiedActorDescView's ActorDesc doesn't have a container, resolve view RuntimeDataLayers with provided ActorDescContainer
+			ModifiedActorDescView.ResolveRuntimeDataLayers(InContainer);
+			return ModifiedActorDescView;
+		};
 		
 		TMap<FGuid, FGuid> ContainerGuidsRemap;
 		for (FActorDescList::TConstIterator<> ActorDescIt(InContainer); ActorDescIt; ++ActorDescIt)
@@ -61,8 +70,8 @@ class FWorldPartitionStreamingGenerator
 					if (bHandleUnsavedActors && (bIsTempContainerPackage || Actor->GetPackage()->IsDirty()))
 					{
 						// Dirty, unsaved actor for PIE
-						FWorldPartitionActorDesc* ActorDesc = ModifiedActorsDescList->AddActor(Actor);
-						OutActorDescViewMap.Emplace(ActorDescIt->GetGuid(), ActorDesc);
+						FWorldPartitionActorDescView ModifiedActorDescView = GetModifiedActorDesc(Actor);
+						OutActorDescViewMap.Emplace(ActorDescIt->GetGuid(), ModifiedActorDescView);
 						continue;
 					}
 				}
@@ -79,14 +88,14 @@ class FWorldPartitionStreamingGenerator
 			{
 				if (IsValid(Actor) && Actor->IsPackageExternal() && Actor->IsMainPackageActor() && !Actor->IsEditorOnly() && !InContainer->GetActorDesc(Actor->GetActorGuid()))
 				{
-					FWorldPartitionActorDesc* ActorDesc = ModifiedActorsDescList->AddActor(Actor);
-					OutActorDescViewMap.Emplace(ActorDesc->GetGuid(), ActorDesc);
+					FWorldPartitionActorDescView ModifiedActorDescView = GetModifiedActorDesc(Actor);
+					OutActorDescViewMap.Emplace(Actor->GetActorGuid(), ModifiedActorDescView);
 				}
 			}
 		}
 	}
 
-	void CreateActorDescriptorViewsRecursive(const UActorDescContainer* InContainer, const FTransform& InTransform, const TSet<FName>& InDataLayers, const FActorContainerID& InContainerID, const FActorContainerID& InParentContainerID, EContainerClusterMode InClusterMode, const TCHAR* OwnerName)
+	void CreateActorDescriptorViewsRecursive(const UActorDescContainer* InContainer, const FTransform& InTransform, const TSet<FName>& InRuntimeDataLayers, const FActorContainerID& InContainerID, const FActorContainerID& InParentContainerID, EContainerClusterMode InClusterMode, const TCHAR* OwnerName)
 	{
 		TMap<FGuid, FWorldPartitionActorDescView> ActorDescViewMap;
 		
@@ -109,16 +118,16 @@ class FWorldPartitionStreamingGenerator
 				const FActorContainerID SubContainerID(InContainerID, ActorGuid);
 
 				// Only propagate data layers from the root container as we don't support sub containers to set their own
-				const TSet<FName>* SubDataLayers = &InDataLayers;
-				TSet<FName> ParentDataLayers;
+				const TSet<FName>* SubRuntimeDataLayers = &InRuntimeDataLayers;
+				TSet<FName> ParentRuntimeDataLayers;
 
 				if (InContainerID.IsMainContainer())
 				{
-					ParentDataLayers.Append(ActorDescView.GetDataLayers());
-					SubDataLayers = &ParentDataLayers;
+					ParentRuntimeDataLayers.Append(ActorDescView.GetRuntimeDataLayers());
+					SubRuntimeDataLayers = &ParentRuntimeDataLayers;
 				}
 
-				CreateActorDescriptorViewsRecursive(SubContainer, SubTransform * InTransform, *SubDataLayers, SubContainerID, InContainerID, SubClusterMode, *ActorDescView.GetActorLabelOrName().ToString());
+				CreateActorDescriptorViewsRecursive(SubContainer, SubTransform * InTransform, *SubRuntimeDataLayers, SubContainerID, InContainerID, SubClusterMode, *ActorDescView.GetActorLabelOrName().ToString());
 
 				// The container actor can be removed now that its contained actors has been registered
 				It.RemoveCurrent();
@@ -138,7 +147,7 @@ class FWorldPartitionStreamingGenerator
 		ContainerDescriptor.Transform = InTransform;
 		ContainerDescriptor.ClusterMode = InClusterMode;
 		ContainerDescriptor.ActorDescViewMap = MoveTemp(ActorDescViewMap);
-		ContainerDescriptor.DataLayers = InDataLayers;
+		ContainerDescriptor.RuntimeDataLayers = InRuntimeDataLayers;
 		ContainerDescriptor.OwnerName = OwnerName;
 
 		// Maintain containers hierarchy, bottom up
@@ -165,10 +174,10 @@ class FWorldPartitionStreamingGenerator
 		// Validate data layers
 		auto IsReferenceDataLayersValid = [](const FWorldPartitionActorDescView& ActorDescView, const FWorldPartitionActorDescView& ReferenceActorDescView)
 		{
-			if (ActorDescView.GetDataLayers().Num() == ReferenceActorDescView.GetDataLayers().Num())
+			if (ActorDescView.GetRuntimeDataLayers().Num() == ReferenceActorDescView.GetRuntimeDataLayers().Num())
 			{
-				const TSet<FName> ActorDescDataLayers(ActorDescView.GetDataLayers());
-				const TSet<FName> ReferenceActorDescDataLayers(ReferenceActorDescView.GetDataLayers());
+				const TSet<FName> ActorDescDataLayers(ActorDescView.GetRuntimeDataLayers());
+				const TSet<FName> ReferenceActorDescDataLayers(ReferenceActorDescView.GetRuntimeDataLayers());
 
 				return ActorDescDataLayers.Includes(ReferenceActorDescDataLayers);
 			}
@@ -198,7 +207,7 @@ class FWorldPartitionStreamingGenerator
 							ActorDescView.SetForcedNonSpatiallyLoaded();
 						}
 
-						if (ActorDescView.GetDataLayers().Num())
+						if (ActorDescView.GetRuntimeDataLayers().Num())
 						{
 							ErrorHandler->OnInvalidReferenceLevelScriptDataLayers(ActorDescView);
 							ActorDescView.SetInvalidDataLayers();
@@ -355,6 +364,16 @@ class FWorldPartitionStreamingGenerator
 				NbValidationPasses++;
 			}
 			while (NbErrorsDetected);
+
+			// Report ActorDescs that need to be resaved
+			for (auto ActorDescIt = ContainerDescriptor.ActorDescViewMap.CreateIterator(); ActorDescIt; ++ActorDescIt)
+			{
+				FWorldPartitionActorDescView& ActorDescView = ActorDescIt.Value();
+				if (ActorDescView.IsResaveNeeded())
+				{
+					ErrorHandler->OnActorNeedsResave(ActorDescView);
+				}
+			}
 		}
 	}
 
@@ -442,7 +461,7 @@ public:
 			const FActorContainerID& ContainerID = It.Key();
 			const FContainerDescriptor& ContainerDescriptor = It.Value();
 
-			ContainerInstances.Emplace(ContainerID, ContainerDescriptor.Transform, ContainerDescriptor.Bounds, ContainerDescriptor.DataLayers, ContainerDescriptor.ClusterMode, ContainerDescriptor.Container, ContainerDescriptor.ActorDescViewMap);
+			ContainerInstances.Emplace(ContainerID, ContainerDescriptor.Transform, ContainerDescriptor.Bounds, ContainerDescriptor.RuntimeDataLayers, ContainerDescriptor.ClusterMode, ContainerDescriptor.Container, ContainerDescriptor.ActorDescViewMap);
 		}
 
 		return FActorClusterContext(MoveTemp(ContainerInstances), InFilterActorDescViewFunc);
@@ -543,7 +562,7 @@ private:
 		const UActorDescContainer* Container;
 		EContainerClusterMode ClusterMode;
 		TMap<FGuid, FWorldPartitionActorDescView> ActorDescViewMap;
-		TSet<FName> DataLayers;
+		TSet<FName> RuntimeDataLayers;
 		FString OwnerName;
 	};
 
diff --git a/Engine/Source/Runtime/Engine/Public/WorldPartition/ActorDescContainer.h b/Engine/Source/Runtime/Engine/Public/WorldPartition/ActorDescContainer.h
--- a/Engine/Source/Runtime/Engine/Public/WorldPartition/ActorDescContainer.h
+++ b/Engine/Source/Runtime/Engine/Public/WorldPartition/ActorDescContainer.h
@@ -44,6 +44,9 @@ public:
 	
 	DECLARE_EVENT_OneParam(UWorldPartition, FActorDescRemovedEvent, FWorldPartitionActorDesc*);
 	FActorDescRemovedEvent OnActorDescRemovedEvent;
+
+	DECLARE_MULTICAST_DELEGATE_OneParam(FActorDescContainerInitializeDelegate, UActorDescContainer*);
+	static FActorDescContainerInitializeDelegate OnActorDescContainerInitialized;
 #endif
 
 	UPROPERTY(Transient)
diff --git a/Engine/Source/Runtime/Engine/Public/WorldPartition/DataLayer/ActorDataLayer.h b/Engine/Source/Runtime/Engine/Public/WorldPartition/DataLayer/ActorDataLayer.h
--- a/Engine/Source/Runtime/Engine/Public/WorldPartition/DataLayer/ActorDataLayer.h
+++ b/Engine/Source/Runtime/Engine/Public/WorldPartition/DataLayer/ActorDataLayer.h
@@ -4,7 +4,7 @@
 
 #include "CoreMinimal.h"
 
-#include "WorldPartition/DataLayer/DataLayerUtils.h"
+#include "WorldPartition/DataLayer/DataLayerInstance.h"
 
 #include "ActorDataLayer.generated.h"
 
diff --git a/Engine/Source/Runtime/Engine/Public/WorldPartition/DataLayer/DataLayerAsset.h b/Engine/Source/Runtime/Engine/Public/WorldPartition/DataLayer/DataLayerAsset.h
--- a/Engine/Source/Runtime/Engine/Public/WorldPartition/DataLayer/DataLayerAsset.h
+++ b/Engine/Source/Runtime/Engine/Public/WorldPartition/DataLayer/DataLayerAsset.h
@@ -7,7 +7,7 @@
 #include "UObject/UObjectGlobals.h"
 #include "UObject/Object.h"
 
-#include "WorldPartition/DataLayer/DataLayerUtils.h" 
+#include "WorldPartition/DataLayer/DataLayerType.h" 
 
 #include "DataLayerAsset.generated.h"
 
diff --git a/Engine/Source/Runtime/Engine/Public/WorldPartition/DataLayer/DataLayerInstance.h b/Engine/Source/Runtime/Engine/Public/WorldPartition/DataLayer/DataLayerInstance.h
--- a/Engine/Source/Runtime/Engine/Public/WorldPartition/DataLayer/DataLayerInstance.h
+++ b/Engine/Source/Runtime/Engine/Public/WorldPartition/DataLayer/DataLayerInstance.h
@@ -6,10 +6,12 @@
 #include "UObject/ObjectMacros.h"
 #include "UObject/UObjectGlobals.h"
 #include "UObject/Object.h"
-#include "WorldPartition/DataLayer/DataLayerAsset.h"
+#include "WorldPartition/DataLayer/DataLayerType.h"
 
 #include "DataLayerInstance.generated.h"
 
+#define DATALAYER_TO_INSTANCE_RUNTIME_CONVERSION_ENABLED 1
+
 class IStreamingGenerationErrorHandler;
 
 UENUM(BlueprintType)
@@ -55,7 +57,7 @@ public:
 	bool IsLoadedInEditorChangedByUserOperation() const { return bIsLoadedInEditorChangedByUserOperation; }
 	void ClearLoadedInEditorChangedByUserOperation() { bIsLoadedInEditorChangedByUserOperation = false; }
 
-	const TCHAR* GetDataLayerIconName() const { return DataLayerUtils::GetDataLayerIconName(GetType()); }
+	const TCHAR* GetDataLayerIconName() const;
 
 	bool CanParent(const UDataLayerInstance* InParent) const;
 	bool IsDataLayerTypeValidToParent(EDataLayerType ParentDataLayerType) const;
@@ -78,7 +80,6 @@ public:
 #if DATALAYER_TO_INSTANCE_RUNTIME_CONVERSION_ENABLED
 	virtual bool SupportRelabeling() const { return false; }
 	virtual bool RelabelDataLayer(FName NewDataLayerLabel) { return false; }
-
 #endif // DATALAYER_TO_INSTANCE_RUNTIME_CONVERSION_ENABLED
 
 #endif // WITH_EDITOR
diff --git a/Engine/Source/Runtime/Engine/Public/WorldPartition/DataLayer/DataLayerInstanceWithAsset.h b/Engine/Source/Runtime/Engine/Public/WorldPartition/DataLayer/DataLayerInstanceWithAsset.h
--- a/Engine/Source/Runtime/Engine/Public/WorldPartition/DataLayer/DataLayerInstanceWithAsset.h
+++ b/Engine/Source/Runtime/Engine/Public/WorldPartition/DataLayer/DataLayerInstanceWithAsset.h
@@ -2,11 +2,10 @@
 #pragma once
 
 #include "WorldPartition/DataLayer/DataLayerInstance.h"
+#include "WorldPartition/DataLayer/DataLayerAsset.h"
 
 #include "DataLayerInstanceWithAsset.generated.h"
 
-class UDataLayerAsset;
-
 UCLASS(Config = Engine, PerObjectConfig, Within = WorldDataLayers, AutoCollapseCategories = ("Data Layer|Advanced"), AutoExpandCategories = ("Data Layer|Editor", "Data Layer|Advanced|Runtime"))
 class ENGINE_API UDataLayerInstanceWithAsset : public UDataLayerInstance
 {
diff --git a/Engine/Source/Runtime/Engine/Public/WorldPartition/DataLayer/DataLayerType.h b/Engine/Source/Runtime/Engine/Public/WorldPartition/DataLayer/DataLayerType.h
new file mode 100644
--- /dev/null
+++ b/Engine/Source/Runtime/Engine/Public/WorldPartition/DataLayer/DataLayerType.h
@@ -0,0 +1,15 @@
+// Copyright Epic Games, Inc. All Rights Reserved.
+
+#pragma once
+
+#include "DataLayerType.generated.h"
+
+UENUM(BlueprintType)
+enum class EDataLayerType : uint8
+{
+	Runtime,
+	Editor,
+	Unknown UMETA(Hidden),
+
+	Size UMETA(Hidden)
+};
\ No newline at end of file
diff --git a/Engine/Source/Runtime/Engine/Public/WorldPartition/DataLayer/DataLayerUtils.h b/Engine/Source/Runtime/Engine/Public/WorldPartition/DataLayer/DataLayerUtils.h
--- a/Engine/Source/Runtime/Engine/Public/WorldPartition/DataLayer/DataLayerUtils.h
+++ b/Engine/Source/Runtime/Engine/Public/WorldPartition/DataLayer/DataLayerUtils.h
@@ -5,30 +5,26 @@
 #include "CoreMinimal.h"
 #include "UObject/ObjectMacros.h"
 #include "UObject/UObjectGlobals.h"
+#include "WorldPartition/DataLayer/DataLayerType.h"
+#include "WorldPartition/DataLayer/DataLayerInstance.h"
 
-#include "DataLayerUtils.generated.h"
+class FWorldDataLayersActorDesc;
+class FWorldPartitionActorDesc;
+class UActorDescContainer;
+class AWorldDataLayers;
 
-#define DATALAYER_TO_INSTANCE_RUNTIME_CONVERSION_ENABLED 1
-
-UENUM(BlueprintType)
-enum class EDataLayerType : uint8
-{
-	Runtime,
-	Editor,
-	Unknown UMETA(Hidden),
-
-	Size UMETA(Hidden)
-};
-
-namespace DataLayerUtils
+class ENGINE_API FDataLayerUtils
 {
+public:
 #if WITH_EDITOR
-	constexpr const TCHAR* IconNameByType[static_cast<int>(EDataLayerType::Size)] = { TEXT("DataLayer.Runtime") , TEXT("DataLayer.Editor"), TEXT("") };
-
-	constexpr const TCHAR* GetDataLayerIconName(EDataLayerType DataLayerType)
+	static const TCHAR* GetDataLayerIconName(EDataLayerType DataLayerType)
 	{
+		static constexpr const TCHAR* IconNameByType[static_cast<int>(EDataLayerType::Size)] = { TEXT("DataLayer.Runtime") , TEXT("DataLayer.Editor"), TEXT("") };
 		return IconNameByType[static_cast<uint32>(DataLayerType)];
 	}
+	static FWorldDataLayersActorDesc* GetWorldDataLayersActorDesc(const UActorDescContainer* InContainer, bool bInCheckValid = true);
+	static TArray<FName> ResolvedDataLayerInstanceNames(const FWorldPartitionActorDesc* InActorDesc, const AWorldDataLayers* InWorldDataLayers = nullptr, const FWorldDataLayersActorDesc* InWorldDataLayersActorDesc = nullptr, bool* bOutIsResultValid = nullptr);
+	static TArray<FName> ResolveRuntimeDataLayerInstanceNames(const FWorldPartitionActorDesc* InActorDesc, const UActorDescContainer* InContainer, bool* bOutIsResultValid = nullptr);
 #endif
 
 #if DATALAYER_TO_INSTANCE_RUNTIME_CONVERSION_ENABLED
@@ -37,6 +33,5 @@ namespace DataLayerUtils
 	{
 		return FName(InDataLayerLabel.ToString().TrimStartAndEnd().Replace(TEXT("\""), TEXT("")));
 	}
-
 #endif
 };
\ No newline at end of file
diff --git a/Engine/Source/Runtime/Engine/Public/WorldPartition/DataLayer/WorldDataLayersActorDesc.h b/Engine/Source/Runtime/Engine/Public/WorldPartition/DataLayer/WorldDataLayersActorDesc.h
--- a/Engine/Source/Runtime/Engine/Public/WorldPartition/DataLayer/WorldDataLayersActorDesc.h
+++ b/Engine/Source/Runtime/Engine/Public/WorldPartition/DataLayer/WorldDataLayersActorDesc.h
@@ -3,7 +3,7 @@
 #pragma once
 
 #include "WorldPartition/WorldPartitionActorDesc.h"
-#include "WorldPartition/DataLayer/DataLayerUtils.h"
+#include "WorldPartition/DataLayer/DataLayerType.h"
 
 #if WITH_EDITOR
 
@@ -64,6 +64,7 @@ protected:
 	virtual void Init(const AActor* InActor) override;
 	virtual bool Equals(const FWorldPartitionActorDesc* Other) const override;
 	virtual void Serialize(FArchive& Ar) override;
+	virtual bool IsResaveNeeded() const override { return !IsValid(); }
 
 private:
 	TArray<FDataLayerInstanceDesc> DataLayerInstances;
diff --git a/Engine/Source/Runtime/Engine/Public/WorldPartition/ErrorHandling/WorldPartitionStreamingGenerationErrorHandler.h b/Engine/Source/Runtime/Engine/Public/WorldPartition/ErrorHandling/WorldPartitionStreamingGenerationErrorHandler.h
--- a/Engine/Source/Runtime/Engine/Public/WorldPartition/ErrorHandling/WorldPartitionStreamingGenerationErrorHandler.h
+++ b/Engine/Source/Runtime/Engine/Public/WorldPartition/ErrorHandling/WorldPartitionStreamingGenerationErrorHandler.h
@@ -57,5 +57,10 @@ public:
 	 * Called when two data layer instances share the same asset
 	 */
 	virtual void OnDataLayerAssetConflict(const UDataLayerInstanceWithAsset* DataLayerInstance, const UDataLayerInstanceWithAsset* ConflictingDataLayerInstance) = 0;
+
+	/**
+	 * Called when an actor needs to be resaved.
+	 */
+	virtual void OnActorNeedsResave(const FWorldPartitionActorDescView& ActorDescView) = 0;
 };
 #endif
diff --git a/Engine/Source/Runtime/Engine/Public/WorldPartition/ErrorHandling/WorldPartitionStreamingGenerationLogErrorHandler.h b/Engine/Source/Runtime/Engine/Public/WorldPartition/ErrorHandling/WorldPartitionStreamingGenerationLogErrorHandler.h
--- a/Engine/Source/Runtime/Engine/Public/WorldPartition/ErrorHandling/WorldPartitionStreamingGenerationLogErrorHandler.h
+++ b/Engine/Source/Runtime/Engine/Public/WorldPartition/ErrorHandling/WorldPartitionStreamingGenerationLogErrorHandler.h
@@ -16,5 +16,6 @@ public:
 	virtual void OnInvalidReferenceDataLayerAsset(const UDataLayerInstanceWithAsset* DataLayerInstance) override;
 	virtual void OnDataLayerHierarchyTypeMismatch(const UDataLayerInstance* DataLayerInstance, const UDataLayerInstance* Parent) override;
 	virtual void OnDataLayerAssetConflict(const UDataLayerInstanceWithAsset* DataLayerInstance, const UDataLayerInstanceWithAsset* ConflictingDataLayerInstance) override;
+	virtual void OnActorNeedsResave(const FWorldPartitionActorDescView& ActorDescView) override;
 };
 #endif
diff --git a/Engine/Source/Runtime/Engine/Public/WorldPartition/ErrorHandling/WorldPartitionStreamingGenerationNullErrorHandler.h b/Engine/Source/Runtime/Engine/Public/WorldPartition/ErrorHandling/WorldPartitionStreamingGenerationNullErrorHandler.h
--- a/Engine/Source/Runtime/Engine/Public/WorldPartition/ErrorHandling/WorldPartitionStreamingGenerationNullErrorHandler.h
+++ b/Engine/Source/Runtime/Engine/Public/WorldPartition/ErrorHandling/WorldPartitionStreamingGenerationNullErrorHandler.h
@@ -16,5 +16,6 @@ public:
 	virtual void OnInvalidReferenceDataLayerAsset(const UDataLayerInstanceWithAsset* DataLayerInstance) override {}
 	virtual void OnDataLayerHierarchyTypeMismatch(const UDataLayerInstance* DataLayerInstance, const UDataLayerInstance* Parent) override {};
 	virtual void OnDataLayerAssetConflict(const UDataLayerInstanceWithAsset* DataLayerInstance, const UDataLayerInstanceWithAsset* ConflictingDataLayerInstance) override {};
+	virtual void OnActorNeedsResave(const FWorldPartitionActorDescView& ActorDescView) override {};
 };
 #endif
diff --git a/Engine/Source/Runtime/Engine/Public/WorldPartition/ErrorHandling/WorldPartitionStreamingGenerationTokenizedMessageErrorHandler.h b/Engine/Source/Runtime/Engine/Public/WorldPartition/ErrorHandling/WorldPartitionStreamingGenerationTokenizedMessageErrorHandler.h
--- a/Engine/Source/Runtime/Engine/Public/WorldPartition/ErrorHandling/WorldPartitionStreamingGenerationTokenizedMessageErrorHandler.h
+++ b/Engine/Source/Runtime/Engine/Public/WorldPartition/ErrorHandling/WorldPartitionStreamingGenerationTokenizedMessageErrorHandler.h
@@ -17,6 +17,7 @@ class ENGINE_API ITokenizedMessageErrorHandler : public IStreamingGenerationErro
 	virtual void OnInvalidReferenceDataLayerAsset(const UDataLayerInstanceWithAsset* DataLayerInstance) override;
 	virtual void OnDataLayerHierarchyTypeMismatch(const UDataLayerInstance* DataLayerInstance, const UDataLayerInstance* Parent) override;
 	virtual void OnDataLayerAssetConflict(const UDataLayerInstanceWithAsset* DataLayerInstance, const UDataLayerInstanceWithAsset* ConflictingDataLayerInstance) override;
+	virtual void OnActorNeedsResave(const FWorldPartitionActorDescView& ActorDescView) override;
 
 protected:
 	virtual void AddAdditionalNameToken(TSharedRef<FTokenizedMessage>& InMessage, const FName& InErrorName) {}
diff --git a/Engine/Source/Runtime/Engine/Public/WorldPartition/WorldPartitionActorCluster.h b/Engine/Source/Runtime/Engine/Public/WorldPartition/WorldPartitionActorCluster.h
--- a/Engine/Source/Runtime/Engine/Public/WorldPartition/WorldPartitionActorCluster.h
+++ b/Engine/Source/Runtime/Engine/Public/WorldPartition/WorldPartitionActorCluster.h
@@ -86,7 +86,7 @@ struct FActorCluster
  */
 struct FActorContainerInstance
 {
-	FActorContainerInstance(const FActorContainerID& InID, const FTransform& InTransform, const FBox& InBounds, const TSet<FName>& InDataLayers, EContainerClusterMode InClusterMode, const UActorDescContainer* InContainer, TMap<FGuid, FWorldPartitionActorDescView> InActorDescViewMap);
+	FActorContainerInstance(const FActorContainerID& InID, const FTransform& InTransform, const FBox& InBounds, const TSet<FName>& InRuntimeDataLayers, EContainerClusterMode InClusterMode, const UActorDescContainer* InContainer, TMap<FGuid, FWorldPartitionActorDescView> InActorDescViewMap);
 	
 	FActorContainerID			ID;
 	FTransform					Transform;
diff --git a/Engine/Source/Runtime/Engine/Public/WorldPartition/WorldPartitionActorDesc.h b/Engine/Source/Runtime/Engine/Public/WorldPartition/WorldPartitionActorDesc.h
--- a/Engine/Source/Runtime/Engine/Public/WorldPartition/WorldPartitionActorDesc.h
+++ b/Engine/Source/Runtime/Engine/Public/WorldPartition/WorldPartitionActorDesc.h
@@ -37,6 +37,8 @@ class ENGINE_API FWorldPartitionActorDesc
 	friend class UActorDescContainer;
 	friend struct FWorldPartitionHandleImpl;
 	friend struct FWorldPartitionReferenceImpl;
+	friend class UDataLayerEditorSubsystem;
+	friend class FDataLayerUtils;
 
 public:
 	virtual ~FWorldPartitionActorDesc() {}
@@ -52,7 +54,7 @@ public:
 	inline bool GetLevelBoundsRelevant() const { return bLevelBoundsRelevant; }
 	inline bool GetActorIsHLODRelevant() const { return bActorIsHLODRelevant; }
 	inline FName GetHLODLayer() const { return HLODLayer; }
-	inline const TArray<FName>& GetDataLayerInstanceNames() const { return DataLayers; }
+	inline const TArray<FName>& GetDataLayerInstanceNames() const { return DataLayerInstanceNames; }
 	inline FName GetActorPackage() const { return ActorPackage; }
 	inline FName GetActorPath() const { return ActorPath; }
 	inline FName GetActorLabel() const { return ActorLabel; }
@@ -60,12 +62,14 @@ public:
 	inline const FGuid& GetFolderGuid() const { return FolderGuid; }
 	FBox GetBounds() const;
 	inline const FGuid& GetParentActor() const { return ParentActor; }
+	inline bool IsUsingDataLayerAsset() const { return bIsUsingDataLayerAsset; }
 
 	FName GetActorName() const;
 	FName GetActorLabelOrName() const;
 
 	virtual bool GetContainerInstance(const UActorDescContainer*& OutLevelContainer, FTransform& OutLevelTransform, EContainerClusterMode& OutClusterMode) const { return false; }
 	virtual const FGuid& GetSceneOutlinerParent() const { return GetParentActor(); }
+	virtual bool IsResaveNeeded() const { return false; }
 
 	bool operator==(const FWorldPartitionActorDesc& Other) const
 	{
@@ -178,6 +182,7 @@ protected:
 	bool							bActorIsEditorOnly;
 	bool							bLevelBoundsRelevant;
 	bool							bActorIsHLODRelevant;
+	bool							bIsUsingDataLayerAsset; // Used to know if DataLayers array represents DataLayers Asset paths or the FNames of the deprecated version of Data Layers
 	FName							HLODLayer;
 	TArray<FName>					DataLayers;
 	TArray<FGuid>					References;
@@ -192,6 +197,7 @@ protected:
 	UClass*							ActorClass;
 	mutable TWeakObjectPtr<AActor>	ActorPtr;
 	UActorDescContainer*			Container;
+	TArray<FName>					DataLayerInstanceNames;
 
 public:
 	// Tagging
diff --git a/Engine/Source/Runtime/Engine/Public/WorldPartition/WorldPartitionActorDescView.h b/Engine/Source/Runtime/Engine/Public/WorldPartition/WorldPartitionActorDescView.h
--- a/Engine/Source/Runtime/Engine/Public/WorldPartition/WorldPartitionActorDescView.h
+++ b/Engine/Source/Runtime/Engine/Public/WorldPartition/WorldPartitionActorDescView.h
@@ -2,6 +2,7 @@
 #pragma once
 
 #include "CoreMinimal.h"
+#include "Misc/Optional.h"
 
 #if WITH_EDITOR
 class FWorldPartitionActorDesc;
@@ -32,6 +33,8 @@ public:
 	bool GetActorIsHLODRelevant() const;
 	FName GetHLODLayer() const;
 	const TArray<FName>& GetDataLayers() const;
+	const TArray<FName>& GetRuntimeDataLayers() const;
+	void ResolveRuntimeDataLayers(const UActorDescContainer* InContainer = nullptr);
 	FName GetActorPackage() const;
 	FName GetActorPath() const;
 	FName GetActorLabel() const;
@@ -53,6 +56,8 @@ public:
 
 	void SetInvalidDataLayers();
 
+	bool IsResaveNeeded() const;
+
 	bool operator==(const FWorldPartitionActorDescView& Other) const
 	{
 		return GetGuid() == Other.GetGuid();
@@ -68,5 +73,6 @@ protected:
 	bool bIsForcedNonSpatiallyLoaded;
 	bool bInvalidDataLayers;
 	bool bInvalidRuntimeGrid;
+	TOptional<TArray<FName>> RuntimeDataLayers;
 };
 #endif
\ No newline at end of file
