World Partition - Fixed PIE streaming generation issues with the "editor only" actor property
* Unsaved actors whose editor only property was modified were wrongly included or excluded from the streaming generation

diff --git a/Engine/Source/Runtime/Engine/Private/WorldPartition/WorldPartitionStreamingGeneration.cpp b/Engine/Source/Runtime/Engine/Private/WorldPartition/WorldPartitionStreamingGeneration.cpp
--- a/Engine/Source/Runtime/Engine/Private/WorldPartition/WorldPartitionStreamingGeneration.cpp
+++ b/Engine/Source/Runtime/Engine/Private/WorldPartition/WorldPartitionStreamingGeneration.cpp
@@ -42,6 +42,9 @@ class FWorldPartitionStreamingGenerator
 
 		// Consider all actors of a /Temp/ container package as Unsaved because loading them from disk will fail (Outer world name mismatch)
 		const bool bIsTempContainerPackage = FPackageName::IsTempPackage(InContainer->GetPackage()->GetName());
+		
+		// Test whether an actor is editor only. Will fallback to the actor descriptor only if the actor is not loaded.
+		auto IsActorEditorOnly = [](const FWorldPartitionActorDesc* ActorDesc) { return ActorDesc->IsLoaded() ? ActorDesc->GetActor()->IsEditorOnly() : ActorDesc->GetActorIsEditorOnly(); };
 
 		auto GetModifiedActorDesc = [this, InContainer](AActor* InActor)
 		{
@@ -55,7 +58,7 @@ class FWorldPartitionStreamingGenerator
 		TMap<FGuid, FGuid> ContainerGuidsRemap;
 		for (FActorDescList::TConstIterator<> ActorDescIt(InContainer); ActorDescIt; ++ActorDescIt)
 		{
-			if (!ActorDescIt->GetActorIsEditorOnly())
+			if (!IsActorEditorOnly(*ActorDescIt))
 			{
 				// Handle unsaved actors
 				if (AActor* Actor = ActorDescIt->GetActor())
