[VLogger] fixed serialization by bumping version and adding support for LargeWorldCoordinates
#rnx

diff --git a/Engine/Source/Runtime/Engine/Private/VisualLogger/VisualLoggerTypes.cpp b/Engine/Source/Runtime/Engine/Private/VisualLogger/VisualLoggerTypes.cpp
--- a/Engine/Source/Runtime/Engine/Private/VisualLogger/VisualLoggerTypes.cpp
+++ b/Engine/Source/Runtime/Engine/Private/VisualLogger/VisualLoggerTypes.cpp
@@ -394,7 +394,18 @@ FArchive& operator<<(FArchive& Ar, FVisualLogHistogramSample& Sample)
 	FVisualLoggerHelpers::Serialize(Ar, Sample.GraphName);
 	FVisualLoggerHelpers::Serialize(Ar, Sample.DataName);
 	Ar << Sample.Verbosity;
-	Ar << Sample.SampleValue;
+
+	if (Ar.CustomVer(EVisualLoggerVersion::GUID) >= EVisualLoggerVersion::LargeWorldCoordinatesAndLocationValidityFlag)
+	{
+		Ar << Sample.SampleValue;
+	}
+	else
+	{
+		FVector2f SampleValueFlt;
+		Ar << SampleValueFlt;
+		Sample.SampleValue = FVector2D(SampleValueFlt);
+	}
+
 	Ar << Sample.UniqueId;
 
 	return Ar;
@@ -420,12 +431,37 @@ FArchive& operator<<(FArchive& Ar, FVisualLogShapeElement& Element)
 	Ar << Element.Verbosity;
 	const int32 VLogsVer = Ar.CustomVer(EVisualLoggerVersion::GUID);
 
+	const bool bUseLargeWorldCoordinates = (VLogsVer >= EVisualLoggerVersion::LargeWorldCoordinatesAndLocationValidityFlag);
+	
 	if (VLogsVer >= EVisualLoggerVersion::TransformationForShapes)
 	{
-		Ar << Element.TransformationMatrix;
+		if (bUseLargeWorldCoordinates)
+		{
+			Ar << Element.TransformationMatrix;
+		}
+		else
+		{
+			FMatrix44f TransformationMatrixFlt;
+        	Ar << TransformationMatrixFlt;
+        	Element.TransformationMatrix = FMatrix(TransformationMatrixFlt);
+		}
+	}
+
+	if (bUseLargeWorldCoordinates)
+	{
+		Ar << Element.Points;
+	}
+	else
+	{
+		TArray<FVector3f> FltPoints;
+		Ar << FltPoints;
+		Element.Points.Reserve(FltPoints.Num());
+		for (FVector3f Point : FltPoints)
+		{
+			Element.Points.Emplace(Point);
+		}
 	}
 
-	Ar << Element.Points;
 	Ar << Element.UniqueId;
 	Ar << Element.Type;
 	Ar << Element.Color;
@@ -513,18 +549,29 @@ FArchive& operator<<(FArchive& Ar, FVisualLogStatusCategory& Status)
 FArchive& operator<<(FArchive& Ar, FVisualLogEntry& LogEntry)
 {
 	Ar << LogEntry.TimeStamp;
-	Ar << LogEntry.Location;
+
+	const int32 VLogsVer = Ar.CustomVer(EVisualLoggerVersion::GUID);
+	if (VLogsVer >= EVisualLoggerVersion::LargeWorldCoordinatesAndLocationValidityFlag)
+	{
+		Ar << LogEntry.Location;
+
+		uint8 bTempIsLocationValid = (LogEntry.bIsLocationValid != 0);
+		Ar.SerializeBits(&bTempIsLocationValid, 1);
+		LogEntry.bIsLocationValid = bTempIsLocationValid != 0;
+	}
+	else
+	{
+		FVector3f LocationFlt(LogEntry.Location);
+		Ar << LocationFlt;
+		LogEntry.Location = FVector(LocationFlt);
+	}
+
 	Ar << LogEntry.LogLines;
 	Ar << LogEntry.Status;
 	Ar << LogEntry.Events;
 	Ar << LogEntry.ElementsToDraw;
 	Ar << LogEntry.DataBlocks;
 	
-	uint8 bTempIsLocationValid = (LogEntry.bIsLocationValid != 0);
-	Ar.SerializeBits(&bTempIsLocationValid, 1);
-	LogEntry.bIsLocationValid = bTempIsLocationValid != 0;	
-
-	const int32 VLogsVer = Ar.CustomVer(EVisualLoggerVersion::GUID);
 	if (VLogsVer > EVisualLoggerVersion::Initial)
 	{
 		Ar << LogEntry.HistogramSamples;
diff --git a/Engine/Source/Runtime/Engine/Public/VisualLogger/VisualLoggerCustomVersion.h b/Engine/Source/Runtime/Engine/Public/VisualLogger/VisualLoggerCustomVersion.h
--- a/Engine/Source/Runtime/Engine/Public/VisualLogger/VisualLoggerCustomVersion.h
+++ b/Engine/Source/Runtime/Engine/Public/VisualLogger/VisualLoggerCustomVersion.h
@@ -14,6 +14,7 @@ namespace EVisualLoggerVersion
 		AddedOwnerClassName = 2,
 		StatusCategoryWithChildren = 3,
 		TransformationForShapes = 4,
+		LargeWorldCoordinatesAndLocationValidityFlag = 5,
 		// -----<new versions can be added before this line>-------------------------------------------------
 		// - this needs to be the last line (see note below)
 		VersionPlusOne,
