Fix structs with no properties not being copied in Arrays like they are in root components during construction scripts.
Extending the fix in cl 11687390

#rnx

diff --git a/Engine/Source/Runtime/Engine/Classes/Engine/BlueprintGeneratedClass.h b/Engine/Source/Runtime/Engine/Classes/Engine/BlueprintGeneratedClass.h
--- a/Engine/Source/Runtime/Engine/Classes/Engine/BlueprintGeneratedClass.h
+++ b/Engine/Source/Runtime/Engine/Classes/Engine/BlueprintGeneratedClass.h
@@ -809,6 +809,18 @@ protected:
 	*/
 	static void InitArrayPropertyFromCustomList(const FArrayProperty* ArrayProperty, const FCustomPropertyListNode* InPropertyList, uint8* DataPtr, const uint8* DefaultDataPtr);
 
+
+	/**
+	* Helper method to assist with initializing a property from sub property list
+	*
+	* @param	Property					the property to copy
+	* @param	CustomPropertySubListNode	the list node of the sub property to copy
+	* @param	PropertyValue				destination address (where to start copying values to)
+	* @param	DefaultPropertyValue		source address (where to start copying the defaults data from)
+	* @return	true if the method was able to copy the property successfully via the sub custom property list
+	*/
+	static bool InitPropertyFromSubPropertyList(const FProperty* Property, const FCustomPropertyListNode* SubPropertyList, uint8* PropertyValue, const uint8* DefaultPropertyValue);
+
 	// @todo: BP2CPP_remove
 	/** Check for and handle manual application of default value overrides to component subobjects that were inherited from a nativized parent class */
 	UE_DEPRECATED(5.0, "This API is no longer in use and will be removed.")
diff --git a/Engine/Source/Runtime/Engine/Private/BlueprintGeneratedClass.cpp b/Engine/Source/Runtime/Engine/Private/BlueprintGeneratedClass.cpp
--- a/Engine/Source/Runtime/Engine/Private/BlueprintGeneratedClass.cpp
+++ b/Engine/Source/Runtime/Engine/Private/BlueprintGeneratedClass.cpp
@@ -825,25 +825,9 @@ void UBlueprintGeneratedClass::InitPropertiesFromCustomList(const FCustomPropert
 		uint8* PropertyValue = CustomPropertyListNode->Property->ContainerPtrToValuePtr<uint8>(DataPtr, CustomPropertyListNode->ArrayIndex);
 		const uint8* DefaultPropertyValue = CustomPropertyListNode->Property->ContainerPtrToValuePtr<uint8>(DefaultDataPtr, CustomPropertyListNode->ArrayIndex);
 
-		if (const FStructProperty* StructProperty = CastField<FStructProperty>(CustomPropertyListNode->Property))
-		{
-			if (CustomPropertyListNode->SubPropertyList != nullptr)
-			{
-				InitPropertiesFromCustomList(CustomPropertyListNode->SubPropertyList, StructProperty->Struct, PropertyValue, DefaultPropertyValue);
-			}
-			else
-			{
-				// A NULL sub-property list indicates that we should copy the entire default value (e.g. a struct with one or more non-reflected fields).
-				StructProperty->CopySingleValue(PropertyValue, DefaultPropertyValue);
-			}
-		}
-		else if (const FArrayProperty* ArrayProperty = CastField<FArrayProperty>(CustomPropertyListNode->Property))
-		{
-			// Note: The sub-property list can be NULL here; in that case only the array size will differ from the default value, but the elements themselves will simply be initialized to defaults.
-			InitArrayPropertyFromCustomList(ArrayProperty, CustomPropertyListNode->SubPropertyList, PropertyValue, DefaultPropertyValue);
-		}
-		else
+		if (!InitPropertyFromSubPropertyList(CustomPropertyListNode->Property, CustomPropertyListNode->SubPropertyList, PropertyValue, DefaultPropertyValue))
 		{
+			// Unable to init properties from sub custom property list, fall back to the default copy value behavior
 			CustomPropertyListNode->Property->CopySingleValue(PropertyValue, DefaultPropertyValue);
 		}
 	}
@@ -878,19 +862,31 @@ void UBlueprintGeneratedClass::InitArrayPropertyFromCustomList(const FArrayPrope
 			continue;
 		}
 
-		if (const FStructProperty* InnerStructProperty = CastField<FStructProperty>(ArrayProperty->Inner))
+		if (!InitPropertyFromSubPropertyList(ArrayProperty->Inner, CustomArrayPropertyListNode->SubPropertyList, DstArrayItemValue, SrcArrayItemValue))
 		{
-			InitPropertiesFromCustomList(CustomArrayPropertyListNode->SubPropertyList, InnerStructProperty->Struct, DstArrayItemValue, SrcArrayItemValue);
+			// Unable to init properties from sub custom property list, fall back to the default copy value behavior
+			ArrayProperty->Inner->CopyCompleteValue(DstArrayItemValue, SrcArrayItemValue);
 		}
-		else if (const FArrayProperty* InnerArrayProperty = CastField<FArrayProperty>(ArrayProperty->Inner))
+	}
+}
+
+bool UBlueprintGeneratedClass::InitPropertyFromSubPropertyList(const FProperty* Property, const FCustomPropertyListNode* SubPropertyList, uint8* PropertyValue, const uint8* DefaultPropertyValue)
+{
+	if (SubPropertyList != nullptr)
+	{
+		if (const FStructProperty* StructProperty = CastField<FStructProperty>(Property))
 		{
-			InitArrayPropertyFromCustomList(InnerArrayProperty, CustomArrayPropertyListNode->SubPropertyList, DstArrayItemValue, SrcArrayItemValue);
+			InitPropertiesFromCustomList(SubPropertyList, StructProperty->Struct, PropertyValue, DefaultPropertyValue);
+			return true;
 		}
-		else
+		else if (const FArrayProperty* ArrayProperty = CastField<FArrayProperty>(Property))
 		{
-			ArrayProperty->Inner->CopyCompleteValue(DstArrayItemValue, SrcArrayItemValue);
+			InitArrayPropertyFromCustomList(ArrayProperty, SubPropertyList, PropertyValue, DefaultPropertyValue);
+			return true;
 		}
+		// No need to handle Sets and Maps as they are not optimized through the custom property list
 	}
+	return false;
 }
 
 bool UBlueprintGeneratedClass::IsFunctionImplementedInScript(FName InFunctionName) const
