Move BP extensions to private, expose getters / setters instead.

diff --git a/Engine/Source/Editor/AnimGraph/Private/AnimBlueprintExtension.cpp b/Engine/Source/Editor/AnimGraph/Private/AnimBlueprintExtension.cpp
--- a/Engine/Source/Editor/AnimGraph/Private/AnimBlueprintExtension.cpp
+++ b/Engine/Source/Editor/AnimGraph/Private/AnimBlueprintExtension.cpp
@@ -35,14 +35,14 @@ UAnimBlueprintExtension* UAnimBlueprintExtension::RequestExtension(UAnimBlueprin
 
 	// Not found, create one
 	UAnimBlueprintExtension* NewExtension = NewObject<UAnimBlueprintExtension>(InAnimBlueprint, InExtensionType.Get());
-	InAnimBlueprint->Extensions.Add(NewExtension);
+	InAnimBlueprint->AddExtension(NewExtension);
 	return NewExtension;
 }
 
 UAnimBlueprintExtension* UAnimBlueprintExtension::GetExtension(UAnimBlueprint* InAnimBlueprint, TSubclassOf<UAnimBlueprintExtension> InExtensionType)
 {
 	// Look for an existing extension
-	for(UBlueprintExtension* Extension : InAnimBlueprint->Extensions)
+	for(TObjectPtr<UBlueprintExtension> Extension : InAnimBlueprint->GetExtensions())
 	{
 		if(Extension && Extension->GetClass() == InExtensionType)
 		{
@@ -57,7 +57,7 @@ TArray<UAnimBlueprintExtension*> UAnimBlueprintExtension::GetExtensions(UAnimBlu
 {
 	TArray<UAnimBlueprintExtension*> Extensions;
 	
-	for(UBlueprintExtension* Extension : InAnimBlueprint->Extensions)
+	for(TObjectPtr<UBlueprintExtension> Extension : InAnimBlueprint->GetExtensions())
 	{
 		if(Extension && Extension->GetClass()->IsChildOf(UAnimBlueprintExtension::StaticClass()))
 		{
@@ -117,7 +117,7 @@ void UAnimBlueprintExtension::RefreshExtensions(UAnimBlueprint* InAnimBlueprint)
 	}
 
 	// Remove all extensions that are no longer needed
-	InAnimBlueprint->Extensions.RemoveAll([](UBlueprintExtension* InExtension)
+	InAnimBlueprint->RemoveAllExtension([](UBlueprintExtension* InExtension)
 	{
 		if(UAnimBlueprintExtension* AnimBlueprintExtension = Cast<UAnimBlueprintExtension>(InExtension))
 		{
@@ -133,7 +133,7 @@ void UAnimBlueprintExtension::RefreshExtensions(UAnimBlueprint* InAnimBlueprint)
 
 void UAnimBlueprintExtension::ForEachExtension(UAnimBlueprint* InAnimBlueprint, TFunctionRef<void(UAnimBlueprintExtension*)> InFunction)
 {
-	for (UBlueprintExtension* BlueprintExtension : InAnimBlueprint->Extensions)
+	for (TObjectPtr<UBlueprintExtension> BlueprintExtension : InAnimBlueprint->GetExtensions())
 	{
 		if(UAnimBlueprintExtension* AnimBlueprintExtension = Cast<UAnimBlueprintExtension>(BlueprintExtension))
 		{
diff --git a/Engine/Source/Editor/KismetCompiler/Private/KismetCompiler.cpp b/Engine/Source/Editor/KismetCompiler/Private/KismetCompiler.cpp
--- a/Engine/Source/Editor/KismetCompiler/Private/KismetCompiler.cpp
+++ b/Engine/Source/Editor/KismetCompiler/Private/KismetCompiler.cpp
@@ -4341,7 +4341,7 @@ void FKismetCompilerContext::CreateFunctionList()
 		BP_SCOPED_COMPILER_EVENT_STAT(EKismetCompilerStats_GenerateFunctionGraphs);
 
 		// Allow blueprint extensions for the blueprint to generate function graphs
-		for (UBlueprintExtension* Extension : Blueprint->Extensions)
+		for (TObjectPtr<UBlueprintExtension> Extension : Blueprint->GetExtensions())
 		{
 			Extension->GenerateFunctionGraphs(this);
 		}
diff --git a/Engine/Source/Editor/MovieSceneTools/Private/MovieSceneEventUtils.cpp b/Engine/Source/Editor/MovieSceneTools/Private/MovieSceneEventUtils.cpp
--- a/Engine/Source/Editor/MovieSceneTools/Private/MovieSceneEventUtils.cpp
+++ b/Engine/Source/Editor/MovieSceneTools/Private/MovieSceneEventUtils.cpp
@@ -481,7 +481,7 @@ void FMovieSceneEventUtils::BindEventSectionToBlueprint(UMovieSceneEventSectionB
 {
 	check(EventSection && DirectorBP);
 
-	for (UBlueprintExtension* Extension : DirectorBP->Extensions)
+	for (TObjectPtr<UBlueprintExtension> Extension : DirectorBP->GetExtensions())
 	{
 		UMovieSceneEventBlueprintExtension* EventExtension = Cast<UMovieSceneEventBlueprintExtension>(Extension);
 		if (EventExtension)
@@ -493,7 +493,7 @@ void FMovieSceneEventUtils::BindEventSectionToBlueprint(UMovieSceneEventSectionB
 
 	UMovieSceneEventBlueprintExtension* EventExtension = NewObject<UMovieSceneEventBlueprintExtension>(DirectorBP);
 	EventExtension->Add(EventSection);
-	DirectorBP->Extensions.Add(EventExtension);
+	DirectorBP->AddExtension(EventExtension);
 }
 
 void FMovieSceneEventUtils::RemoveEndpointsForEventSection(UMovieSceneEventSectionBase* EventSection, UBlueprint* DirectorBP)
diff --git a/Engine/Source/Editor/UMGEditor/Private/WidgetBlueprintExtension.cpp b/Engine/Source/Editor/UMGEditor/Private/WidgetBlueprintExtension.cpp
--- a/Engine/Source/Editor/UMGEditor/Private/WidgetBlueprintExtension.cpp
+++ b/Engine/Source/Editor/UMGEditor/Private/WidgetBlueprintExtension.cpp
@@ -7,7 +7,6 @@
 
 UWidgetBlueprintExtension* UWidgetBlueprintExtension::RequestExtension(UWidgetBlueprint* InBlueprint, TSubclassOf<UWidgetBlueprintExtension> InExtensionType)
 {
-	
 	checkf(!InBlueprint->bBeingCompiled, TEXT("Do not use RequestExtension when a blueprint is being compiled."));
 
 	// Look for an existing extension
@@ -18,7 +17,7 @@ UWidgetBlueprintExtension* UWidgetBlueprintExtension::RequestExtension(UWidgetBl
 
 	// Not found, create one
 	UWidgetBlueprintExtension* NewExtension = NewObject<UWidgetBlueprintExtension>(InBlueprint, InExtensionType.Get());
-	InBlueprint->Extensions.Add(NewExtension);
+	InBlueprint->AddExtension(NewExtension);
 	return NewExtension;
 }
 
@@ -26,7 +25,7 @@ UWidgetBlueprintExtension* UWidgetBlueprintExtension::RequestExtension(UWidgetBl
 UWidgetBlueprintExtension* UWidgetBlueprintExtension::GetExtension(const UWidgetBlueprint* InBlueprint, TSubclassOf<UWidgetBlueprintExtension> InExtensionType)
 {
 	// Look for an existing extension
-	for (UBlueprintExtension* Extension : InBlueprint->Extensions)
+	for (TObjectPtr<UBlueprintExtension> Extension : InBlueprint->GetExtensions())
 	{
 		if (Extension && Extension->GetClass() == InExtensionType)
 		{
@@ -42,7 +41,7 @@ TArray<UWidgetBlueprintExtension*> UWidgetBlueprintExtension::GetExtensions(cons
 {
 	TArray<UWidgetBlueprintExtension*> Extensions;
 
-	for (UBlueprintExtension* Extension : InBlueprint->Extensions)
+	for (TObjectPtr<UBlueprintExtension> Extension : InBlueprint->GetExtensions())
 	{
 		if (Extension && Extension->GetClass()->IsChildOf(UWidgetBlueprintExtension::StaticClass()))
 		{
@@ -53,7 +52,6 @@ TArray<UWidgetBlueprintExtension*> UWidgetBlueprintExtension::GetExtensions(cons
 	return Extensions;
 }
 
-
 UWidgetBlueprint* UWidgetBlueprintExtension::GetWidgetBlueprint() const
 {
 	return GetOuterUWidgetBlueprint();
diff --git a/Engine/Source/Editor/UMGEditor/Public/WidgetBlueprintExtension.h b/Engine/Source/Editor/UMGEditor/Public/WidgetBlueprintExtension.h
--- a/Engine/Source/Editor/UMGEditor/Public/WidgetBlueprintExtension.h
+++ b/Engine/Source/Editor/UMGEditor/Public/WidgetBlueprintExtension.h
@@ -53,7 +53,7 @@ public:
 	template<typename Predicate>
 	static void ForEachExtension(const UWidgetBlueprint* InWidgetBlueprint, Predicate Pred)
 	{
-		for (UBlueprintExtension* BlueprintExtension : InWidgetBlueprint->Extensions)
+		for (TObjectPtr<UBlueprintExtension> BlueprintExtension : InWidgetBlueprint->GetExtensions())
 		{
 			if (UWidgetBlueprintExtension* WidgetBlueprintExtension = Cast<UWidgetBlueprintExtension>(BlueprintExtension))
 			{
diff --git a/Engine/Source/Runtime/Engine/Classes/Engine/Blueprint.h b/Engine/Source/Runtime/Engine/Classes/Engine/Blueprint.h
--- a/Engine/Source/Runtime/Engine/Classes/Engine/Blueprint.h
+++ b/Engine/Source/Runtime/Engine/Classes/Engine/Blueprint.h
@@ -621,6 +621,9 @@ PRAGMA_ENABLE_DEPRECATION_WARNINGS
 	TObjectPtr<class UInheritableComponentHandler> InheritableComponentHandler;
 
 #if WITH_EDITORONLY_DATA
+	DECLARE_MULTICAST_DELEGATE_OneParam(FOnExtensionAdded, TObjectPtr<UBlueprintExtension>);
+	DECLARE_MULTICAST_DELEGATE_OneParam(FOnExtensionRemoved, TObjectPtr<UBlueprintExtension>);
+
 	/** Array of new variables to be added to generated class */
 	UPROPERTY()
 	TArray<struct FBPVariableDescription> NewVariables;
@@ -669,9 +672,15 @@ PRAGMA_ENABLE_DEPRECATION_WARNINGS
 	TMap<FName, FName> OldToNewComponentTemplateNames;
 
 	/** Array of extensions for this blueprint */
+	UE_DEPRECATED(5.1, "Please do not access this member directly; Instead use: UBlueprint::GetExtensions / UBlueprint::AddExtension / UBlueprint::RemoveExtension[At].")
 	UPROPERTY()
 	TArray<TObjectPtr<UBlueprintExtension>> Extensions;
 
+	/** Fires whenever BP extension added */
+	FOnExtensionAdded OnExtensionAdded;
+
+	/** Fires whenever BP extension removed */
+	FOnExtensionRemoved OnExtensionRemoved;
 #endif // WITH_EDITORONLY_DATA
 
 public:
@@ -688,7 +697,38 @@ public:
 	DECLARE_EVENT_OneParam(UBlueprint, FCompiledEvent, class UBlueprint*);
 	FCompiledEvent& OnCompiled() { return CompiledEvent; }
 	void BroadcastCompiled() { CompiledEvent.Broadcast(this); }
-#endif
+
+	/** Gives const access to extensions. */
+	TArrayView<const TObjectPtr<UBlueprintExtension>> GetExtensions() const;
+
+	/** Adds given extension, broadcasting on add. */
+	int32 AddExtension(const TObjectPtr<UBlueprintExtension>& InExtension);
+
+	/** Removes given extension, broadcasting on remove. */
+	int32 RemoveExtension(const TObjectPtr<UBlueprintExtension>& InExtension);
+
+PRAGMA_DISABLE_DEPRECATION_WARNINGS
+	/** Removes all extensions matching the predicate, broadcasting on remove if one exists. */
+	template <class PREDICATE_CLASS>
+	int32 RemoveAllExtension(const PREDICATE_CLASS& Predicate)
+	{
+		auto BroadcastRemovePredicate = [&Predicate, this](UBlueprintExtension* InExtension)
+		{
+			bool NotMatch = !::Invoke(Predicate, InExtension); // use a ! to guarantee it can't be anything other than zero or one
+
+			if (!NotMatch)
+			{
+				OnExtensionRemoved.Broadcast(InExtension);
+			}
+
+			return !NotMatch;
+		};
+
+		return Extensions.RemoveAll(BroadcastRemovePredicate);
+	}
+PRAGMA_ENABLE_DEPRECATION_WARNINGS
+
+#endif // WITH_EDITOR
 
 	/** Whether or not this blueprint can be considered for a bytecode only compile */
 	virtual bool IsValidForBytecodeOnlyRecompile() const { return true; }
diff --git a/Engine/Source/Runtime/Engine/Private/Blueprint.cpp b/Engine/Source/Runtime/Engine/Private/Blueprint.cpp
--- a/Engine/Source/Runtime/Engine/Private/Blueprint.cpp
+++ b/Engine/Source/Runtime/Engine/Private/Blueprint.cpp
@@ -608,11 +608,13 @@ UClass* UBlueprint::RegenerateClass(UClass* ClassToRegenerate, UObject* Previous
 	}
 	UBlueprint::ForceLoadMembers(this);
 
+PRAGMA_DISABLE_DEPRECATION_WARNINGS
 	for (UBlueprintExtension* Extension : Extensions)
 	{
 		ForceLoad(Extension);
 		Extension->PreloadObjectsForCompilation(this);
 	}
+PRAGMA_ENABLE_DEPRECATION_WARNINGS
 
 	FBlueprintEditorUtils::PreloadConstructionScript( this );
 
@@ -1932,6 +1934,32 @@ UEdGraph* UBlueprint::GetLastEditedUberGraph() const
 	return nullptr;
 }
 
+PRAGMA_DISABLE_DEPRECATION_WARNINGS
+
+TArrayView<const TObjectPtr<UBlueprintExtension>> UBlueprint::GetExtensions() const
+{
+	return Extensions;
+}
+
+int32 UBlueprint::AddExtension(const TObjectPtr<UBlueprintExtension>& InExtension)
+{
+	int32 Index = Extensions.Add(InExtension);
+	OnExtensionAdded.Broadcast(InExtension);
+	return Index;
+}
+
+int32 UBlueprint::RemoveExtension(const TObjectPtr<UBlueprintExtension>& InExtension)
+{
+	int32 NumRemoved = Extensions.RemoveSingleSwap(InExtension);
+	if (NumRemoved > 0)
+	{
+		OnExtensionRemoved.Broadcast(InExtension);
+	}
+	return NumRemoved;
+}
+
+PRAGMA_ENABLE_DEPRECATION_WARNINGS
+
 #endif //WITH_EDITOR
 
 #if WITH_EDITORONLY_DATA
@@ -1944,5 +1972,3 @@ void UBlueprint::LoadModulesRequiredForCompilation()
 	FModuleManager::Get().LoadModule(MovieSceneToolsModuleName);
 }
 #endif //WITH_EDITORONLY_DATA
-
-
