Exposed functions for fetching a scalar material parameter index and then setting the parameters by index to blueprint in order to be able to optimize cases with extreme amounts of material parameters where many are set each frame.

#UE5

diff --git a/Engine/Source/Runtime/Engine/Classes/Materials/MaterialInstanceDynamic.h b/Engine/Source/Runtime/Engine/Classes/Materials/MaterialInstanceDynamic.h
--- a/Engine/Source/Runtime/Engine/Classes/Materials/MaterialInstanceDynamic.h
+++ b/Engine/Source/Runtime/Engine/Classes/Materials/MaterialInstanceDynamic.h
@@ -25,13 +25,24 @@ class ENGINE_API UMaterialInstanceDynamic : public UMaterialInstance
 	UFUNCTION(BlueprintCallable, meta = (Keywords = "SetFloatParameterValue"), Category = "Rendering|Material")
 	void SetScalarParameterValueByInfo(const FMaterialParameterInfo& ParameterInfo, float Value);
 
-	// NOTE: These Index-related functions should be used VERY carefully, and only in cases where optimization is
-	// critical.  Generally that's only if you're using an unusually high number of parameters in a material AND
-	// setting a huge number of parameters in the same frame.
-
-	// Use this function to set an initial value and fetch the index for use in the following function.
+	//~ NOTE: These Index-related functions should be used VERY carefully, and only in cases where optimization is
+	//~ critical.  Generally that's only if you're using an unusually high number of parameters in a material AND
+	//~ setting a huge number of parameters in the same frame.
+
+	// Use this function to set an initial value and fetch the index for use in SetScalarParameterByIndex.  This
+	// function should only be called once for a particular name, and then use SetScalarParameterByIndex for subsequent
+	// calls.  However, beware using this except in cases where optimization is critical, which is generally only if
+	// you're using an unusually high number of parameters in a material and setting a large number of parameters in the
+	// same frame.  Also, if the material is changed in any way that can change the parameter list, the index can be
+	// invalidated.
+	UFUNCTION(BlueprintCallable, meta = (Keywords = "SetFloatParameterValue"), Category = "Rendering|Material")
 	bool InitializeScalarParameterAndGetIndex(const FName& ParameterName, float Value, int32& OutParameterIndex);
-	// Use the cached value of OutParameterIndex above to set the scalar parameter ONLY on the exact same MID
+
+	// Use the cached value of OutParameterIndex from InitializeScalarParameterAndGetIndex to set the scalar parameter
+	// ONLY on the exact same MID.  Do NOT presume the index can be used from one instance on another.  Only use this
+	// pair of functions when optimization is critical; otherwise use either SetScalarParameterValueByInfo or
+	// SetScalarParameterValue.
+	UFUNCTION(BlueprintCallable, meta = (Keywords = "SetFloatParameterValue"), Category = "Rendering|Material")
 	bool SetScalarParameterByIndex(int32 ParameterIndex, float Value);
 
 	// Use this function to set an initial value and fetch the index for use in the following function.
