Fix FReplayExternalData constructor to not construct then move the internal bitreader

diff --git a/Engine/Source/Runtime/Engine/Public/ReplayTypes.h b/Engine/Source/Runtime/Engine/Public/ReplayTypes.h
--- a/Engine/Source/Runtime/Engine/Public/ReplayTypes.h
+++ b/Engine/Source/Runtime/Engine/Public/ReplayTypes.h
@@ -562,9 +562,10 @@ public:
 	{
 	}
 
-	FReplayExternalData(FBitReader&& InReader, const float InTimeSeconds) : TimeSeconds(InTimeSeconds)
+	FReplayExternalData(FBitReader&& InReader, const float InTimeSeconds) 
+		: Reader(MoveTemp(InReader))
+		, TimeSeconds(InTimeSeconds)
 	{
-		Reader = MoveTemp(InReader);
 	}
 
 	FBitReader	Reader;
