Add 'Custom' as a new EQuartzCommandType

diff --git a/Engine/Source/Runtime/Engine/Classes/Sound/QuartzQuantizationUtilities.h b/Engine/Source/Runtime/Engine/Classes/Sound/QuartzQuantizationUtilities.h
--- a/Engine/Source/Runtime/Engine/Classes/Sound/QuartzQuantizationUtilities.h
+++ b/Engine/Source/Runtime/Engine/Classes/Sound/QuartzQuantizationUtilities.h
@@ -206,8 +206,8 @@ enum class EQuartzCommandType : uint8
 	RetriggerSound UMETA(DisplayName = "Re-trigger Sound", ToolTip = "Quantized looping of the target sound (event tells the AudioComponent to play the sound again)"),
 	TickRateChange UMETA(DisplayName = "Tick Rate Change", ToolTip = "Quantized change of the tick-rate (i.e. BPM change)"),
 	TransportReset UMETA(DisplayName = "Transport Reset", ToolTip = "Quantized reset of the clocks transport (back to time = 0 on the boundary)"),
-	StartOtherClock UMETA(DisplayName = "Start Other Clock", ToolTip = "Quantized start of another clock. Useful for sample accurate synchronization of clocks (i.e. to handle time signature changes)")
-
+	StartOtherClock UMETA(DisplayName = "Start Other Clock", ToolTip = "Quantized start of another clock. Useful for sample accurate synchronization of clocks (i.e. to handle time signature changes)"),
+	Custom UMETA(DisplayName = "Custom", ToolTip = "Quantized custom command")
 }; // EQuartzCommandType
 
 
