[UE][FIX] Add cvar to set cloth weight used in skin render to 0
- New CVar "r.SkeletalMeshClothBlend.Enabled" is true by default, when false the cloth weight used in skin render is set to 0

diff --git a/Engine/Source/Runtime/Engine/Private/SkeletalRender.cpp b/Engine/Source/Runtime/Engine/Private/SkeletalRender.cpp
--- a/Engine/Source/Runtime/Engine/Private/SkeletalRender.cpp
+++ b/Engine/Source/Runtime/Engine/Private/SkeletalRender.cpp
@@ -28,6 +28,8 @@ static bool IsSkeletalMeshRayTracingSupported()
 }
 #endif // RHI_RAYTRACING
 
+static TAutoConsoleVariable<bool> CVarSkeletalMeshClothBlendEnabled(TEXT("r.SkeletalMeshClothBlend.Enabled"), true, TEXT("Enable the use of the cloth blend weight value set by the skeletal mesh component. When disabled all cloth blend weight will become 0."));
+
 /*-----------------------------------------------------------------------------
 FSkeletalMeshObject
 -----------------------------------------------------------------------------*/
@@ -444,3 +446,8 @@ void UpdatePreviousRefToLocalMatrices(TArray<FMatrix44f>& ReferenceToLocal, cons
 	UpdateRefToLocalMatricesInner(ReferenceToLocal, ComponentTransform, BoneVisibilityStates, (bIsMasterCompValid) ? &MasterBoneMap : nullptr, RefBasesInvMatrix, RefSkeleton, InSkeletalMeshRenderData, LODIndex, ExtraRequiredBoneIndices);
 }
 
+bool IsSkeletalMeshClothBlendEnabled()
+{
+	return CVarSkeletalMeshClothBlendEnabled.GetValueOnAnyThread();
+}
+
diff --git a/Engine/Source/Runtime/Engine/Private/SkeletalRender.h b/Engine/Source/Runtime/Engine/Private/SkeletalRender.h
--- a/Engine/Source/Runtime/Engine/Private/SkeletalRender.h
+++ b/Engine/Source/Runtime/Engine/Private/SkeletalRender.h
@@ -105,4 +105,9 @@ static FORCEINLINE void Pack4(VectorRegister4Double Normal, uint32* PackedNormal
 	VectorStoreSignedByte4(MakeVectorRegisterFloatFromDouble(Normal), PackedNormal);
 }
 
+/**
+* Returns true if the cloth blend weight defined in the SkeletalMeshComponent should be used by skin rendereding.
+*/
+bool IsSkeletalMeshClothBlendEnabled();
+
 
diff --git a/Engine/Source/Runtime/Engine/Private/SkeletalRenderCPUSkin.cpp b/Engine/Source/Runtime/Engine/Private/SkeletalRenderCPUSkin.cpp
--- a/Engine/Source/Runtime/Engine/Private/SkeletalRenderCPUSkin.cpp
+++ b/Engine/Source/Runtime/Engine/Private/SkeletalRenderCPUSkin.cpp
@@ -1205,7 +1205,7 @@ bool FDynamicSkelMeshObjectDataCPUSkin::UpdateClothSimulationData(USkinnedMeshCo
 		}
 
 		WorldToLocal = SrcComponent->GetRenderMatrix().InverseFast();
-		ClothBlendWeight = SrcComponent->ClothBlendWeight;
+		ClothBlendWeight = IsSkeletalMeshClothBlendEnabled() ? SrcComponent->ClothBlendWeight : 0.0f;
 		SimMeshComponent->GetUpdateClothSimulationData(ClothSimulUpdateData, SrcComponent);
 
 		return true;
@@ -1214,7 +1214,7 @@ bool FDynamicSkelMeshObjectDataCPUSkin::UpdateClothSimulationData(USkinnedMeshCo
 	if (SimMeshComponent)
 	{
 		WorldToLocal = SimMeshComponent->GetRenderMatrix().InverseFast();
-		ClothBlendWeight = SimMeshComponent->ClothBlendWeight;
+		ClothBlendWeight = IsSkeletalMeshClothBlendEnabled() ? SimMeshComponent->ClothBlendWeight : 0.0f;
 		SimMeshComponent->GetUpdateClothSimulationData(ClothSimulUpdateData);
 		return true;
 	}
diff --git a/Engine/Source/Runtime/Engine/Private/SkeletalRenderGPUSkin.cpp b/Engine/Source/Runtime/Engine/Private/SkeletalRenderGPUSkin.cpp
--- a/Engine/Source/Runtime/Engine/Private/SkeletalRenderGPUSkin.cpp
+++ b/Engine/Source/Runtime/Engine/Private/SkeletalRenderGPUSkin.cpp
@@ -2340,7 +2340,7 @@ bool FDynamicSkelMeshObjectDataGPUSkin::UpdateClothSimulationData(USkinnedMeshCo
 		}
 
 		ClothObjectLocalToWorld = SrcComponent->GetComponentToWorld().ToMatrixWithScale();
-		ClothBlendWeight = SimMeshComponent->ClothBlendWeight;
+		ClothBlendWeight = IsSkeletalMeshClothBlendEnabled() ? SimMeshComponent->ClothBlendWeight : 0.0f;
 		ClothingSimData = SimMeshComponent->GetCurrentClothingData_AnyThread();
 		return true;
 	}
@@ -2355,7 +2355,7 @@ bool FDynamicSkelMeshObjectDataGPUSkin::UpdateClothSimulationData(USkinnedMeshCo
 		}
 		else
 		{
-			ClothBlendWeight = SimMeshComponent->ClothBlendWeight;
+			ClothBlendWeight = IsSkeletalMeshClothBlendEnabled() ? SimMeshComponent->ClothBlendWeight : 0.0f;
 			ClothingSimData = SimMeshComponent->GetCurrentClothingData_AnyThread();
 		}
 
