Fix missing static mesh property change notification after copy-paste

#rnx

diff --git a/Engine/Source/Runtime/Engine/Private/Components/StaticMeshComponent.cpp b/Engine/Source/Runtime/Engine/Private/Components/StaticMeshComponent.cpp
--- a/Engine/Source/Runtime/Engine/Private/Components/StaticMeshComponent.cpp
+++ b/Engine/Source/Runtime/Engine/Private/Components/StaticMeshComponent.cpp
@@ -1679,6 +1679,10 @@ void UStaticMeshComponent::ExportCustomProperties(FOutputDevice& Out, uint32 Ind
 
 void UStaticMeshComponent::ImportCustomProperties(const TCHAR* SourceText, FFeedbackContext* Warn)
 {
+	// First thing that should be done after importing properties is to
+	// make sure notification is sent if the static mesh property has been modified.
+	NotifyIfStaticMeshChanged();
+
 	check(SourceText);
 	check(Warn);
 
