World Partition: remove a warning that happens when the world has streaming disabled.

#rnx

diff --git a/Engine/Source/Runtime/Engine/Private/WorldPartition/RuntimeSpatialHash/RuntimeSpatialHashGridHelper.cpp b/Engine/Source/Runtime/Engine/Private/WorldPartition/RuntimeSpatialHash/RuntimeSpatialHashGridHelper.cpp
--- a/Engine/Source/Runtime/Engine/Private/WorldPartition/RuntimeSpatialHash/RuntimeSpatialHashGridHelper.cpp
+++ b/Engine/Source/Runtime/Engine/Private/WorldPartition/RuntimeSpatialHash/RuntimeSpatialHashGridHelper.cpp
@@ -29,27 +29,24 @@ FSquare2DGridHelper::FSquare2DGridHelper(const FBox& InWorldBounds, const FVecto
 	, CellSize(InCellSize)
 {
 	// Compute Grid's size and level count based on World bounds
-	float WorldBoundsMaxExtent = 0.f;
+	int32 GridSize = 1;
+	int32 GridLevelCount = 1;
+
 	if (WorldBounds.IsValid)
 	{
 		const FVector2D DistMin = FMath::Abs(FVector2D(WorldBounds.Min - Origin));
 		const FVector2D DistMax = FMath::Abs(FVector2D(WorldBounds.Max - Origin));
-		WorldBoundsMaxExtent = FMath::Max(DistMin.GetMax(), DistMax.GetMax());
-	}
-	int32 GridSize = 1;
-	int32 GridLevelCount = 1;
-	if (WorldBoundsMaxExtent > 0.f)
-	{
-		GridSize = 2.f * FMath::CeilToFloat(WorldBoundsMaxExtent / CellSize); 
-		if (!FMath::IsPowerOfTwo(GridSize))
+		const float WorldBoundsMaxExtent = FMath::Max(DistMin.GetMax(), DistMax.GetMax());
+
+		if (WorldBoundsMaxExtent > 0.f)
 		{
-			GridSize = FMath::Pow(2.f, FMath::CeilToFloat(FMath::Log2(static_cast<float>(GridSize))));
+			GridSize = 2.f * FMath::CeilToFloat(WorldBoundsMaxExtent / CellSize); 
+			if (!FMath::IsPowerOfTwo(GridSize))
+			{
+				GridSize = FMath::Pow(2.f, FMath::CeilToFloat(FMath::Log2(static_cast<float>(GridSize))));
+			}
+			GridLevelCount = FMath::FloorLog2(GridSize) + 1;
 		}
-		GridLevelCount = FMath::FloorLog2(GridSize) + 1;
-	}
-	else
-	{
-		UE_LOG(LogWorldPartition, Warning, TEXT("Invalid world bounds, grid partitioning will use a runtime grid with 1 cell."));
 	}
 
 	check(FMath::IsPowerOfTwo(GridSize));
