Add : to filename character exceptions when creating shader debug path.

diff --git a/Engine/Source/Runtime/Engine/Private/ShaderCompiler/ShaderCompiler.cpp b/Engine/Source/Runtime/Engine/Private/ShaderCompiler/ShaderCompiler.cpp
--- a/Engine/Source/Runtime/Engine/Private/ShaderCompiler/ShaderCompiler.cpp
+++ b/Engine/Source/Runtime/Engine/Private/ShaderCompiler/ShaderCompiler.cpp
@@ -3743,6 +3743,7 @@ FString FShaderCompilingManager::CreateShaderDebugInfoPath(const FShaderCompiler
 	DumpDebugInfoPath.ReplaceInline(TEXT(">"), TEXT(")"));
 	DumpDebugInfoPath.ReplaceInline(TEXT("::"), TEXT("=="));
 	DumpDebugInfoPath.ReplaceInline(TEXT("|"), TEXT("_"));
+	DumpDebugInfoPath.ReplaceInline(TEXT(":"), TEXT("_"));
 	DumpDebugInfoPath.ReplaceInline(TEXT("*"), TEXT("-"));
 	DumpDebugInfoPath.ReplaceInline(TEXT("?"), TEXT("!"));
 	DumpDebugInfoPath.ReplaceInline(TEXT("\""), TEXT("\'"));
