Avoid updating the entire anim graph from the game thread when using root motion from montage only

diff --git a/Engine/Source/Runtime/Engine/Classes/Animation/AnimInstance.h b/Engine/Source/Runtime/Engine/Classes/Animation/AnimInstance.h
--- a/Engine/Source/Runtime/Engine/Classes/Animation/AnimInstance.h
+++ b/Engine/Source/Runtime/Engine/Classes/Animation/AnimInstance.h
@@ -529,7 +529,7 @@ public:
 
 private:
 	// Does this anim instance need immediate update (rather than parallel)?
-	bool NeedsImmediateUpdate(float DeltaSeconds) const;
+	bool NeedsImmediateUpdate(float DeltaSeconds, bool bNeedsValidRootMotion) const;
 
 public:
 	/** Returns the owning actor of this AnimInstance */
diff --git a/Engine/Source/Runtime/Engine/Private/Animation/AnimInstance.cpp b/Engine/Source/Runtime/Engine/Private/Animation/AnimInstance.cpp
--- a/Engine/Source/Runtime/Engine/Private/Animation/AnimInstance.cpp
+++ b/Engine/Source/Runtime/Engine/Private/Animation/AnimInstance.cpp
@@ -536,7 +536,7 @@ void UAnimInstance::UpdateAnimation(float DeltaSeconds, bool bNeedsValidRootMoti
 	}
 
 	// Determine whether or not the animation should be immediately updated according to current state
-	const bool bWantsImmediateUpdate = bNeedsValidRootMotion || NeedsImmediateUpdate(DeltaSeconds);
+	const bool bWantsImmediateUpdate = NeedsImmediateUpdate(DeltaSeconds, bNeedsValidRootMotion);
 
 	// Determine whether or not we can or should actually immediately update the animation state
 	bool bShouldImmediateUpdate = bWantsImmediateUpdate;
@@ -700,18 +700,18 @@ void UAnimInstance::ParallelUpdateAnimation()
 	GetProxyOnAnyThread<FAnimInstanceProxy>().UpdateAnimation();
 }
 
-bool UAnimInstance::NeedsImmediateUpdate(float DeltaSeconds) const
+bool UAnimInstance::NeedsImmediateUpdate(float DeltaSeconds, bool bNeedsValidRootMotion) const
 {
 	const bool bUseParallelUpdateAnimation = (GetDefault<UEngine>()->bAllowMultiThreadedAnimationUpdate && bUseMultiThreadedAnimationUpdate) || (CVarForceUseParallelAnimUpdate.GetValueOnGameThread() != 0);
 
 	return
+		(bNeedsValidRootMotion && RootMotionMode == ERootMotionMode::RootMotionFromEverything) ||
 		!CanRunParallelWork() ||
 		GIntraFrameDebuggingGameThread ||
 		CVarUseParallelAnimUpdate.GetValueOnGameThread() == 0 ||
 		CVarUseParallelAnimationEvaluation.GetValueOnGameThread() == 0 ||
 		!bUseParallelUpdateAnimation ||
-		DeltaSeconds == 0.0f ||
-		RootMotionMode == ERootMotionMode::RootMotionFromEverything;
+		DeltaSeconds == 0.0f;
 }
 
 bool UAnimInstance::NeedsUpdate() const
