Non unity build fix

Since both UMaterial and UMaterialInterface were both forward declared, so assignation was failing

diff --git a/Engine/Source/Runtime/Engine/Public/WorldPartition/HLOD/HLODLayer.h b/Engine/Source/Runtime/Engine/Public/WorldPartition/HLOD/HLODLayer.h
--- a/Engine/Source/Runtime/Engine/Public/WorldPartition/HLOD/HLODLayer.h
+++ b/Engine/Source/Runtime/Engine/Public/WorldPartition/HLOD/HLODLayer.h
@@ -115,6 +115,6 @@ public:
 	UPROPERTY()
 	FMeshApproximationSettings MeshApproximationSettings_DEPRECATED;
 	UPROPERTY()
-	TSoftObjectPtr<UMaterial> HLODMaterial_DEPRECATED;
+	TSoftObjectPtr<UMaterialInterface> HLODMaterial_DEPRECATED;
 #endif
 };
