Allow for shader compilation retry when the following conditions are true:

- r.ShaderDevelopmentMode=1
- bSpecialEngineMaterial == true

- This allows you to retry compilation of special engine materials, like WorldGridMaterial, even if you have bPromptToRetryFailedShaderCompiles=False.

diff --git a/Engine/Source/Runtime/Engine/Private/ShaderCompiler/ShaderCompiler.cpp b/Engine/Source/Runtime/Engine/Private/ShaderCompiler/ShaderCompiler.cpp
--- a/Engine/Source/Runtime/Engine/Private/ShaderCompiler/ShaderCompiler.cpp
+++ b/Engine/Source/Runtime/Engine/Private/ShaderCompiler/ShaderCompiler.cpp
@@ -4771,7 +4771,7 @@ bool FShaderCompilingManager::HandlePotentialRetryOnError(TMap<int32, FShaderMap
 					}
 				}
 
-				if (UE_LOG_ACTIVE(LogShaders, Log) && (bAnyErrorLikelyToBeCodeError || bPromptToRetryFailedShaderCompiles))
+				if (UE_LOG_ACTIVE(LogShaders, Log) && (bAnyErrorLikelyToBeCodeError || bPromptToRetryFailedShaderCompiles || bSpecialEngineMaterial))
 				{
 #if UE_BUILD_DEBUG
 					// Use debug break in debug with the debugger attached, otherwise message box
