Add PreUpdateLinkedInstances function to UAninInstance API. Called by the skeletal mesh before linked animation instances are updated

#Tests PIE

diff --git a/Engine/Source/Runtime/Engine/Classes/Animation/AnimInstance.h b/Engine/Source/Runtime/Engine/Classes/Animation/AnimInstance.h
--- a/Engine/Source/Runtime/Engine/Classes/Animation/AnimInstance.h
+++ b/Engine/Source/Runtime/Engine/Classes/Animation/AnimInstance.h
@@ -1215,6 +1215,9 @@ public:
 	// @param	bInDeferRootNodeInitialization	When set to true, defer init of the blend tree until the first Update() call
 	void InitializeAnimation(bool bInDeferRootNodeInitialization = false);
 
+	/** Called on the game thread before UpdateAnimation is called on linked instances */
+	virtual void PreUpdateLinkedInstances(float DeltaSeconds) {}
+
 	/** Update Animation code-paths, updates and advances animation state, returns whether or not the actual update should have been called immediately */
 	void UpdateAnimation(float DeltaSeconds, bool bNeedsValidRootMotion, EUpdateAnimationFlag UpdateFlag = EUpdateAnimationFlag::Default );
 
diff --git a/Engine/Source/Runtime/Engine/Private/Components/SkeletalMeshComponent.cpp b/Engine/Source/Runtime/Engine/Private/Components/SkeletalMeshComponent.cpp
--- a/Engine/Source/Runtime/Engine/Private/Components/SkeletalMeshComponent.cpp
+++ b/Engine/Source/Runtime/Engine/Private/Components/SkeletalMeshComponent.cpp
@@ -1210,6 +1210,9 @@ void USkeletalMeshComponent::SetPredictedLODLevel(int32 InPredictedLODLevel)
 
 void USkeletalMeshComponent::TickAnimInstances(float DeltaTime, bool bNeedsValidRootMotion)
 {
+	// Allow animation instance to do some processing before the linked instances update
+	AnimScriptInstance->PreUpdateLinkedInstances(DeltaTime);
+
 	// We update linked instances first incase we're using either root motion or non-threaded update.
 	// This ensures that we go through the pre update process and initialize the proxies correctly.
 	for (UAnimInstance* LinkedInstance : LinkedInstances)
