Fixed incorrect timestamp check, preventing disabling of challenge ack resend code path serverside. Fix submitted by [REDACTED_PERSON].

diff --git a/Engine/Source/Runtime/Engine/Private/PacketHandlers/StatelessConnectHandlerComponent.cpp b/Engine/Source/Runtime/Engine/Private/PacketHandlers/StatelessConnectHandlerComponent.cpp
--- a/Engine/Source/Runtime/Engine/Private/PacketHandlers/StatelessConnectHandlerComponent.cpp
+++ b/Engine/Source/Runtime/Engine/Private/PacketHandlers/StatelessConnectHandlerComponent.cpp
@@ -964,7 +964,7 @@ void StatelessConnectHandlerComponent::Incoming(FBitReader& Packet)
 		const double RestartHandshakeAckResendWindow = 10.0;
 		double CurTime = Driver != nullptr ? Driver->GetElapsedTime() : 0.0;
 
-		if (LastInitTimestamp - CurTime >= RestartHandshakeAckResendWindow)
+		if (CurTime - LastInitTimestamp >= RestartHandshakeAckResendWindow)
 		{
 			LastChallengeSuccessAddress.Reset();
 			LastInitTimestamp = 0.0;
