Don't create destruction infos for levels during replay playback

diff --git a/Engine/Source/Runtime/Engine/Private/LevelActor.cpp b/Engine/Source/Runtime/Engine/Private/LevelActor.cpp
--- a/Engine/Source/Runtime/Engine/Private/LevelActor.cpp
+++ b/Engine/Source/Runtime/Engine/Private/LevelActor.cpp
@@ -38,6 +38,7 @@
 #include "Misc/MapErrors.h"
 #include "GameFramework/WorldSettings.h"
 #include "Engine/NetDriver.h"
+#include "Engine/DemoNetDriver.h"
 #include "Engine/Player.h"
 #include "AssetRegistryModule.h"
 
@@ -919,10 +920,24 @@ bool UWorld::DestroyActor( AActor* ThisActor, bool bNetForce, bool bShouldModify
 		ThisActor->SetOwner(NULL);
 	}
 
-	ULevel* const ActorLevel = ThisActor->GetLevel();
-	if (ActorLevel)
+	if (ULevel* const ActorLevel = ThisActor->GetLevel())
 	{
-		ActorLevel->CreateReplicatedDestructionInfo(ThisActor);
+		UDemoNetDriver* DemoDriver = nullptr;
+
+		if (const FLevelCollection* LevelCollection = ActorLevel->GetCachedLevelCollection())
+		{
+			DemoDriver = LevelCollection->GetDemoNetDriver();
+		}
+
+		if (!DemoDriver)
+		{
+			DemoDriver = GetDemoNetDriver();
+		}
+
+		if (!DemoDriver || !DemoDriver->IsPlaying())
+		{
+			ActorLevel->CreateReplicatedDestructionInfo(ThisActor);
+		}
 	}
 
 	// Notify net drivers that this guy has been destroyed.
