Check AnimScriptInstance pointer before calling PreUpdateLinkedInstances

#rnx

diff --git a/Engine/Source/Runtime/Engine/Private/Components/SkeletalMeshComponent.cpp b/Engine/Source/Runtime/Engine/Private/Components/SkeletalMeshComponent.cpp
--- a/Engine/Source/Runtime/Engine/Private/Components/SkeletalMeshComponent.cpp
+++ b/Engine/Source/Runtime/Engine/Private/Components/SkeletalMeshComponent.cpp
@@ -1211,7 +1211,10 @@ void USkeletalMeshComponent::SetPredictedLODLevel(int32 InPredictedLODLevel)
 void USkeletalMeshComponent::TickAnimInstances(float DeltaTime, bool bNeedsValidRootMotion)
 {
 	// Allow animation instance to do some processing before the linked instances update
-	AnimScriptInstance->PreUpdateLinkedInstances(DeltaTime);
+	if (AnimScriptInstance != nullptr)
+	{
+		AnimScriptInstance->PreUpdateLinkedInstances(DeltaTime);
+	}
 
 	// We update linked instances first incase we're using either root motion or non-threaded update.
 	// This ensures that we go through the pre update process and initialize the proxies correctly.
