Fix missing declaration

diff --git a/Engine/Source/Runtime/Engine/Classes/Engine/ActorChannel.h b/Engine/Source/Runtime/Engine/Classes/Engine/ActorChannel.h
--- a/Engine/Source/Runtime/Engine/Classes/Engine/ActorChannel.h
+++ b/Engine/Source/Runtime/Engine/Classes/Engine/ActorChannel.h
@@ -22,6 +22,8 @@ class FNetFieldExportGroup;
 class FOutBunch;
 class UNetConnection;
 
+namespace UE::Net { struct FSubObjectRegistry;  }
+
 enum class ESetChannelActorFlags : uint32
 {
 	None					= 0,
