Improve BuildHiddenComponentList performance by ~2ms in split-screen by expanding the hidden actor's components directly instead of doing it every frame. This is still slower than having the component IDs pre-baked and just copied over every frame, but a deeper engine solution is complicated and this one should work well enough.

diff --git a/Engine/Source/Runtime/Engine/Private/PlayerController.cpp b/Engine/Source/Runtime/Engine/Private/PlayerController.cpp
--- a/Engine/Source/Runtime/Engine/Private/PlayerController.cpp
+++ b/Engine/Source/Runtime/Engine/Private/PlayerController.cpp
@@ -5700,12 +5700,11 @@ void APlayerController::BuildHiddenComponentList(const FVector& ViewLocation, TS
 	}
 
 	// iterate backwards so we can remove as we go
+	HiddenComponentsOut.Reserve(HiddenComponentsOut.Num() + HiddenPrimitiveComponents.Num());
 	for (int32 ComponentIndx = HiddenPrimitiveComponents.Num() - 1; ComponentIndx >= 0; --ComponentIndx)
 	{
-		TWeakObjectPtr<UPrimitiveComponent> ComponentPtr = HiddenPrimitiveComponents[ComponentIndx];
-		if (ComponentPtr.IsValid())
+		if (UPrimitiveComponent* Component = HiddenPrimitiveComponents[ComponentIndx].Get())
 		{
-			UPrimitiveComponent* Component = ComponentPtr.Get();
 			if (Component->IsRegistered())
 			{
 				HiddenComponentsOut.Add(Component->ComponentId);
@@ -5713,7 +5712,7 @@ void APlayerController::BuildHiddenComponentList(const FVector& ViewLocation, TS
 		}
 		else
 		{
-			HiddenPrimitiveComponents.RemoveAt(ComponentIndx);
+			HiddenPrimitiveComponents.RemoveAtSwap(ComponentIndx);
 		}
 	}
 
