Fixed Static Analysis

diff --git a/Engine/Source/Runtime/Engine/Private/WorldPartition/DataLayer/DataLayerUtils.cpp b/Engine/Source/Runtime/Engine/Private/WorldPartition/DataLayer/DataLayerUtils.cpp
--- a/Engine/Source/Runtime/Engine/Private/WorldPartition/DataLayer/DataLayerUtils.cpp
+++ b/Engine/Source/Runtime/Engine/Private/WorldPartition/DataLayer/DataLayerUtils.cpp
@@ -31,7 +31,7 @@ FWorldDataLayersActorDesc* FDataLayerUtils::GetWorldDataLayersActorDesc(const UA
 
 TArray<FName> FDataLayerUtils::ResolvedDataLayerInstanceNames(const FWorldPartitionActorDesc* InActorDesc, const AWorldDataLayers* InWorldDataLayers, const FWorldDataLayersActorDesc* InWorldDataLayersActorDesc, bool* bOutIsResultValid)
 {
-	bool bLocalIsSuccess;
+	bool bLocalIsSuccess = true;
 	bool& bIsSuccess = bOutIsResultValid ? *bOutIsResultValid : bLocalIsSuccess;
 	bIsSuccess = true;
 
@@ -116,7 +116,7 @@ TArray<FName> FDataLayerUtils::ResolvedDataLayerInstanceNames(const FWorldPartit
 // For performance reasons, this function assumes that InActorDesc's DataLayerInstanceNames was already resolved.
 TArray<FName> FDataLayerUtils::ResolveRuntimeDataLayerInstanceNames(const FWorldPartitionActorDesc* InActorDesc, const UActorDescContainer* InContainer, bool* bOutIsResultValid)
 {
-	bool bLocalIsSuccess;
+	bool bLocalIsSuccess = false;
 	bool& bIsSuccess = bOutIsResultValid ? *bOutIsResultValid : bLocalIsSuccess;
 	bIsSuccess = false;
 
