Generate if{}else{} control flow rather than ?: for material switch expressions
Add non-static bool MCT type to allow plugging bool expressions into switch nodes when using new HLSL translator

diff --git a/Engine/Source/Editor/UnrealEd/Private/MaterialGraph.cpp b/Engine/Source/Editor/UnrealEd/Private/MaterialGraph.cpp
--- a/Engine/Source/Editor/UnrealEd/Private/MaterialGraph.cpp
+++ b/Engine/Source/Editor/UnrealEd/Private/MaterialGraph.cpp
@@ -127,7 +127,10 @@ UMaterialGraphNode_Base* FindGraphNodeForObject(const UObject* Object)
 {
 	if (const UMaterialExpression* MaterialExpression = Cast<UMaterialExpression>(Object))
 	{
-		return CastChecked<UMaterialGraphNode_Base>(MaterialExpression->GraphNode);
+		if (MaterialExpression->GraphNode)
+		{
+			return CastChecked<UMaterialGraphNode_Base>(MaterialExpression->GraphNode);
+		}
 	}
 	else if (const UMaterial* Material = Cast<UMaterial>(Object))
 	{
diff --git a/Engine/Source/Runtime/Engine/Private/HLSLTree/HLSLTree.cpp b/Engine/Source/Runtime/Engine/Private/HLSLTree/HLSLTree.cpp
--- a/Engine/Source/Runtime/Engine/Private/HLSLTree/HLSLTree.cpp
+++ b/Engine/Source/Runtime/Engine/Private/HLSLTree/HLSLTree.cpp
@@ -1261,6 +1261,16 @@ Shader::FType FExpression::GetValuePreshader(FEmitContext& Context, FEmitScope&
 	return GetValuePreshader(Context, Scope, RequestedType, PreparedType, PreparedType.GetType(), OutPreshader);
 }
 
+Shader::FType FExpression::GetValuePreshader(FEmitContext& Context, FEmitScope& Scope, const Shader::FType& ResultType, Shader::FPreshaderData& OutPreshader) const
+{
+	return GetValuePreshader(Context, Scope, ResultType, ResultType, OutPreshader);
+}
+
+Shader::FType FExpression::GetValuePreshader(FEmitContext& Context, FEmitScope& Scope, Shader::EValueType ResultType, Shader::FPreshaderData& OutPreshader) const\
+{
+	return GetValuePreshader(Context, Scope, ResultType, ResultType, OutPreshader);
+}
+
 Shader::FValue FExpression::GetValueConstant(FEmitContext& Context, FEmitScope& Scope, const FRequestedType& RequestedType, const FPreparedType& PreparedType, const Shader::FType& ResultType) const
 {
 	FEmitOwnerScope OwnerScope(Context, this);
diff --git a/Engine/Source/Runtime/Engine/Private/HLSLTree/HLSLTreeCommon.cpp b/Engine/Source/Runtime/Engine/Private/HLSLTree/HLSLTreeCommon.cpp
--- a/Engine/Source/Runtime/Engine/Private/HLSLTree/HLSLTreeCommon.cpp
+++ b/Engine/Source/Runtime/Engine/Private/HLSLTree/HLSLTreeCommon.cpp
@@ -780,10 +780,7 @@ bool FExpressionSwitchBase::PrepareValue(FEmitContext& Context, FEmitScope& Scop
 		if (IsInputActive(Context, InputIndex))
 		{
 			const FPreparedType& InputType = Context.PrepareExpression(Input[InputIndex], Scope, RequestedType);
-			if (ResultType.IsVoid() && !InputType.IsVoid())
-			{
-				ResultType = InputType;
-			}
+			ResultType = MergePreparedTypes(ResultType, InputType);
 		}
 	}
 
diff --git a/Engine/Source/Runtime/Engine/Private/HLSLTree/HLSLTreeEmit.cpp b/Engine/Source/Runtime/Engine/Private/HLSLTree/HLSLTreeEmit.cpp
--- a/Engine/Source/Runtime/Engine/Private/HLSLTree/HLSLTreeEmit.cpp
+++ b/Engine/Source/Runtime/Engine/Private/HLSLTree/HLSLTreeEmit.cpp
@@ -400,31 +400,17 @@ FPreparedType FEmitContext::PrepareExpression(const FExpression* InExpression, F
 		PrepareValueMap.Add(InExpression, Result);
 	}
 
-	if (bMarkLiveValues)
-	{
-		if (Result->bPreparingValue)
-		{
-			return Result->PreparedType;
-		}
-
-		const EExpressionEvaluation Evaluation = Result->PreparedType.GetEvaluation(Scope, RequestedType);
-		if (Evaluation == EExpressionEvaluation::None || IsConstantEvaluation(Evaluation))
-		{
-			// Don't need to prepare constant values, since we are just replacing the expression with a constant
-			// Can't set the 'bIsLive' flag here, since expression may not be constant in all contexts
-			return Result->PreparedType;
-		}
-	}
-
 	if (Result->bPreparingValue)
 	{
 		// Valid for this to be called reentrantly
 		// Code should ensure that the type is set before the reentrant call, otherwise type will not be valid here
 		// LocalPHI nodes rely on this to break loops
-		check(!bMarkLiveValues);
-		FEmitScope* LoopScope = Scope.FindLoop();
-		check(LoopScope);
-		Result->PreparedType.SetLoopEvaluation(*LoopScope, RequestedType);
+		if (!bMarkLiveValues)
+		{
+			FEmitScope* LoopScope = Scope.FindLoop();
+			check(LoopScope);
+			Result->PreparedType.SetLoopEvaluation(*LoopScope, RequestedType);
+		}
 		return Result->PreparedType;
 	}
 
@@ -774,11 +760,6 @@ FEmitShaderExpression* FEmitContext::InternalEmitExpression(FEmitScope& Scope, T
 {
 	FEmitShaderExpression* ShaderValue = nullptr;
 
-	if (Code.Contains(TEXT("(Local253.xy - float3(0.00000000f, 0.00000000f, 1.00000000f))")))
-	{
-		int a = 0;
-	}
-
 	FXxHash64Builder Hasher;
 	Hasher.Update(Code.GetData(), Code.Len() * sizeof(TCHAR));
 	if (bInline)
@@ -795,6 +776,27 @@ FEmitShaderExpression* FEmitContext::InternalEmitExpression(FEmitScope& Scope, T
 		ShaderValue = *PrevShaderValue;
 		check(ShaderValue->Type == Type);
 		Private::MoveToScope(ShaderValue, Scope);
+
+		// Check to see if we have any new dependencies to add to the previous expression
+		TArray<FEmitShaderNode*, TInlineAllocator<32>> NewDependencies;
+		bool bHasNewDependencies = false;
+		for (FEmitShaderNode* Dependency : Dependencies)
+		{
+			if (!ShaderValue->Dependencies.Contains(Dependency))
+			{
+				if (!bHasNewDependencies)
+				{
+					NewDependencies.Reserve(ShaderValue->Dependencies.Num() + Dependencies.Num());
+					NewDependencies.Append(ShaderValue->Dependencies.GetData(), ShaderValue->Dependencies.Num());
+					bHasNewDependencies = true;
+				}
+				NewDependencies.Add(Dependency);
+			}
+		}
+		if (bHasNewDependencies)
+		{
+			ShaderValue->Dependencies = MemStack::AllocateArrayView(*Allocator, MakeArrayView(NewDependencies));
+		}
 	}
 	else
 	{
diff --git a/Engine/Source/Runtime/Engine/Private/HLSLTree/HLSLTreeOperation.cpp b/Engine/Source/Runtime/Engine/Private/HLSLTree/HLSLTreeOperation.cpp
--- a/Engine/Source/Runtime/Engine/Private/HLSLTree/HLSLTreeOperation.cpp
+++ b/Engine/Source/Runtime/Engine/Private/HLSLTree/HLSLTreeOperation.cpp
@@ -548,9 +548,19 @@ bool FExpressionOperation::PrepareValue(FEmitContext& Context, FEmitScope& Scope
 	FPreparedType InputPreparedType[MaxInputs];
 	Shader::FValue ConstantInput[MaxInputs];
 	bool bConstantZeroInput[MaxInputs] = { false };
+	bool bMarkLiveInput[MaxInputs] = { false };
 	for (int32 Index = 0; Index < OpDesc.NumInputs; ++Index)
 	{
-		InputPreparedType[Index] = Context.PrepareExpression(Inputs[Index], Scope, RequestedTypes.InputType[Index]);
+		if (Context.bMarkLiveValues)
+		{
+			InputPreparedType[Index] = Context.GetPreparedType(Inputs[Index]);
+			bMarkLiveInput[Index] = true;
+		}
+		else
+		{
+			InputPreparedType[Index] = Context.PrepareExpression(Inputs[Index], Scope, RequestedTypes.InputType[Index]);
+		}
+
 		if (InputPreparedType[Index].IsVoid())
 		{
 			return false;
@@ -581,6 +591,7 @@ bool FExpressionOperation::PrepareValue(FEmitContext& Context, FEmitScope& Scope
 		Context.MarkInputType(Inputs[Index], Types.InputType[Index]);
 	}
 
+	bool bMarkInputsLive = Context.bMarkLiveValues;
 	switch (Op)
 	{
 	case EOperation::Mul:
@@ -588,12 +599,33 @@ bool FExpressionOperation::PrepareValue(FEmitContext& Context, FEmitScope& Scope
 		{
 			// X * 0 == 0
 			Types.ResultType.SetEvaluation(EExpressionEvaluation::ConstantZero);
+			// If one input is 0, we can avoid marking the other input as live (since we don't care what it is)
+			// In the case 'both' inputs are 0, we still need to mark one of them live...in this case we could potentially somehow choose the simplest input to mark live, but for now just pick one
+			if (bConstantZeroInput[0])
+			{
+				bMarkLiveInput[1] = false;
+			}
+			else
+			{
+				bMarkLiveInput[0] = false;
+			}
 		}
 		break;
 	default:
 		break;
 	}
 
+	if (Context.bMarkLiveValues)
+	{
+		for (int32 Index = 0; Index < OpDesc.NumInputs; ++Index)
+		{
+			if (bMarkLiveInput[Index])
+			{
+				Context.PrepareExpression(Inputs[Index], Scope, RequestedTypes.InputType[Index]);
+			}
+		}
+	}
+
 	return OutResult.SetType(Context, RequestedType, Types.ResultType);
 }
 
diff --git a/Engine/Source/Runtime/Engine/Private/Materials/HLSLMaterialTranslator.cpp b/Engine/Source/Runtime/Engine/Private/Materials/HLSLMaterialTranslator.cpp
--- a/Engine/Source/Runtime/Engine/Private/Materials/HLSLMaterialTranslator.cpp
+++ b/Engine/Source/Runtime/Engine/Private/Materials/HLSLMaterialTranslator.cpp
@@ -2910,6 +2910,7 @@ const TCHAR* FHLSLMaterialTranslator::DescribeType(EMaterialValueType Type) cons
 	case MCT_TextureCubeArray:		return TEXT("textureCubeArray");
 	case MCT_VolumeTexture:			return TEXT("volumeTexture");
 	case MCT_StaticBool:			return TEXT("static bool");
+	case MCT_Bool:					return TEXT("bool");
 	case MCT_MaterialAttributes:	return TEXT("MaterialAttributes");
 	case MCT_TextureExternal:		return TEXT("TextureExternal");
 	case MCT_TextureVirtual:		return TEXT("TextureVirtual");
@@ -2940,6 +2941,7 @@ const TCHAR* FHLSLMaterialTranslator::HLSLTypeString(EMaterialValueType Type) co
 	case MCT_TextureCubeArray:		return TEXT("textureCubeArray");
 	case MCT_VolumeTexture:			return TEXT("volumeTexture");
 	case MCT_StaticBool:			return TEXT("static bool");
+	case MCT_Bool:					return TEXT("bool");
 	case MCT_MaterialAttributes:	return TEXT("FMaterialAttributes");
 	case MCT_TextureExternal:		return TEXT("TextureExternal");
 	case MCT_TextureVirtual:		return TEXT("TextureVirtual");
@@ -2970,6 +2972,7 @@ const TCHAR* FHLSLMaterialTranslator::HLSLTypeStringDeriv(EMaterialValueType Typ
 	case MCT_TextureCubeArray:		return TEXT("textureCubeArray");
 	case MCT_VolumeTexture:			return TEXT("volumeTexture");
 	case MCT_StaticBool:			return TEXT("static bool");
+	case MCT_Bool:					return TEXT("bool");
 	case MCT_MaterialAttributes:	return TEXT("FMaterialAttributes");
 	case MCT_TextureExternal:		return TEXT("TextureExternal");
 	case MCT_TextureVirtual:		return TEXT("TextureVirtual");
diff --git a/Engine/Source/Runtime/Engine/Private/Materials/MaterialCachedHLSLTree.cpp b/Engine/Source/Runtime/Engine/Private/Materials/MaterialCachedHLSLTree.cpp
--- a/Engine/Source/Runtime/Engine/Private/Materials/MaterialCachedHLSLTree.cpp
+++ b/Engine/Source/Runtime/Engine/Private/Materials/MaterialCachedHLSLTree.cpp
@@ -39,6 +39,7 @@ static UE::Shader::EValueType GetShaderType(EMaterialValueType MaterialType)
 	case MCT_Float4: return UE::Shader::EValueType::Float4;
 	case MCT_Float: return UE::Shader::EValueType::Float1;
 	case MCT_StaticBool: return UE::Shader::EValueType::Bool1;
+	case MCT_Bool: return UE::Shader::EValueType::Bool1;
 	case MCT_MaterialAttributes: return UE::Shader::EValueType::Struct;
 	case MCT_ShadingModel: return UE::Shader::EValueType::Int1;
 	case MCT_LWCScalar: return UE::Shader::EValueType::Double1;
diff --git a/Engine/Source/Runtime/Engine/Private/Materials/MaterialExpressionHLSL.cpp b/Engine/Source/Runtime/Engine/Private/Materials/MaterialExpressionHLSL.cpp
--- a/Engine/Source/Runtime/Engine/Private/Materials/MaterialExpressionHLSL.cpp
+++ b/Engine/Source/Runtime/Engine/Private/Materials/MaterialExpressionHLSL.cpp
@@ -221,11 +221,17 @@ bool UMaterialExpressionShadingModel::GenerateHLSLExpression(FMaterialHLSLGenera
 
 bool UMaterialExpressionStaticSwitch::GenerateHLSLExpression(FMaterialHLSLGenerator& Generator, UE::HLSLTree::FScope& Scope, int32 OutputIndex, UE::HLSLTree::FExpression const*& OutExpression) const
 {
-	const UE::HLSLTree::FExpression* ConditionExpression = Value.AcquireHLSLExpressionOrConstant(Generator, Scope, (bool)DefaultValue);
-	const UE::HLSLTree::FExpression* TrueExpression = A.TryAcquireHLSLExpression(Generator, Scope);
-	const UE::HLSLTree::FExpression* FalseExpression = B.TryAcquireHLSLExpression(Generator, Scope);
+	using namespace UE::HLSLTree;
+
+	const FExpression* ConditionExpression = Value.AcquireHLSLExpressionOrConstant(Generator, Scope, (bool)DefaultValue);
+	const FExpression* TrueExpression = A.TryAcquireHLSLExpression(Generator, Scope);
+	const FExpression* FalseExpression = B.TryAcquireHLSLExpression(Generator, Scope);
+	if (!ConditionExpression)
+	{
+		return false;
+	}
 
-	OutExpression = Generator.GetTree().NewExpression<UE::HLSLTree::FExpressionSelect>(ConditionExpression, TrueExpression, FalseExpression);
+	OutExpression = Generator.GenerateBranch(Scope, ConditionExpression, TrueExpression, FalseExpression);
 	return true;
 }
 
@@ -371,8 +377,9 @@ bool UMaterialExpressionStaticSwitchParameter::GenerateHLSLExpression(FMaterialH
 		return false;
 	}
 
+	// No reason to generate a dynamic branch here, since the condition is always a static switch parameter
 	const FExpression* ExpressionSwitch = Generator.GenerateMaterialParameter(ParameterName, ParameterMeta);
-	OutExpression = Generator.GetTree().NewExpression <FExpressionSelect>(ExpressionSwitch, ExpressionA, ExpressionB);
+	OutExpression = Generator.GetTree().NewExpression<FExpressionSelect>(ExpressionSwitch, ExpressionA, ExpressionB);
 	return true;
 }
 
@@ -1966,7 +1973,7 @@ bool UMaterialExpressionIf::GenerateHLSLExpression(FMaterialHLSLGenerator& Gener
 	using namespace UE::Shader;
 
 	const FExpression* ExpressionA = A.AcquireHLSLExpression(Generator, Scope);
-	const FExpression* ExpressionB = A.AcquireHLSLExpressionOrConstant(Generator, Scope, ConstB);
+	const FExpression* ExpressionB = B.AcquireHLSLExpressionOrConstant(Generator, Scope, ConstB);
 	const FExpression* ExpressionAGreaterThanB = AGreaterThanB.AcquireHLSLExpression(Generator, Scope);
 	const FExpression* ExpressionALessThanB = ALessThanB.AcquireHLSLExpression(Generator, Scope);
 	if (!ExpressionA || !ExpressionB || !ExpressionAGreaterThanB || !ExpressionALessThanB)
@@ -1976,14 +1983,14 @@ bool UMaterialExpressionIf::GenerateHLSLExpression(FMaterialHLSLGenerator& Gener
 
 	UE::HLSLTree::FTree& Tree = Generator.GetTree();
 	const FExpression* ExpressionCondAGreaterEqualB = Tree.NewGreaterEqual(ExpressionA, ExpressionB);
-	OutExpression = Tree.NewExpression<FExpressionSelect>(ExpressionCondAGreaterEqualB, ExpressionAGreaterThanB, ExpressionALessThanB);
+	OutExpression = Generator.GenerateBranch(Scope, ExpressionCondAGreaterEqualB, ExpressionAGreaterThanB, ExpressionALessThanB);
 
 	const FExpression* ExpressionAEqualsB = AEqualsB.TryAcquireHLSLExpression(Generator, Scope);
 	if (ExpressionAEqualsB)
 	{
 		const FExpression* ExpressionThreshold = Generator.NewConstant(EqualsThreshold);
 		const FExpression* ExpressionCondANotEqualsB = Tree.NewGreater(Tree.NewAbs(Tree.NewSub(ExpressionA, ExpressionB)), ExpressionThreshold);
-		OutExpression = Tree.NewExpression<FExpressionSelect>(ExpressionCondANotEqualsB, OutExpression, ExpressionAEqualsB);
+		OutExpression = Generator.GenerateBranch(Scope, ExpressionCondANotEqualsB, OutExpression, ExpressionAEqualsB);
 	}
 
 	return true;
diff --git a/Engine/Source/Runtime/Engine/Private/Materials/MaterialExpressions.cpp b/Engine/Source/Runtime/Engine/Private/Materials/MaterialExpressions.cpp
--- a/Engine/Source/Runtime/Engine/Private/Materials/MaterialExpressions.cpp
+++ b/Engine/Source/Runtime/Engine/Private/Materials/MaterialExpressions.cpp
@@ -458,10 +458,15 @@ bool CanConnectMaterialValueTypes(const uint32 InputType, const uint32 OutputTyp
 	}
 	// Need to do more checks here to see whether types can be cast
 	// just check if both are float for now
-	if (InputType & MCT_Float && OutputType & MCT_Float)
+	if ((InputType & MCT_Numeric) && (OutputType & MCT_Numeric))
 	{
 		return true;
 	}
+	if (InputType == MCT_StaticBool && OutputType == MCT_Bool)
+	{
+		// StaticBool is allowed to connect to Bool (but not the other way around)
+		return true;
+	}
 	return false;
 }
 
@@ -8660,7 +8665,15 @@ uint32 UMaterialExpressionStaticSwitch::GetInputType(int32 InputIndex)
 	}
 	else
 	{
-		return MCT_StaticBool;
+		if (IsUsingNewHLSLGenerator())
+		{
+			// Allow non-static bool when using new HLSL generator
+			return MCT_Bool;
+		}
+		else
+		{
+			return MCT_StaticBool;
+		}
 	}
 }
 
diff --git a/Engine/Source/Runtime/Engine/Private/Materials/MaterialHLSLEmitter.cpp b/Engine/Source/Runtime/Engine/Private/Materials/MaterialHLSLEmitter.cpp
--- a/Engine/Source/Runtime/Engine/Private/Materials/MaterialHLSLEmitter.cpp
+++ b/Engine/Source/Runtime/Engine/Private/Materials/MaterialHLSLEmitter.cpp
@@ -36,6 +36,7 @@ static const TCHAR* HLSLTypeString(EMaterialValueType Type)
 	case MCT_Texture2DArray:		return TEXT("texture2DArray");
 	case MCT_VolumeTexture:			return TEXT("volumeTexture");
 	case MCT_StaticBool:			return TEXT("static bool");
+	case MCT_Bool:					return TEXT("bool");
 	case MCT_MaterialAttributes:	return TEXT("FMaterialAttributes");
 	case MCT_TextureExternal:		return TEXT("TextureExternal");
 	case MCT_TextureVirtual:		return TEXT("TextureVirtual");
diff --git a/Engine/Source/Runtime/Engine/Private/Materials/MaterialHLSLGenerator.cpp b/Engine/Source/Runtime/Engine/Private/Materials/MaterialHLSLGenerator.cpp
--- a/Engine/Source/Runtime/Engine/Private/Materials/MaterialHLSLGenerator.cpp
+++ b/Engine/Source/Runtime/Engine/Private/Materials/MaterialHLSLGenerator.cpp
@@ -701,6 +701,51 @@ const UE::HLSLTree::FExpression* FMaterialHLSLGenerator::GenerateFunctionCall(UE
 	return Result;
 }
 
+const UE::HLSLTree::FExpression* FMaterialHLSLGenerator::GenerateBranch(UE::HLSLTree::FScope& Scope,
+	const UE::HLSLTree::FExpression* ConditionExpression,
+	const UE::HLSLTree::FExpression* TrueExpression,
+	const UE::HLSLTree::FExpression* FalseExpression)
+{
+	using namespace UE::HLSLTree;
+
+	FXxHash64 Hash;
+	{
+		FXxHash64Builder Hasher;
+		Hasher.Update(&ConditionExpression, sizeof(FExpression*));
+		Hasher.Update(&TrueExpression, sizeof(FExpression*));
+		Hasher.Update(&FalseExpression, sizeof(FExpression));
+		Hash = Hasher.Finalize();
+	}
+
+	FExpression const* const* PrevExpression = BranchMap.Find(Hash);
+	if (PrevExpression)
+	{
+		return *PrevExpression;
+	}
+
+	TStringBuilder<64> LocalNameBuilder;
+	LocalNameBuilder.Appendf(TEXT("__InternalBranch%d"), BranchMap.Num());
+	const FName LocalName(LocalNameBuilder.ToString());
+
+	FFunction* Function = GetTree().NewFunction();
+	FStatementIf* IfStatement = GetTree().NewStatement<FStatementIf>(Function->GetRootScope());
+	IfStatement->ConditionExpression = ConditionExpression;
+	IfStatement->ThenScope = NewOwnedScope(*IfStatement);
+	IfStatement->ElseScope = NewOwnedScope(*IfStatement);
+	IfStatement->NextScope = NewScope(Function->GetRootScope(), EMaterialNewScopeFlag::NoPreviousScope);
+	IfStatement->NextScope->AddPreviousScope(*IfStatement->ThenScope);
+	IfStatement->NextScope->AddPreviousScope(*IfStatement->ElseScope);
+
+	GetTree().AssignLocal(*IfStatement->ThenScope, LocalName, TrueExpression);
+	GetTree().AssignLocal(*IfStatement->ElseScope, LocalName, FalseExpression);
+	const FExpression* ResultExpression = GetTree().AcquireLocal(*IfStatement->NextScope, LocalName);
+
+	Function->OutputExpressions.Add(ResultExpression);
+	const FExpression* Result = GetTree().NewFunctionCall(Scope, Function, 0);
+	BranchMap.Add(Hash, Result);
+	return Result;
+}
+
 bool FMaterialHLSLGenerator::GetParameterOverrideValueForCurrentFunction(EMaterialParameterType ParameterType, FName ParameterName, FMaterialParameterMetadata& OutResult) const
 {
 	bool bResult = false;
diff --git a/Engine/Source/Runtime/Engine/Private/Materials/MaterialHLSLTree.cpp b/Engine/Source/Runtime/Engine/Private/Materials/MaterialHLSLTree.cpp
--- a/Engine/Source/Runtime/Engine/Private/Materials/MaterialHLSLTree.cpp
+++ b/Engine/Source/Runtime/Engine/Private/Materials/MaterialHLSLTree.cpp
@@ -446,24 +446,21 @@ void FExpressionParameter::ComputeAnalyticDerivatives(FTree& Tree, FExpressionDe
 
 bool FExpressionParameter::PrepareValue(FEmitContext& Context, FEmitScope& Scope, const FRequestedType& RequestedType, FPrepareValueResult& OutResult) const
 {
+	FEmitData& EmitData = Context.FindData<FEmitData>();
 	const EMaterialParameterType ParameterType = ParameterMeta.Value.Type;
 	const Shader::FType ResultType = GetShaderValueType(ParameterType);
 
-	if (Context.bMarkLiveValues)
+	if (Context.bMarkLiveValues && EmitData.CachedExpressionData)
 	{
-		FEmitData& EmitData = Context.FindData<FEmitData>();
-		if (EmitData.CachedExpressionData)
+		if (!ParameterInfo.Name.IsNone())
 		{
-			if (!ParameterInfo.Name.IsNone())
-			{
-				EmitData.CachedExpressionData->Parameters.AddParameter(ParameterInfo, ParameterMeta);
-			}
+			EmitData.CachedExpressionData->Parameters.AddParameter(ParameterInfo, ParameterMeta);
+		}
 
-			UObject* ReferencedTexture = ParameterMeta.Value.AsTextureObject();
-			if (ReferencedTexture)
-			{
-				EmitData.CachedExpressionData->ReferencedTextures.AddUnique(ReferencedTexture);
-			}
+		UObject* ReferencedTexture = ParameterMeta.Value.AsTextureObject();
+		if (ReferencedTexture)
+		{
+			EmitData.CachedExpressionData->ReferencedTextures.AddUnique(ReferencedTexture);
 		}
 	}
 
@@ -1115,7 +1112,7 @@ bool FExpressionMaterialLayers::PrepareValue(FEmitContext& Context, FEmitScope&
 bool FExpressionSceneTexture::PrepareValue(FEmitContext& Context, FEmitScope& Scope, const FRequestedType& RequestedType, FPrepareValueResult& OutResult) const
 {
 	Context.PrepareExpression(TexCoordExpression, Scope, Shader::EValueType::Float2);
-	if (Context.bMarkLiveValues)
+	if (Context.bMarkLiveValues && Context.MaterialCompilationOutput)
 	{
 		Context.MaterialCompilationOutput->bNeedsSceneTextures = true;
 		Context.MaterialCompilationOutput->SetIsSceneTextureUsed((ESceneTextureId)SceneTextureId);
diff --git a/Engine/Source/Runtime/Engine/Private/Shader/ShaderValue.cpp b/Engine/Source/Runtime/Engine/Private/Shader/ShaderValue.cpp
--- a/Engine/Source/Runtime/Engine/Private/Shader/ShaderValue.cpp
+++ b/Engine/Source/Runtime/Engine/Private/Shader/ShaderValue.cpp
@@ -676,6 +676,7 @@ EValueType MakeValueType(EValueComponentType ComponentType, int32 NumComponents)
 		case 2: return EValueType::Int2;
 		case 3: return EValueType::Int3;
 		case 4: return EValueType::Int4;
+		case 16: return EValueType::Float4x4; // no explicit int/bool matrix types...so use float instead
 		default: break;
 		}
 		break;
@@ -686,6 +687,7 @@ EValueType MakeValueType(EValueComponentType ComponentType, int32 NumComponents)
 		case 2: return EValueType::Bool2;
 		case 3: return EValueType::Bool3;
 		case 4: return EValueType::Bool4;
+		case 16: return EValueType::Float4x4;
 		default: break;
 		}
 		break;
diff --git a/Engine/Source/Runtime/Engine/Public/HLSLTree/HLSLTree.h b/Engine/Source/Runtime/Engine/Public/HLSLTree/HLSLTree.h
--- a/Engine/Source/Runtime/Engine/Public/HLSLTree/HLSLTree.h
+++ b/Engine/Source/Runtime/Engine/Public/HLSLTree/HLSLTree.h
@@ -392,6 +392,8 @@ public:
 
 	Shader::FType GetValuePreshader(FEmitContext& Context, FEmitScope& Scope, const FRequestedType& RequestedType, const Shader::FType& ResultType, Shader::FPreshaderData& OutPreshader) const;
 	Shader::FType GetValuePreshader(FEmitContext& Context, FEmitScope& Scope, const FRequestedType& RequestedType, Shader::FPreshaderData& OutPreshader) const;
+	Shader::FType GetValuePreshader(FEmitContext& Context, FEmitScope& Scope, const Shader::FType& ResultType, Shader::FPreshaderData& OutPreshader) const;
+	Shader::FType GetValuePreshader(FEmitContext& Context, FEmitScope& Scope, Shader::EValueType ResultType, Shader::FPreshaderData& OutPreshader) const;
 
 	Shader::FValue GetValueConstant(FEmitContext& Context, FEmitScope& Scope, const FRequestedType& RequestedType, const FPreparedType& PreparedType) const;
 	Shader::FValue GetValueConstant(FEmitContext& Context, FEmitScope& Scope, const FRequestedType& RequestedType, const Shader::FType& ResultType) const;
diff --git a/Engine/Source/Runtime/Engine/Public/MaterialHLSLGenerator.h b/Engine/Source/Runtime/Engine/Public/MaterialHLSLGenerator.h
--- a/Engine/Source/Runtime/Engine/Public/MaterialHLSLGenerator.h
+++ b/Engine/Source/Runtime/Engine/Public/MaterialHLSLGenerator.h
@@ -152,6 +152,11 @@ public:
 		TArrayView<const UE::HLSLTree::FExpression*> ConnectedInputs,
 		int32 OutputIndex);
 
+	const UE::HLSLTree::FExpression* GenerateBranch(UE::HLSLTree::FScope& Scope,
+		const UE::HLSLTree::FExpression* ConditionExpression,
+		const UE::HLSLTree::FExpression* TrueExpression,
+		const UE::HLSLTree::FExpression* FalseExpression);
+
 	template<typename T, typename... ArgTypes>
 	T* NewExpressionData(const UMaterialExpression* MaterialExpression, ArgTypes... Args)
 	{
@@ -254,6 +259,7 @@ private:
 	TArray<FFunctionCallEntry*, TInlineAllocator<8>> FunctionCallStack;
 	TArray<UE::HLSLTree::FScope*> JoinedScopeStack;
 	TMap<FXxHash64, TUniquePtr<FFunctionCallEntry>> FunctionCallMap;
+	TMap<FXxHash64, const UE::HLSLTree::FExpression*> BranchMap;
 	TMap<UMaterialFunctionInterface*, UE::HLSLTree::FFunction*> FunctionMap;
 	TMap<UMaterialExpression*, FStatementEntry> StatementMap;
 	TMap<FExpressionDataKey, void*> ExpressionDataMap;
diff --git a/Engine/Source/Runtime/Engine/Public/MaterialShared.h b/Engine/Source/Runtime/Engine/Public/MaterialShared.h
--- a/Engine/Source/Runtime/Engine/Public/MaterialShared.h
+++ b/Engine/Source/Runtime/Engine/Public/MaterialShared.h
@@ -178,6 +178,11 @@ enum EMaterialValueType
 
 	/** Used for code chunks that are statements with no value, rather than expressions */
 	MCT_VoidStatement         = 1 << 22,
+
+	/** Non-static bool, only used in new HLSL translator */
+	MCT_Bool = 1 << 23,
+
+	MCT_Numeric = MCT_Float | MCT_LWCType | MCT_Bool,
 };
 
 /** @return the number of components in a vector type. */
