Fix issue with certain old texture versions having incorrect compression format set

diff --git a/Engine/Source/Runtime/Engine/Private/Texture.cpp b/Engine/Source/Runtime/Engine/Private/Texture.cpp
--- a/Engine/Source/Runtime/Engine/Private/Texture.cpp
+++ b/Engine/Source/Runtime/Engine/Private/Texture.cpp
@@ -664,6 +664,26 @@ void UTexture::Serialize(FArchive& Ar)
 	
 	if (Ar.IsLoading())
 	{
+		// For a while, we could end up writing textures that were flagged as !bPNGCompressed, had their
+		// compression format set to PNG, but did not actually contain compressed data. Fix these up.
+		if (Ar.CustomVer(FUE5MainStreamObjectVersion::GUID) < FUE5MainStreamObjectVersion::TextureSourceVirtualization
+			&& !Source.bPNGCompressed_DEPRECATED
+			&& Source.CompressionFormat == TSCF_PNG)
+		{
+			UE_LOG(LogTexture, Display, TEXT("Texture \"%s\" has CompressionFormat=PNG but not bPNGCompressed, assuming texture is actually uncompressed."), *GetPathName());
+			Source.CompressionFormat = TSCF_None;
+		}
+
+		// Along with the above, the rules implemented here are this:
+		// - In older versions, bPNGCompressed is the authoritative source of PNG-ness. In those versions,
+		//   if the compression format is set to PNG but bPNGCompressed isn't set, bPNGCompressed wins.
+		// - In new versions, bPNGCompressed is deprecated, never written as true, and CompressionFormat
+		//   is the authoritative source on whether something is a PNG or not.
+		// - Finally, in between, some versions would in certain cases write textures with bPNGCompressed set
+		//   and the CompressionFormat being None, and those are actually PNGs.
+		//
+		// Now, the separate bPNGCompressed is gone (to avoid further desyncs like this) and we make sure
+		// that CompressionFormat always matches the contents.
 		if ( Source.bPNGCompressed_DEPRECATED )
 		{
 			// loaded with deprecated "bPNGCompressed"
