Add explicit namespace to GetVTAddressMode() calls, avoid conflict with unity build

diff --git a/Engine/Source/Runtime/Engine/Private/Materials/MaterialHLSLTree.cpp b/Engine/Source/Runtime/Engine/Private/Materials/MaterialHLSLTree.cpp
--- a/Engine/Source/Runtime/Engine/Private/Materials/MaterialHLSLTree.cpp
+++ b/Engine/Source/Runtime/Engine/Private/Materials/MaterialHLSLTree.cpp
@@ -778,7 +778,7 @@ uint32 AcquireVTStackIndex(
 			"%U + LIGHTMAP_VT_ENABLED, Parameters.VirtualTextureFeedback)"),
 			BaseFunctionName,
 			StackIndex, StackIndex, StackIndex,
-			EmitTexCoordValue, GetVTAddressMode(AddressU), GetVTAddressMode(AddressV),
+			EmitTexCoordValue, Private::GetVTAddressMode(AddressU), Private::GetVTAddressMode(AddressV),
 			EmitMipValue,
 			StackIndex);
 		break;
@@ -792,7 +792,7 @@ uint32 AcquireVTStackIndex(
 			"%)"),
 			BaseFunctionName,
 			StackIndex, StackIndex, StackIndex,
-			EmitTexCoordValue, GetVTAddressMode(AddressU), GetVTAddressMode(AddressV),
+			EmitTexCoordValue, Private::GetVTAddressMode(AddressU), Private::GetVTAddressMode(AddressV),
 			EmitMipValue,
 			FormattedFeedback.ToString());
 		break;
@@ -806,7 +806,7 @@ uint32 AcquireVTStackIndex(
 			"%U + LIGHTMAP_VT_ENABLED, Parameters.VirtualTextureFeedback)"),
 			BaseFunctionName,
 			StackIndex, StackIndex, StackIndex,
-			EmitTexCoordValue, GetVTAddressMode(AddressU), GetVTAddressMode(AddressV),
+			EmitTexCoordValue, Private::GetVTAddressMode(AddressU), Private::GetVTAddressMode(AddressV),
 			EmitTexCoordValueDdx, EmitTexCoordValueDdy,
 			StackIndex);
 		break;
