World Partition
- Edigrate 19149766 from Main: wp by default with disabled streaming to mimic classic level streaming with data layers as a bonus.
- Revert part of 18994482: removed option to force load everything in the editor.
- Warn the user when the world is large enough to justify enabling streaming.

#rnx

diff --git a/Engine/Source/Editor/DetailCustomizations/Private/DetailCustomizations.cpp b/Engine/Source/Editor/DetailCustomizations/Private/DetailCustomizations.cpp
--- a/Engine/Source/Editor/DetailCustomizations/Private/DetailCustomizations.cpp
+++ b/Engine/Source/Editor/DetailCustomizations/Private/DetailCustomizations.cpp
@@ -458,6 +458,7 @@ void FDetailCustomizationsModule::RegisterSectionMappings()
 		{
 			TSharedRef<FPropertySection> Section = PropertyModule.FindOrCreateSection("Actor", "Streaming", LOCTEXT("Streaming", "Streaming"));
 			Section->AddCategory("World Partition");
+			Section->AddCategory("Data Layers");
 			Section->AddCategory("HLOD");
 		}
 	}
@@ -479,12 +480,6 @@ void FDetailCustomizationsModule::RegisterSectionMappings()
 			Section->AddCategory("Cooking");
 			Section->AddCategory("Tags");
 		}
-
-		{
-			TSharedRef<FPropertySection> Section = PropertyModule.FindOrCreateSection("ActorComponent", "Streaming", LOCTEXT("Streaming", "Streaming"));
-			Section->AddCategory("World Partition");
-			Section->AddCategory("Data Layers");
-		}
 	}
 
 	// SceneComponent
diff --git a/Engine/Source/Editor/WorldPartitionEditor/Private/WorldPartition/Customizations/WorldPartitionDetailsCustomization.cpp b/Engine/Source/Editor/WorldPartitionEditor/Private/WorldPartition/Customizations/WorldPartitionDetailsCustomization.cpp
--- a/Engine/Source/Editor/WorldPartitionEditor/Private/WorldPartition/Customizations/WorldPartitionDetailsCustomization.cpp
+++ b/Engine/Source/Editor/WorldPartitionEditor/Private/WorldPartition/Customizations/WorldPartitionDetailsCustomization.cpp
@@ -34,9 +34,11 @@ void FWorldPartitionDetails::CustomizeDetails(IDetailLayoutBuilder& DetailBuilde
 	{
 		FAddPropertyParams Params;
 		Params.HideRootObjectNode(true);
+		Params.UniqueId(TEXT("RuntimeHash"));
 
 		if (IDetailPropertyRow* RuntimeHashRow = WorldPartitionCategory.AddExternalObjects({ WorldPartition->RuntimeHash }, EPropertyLocation::Default, Params))
 		{
+			// use TAttribute::CreateLambda?
 			RuntimeHashRow->ShouldAutoExpand(true)
 				.DisplayName(LOCTEXT("RuntimeHash", "Runtime Hash"))
 				.Visibility(TAttribute<EVisibility>::Create(TAttribute<EVisibility>::FGetter::CreateLambda([this]() { return WorldPartition->bEnableStreaming ? EVisibility::Visible : EVisibility::Hidden; })));
@@ -54,7 +56,7 @@ void FWorldPartitionDetails::CustomizeDetails(IDetailLayoutBuilder& DetailBuilde
 		.ValueContent()
 		[
 			SNew(SHorizontalBox)
-			+ SHorizontalBox::Slot()
+			+SHorizontalBox::Slot()
 			.AutoWidth()
 			.VAlign(VAlign_Center)
 			[
diff --git a/Engine/Source/Runtime/Engine/Private/WorldPartition/WorldPartition.cpp b/Engine/Source/Runtime/Engine/Private/WorldPartition/WorldPartition.cpp
--- a/Engine/Source/Runtime/Engine/Private/WorldPartition/WorldPartition.cpp
+++ b/Engine/Source/Runtime/Engine/Private/WorldPartition/WorldPartition.cpp
@@ -259,6 +259,8 @@ UWorldPartition::UWorldPartition(const FObjectInitializer& ObjectInitializer)
 	, WorldPartitionEditor(nullptr)
 	, bEnableStreaming(true)
 	, bStreamingWasEnabled(true)
+	, bShouldEnableStreamingWarned(false)
+	, bShouldCheckEnableStreamingWarning(false)
 	, bForceGarbageCollection(false)
 	, bForceGarbageCollectionPurge(false)
 	, bIsPIE(false)
@@ -1142,12 +1144,23 @@ bool UWorldPartition::GetInstancingContext(const FLinkerInstancingContext*& OutI
 	return false;
 }
 
+void UWorldPartition::OnEnableStreaming()
+{
+	if (!bStreamingWasEnabled)
+	{
+		bStreamingWasEnabled = true;
+		FMessageDialog::Open(EAppMsgType::Ok, LOCTEXT("WorldPartitionEnableStreamingDialolg", "Please refer to https://docs.unrealengine.com/5.0/en-US/building-virtual-worlds/world-partition for how to set up streaming."));
+	}
+}
+
 void UWorldPartition::HashActorDesc(FWorldPartitionActorDesc* ActorDesc)
 {
 	check(ActorDesc);
 	check(EditorHash);
 	FWorldPartitionHandle ActorHandle(this, ActorDesc->GetGuid());
 	EditorHash->HashActor(ActorHandle);
+
+	bShouldCheckEnableStreamingWarning = true;
 }
 
 void UWorldPartition::UnhashActorDesc(FWorldPartitionActorDesc* ActorDesc)
@@ -1301,6 +1314,32 @@ void UWorldPartition::Tick(float DeltaSeconds)
 		bForceGarbageCollection = false;
 		bForceGarbageCollectionPurge = false;
 	}
+
+	if (bShouldCheckEnableStreamingWarning)
+	{
+		bShouldCheckEnableStreamingWarning = false;
+
+		if (!bEnableStreaming && !bStreamingWasEnabled && !bShouldEnableStreamingWarned)
+		{
+			FBox AllActorsBounds(ForceInit);
+			for (UActorDescContainer::TConstIterator<> ActorDescIterator(this); ActorDescIterator; ++ActorDescIterator)
+			{
+				AllActorsBounds += ActorDescIterator->GetBounds();
+
+				// Warn the user if the world becomes larger that 4km in any axis
+				if (AllActorsBounds.GetSize().GetMax() >= 400000.0f)
+				{
+					bShouldEnableStreamingWarned = true;
+					if (FMessageDialog::Open(EAppMsgType::YesNo, LOCTEXT("WorldPartitionShouldEnableStreaming", "The size of your world has grown enough to justify enabling streaming. Enable streaming now?")) == EAppReturnType::Yes)
+					{
+						OnEnableStreaming();
+					}
+
+					break;
+				}
+			}
+		}
+	}
 #endif
 }
 
diff --git a/Engine/Source/Runtime/Engine/Private/WorldPartition/WorldPartitionConverter.cpp b/Engine/Source/Runtime/Engine/Private/WorldPartition/WorldPartitionConverter.cpp
--- a/Engine/Source/Runtime/Engine/Private/WorldPartition/WorldPartitionConverter.cpp
+++ b/Engine/Source/Runtime/Engine/Private/WorldPartition/WorldPartitionConverter.cpp
@@ -57,6 +57,8 @@ bool FWorldPartitionConverter::Convert()
 		{
 			return false;
 		}
+		WorldPartition->bEnableStreaming = Parameters.bEnableStreaming;
+		WorldPartition->bStreamingWasEnabled = Parameters.bEnableStreaming;
 		bCreatedWorldPartition = true;
 	}
 
diff --git a/Engine/Source/Runtime/Engine/Public/WorldPartition/WorldPartition.h b/Engine/Source/Runtime/Engine/Public/WorldPartition/WorldPartition.h
--- a/Engine/Source/Runtime/Engine/Public/WorldPartition/WorldPartition.h
+++ b/Engine/Source/Runtime/Engine/Public/WorldPartition/WorldPartition.h
@@ -121,6 +121,8 @@ private:
 
 	virtual bool GetInstancingContext(const FLinkerInstancingContext*& OutInstancingContext, FSoftObjectPathFixupArchive*& OutSoftObjectPathFixupArchive) const override;
 	//~ End UActorDescContainer Interface
+
+	void OnEnableStreaming();
 #endif
 
 public:
@@ -210,9 +212,15 @@ public:
 	UPROPERTY(EditAnywhere, Category=WorldPartition)
 	bool bEnableStreaming;
 
-	/** Used to know if it's the first time streaming is enabled on this world */
+	/** Used to know if it's the first time streaming is enabled on this world. */
 	UPROPERTY()
 	bool bStreamingWasEnabled;
+
+	/** Used to know if the user has already been warned about that it should enable streaming based on world size. */
+	bool bShouldEnableStreamingWarned;
+
+	/** Used to know if we need to recheck if the user should enable streaming based on world size. */
+	bool bShouldCheckEnableStreamingWarning;
 #endif
 
 	UPROPERTY()
diff --git a/Engine/Source/Runtime/Engine/Public/WorldPartition/WorldPartitionConverter.h b/Engine/Source/Runtime/Engine/Public/WorldPartition/WorldPartitionConverter.h
--- a/Engine/Source/Runtime/Engine/Public/WorldPartition/WorldPartitionConverter.h
+++ b/Engine/Source/Runtime/Engine/Public/WorldPartition/WorldPartitionConverter.h
@@ -20,12 +20,13 @@ public:
 	{
 		FParameters()
 		: bConvertSubLevels(true)
+		, bEnableStreaming(true)
 		, EditorHashClass(nullptr)
 		, RuntimeHashClass(nullptr)
-		{
-		}
+		{}
 
 		bool bConvertSubLevels;
+		bool bEnableStreaming;
 		TSubclassOf<UWorldPartitionEditorHash> EditorHashClass;
 		TSubclassOf<UWorldPartitionRuntimeHash> RuntimeHashClass;
 	};
diff --git a/Engine/Source/Runtime/Engine/Public/WorldPartition/WorldPartitionRuntimeSpatialHash.h b/Engine/Source/Runtime/Engine/Public/WorldPartition/WorldPartitionRuntimeSpatialHash.h
--- a/Engine/Source/Runtime/Engine/Public/WorldPartition/WorldPartitionRuntimeSpatialHash.h
+++ b/Engine/Source/Runtime/Engine/Public/WorldPartition/WorldPartitionRuntimeSpatialHash.h
@@ -234,7 +234,6 @@ public:
 	virtual void FlushStreaming() override;
 	virtual bool GenerateHLOD(ISourceControlHelper* SourceControlHelper, FActorClusterContext& ActorClusterContext, bool bCreateActorsOnly) override;
 	virtual void DrawPreview() const override;
-	TArray<FSpatialHashRuntimeGrid>& GetGrids() { return Grids; }
 
 	FName GetCellName(FName InGridName, const FIntVector& InCellGlobalCoord, const FDataLayersID& InDataLayerID) const;
 	static FName GetCellName(UWorldPartition* WorldPartition, FName InGridName, const FIntVector& InCellGlobalCoord, const FDataLayersID& InDataLayerID);
