Disable a bogus static analysis warning

diff --git a/Engine/Source/Runtime/Engine/Private/HLSLTree/HLSLTreeEmit.cpp b/Engine/Source/Runtime/Engine/Private/HLSLTree/HLSLTreeEmit.cpp
--- a/Engine/Source/Runtime/Engine/Private/HLSLTree/HLSLTreeEmit.cpp
+++ b/Engine/Source/Runtime/Engine/Private/HLSLTree/HLSLTreeEmit.cpp
@@ -523,7 +523,9 @@ FEmitScope* FEmitContext::InternalPrepareScope(FScope* InScope, FScope* InParent
 				check(EmitParentScope);
 				if (Statement->Prepare(*this, *EmitParentScope))
 				{
-					if (EmitScope->State == EEmitScopeState::Initializing)
+					// Disable 'V547: Expression is always true'
+					// EmitScope->State may be set to 'Dead' via Statement->Prepare
+					if (EmitScope->State == EEmitScopeState::Initializing)  //-V547
 					{
 						check(EmitScope->Evaluation != EExpressionEvaluation::None);
 						EmitScope->State = EEmitScopeState::Live;
