Add check() to make static analysis happy

diff --git a/Engine/Source/Runtime/Engine/Private/HLSLTree/HLSLTree.cpp b/Engine/Source/Runtime/Engine/Private/HLSLTree/HLSLTree.cpp
--- a/Engine/Source/Runtime/Engine/Private/HLSLTree/HLSLTree.cpp
+++ b/Engine/Source/Runtime/Engine/Private/HLSLTree/HLSLTree.cpp
@@ -343,6 +343,7 @@ bool GetLiveScopes(FEmitContext& Context, const FExpressionLocalPHI& Expression,
 			}
 
 			const int32 LiveScopeIndex = OutLiveScopes.NumScopes++;
+			check(LiveScopeIndex < MaxNumPreviousScopes);
 			OutLiveScopes.LiveValues[LiveScopeIndex] = Expression.Values[ScopeIndex];
 			OutLiveScopes.EmitValueScopes[LiveScopeIndex] = EmitValueScope;
 			OutLiveScopes.EmitDeclarationScope = FEmitScope::FindSharedParent(OutLiveScopes.EmitDeclarationScope, EmitValueScope);
