Improve some memory usage in Networking Code

diff --git a/Engine/Plugins/Runtime/ReplicationGraph/Source/Public/ReplicationGraph.h b/Engine/Plugins/Runtime/ReplicationGraph/Source/Public/ReplicationGraph.h
--- a/Engine/Plugins/Runtime/ReplicationGraph/Source/Public/ReplicationGraph.h
+++ b/Engine/Plugins/Runtime/ReplicationGraph/Source/Public/ReplicationGraph.h
@@ -1112,6 +1112,11 @@ protected:
 	bool bEnableFullActorPrioritizationDetailsAllConnections = false;
 #endif
 
+private:
+	/** Whether or not a connection was saturated during an update. */
+	bool bWasConnectionSaturated = false;
+
+protected:
 	/** Default Replication Path */
 	void ReplicateActorListsForConnections_Default(UNetReplicationGraphConnection* ConnectionManager, FGatheredReplicationActorLists& GatheredReplicationListsForConnection, FNetViewerArray& Viewers);
 
@@ -1156,9 +1161,6 @@ private:
 	/** Collect replication data during ServerReplicateActors */
 	FFrameReplicationStats FrameReplicationStats;
 
-	/** Whether or not a connection was saturated during an update. */
-	bool bWasConnectionSaturated = false;
-
 	/** Internal frame counter for replication. This is only updated by us. The one of UNetDriver can be updated by RPC calls and is only used to invalidate shared property CLs/serialiation data. */
 	uint32 ReplicationGraphFrame = 0;
 
@@ -1235,11 +1237,11 @@ public:
 	bool bEnableFullActorPrioritizationDetails = false;
 #endif
 
+	bool bEnableDebugging;
+
 	UPROPERTY()
 	TObjectPtr<class AReplicationGraphDebugActor> DebugActor = nullptr;
 
-	bool bEnableDebugging;
-
 	/** Index of the connection in the global list. Will be reassigned when any client disconnects so it is a key that can be referenced only during a single tick */
 	int32 ConnectionOrderNum;
 
diff --git a/Engine/Source/Runtime/Engine/Classes/Engine/DemoNetDriver.h b/Engine/Source/Runtime/Engine/Classes/Engine/DemoNetDriver.h
--- a/Engine/Source/Runtime/Engine/Classes/Engine/DemoNetDriver.h
+++ b/Engine/Source/Runtime/Engine/Classes/Engine/DemoNetDriver.h
@@ -170,8 +170,6 @@ public:
 	APlayerController* GetSpectatorController() const { return SpectatorController; }
 
 private:
-	/** True if as have paused all of the channels */
-	bool bChannelsArePaused;
 
 	/** This is the main spectator controller that is used to view the demo world from */
 	APlayerController* SpectatorController;
@@ -248,11 +246,22 @@ private:
 	};
 
 	bool bIsFastForwarding;
-	bool bIsFastForwardingForCheckpoint;
-	bool bWasStartStreamingSuccessful;
 	bool bIsFinalizingFastForward;
 	bool bIsRestoringStartupActors;
 
+	/** True if as have paused all of the channels */
+	uint8 bChannelsArePaused : 1;
+	uint8 bIsFastForwardingForCheckpoint : 1;
+	uint8 bWasStartStreamingSuccessful : 1;
+
+	/** If true, recording will prioritize replicating actors based on the value that AActor::GetReplayPriority returns. */
+	uint8 bPrioritizeActors : 1;
+
+protected:
+	uint8 bIsWaitingForHeaderDownload : 1;
+	uint8 bIsWaitingForStream : 1;
+
+private:
 	TArray<FNetworkGUID> NonQueuedGUIDsForScrubbing;
 
 	// Replay tasks
@@ -288,9 +297,6 @@ private:
 	/** Array of prioritized actors, used in TickDemoRecord. Stored as a member so that its storage doesn't have to be re-allocated each frame. */
 	TArray<FDemoActorPriority> PrioritizedActors;
 
-	/** If true, recording will prioritize replicating actors based on the value that AActor::GetReplayPriority returns. */
-	bool bPrioritizeActors;
-
 	/** Does the actual work of TickFlush, either on the main thread or in a task thread in parallel with Slate. */
 	void TickFlushInternal(float DeltaSeconds);
 
@@ -680,9 +686,6 @@ protected:
 
 	TArray<FQueuedDemoPacket> QueuedPacketsBeforeTravel;
 
-	bool bIsWaitingForHeaderDownload;
-	bool bIsWaitingForStream;
-
 	int64 MaxArchiveReadPos;
 
 private:
diff --git a/Engine/Source/Runtime/Engine/Classes/Engine/NetDriver.h b/Engine/Source/Runtime/Engine/Classes/Engine/NetDriver.h
--- a/Engine/Source/Runtime/Engine/Classes/Engine/NetDriver.h
+++ b/Engine/Source/Runtime/Engine/Classes/Engine/NetDriver.h
@@ -614,19 +614,19 @@ struct ENGINE_API FChannelDefinition
 	int32 StaticChannelIndex;	// Channel always uses this index, INDEX_NONE if dynamically chosen
 
 	UPROPERTY()
-	bool bTickOnCreate;			// Whether to immediately begin ticking the channel after creation
+	uint8 bTickOnCreate : 1;			// Whether to immediately begin ticking the channel after creation
 
 	UPROPERTY()
-	bool bServerOpen;			// Channel opened by the server
+	uint8 bServerOpen : 1;			// Channel opened by the server
 
 	UPROPERTY()
-	bool bClientOpen;			// Channel opened by the client
+	uint8 bClientOpen : 1;			// Channel opened by the client
 
 	UPROPERTY()
-	bool bInitialServer;		// Channel created on server when connection is established
+	uint8 bInitialServer : 1;		// Channel created on server when connection is established
 
 	UPROPERTY()
-	bool bInitialClient;		// Channel created on client before connecting
+	uint8 bInitialClient : 1;		// Channel created on client before connecting
 
 	FChannelDefinition() : 
 		ChannelName(NAME_None),
@@ -770,19 +770,6 @@ public:
 	UPROPERTY(Config)
 	float TimeoutMultiplierForUnoptimizedBuilds;
 
-	/**
-	 * If true, ignore timeouts completely.  Should be used only in development
-	 */
-	UPROPERTY(Config)
-	bool bNoTimeouts;
-
-	/**
-	 * If true this NetDriver will not apply the network emulation settings that simulate
-	 * latency and packet loss in non-shippable builds
-	 */
-	UPROPERTY(Config)
-	bool bNeverApplyNetworkEmulationSettings;
-
 	/** Connection to the server (this net driver is a client) */
 	UPROPERTY()
 	TObjectPtr<class UNetConnection> ServerConnection;
@@ -918,22 +905,55 @@ public:
 private:
 	double						ElapsedTime;
 	
-
 	/** Whether or not the NetDriver is ticking */
 	bool bInTick;
 
-	bool bPendingDestruction;
+	uint8 bPendingDestruction : 1;
+
+#if DO_ENABLE_NET_TEST
+	/** Dont load packet settings from config or cmdline when true*/
+	uint8 bForcedPacketSettings : 1;
+#endif 
+
+	uint8 bDidHitchLastFrame : 1;
+
+	/** cache whether or not we have a replay connection, updated when a connection is added or removed */
+	uint8 bHasReplayConnection : 1;
+
+protected:
+	uint8 bMaySendProperties : 1;
+
+	uint8 bSkipServerReplicateActors : 1;
 
 public:
-	/** Last realtime a tick dispatch occurred. Used currently to try and diagnose timeout issues */
-	double						LastTickDispatchRealtime;
+	/**
+	 * If true, ignore timeouts completely.  Should be used only in development
+	 */
+	UPROPERTY(Config)
+	uint8 bNoTimeouts : 1;
+
+	/**
+	 * If true this NetDriver will not apply the network emulation settings that simulate
+	 * latency and packet loss in non-shippable builds
+	 */
+	UPROPERTY(Config)
+	uint8 bNeverApplyNetworkEmulationSettings : 1;
 
 	/** If true then client connections are to other client peers */
-	bool						bIsPeer;
+	uint8						bIsPeer : 1;
 	/** @todo document */
-	bool						ProfileStats;
+	uint8						ProfileStats : 1;
 	/** If true, it assumes the stats are being set by server data */
-	bool						bSkipLocalStats;
+	uint8						bSkipLocalStats : 1;
+	/** Collect net stats even if not FThreadStats::IsCollectingData(). */
+	uint8 bCollectNetStats : 1;
+	/** Used to determine if checking for standby cheats should occur */
+	uint8						bIsStandbyCheckingEnabled : 1;
+	/** Used to determine whether we've already caught a cheat or not */
+	uint8						bHasStandbyCheatTriggered : 1;
+
+	/** Last realtime a tick dispatch occurred. Used currently to try and diagnose timeout issues */
+	double						LastTickDispatchRealtime;
 	/** Timings for Socket::SendTo() */
 	int32						SendCycles;
 	/** Stats for network perf */
@@ -1001,14 +1021,8 @@ public:
 	/** Total acks sent since the net driver's creation  */
 	uint32						OutTotalAcks;
 
-	/** Collect net stats even if not FThreadStats::IsCollectingData(). */
-	bool bCollectNetStats;
 	/** Time of last netdriver cleanup pass */
 	double						LastCleanupTime;
-	/** Used to determine if checking for standby cheats should occur */
-	bool						bIsStandbyCheckingEnabled;
-	/** Used to determine whether we've already caught a cheat or not */
-	bool						bHasStandbyCheatTriggered;
 	/** The amount of time without packets before triggering the cheat code */
 	float						StandbyRxCheatTime;
 	/** todo document */
@@ -1816,10 +1830,6 @@ public:
 	bool HasReplayConnection() const { return bHasReplayConnection; }
 
 protected:
-
-	bool bMaySendProperties;
-
-	bool bSkipServerReplicateActors = false;
 	
 	/** Stream of random numbers to be used by this instance of UNetDriver */
 	FRandomStream UpdateDelayRandomStream;
@@ -1874,16 +1884,6 @@ private:
 	/** Unique id used by NetTrace to identify driver */
 	uint32 NetTraceId = 0;
 
-#if DO_ENABLE_NET_TEST
-	/** Dont load packet settings from config or cmdline when true*/
-	bool bForcedPacketSettings;
-#endif 
-
-	bool bDidHitchLastFrame = false;
-
-	/** cache whether or not we have a replay connection, updated when a connection is added or removed */
-	bool bHasReplayConnection;
-
 	/** Stat tracking for the total number of out of order packets lost */
 	int32 TotalOutOfOrderPacketsLost = 0;
 
diff --git a/Engine/Source/Runtime/Engine/Private/NetDriver.cpp b/Engine/Source/Runtime/Engine/Private/NetDriver.cpp
--- a/Engine/Source/Runtime/Engine/Private/NetDriver.cpp
+++ b/Engine/Source/Runtime/Engine/Private/NetDriver.cpp
@@ -356,7 +356,6 @@ UNetDriver::UNetDriver(const FObjectInitializer& ObjectInitializer)
 :	UObject(ObjectInitializer)
 ,	MaxInternetClientRate(10000)
 , 	MaxClientRate(15000)
-,   bNoTimeouts(false)
 ,   ServerConnection(nullptr)
 ,	ClientConnections()
 ,	MappedClientConnections()
@@ -373,9 +372,22 @@ UNetDriver::UNetDriver(const FObjectInitializer& ObjectInitializer)
 ,	ElapsedTime( 0.0 )
 ,	bInTick(false)
 ,	bPendingDestruction(false)
-,	LastTickDispatchRealtime( 0.f )
+#if DO_ENABLE_NET_TEST
+,   bForcedPacketSettings(false)
+#endif
+,   bDidHitchLastFrame(false)
+,   bHasReplayConnection(false)
+,   bMaySendProperties(false)
+,   bSkipServerReplicateActors(false)
+,   bNoTimeouts(false)
+,   bNeverApplyNetworkEmulationSettings(false)
 ,   bIsPeer(false)
-,	bSkipLocalStats(false)
+,   ProfileStats(false)
+,   bSkipLocalStats(false)
+,   bCollectNetStats(false)
+,   bIsStandbyCheckingEnabled(false)
+,   bHasStandbyCheatTriggered(false)
+,	LastTickDispatchRealtime( 0.f )
 ,	InBytes(0)
 ,	InTotalBytes(0)
 ,	OutBytes(0)
@@ -398,7 +410,6 @@ UNetDriver::UNetDriver(const FObjectInitializer& ObjectInitializer)
 ,	StatPeriod(1.f)
 ,	TotalRPCsCalled(0)
 ,	OutTotalAcks(0)
-,	bCollectNetStats(false)
 ,	LastCleanupTime(0.0)
 ,	NetTag(0)
 #if NET_DEBUG_RELEVANT_ACTORS
diff --git a/Engine/Source/Runtime/NetworkReplayStreaming/NetworkReplayStreaming/Public/NetworkReplayStreaming.h b/Engine/Source/Runtime/NetworkReplayStreaming/NetworkReplayStreaming/Public/NetworkReplayStreaming.h
--- a/Engine/Source/Runtime/NetworkReplayStreaming/NetworkReplayStreaming/Public/NetworkReplayStreaming.h
+++ b/Engine/Source/Runtime/NetworkReplayStreaming/NetworkReplayStreaming/Public/NetworkReplayStreaming.h
@@ -53,7 +53,7 @@ public:
 /** Struct to store information about a stream, returned from search results. */
 struct FNetworkReplayStreamInfo
 {
-	FNetworkReplayStreamInfo() : Timestamp( 0 ), SizeInBytes( 0 ), LengthInMS( 0 ), NumViewers( 0 ), bIsLive( false ), Changelist( 0 ), bShouldKeep( false ) {}
+	FNetworkReplayStreamInfo() : Timestamp( 0 ), SizeInBytes( 0 ), LengthInMS( 0 ), NumViewers( 0 ), Changelist( 0 ), bIsLive(false), bShouldKeep( false ) {}
 
 	/** The name of the stream (generally this is auto generated, refer to friendly name for UI) */
 	FString Name;
@@ -73,12 +73,12 @@ struct FNetworkReplayStreamInfo
 	/** Number of viewers viewing this stream */
 	int32 NumViewers;
 
-	/** True if the stream is live and the game hasn't completed yet */
-	bool bIsLive;
-
 	/** The changelist of the replay */
 	int32 Changelist;
 
+	/** True if the stream is live and the game hasn't completed yet */
+	bool bIsLive;
+
 	/** Debug feature that allows us to mark replays to not be deleted. True if this replay has been marked as such */
 	bool bShouldKeep;
 };
