PR #9011: Fixed a typo in a comment (Contributed by [REDACTED_PERSON])

diff --git a/Engine/Source/Runtime/Engine/Classes/Engine/EngineTypes.h b/Engine/Source/Runtime/Engine/Classes/Engine/EngineTypes.h
--- a/Engine/Source/Runtime/Engine/Classes/Engine/EngineTypes.h
+++ b/Engine/Source/Runtime/Engine/Classes/Engine/EngineTypes.h
@@ -3395,7 +3395,7 @@ namespace EAutoReceiveInput
 	};
 }
 
-/** Specifies if an AI pawn will automatically be possed by an AI controller */
+/** Specifies if an AI pawn will automatically be possessed by an AI controller */
 UENUM()
 enum class EAutoPossessAI : uint8
 {
