Level Instance: Selection Fixes

- When editing a Level Instance it wasn't possible to select one of its level actors if it was a Level instance itself
- Right-click edit context actor wasn't forwarded to child edit level instance

#rnx

diff --git a/Engine/Source/Editor/LevelInstanceEditor/Private/LevelInstanceEditorMode.cpp b/Engine/Source/Editor/LevelInstanceEditor/Private/LevelInstanceEditorMode.cpp
--- a/Engine/Source/Editor/LevelInstanceEditor/Private/LevelInstanceEditorMode.cpp
+++ b/Engine/Source/Editor/LevelInstanceEditor/Private/LevelInstanceEditorMode.cpp
@@ -120,8 +120,19 @@ bool ULevelInstanceEditorMode::IsSelectionDisallowed(AActor* InActor, bool bInSe
 			{
 				ALevelInstance* EditingLevelInstance = LevelInstanceSubsystem->GetEditingLevelInstance();
 				ALevelInstance* LevelInstance = LevelInstanceSubsystem->GetParentLevelInstance(InActor);
+				// Allow selection on actors that are part of the currently edited Level Instance hierarchy because AActor::GetRootSelectionParent() will eventually
+				// Bubble up the selection to its parent.
+				while (LevelInstance != nullptr)
+				{
+					if (LevelInstance == EditingLevelInstance)
+					{
+						return false;
+					}
 
-				return EditingLevelInstance != LevelInstance;
+					LevelInstance = LevelInstanceSubsystem->GetParentLevelInstance(LevelInstance);
+				}
+								
+				return EditingLevelInstance != nullptr;
 			}
 		}
 	}
diff --git a/Engine/Source/Runtime/Engine/Private/LevelInstance/LevelInstanceSubsystem.cpp b/Engine/Source/Runtime/Engine/Private/LevelInstance/LevelInstanceSubsystem.cpp
--- a/Engine/Source/Runtime/Engine/Private/LevelInstance/LevelInstanceSubsystem.cpp
+++ b/Engine/Source/Runtime/Engine/Private/LevelInstance/LevelInstanceSubsystem.cpp
@@ -1588,14 +1588,14 @@ bool ULevelInstanceSubsystem::CanDiscardLevelInstance(const ALevelInstance* Leve
 
 void ULevelInstanceSubsystem::EditLevelInstance(ALevelInstance* LevelInstanceActor, TWeakObjectPtr<AActor> ContextActorPtr)
 {
-	if (EditLevelInstanceInternal(LevelInstanceActor, ContextActorPtr, false))
+	if (EditLevelInstanceInternal(LevelInstanceActor, ContextActorPtr, FString(), false))
 	{
 		ILevelInstanceEditorModule& EditorModule = FModuleManager::GetModuleChecked<ILevelInstanceEditorModule>("LevelInstanceEditor");
 		EditorModule.ActivateEditorMode();
 	}
 }
 
-bool ULevelInstanceSubsystem::EditLevelInstanceInternal(ALevelInstance* LevelInstanceActor, TWeakObjectPtr<AActor> ContextActorPtr, bool bRecursive)
+bool ULevelInstanceSubsystem::EditLevelInstanceInternal(ALevelInstance* LevelInstanceActor, TWeakObjectPtr<AActor> ContextActorPtr, const FString& InActorNameToSelect, bool bRecursive)
 {
 	check(CanEditLevelInstance(LevelInstanceActor));
 		
@@ -1603,7 +1603,7 @@ bool ULevelInstanceSubsystem::EditLevelInstanceInternal(ALevelInstance* LevelIns
 	SlowTask.MakeDialog();
 
 	// Gather information from the context actor to try and select something meaningful after the loading
-	FString ActorNameToSelect;
+	FString ActorNameToSelect = InActorNameToSelect;
 	if (AActor* ContextActor = ContextActorPtr.Get())
 	{
 		ActorNameToSelect = ContextActor->GetName();
@@ -1651,7 +1651,7 @@ bool ULevelInstanceSubsystem::EditLevelInstanceInternal(ALevelInstance* LevelIns
 		ALevelInstance* LevelInstanceToEdit = GetLevelInstance(PendingEditId);
 		check(LevelInstanceToEdit);
 
-		return EditLevelInstanceInternal(LevelInstanceToEdit, nullptr, /*bRecursive=*/true);
+		return EditLevelInstanceInternal(LevelInstanceToEdit, nullptr, ActorNameToSelect, /*bRecursive=*/true);
 	}
 
 	// Cleanup async requests in case
diff --git a/Engine/Source/Runtime/Engine/Public/LevelInstance/LevelInstanceSubsystem.h b/Engine/Source/Runtime/Engine/Public/LevelInstance/LevelInstanceSubsystem.h
--- a/Engine/Source/Runtime/Engine/Public/LevelInstance/LevelInstanceSubsystem.h
+++ b/Engine/Source/Runtime/Engine/Public/LevelInstance/LevelInstanceSubsystem.h
@@ -142,7 +142,7 @@ private:
 	};
 
 	void ResetEdit(TUniquePtr<FLevelInstanceEdit>& InLevelInstanceEdit);
-	bool EditLevelInstanceInternal(ALevelInstance* LevelInstanceActor, TWeakObjectPtr<AActor> ContextActorPtr, bool bRecursive);
+	bool EditLevelInstanceInternal(ALevelInstance* LevelInstanceActor, TWeakObjectPtr<AActor> ContextActorPtr, const FString& InActorNameToSelect, bool bRecursive);
 	bool CommitLevelInstanceInternal(TUniquePtr<FLevelInstanceEdit>& InLevelInstanceEdit, bool bDiscardEdits = false, bool bDiscardOnFailure = false, TSet<FName>* DirtyPackages = nullptr);
 	
 	const FLevelInstanceEdit* GetLevelInstanceEdit(const ALevelInstance* LevelInstanceActor) const;
