WP HLOD - Generalized handling of components using batching policies
* Now handled outside of the individual builders

diff --git a/Engine/Plugins/Editor/WorldPartitionHLODUtilities/Source/Private/HLODBuilderInstancing.cpp b/Engine/Plugins/Editor/WorldPartitionHLODUtilities/Source/Private/HLODBuilderInstancing.cpp
--- a/Engine/Plugins/Editor/WorldPartitionHLODUtilities/Source/Private/HLODBuilderInstancing.cpp
+++ b/Engine/Plugins/Editor/WorldPartitionHLODUtilities/Source/Private/HLODBuilderInstancing.cpp
@@ -2,43 +2,6 @@
 
 #include "HLODBuilderInstancing.h"
 
-#include "Misc/HashBuilder.h"
-#include "Engine/StaticMesh.h"
-#include "Components/InstancedStaticMeshComponent.h"
-#include "WorldPartition/HLOD/HLODActor.h"
-#include "WorldPartition/HLOD/HLODLayer.h"
-
-#include "ISMPartition/ISMComponentDescriptor.h"
-
-namespace 
-{
-	// Instance batcher class based on FISMComponentDescriptor
-	struct FCustomISMComponentDescriptor : public FISMComponentDescriptor
-	{
-		FCustomISMComponentDescriptor(UStaticMeshComponent* SMC)
-		{
-			InitFrom(SMC, false);
-
-			// We'll always want to spawn ISMC, even if our source components are all SMC
-			ComponentClass = UInstancedStaticMeshComponent::StaticClass();
-
-			ComputeHash();
-		}
-	};
-
-	// Store batched instances data
-	struct FInstancingData
-	{
-		int32							NumInstances = 0;
-
-		TArray<FTransform>				InstancesTransforms;
-
-		int32							NumCustomDataFloats = 0;
-		TArray<float>					InstancesCustomData;
-
-		TArray<FInstancedStaticMeshRandomSeed> RandomSeeds;
-	};
-}
 
 UHLODBuilderInstancing::UHLODBuilderInstancing(const FObjectInitializer& ObjectInitializer)
 	: Super(ObjectInitializer)
@@ -47,93 +10,5 @@ UHLODBuilderInstancing::UHLODBuilderInstancing(const FObjectInitializer& ObjectI
 
 TArray<UActorComponent*> UHLODBuilderInstancing::Build(const FHLODBuildContext& InHLODBuildContext, const TArray<UActorComponent*>& InSourceComponents) const
 {
-	TRACE_CPUPROFILER_EVENT_SCOPE(UHLODBuilderInstancing::Build);
-
-	TArray<UStaticMeshComponent*> SourceStaticMeshComponents = FilterComponents<UStaticMeshComponent>(InSourceComponents);
-	
-	// Prepare instance batches
-	TMap<FISMComponentDescriptor, FInstancingData> InstancesData;
-	for (UStaticMeshComponent* SMC : SourceStaticMeshComponents)
-	{
-		FCustomISMComponentDescriptor ISMComponentDescriptor(SMC);
-		FInstancingData& InstancingData = InstancesData.FindOrAdd(ISMComponentDescriptor);
-
-		if (UInstancedStaticMeshComponent* ISMC = Cast<UInstancedStaticMeshComponent>(SMC))
-		{
-			InstancingData.NumCustomDataFloats = FMath::Max(InstancingData.NumCustomDataFloats, ISMC->NumCustomDataFloats);
-			InstancingData.RandomSeeds.Add( { InstancingData.NumInstances, ISMC->InstancingRandomSeed } );
-			InstancingData.NumInstances += ISMC->GetInstanceCount();
-		}
-		else
-		{
-			InstancingData.NumInstances++;
-		}
-	}
-
-	// Resize arrays
-	for (auto& Entry : InstancesData)
-	{
-		const FISMComponentDescriptor& ISMComponentDescriptor = Entry.Key;
-		FInstancingData& EntryInstancingData = Entry.Value;
-
-		EntryInstancingData.InstancesTransforms.Reset(EntryInstancingData.NumInstances);
-		EntryInstancingData.InstancesCustomData.Reset(EntryInstancingData.NumInstances * EntryInstancingData.NumCustomDataFloats);
-	}
-	
-	// Append all transforms & per instance custom data
-	for (UStaticMeshComponent* SMC : SourceStaticMeshComponents)
-	{
-		FCustomISMComponentDescriptor ISMComponentDescriptor(SMC);
-		FInstancingData& InstancingData = InstancesData.FindChecked(ISMComponentDescriptor);
-
-		if (UInstancedStaticMeshComponent* ISMC = Cast<UInstancedStaticMeshComponent>(SMC))
-		{
-			// Add transforms
-			for (int32 InstanceIdx = 0; InstanceIdx < ISMC->GetInstanceCount(); InstanceIdx++)
-			{
-				FTransform& InstanceTransform = InstancingData.InstancesTransforms.AddDefaulted_GetRef();
-				ISMC->GetInstanceTransform(InstanceIdx, InstanceTransform, true);
-			}
-
-			// Add per instance custom data
-			int32 NumCustomDataFloatToAdd = ISMC->GetInstanceCount() * InstancingData.NumCustomDataFloats;
-			InstancingData.InstancesCustomData.Append(ISMC->PerInstanceSMCustomData);
-			InstancingData.InstancesCustomData.AddDefaulted(NumCustomDataFloatToAdd - ISMC->PerInstanceSMCustomData.Num());
-		}
-		else
-		{
-			InstancingData.InstancesTransforms.Add(SMC->GetComponentTransform());
-			InstancingData.InstancesCustomData.AddDefaulted(InstancingData.NumCustomDataFloats);
-		}
-	}
-
-	// Create an ISMC for each SM asset we found
-	TArray<UActorComponent*> HLODComponents;
-	for (auto& Entry : InstancesData)
-	{
-		const FISMComponentDescriptor& ISMComponentDescriptor = Entry.Key;
-		FInstancingData& EntryInstancingData = Entry.Value;
-
-		check(EntryInstancingData.InstancesTransforms.Num() * EntryInstancingData.NumCustomDataFloats == EntryInstancingData.InstancesCustomData.Num());
-
-		UInstancedStaticMeshComponent* Component = ISMComponentDescriptor.CreateComponent(GetTransientPackage());
-		Component->SetForcedLodModel(Component->GetStaticMesh()->GetNumLODs());
-		Component->NumCustomDataFloats = EntryInstancingData.NumCustomDataFloats;
-		Component->AddInstances(EntryInstancingData.InstancesTransforms, /*bShouldReturnIndices*/false, /*bWorldSpace*/true);
-		Component->PerInstanceSMCustomData = MoveTemp(EntryInstancingData.InstancesCustomData);
-
-		if (!EntryInstancingData.RandomSeeds.IsEmpty())
-		{
-			Component->InstancingRandomSeed = EntryInstancingData.RandomSeeds[0].RandomSeed;
-		}
-
-		if (EntryInstancingData.RandomSeeds.Num() > 1)
-		{
-			Component->AdditionalRandomSeeds = TArrayView<FInstancedStaticMeshRandomSeed>(&EntryInstancingData.RandomSeeds[1], EntryInstancingData.RandomSeeds.Num() - 1);
-		}
-
-		HLODComponents.Add(Component);
-	};
-
-	return HLODComponents;
+	return UHLODBuilder::BatchInstances(InSourceComponents);
 }
diff --git a/Engine/Plugins/Editor/WorldPartitionHLODUtilities/Source/Private/HLODBuilderMeshApproximate.cpp b/Engine/Plugins/Editor/WorldPartitionHLODUtilities/Source/Private/HLODBuilderMeshApproximate.cpp
--- a/Engine/Plugins/Editor/WorldPartitionHLODUtilities/Source/Private/HLODBuilderMeshApproximate.cpp
+++ b/Engine/Plugins/Editor/WorldPartitionHLODUtilities/Source/Private/HLODBuilderMeshApproximate.cpp
@@ -20,8 +20,6 @@
 #include "Engine/HLODProxy.h"
 #include "Serialization/ArchiveCrc32.h"
 
-#include "HLODBuilderInstancing.h"
-
 
 UHLODBuilderMeshApproximate::UHLODBuilderMeshApproximate(const FObjectInitializer& ObjectInitializer)
 	: Super(ObjectInitializer)
@@ -73,26 +71,7 @@ TArray<UActorComponent*> UHLODBuilderMeshApproximate::Build(const FHLODBuildCont
 {
 	TRACE_CPUPROFILER_EVENT_SCOPE(UHLODBuilderMeshApproximate::Build);
 
-	TArray<UStaticMeshComponent*> StaticMeshComponents;
-	TArray<UActorComponent*> InstancedComponents;
-
-	// Filter the input components
-	for (UStaticMeshComponent* SubComponent : FilterComponents<UStaticMeshComponent>(InSourceComponents))
-	{
-		switch (SubComponent->HLODBatchingPolicy)
-		{
-		case EHLODBatchingPolicy::None:
-			StaticMeshComponents.Add(SubComponent);
-			break;
-		case EHLODBatchingPolicy::Instancing:
-			InstancedComponents.Add(SubComponent);
-			break;
-		case EHLODBatchingPolicy::MeshSection:
-			InstancedComponents.Add(SubComponent);
-			UE_LOG(LogHLODBuilder, Warning, TEXT("EHLODBatchingPolicy::MeshSection is not yet supported by the UHLODBuilderMeshSimplify builder."));
-			break;
-		}
-	}
+	TArray<UStaticMeshComponent*> StaticMeshComponents = FilterComponents<UStaticMeshComponent>(InSourceComponents);
 
 	TSet<AActor*> Actors;
 	Algo::Transform(StaticMeshComponents, Actors, [](UPrimitiveComponent* PrimitiveComponent) { return PrimitiveComponent->GetOwner(); });
@@ -225,12 +204,5 @@ TArray<UActorComponent*> UHLODBuilderMeshApproximate::Build(const FHLODBuildCont
 		}
 	}
 
-	// Batch instances
-	if (InstancedComponents.Num())
-	{
-		UHLODBuilderInstancing* InstancingHLODBuilder = NewObject<UHLODBuilderInstancing>();
-		Components.Append(InstancingHLODBuilder->Build(InHLODBuildContext, InstancedComponents));
-	}
-
 	return Components;
 }
diff --git a/Engine/Plugins/Editor/WorldPartitionHLODUtilities/Source/Private/HLODBuilderMeshMerge.cpp b/Engine/Plugins/Editor/WorldPartitionHLODUtilities/Source/Private/HLODBuilderMeshMerge.cpp
--- a/Engine/Plugins/Editor/WorldPartitionHLODUtilities/Source/Private/HLODBuilderMeshMerge.cpp
+++ b/Engine/Plugins/Editor/WorldPartitionHLODUtilities/Source/Private/HLODBuilderMeshMerge.cpp
@@ -94,5 +94,5 @@ TArray<UActorComponent*> UHLODBuilderMeshMerge::Build(const FHLODBuildContext& I
 		}
 	});
 
-	return TArray<UActorComponent*>({ Component });
+	return { Component };
 }
diff --git a/Engine/Plugins/Editor/WorldPartitionHLODUtilities/Source/Private/HLODBuilderMeshSimplify.cpp b/Engine/Plugins/Editor/WorldPartitionHLODUtilities/Source/Private/HLODBuilderMeshSimplify.cpp
--- a/Engine/Plugins/Editor/WorldPartitionHLODUtilities/Source/Private/HLODBuilderMeshSimplify.cpp
+++ b/Engine/Plugins/Editor/WorldPartitionHLODUtilities/Source/Private/HLODBuilderMeshSimplify.cpp
@@ -19,8 +19,6 @@
 #include "Engine/HLODProxy.h"
 #include "Serialization/ArchiveCrc32.h"
 
-#include "HLODBuilderInstancing.h"
-
 
 UHLODBuilderMeshSimplify::UHLODBuilderMeshSimplify(const FObjectInitializer& ObjectInitializer)
 	: Super(ObjectInitializer)
@@ -72,26 +70,7 @@ TArray<UActorComponent*> UHLODBuilderMeshSimplify::Build(const FHLODBuildContext
 {
 	TRACE_CPUPROFILER_EVENT_SCOPE(UHLODBuilderMeshSimplifySettings::Build);
 
-	TArray<UStaticMeshComponent*> StaticMeshComponents;
-	TArray<UActorComponent*> InstancedComponents;
-
-	// Filter the input components
-	for (UStaticMeshComponent* SubComponent : FilterComponents<UStaticMeshComponent>(InSourceComponents))
-	{
-		switch (SubComponent->HLODBatchingPolicy)
-		{
-		case EHLODBatchingPolicy::None:
-			StaticMeshComponents.Add(SubComponent);
-			break;
-		case EHLODBatchingPolicy::Instancing:
-			InstancedComponents.Add(SubComponent);
-			break;
-		case EHLODBatchingPolicy::MeshSection:
-			InstancedComponents.Add(SubComponent);
-			UE_LOG(LogHLODBuilder, Warning, TEXT("EHLODBatchingPolicy::MeshSection is not yet supported by the UHLODBuilderMeshSimplify builder."));
-			break;
-		}
-	}
+	TArray<UStaticMeshComponent*> StaticMeshComponents = FilterComponents<UStaticMeshComponent>(InSourceComponents);
 
 	const UHLODBuilderMeshSimplifySettings* MeshSimplifySettings = CastChecked<UHLODBuilderMeshSimplifySettings>(HLODBuilderSettings);
 	const FMeshProxySettings& UseSettings = MeshSimplifySettings->MeshSimplifySettings;
@@ -119,12 +98,5 @@ TArray<UActorComponent*> UHLODBuilderMeshSimplify::Build(const FHLODBuildContext
 	TArray<UActorComponent*> Components;
 	Components.Add(Component);
 
-	// Batch instances
-	if (InstancedComponents.Num())
-	{
-		UHLODBuilderInstancing* InstancingHLODBuilder = NewObject<UHLODBuilderInstancing>();
-		Components.Append(InstancingHLODBuilder->Build(InHLODBuildContext, InstancedComponents));
-	}
-
 	return Components;
 }
diff --git a/Engine/Source/Runtime/Engine/Private/WorldPartition/HLOD/HLODBuilder.cpp b/Engine/Source/Runtime/Engine/Private/WorldPartition/HLOD/HLODBuilder.cpp
--- a/Engine/Source/Runtime/Engine/Private/WorldPartition/HLOD/HLODBuilder.cpp
+++ b/Engine/Source/Runtime/Engine/Private/WorldPartition/HLOD/HLODBuilder.cpp
@@ -3,8 +3,11 @@
 #include "WorldPartition/HLOD/HLODBuilder.h"
 
 #include "AssetCompilingManager.h"
-
 #include "Engine/HLODProxy.h"
+#include "Engine/StaticMesh.h"
+#include "Components/InstancedStaticMeshComponent.h"
+#include "ISMPartition/ISMComponentDescriptor.h"
+
 
 DEFINE_LOG_CATEGORY(LogHLODBuilder);
 
@@ -134,12 +137,181 @@ uint32 UHLODBuilder::ComputeHLODHash(const TArray<AActor*>& InSourceActors)
 	return FCrc::MemCrc32(ComponentsCRCs.GetData(), ComponentsCRCs.Num() * ComponentsCRCs.GetTypeSize());
 }
 
+namespace
+{
+	// Instance batcher class based on FISMComponentDescriptor
+	struct FCustomISMComponentDescriptor : public FISMComponentDescriptor
+	{
+		FCustomISMComponentDescriptor(UStaticMeshComponent* SMC)
+		{
+			InitFrom(SMC, false);
+
+			// We'll always want to spawn ISMC, even if our source components are all SMC
+			ComponentClass = UInstancedStaticMeshComponent::StaticClass();
+
+			ComputeHash();
+		}
+	};
+
+	// Store batched instances data
+	struct FInstancingData
+	{
+		int32							NumInstances = 0;
+
+		TArray<FTransform>				InstancesTransforms;
+
+		int32							NumCustomDataFloats = 0;
+		TArray<float>					InstancesCustomData;
+
+		TArray<FInstancedStaticMeshRandomSeed> RandomSeeds;
+	};
+}
+
+TArray<UActorComponent*> UHLODBuilder::BatchInstances(const TArray<UActorComponent*>& InSourceComponents)
+{
+	TRACE_CPUPROFILER_EVENT_SCOPE(UHLODBuilderInstancing::Build);
+
+	TArray<UStaticMeshComponent*> SourceStaticMeshComponents = FilterComponents<UStaticMeshComponent>(InSourceComponents);
+
+	// Prepare instance batches
+	TMap<FISMComponentDescriptor, FInstancingData> InstancesData;
+	for (UStaticMeshComponent* SMC : SourceStaticMeshComponents)
+	{
+		FCustomISMComponentDescriptor ISMComponentDescriptor(SMC);
+		FInstancingData& InstancingData = InstancesData.FindOrAdd(ISMComponentDescriptor);
+
+		if (UInstancedStaticMeshComponent* ISMC = Cast<UInstancedStaticMeshComponent>(SMC))
+		{
+			InstancingData.NumCustomDataFloats = FMath::Max(InstancingData.NumCustomDataFloats, ISMC->NumCustomDataFloats);
+			InstancingData.RandomSeeds.Add({ InstancingData.NumInstances, ISMC->InstancingRandomSeed });
+			InstancingData.NumInstances += ISMC->GetInstanceCount();
+		}
+		else
+		{
+			InstancingData.NumInstances++;
+		}
+	}
+
+	// Resize arrays
+	for (auto& Entry : InstancesData)
+	{
+		const FISMComponentDescriptor& ISMComponentDescriptor = Entry.Key;
+		FInstancingData& EntryInstancingData = Entry.Value;
+
+		EntryInstancingData.InstancesTransforms.Reset(EntryInstancingData.NumInstances);
+		EntryInstancingData.InstancesCustomData.Reset(EntryInstancingData.NumInstances * EntryInstancingData.NumCustomDataFloats);
+	}
+
+	// Append all transforms & per instance custom data
+	for (UStaticMeshComponent* SMC : SourceStaticMeshComponents)
+	{
+		FCustomISMComponentDescriptor ISMComponentDescriptor(SMC);
+		FInstancingData& InstancingData = InstancesData.FindChecked(ISMComponentDescriptor);
+
+		if (UInstancedStaticMeshComponent* ISMC = Cast<UInstancedStaticMeshComponent>(SMC))
+		{
+			// Add transforms
+			for (int32 InstanceIdx = 0; InstanceIdx < ISMC->GetInstanceCount(); InstanceIdx++)
+			{
+				FTransform& InstanceTransform = InstancingData.InstancesTransforms.AddDefaulted_GetRef();
+				ISMC->GetInstanceTransform(InstanceIdx, InstanceTransform, true);
+			}
+
+			// Add per instance custom data
+			int32 NumCustomDataFloatToAdd = ISMC->GetInstanceCount() * InstancingData.NumCustomDataFloats;
+			InstancingData.InstancesCustomData.Append(ISMC->PerInstanceSMCustomData);
+			InstancingData.InstancesCustomData.AddDefaulted(NumCustomDataFloatToAdd - ISMC->PerInstanceSMCustomData.Num());
+		}
+		else
+		{
+			InstancingData.InstancesTransforms.Add(SMC->GetComponentTransform());
+			InstancingData.InstancesCustomData.AddDefaulted(InstancingData.NumCustomDataFloats);
+		}
+	}
+
+	// Create an ISMC for each SM asset we found
+	TArray<UActorComponent*> HLODComponents;
+	for (auto& Entry : InstancesData)
+	{
+		const FISMComponentDescriptor& ISMComponentDescriptor = Entry.Key;
+		FInstancingData& EntryInstancingData = Entry.Value;
+
+		check(EntryInstancingData.InstancesTransforms.Num() * EntryInstancingData.NumCustomDataFloats == EntryInstancingData.InstancesCustomData.Num());
+
+		UInstancedStaticMeshComponent* Component = ISMComponentDescriptor.CreateComponent(GetTransientPackage());
+		Component->SetForcedLodModel(Component->GetStaticMesh()->GetNumLODs());
+		Component->NumCustomDataFloats = EntryInstancingData.NumCustomDataFloats;
+		Component->AddInstances(EntryInstancingData.InstancesTransforms, /*bShouldReturnIndices*/false, /*bWorldSpace*/true);
+		Component->PerInstanceSMCustomData = MoveTemp(EntryInstancingData.InstancesCustomData);
+
+		if (!EntryInstancingData.RandomSeeds.IsEmpty())
+		{
+			Component->InstancingRandomSeed = EntryInstancingData.RandomSeeds[0].RandomSeed;
+		}
+
+		if (EntryInstancingData.RandomSeeds.Num() > 1)
+		{
+			Component->AdditionalRandomSeeds = TArrayView<FInstancedStaticMeshRandomSeed>(&EntryInstancingData.RandomSeeds[1], EntryInstancingData.RandomSeeds.Num() - 1);
+		}
+
+		HLODComponents.Add(Component);
+	};
+
+	return HLODComponents;
+}
+
+static bool ShouldBatchComponent(UActorComponent* ActorComponent)
+{
+	bool bShouldBatch = false;
+
+	if (UPrimitiveComponent* PrimitiveComponent = Cast<UPrimitiveComponent>(ActorComponent))
+	{
+		switch (PrimitiveComponent->HLODBatchingPolicy)
+		{
+		case EHLODBatchingPolicy::None:
+			break;
+		case EHLODBatchingPolicy::Instancing:
+			bShouldBatch = true;
+			break;
+		case EHLODBatchingPolicy::MeshSection:
+			bShouldBatch = true;
+			UE_LOG(LogHLODBuilder, Warning, TEXT("EHLODBatchingPolicy::MeshSection is not yet supported by the HLOD builder, falling back to EHLODBatchingPolicy::Instancing."));
+			break;
+		default:
+			checkNoEntry();
+		}
+	}
+
+	return bShouldBatch;
+}
+
 TArray<UActorComponent*> UHLODBuilder::Build(const FHLODBuildContext& InHLODBuildContext, const TArray<AActor*>& InSourceActors) const
 {
+	TArray<UActorComponent*> HLODRelevantComponents = GatherHLODRelevantComponents(InSourceActors);
+	if (HLODRelevantComponents.IsEmpty())
+	{
+		return {};
+	}
+
+	// Handle components using a batching policy separately
+	TArray<UActorComponent*> InputComponents;
+	TArray<UActorComponent*> ComponentsToBatch;
+	for (UActorComponent* SourceComponent : HLODRelevantComponents)
+	{
+		if (ShouldBatchComponent(SourceComponent))
+		{
+			ComponentsToBatch.Add(SourceComponent);
+		}
+		else
+		{
+			InputComponents.Add(SourceComponent);
+		}
+	}
+
 	TMap<TSubclassOf<UHLODBuilder>, TArray<UActorComponent*>> HLODBuildersForComponents;
 
 	// Gather custom HLOD builders, and regroup all components by builders
-	for (UActorComponent* SourceComponent : GatherHLODRelevantComponents(InSourceActors))
+	for (UActorComponent* SourceComponent : InputComponents)
 	{
 		TSubclassOf<UHLODBuilder> HLODBuilderClass = SourceComponent->GetCustomHLODBuilderClass();
 		HLODBuildersForComponents.FindOrAdd(HLODBuilderClass).Add(SourceComponent);
@@ -168,6 +340,9 @@ TArray<UActorComponent*> UHLODBuilder::Build(const FHLODBuildContext& InHLODBuil
 		HLODComponents.Append(NewComponents);
 	}
 
+	// Append batched components
+	HLODComponents.Append(BatchInstances(ComponentsToBatch));
+
 	// In case a builder returned null entries, clean the array.
 	HLODComponents.RemoveSwap(nullptr);
 
diff --git a/Engine/Source/Runtime/Engine/Public/WorldPartition/HLOD/HLODBuilder.h b/Engine/Source/Runtime/Engine/Public/WorldPartition/HLOD/HLODBuilder.h
--- a/Engine/Source/Runtime/Engine/Public/WorldPartition/HLOD/HLODBuilder.h
+++ b/Engine/Source/Runtime/Engine/Public/WorldPartition/HLOD/HLODBuilder.h
@@ -110,6 +110,8 @@ public:
 	static uint32 ComputeHLODHash(const TArray<AActor*>& InSourceActors);
 
 protected:
+	static TArray<UActorComponent*> BatchInstances(const TArray<UActorComponent*>& InSubComponents);
+
 	template <typename TComponentClass>
 	static inline TArray<TComponentClass*> FilterComponents(const TArray<UActorComponent*>& InSourceComponents)
 	{
