World Partition: fix crash in replay code when streaming partitioned sublevels.

#rnx

diff --git a/Engine/Source/Runtime/Engine/Private/WorldPartition/WorldPartitionReplay.cpp b/Engine/Source/Runtime/Engine/Private/WorldPartition/WorldPartitionReplay.cpp
--- a/Engine/Source/Runtime/Engine/Private/WorldPartition/WorldPartitionReplay.cpp
+++ b/Engine/Source/Runtime/Engine/Private/WorldPartition/WorldPartitionReplay.cpp
@@ -98,13 +98,13 @@ void AWorldPartitionReplay::Initialize(UWorld* World)
 bool AWorldPartitionReplay::IsEnabled(UWorld* World)
 {
 #if !UE_BUILD_SHIPPING
-	check(World->IsGameWorld());
-	check(World->IsPartitionedWorld());
-	UWorldPartition* WorldPartition = World->GetWorldPartition();
-	return WorldPartition->Replay && WorldPartition->Replay->IsEnabled();
-#else
-	return false;
+	if (UWorldPartition* WorldPartition = World->GetWorldPartition())
+	{
+		return WorldPartition->Replay && WorldPartition->Replay->IsEnabled();
+	}
 #endif
+
+	return false;
 }
 
 void AWorldPartitionReplay::BeginPlay()
