Make sure overlay is always rendered after base mesh

diff --git a/Engine/Source/Runtime/Engine/Private/HierarchicalInstancedStaticMesh.cpp b/Engine/Source/Runtime/Engine/Private/HierarchicalInstancedStaticMesh.cpp
--- a/Engine/Source/Runtime/Engine/Private/HierarchicalInstancedStaticMesh.cpp
+++ b/Engine/Source/Runtime/Engine/Private/HierarchicalInstancedStaticMesh.cpp
@@ -1344,6 +1344,8 @@ void FHierarchicalStaticMeshSceneProxy::FillDynamicMeshElements(FMeshElementColl
 							OverlayMeshBatch.CastShadow = false;
 							OverlayMeshBatch.bSelectable = false;
 							OverlayMeshBatch.MaterialRenderProxy = OverlayMaterial->GetRenderProxy();
+							// make sure overlay is always rendered on top of base mesh
+							OverlayMeshBatch.MeshIdInPrimitive += LODModel.Sections.Num();
 							Collector.AddMesh(ElementParams.ViewIndex, OverlayMeshBatch);
 						}
 					}
@@ -1472,6 +1474,8 @@ void FHierarchicalStaticMeshSceneProxy::FillDynamicMeshElements(FMeshElementColl
 								OverlayMeshBatch.CastShadow = false;
 								OverlayMeshBatch.bSelectable = false;
 								OverlayMeshBatch.MaterialRenderProxy = OverlayMaterial->GetRenderProxy();
+								// make sure overlay is always rendered on top of base mesh
+								OverlayMeshBatch.MeshIdInPrimitive += LODModel.Sections.Num();
 								Collector.AddMesh(ElementParams.ViewIndex, OverlayMeshBatch);
 							}
 						}
diff --git a/Engine/Source/Runtime/Engine/Private/InstancedStaticMesh.cpp b/Engine/Source/Runtime/Engine/Private/InstancedStaticMesh.cpp
--- a/Engine/Source/Runtime/Engine/Private/InstancedStaticMesh.cpp
+++ b/Engine/Source/Runtime/Engine/Private/InstancedStaticMesh.cpp
@@ -1315,6 +1315,8 @@ void FInstancedStaticMeshSceneProxy::GetDynamicMeshElements(const TArray<const F
 								OverlayMeshBatch.CastShadow = false;
 								OverlayMeshBatch.bSelectable = false;
 								OverlayMeshBatch.MaterialRenderProxy = OverlayMaterial->GetRenderProxy();
+								// make sure overlay is always rendered on top of base mesh
+								OverlayMeshBatch.MeshIdInPrimitive += LODModel.Sections.Num();
 								Collector.AddMesh(ViewIndex, OverlayMeshBatch);
 								
 								INC_DWORD_STAT_BY(STAT_StaticMeshTriangles, OverlayMeshBatch.GetNumPrimitives());
diff --git a/Engine/Source/Runtime/Engine/Private/SkeletalMesh.cpp b/Engine/Source/Runtime/Engine/Private/SkeletalMesh.cpp
--- a/Engine/Source/Runtime/Engine/Private/SkeletalMesh.cpp
+++ b/Engine/Source/Runtime/Engine/Private/SkeletalMesh.cpp
@@ -6422,6 +6422,8 @@ void FSkeletalMeshSceneProxy::DrawStaticElements(FStaticPrimitiveDrawInterface*
 						OverlayMeshBatch.CastShadow = false;
 						OverlayMeshBatch.bSelectable = false;
 						OverlayMeshBatch.MaterialRenderProxy = OverlayMaterial->GetRenderProxy();
+						// make sure overlay is always rendered on top of base mesh
+						OverlayMeshBatch.MeshIdInPrimitive += LODData.RenderSections.Num();
 						// Reuse mesh ScreenSize as cull distance for an overlay. Overlay does not need to compute LOD so we can avoid adding new members into MeshBatch or MeshRelevance
 						float OverlayMeshScreenSize = OverlayMaterialMaxDrawDistance;
 						PDI->DrawMesh(OverlayMeshBatch, OverlayMeshScreenSize);
@@ -6734,6 +6736,8 @@ void FSkeletalMeshSceneProxy::GetDynamicElementsSection(const TArray<const FScen
 					OverlayMeshBatch.CastShadow = false;
 					OverlayMeshBatch.bSelectable = false;
 					OverlayMeshBatch.MaterialRenderProxy = OverlayMaterial->GetRenderProxy();
+					// make sure overlay is always rendered on top of base mesh
+					OverlayMeshBatch.MeshIdInPrimitive += LODData.RenderSections.Num();
 					Collector.AddMesh(ViewIndex, OverlayMeshBatch);
 				
 					INC_DWORD_STAT_BY(STAT_SkelMeshTriangles, OverlayMeshBatch.GetNumPrimitives());
diff --git a/Engine/Source/Runtime/Engine/Private/StaticMeshRender.cpp b/Engine/Source/Runtime/Engine/Private/StaticMeshRender.cpp
--- a/Engine/Source/Runtime/Engine/Private/StaticMeshRender.cpp
+++ b/Engine/Source/Runtime/Engine/Private/StaticMeshRender.cpp
@@ -1202,6 +1202,8 @@ void FStaticMeshSceneProxy::DrawStaticElements(FStaticPrimitiveDrawInterface* PD
 							OverlayMeshBatch.CastShadow = false;
 							OverlayMeshBatch.bSelectable = false;
 							OverlayMeshBatch.MaterialRenderProxy = OverlayMaterial->GetRenderProxy();
+							// make sure overlay is always rendered on top of base mesh
+							OverlayMeshBatch.MeshIdInPrimitive += LODModel.Sections.Num();
 							PDI->DrawMesh(OverlayMeshBatch, FLT_MAX);
 						}
 					}
@@ -1335,6 +1337,8 @@ void FStaticMeshSceneProxy::DrawStaticElements(FStaticPrimitiveDrawInterface* PD
 								OverlayMeshBatch.CastShadow = false;
 								OverlayMeshBatch.bSelectable = false;
 								OverlayMeshBatch.MaterialRenderProxy = OverlayMaterial->GetRenderProxy();
+								// make sure overlay is always rendered on top of base mesh
+								OverlayMeshBatch.MeshIdInPrimitive += LODModel.Sections.Num();
 								// Reuse mesh ScreenSize as cull distance for an overlay. Overlay does not need to compute LOD so we can avoid adding new members into MeshBatch or MeshRelevance
 								float OverlayMeshScreenSize = OverlayMaterialMaxDrawDistance;
 								PDI->DrawMesh(OverlayMeshBatch, OverlayMeshScreenSize);
