ImageUtils small cleanup in ExportTextureSourceToDDS

diff --git a/Engine/Source/Runtime/Engine/Private/ImageUtils.cpp b/Engine/Source/Runtime/Engine/Private/ImageUtils.cpp
--- a/Engine/Source/Runtime/Engine/Private/ImageUtils.cpp
+++ b/Engine/Source/Runtime/Engine/Private/ImageUtils.cpp
@@ -157,14 +157,15 @@ bool FImageUtils::CompressImage(TArray64<uint8> & OutData, const TCHAR * ToForma
 bool FImageUtils::ExportTextureSourceToDDS(TArray64<uint8> & OutData, UTexture * Texture, int BlockIndex, int LayerIndex)
 {
 #if WITH_EDITORONLY_DATA	
-	UE::DDS::EDDSError Error;
 
-	UTexture2D * Texture2D = Cast<UTexture2D>(Texture);
-	UTextureCube * TextureCube = Cast<UTextureCube>(Texture);
-	UTextureCubeArray * TextureCubeArray = Cast<UTextureCubeArray>(Texture);
-	UVolumeTexture * TextureVolume = Cast<UVolumeTexture>(Texture);
 	const FTextureSource & Source = Texture->Source;
-	
+
+	if ( ! Source.IsValid() )
+	{
+		UE_LOG(LogImageUtils,Warning,TEXT("ExportTextureSourceToDDS: Texture does not have source (%s)"),*Texture->GetName());
+		return false;
+	}
+
 	check( BlockIndex < Source.GetNumBlocks() );
 	check( LayerIndex < Source.GetNumLayers() );
 	
@@ -180,13 +181,14 @@ bool FImageUtils::ExportTextureSourceToDDS(TArray64<uint8> & OutData, UTexture *
 	int32 SizeZ;
 	int32 ArraySize;
 	uint32 CreateFlags = 0;
-	if ( Texture2D != nullptr )
+	ETextureClass TextureClass = Texture->GetTextureClass();
+	if ( TextureClass == ETextureClass::TwoD )
 	{
 		Dimension = 2;
 		SizeZ = 1;
 		ArraySize = NumSlices;
 	}
-	else if ( TextureCube != nullptr )
+	else if ( TextureClass == ETextureClass::Cube )
 	{
 		Dimension = 2;
 		SizeZ = 1;
@@ -210,7 +212,7 @@ bool FImageUtils::ExportTextureSourceToDDS(TArray64<uint8> & OutData, UTexture *
 			CreateFlags = UE::DDS::FDDSFile::CREATE_FLAG_CUBEMAP;
 		}
 	}
-	else if ( TextureCubeArray != nullptr )
+	else if ( TextureClass == ETextureClass::CubeArray )
 	{
 		Dimension = 2;
 		SizeZ = 1;
@@ -230,7 +232,7 @@ bool FImageUtils::ExportTextureSourceToDDS(TArray64<uint8> & OutData, UTexture *
 			CreateFlags = UE::DDS::FDDSFile::CREATE_FLAG_CUBEMAP;
 		}
 	}
-	else if ( TextureVolume != nullptr )
+	else if ( TextureClass == ETextureClass::Volume )
 	{
 		Dimension = 3;
 		SizeZ = NumSlices;
@@ -240,7 +242,7 @@ bool FImageUtils::ExportTextureSourceToDDS(TArray64<uint8> & OutData, UTexture *
 	}
 	else
 	{
-		check(0);
+		checkf(false, TEXT("unexpected TextureClass"));
 		return false;
 	}
 	
@@ -255,7 +257,8 @@ bool FImageUtils::ExportTextureSourceToDDS(TArray64<uint8> & OutData, UTexture *
 	}
 
 	UE_LOG(LogImageUtils,Display,TEXT("Exporting DDS Dimension=%d SizeX=%d SizeY=%d SizeZ=%d NumMips=%d ArraySize=%d"),Dimension, SizeX,SizeY,SizeZ,NumMips,ArraySize);	
-
+	
+	UE::DDS::EDDSError Error;
 	UE::DDS::FDDSFile * DDS = UE::DDS::FDDSFile::CreateEmpty(Dimension, SizeX,SizeY,SizeZ,NumMips,ArraySize, DXGIFormat,CreateFlags, &Error);
 	if ( DDS == nullptr || Error != UE::DDS::EDDSError::OK )
 	{
