Fixing UDemoNetDriver::FinalizeFastForward and FastForwardLevels not calling DiffProperties on actor subobjects

diff --git a/Engine/Source/Runtime/Engine/Classes/Engine/DemoNetDriver.h b/Engine/Source/Runtime/Engine/Classes/Engine/DemoNetDriver.h
--- a/Engine/Source/Runtime/Engine/Classes/Engine/DemoNetDriver.h
+++ b/Engine/Source/Runtime/Engine/Classes/Engine/DemoNetDriver.h
@@ -676,6 +676,9 @@ private:
 
 	void OnPostLoadMapWithWorld(UWorld* World);
 
+	// Diff Actor plus it's Components and Subobjects from given ActorChannel
+	void DiffActorProperties(UActorChannel* const ActorChannel);
+
 protected:
 
 	void ProcessSeamlessTravel(int32 LevelIndex);
diff --git a/Engine/Source/Runtime/Engine/Private/DemoNetDriver.cpp b/Engine/Source/Runtime/Engine/Private/DemoNetDriver.cpp
--- a/Engine/Source/Runtime/Engine/Private/DemoNetDriver.cpp
+++ b/Engine/Source/Runtime/Engine/Private/DemoNetDriver.cpp
@@ -1149,6 +1149,42 @@ void UDemoNetDriver::OnPostLoadMapWithWorld(UWorld* InWorld)
 	}
 }
 
+void UDemoNetDriver::DiffActorProperties(UActorChannel* const ActorChannel)
+{
+	if (!ActorChannel || !ActorChannel->GetActor())
+	{
+		return;
+	}
+
+	auto DiffObjectProperties = [](const FObjectReplicator& ObjectReplicator)
+	{
+		if (const UObject* const ReplicatedObject = ObjectReplicator.GetObject())
+		{
+			FReceivingRepState* const ReceivingRepState = ObjectReplicator.RepState->GetReceivingRepState();
+			const FRepShadowDataBuffer ShadowData(ReceivingRepState->StaticBuffer.GetData());
+			const FConstRepObjectDataBuffer RepObjectData(ReplicatedObject);
+
+			ObjectReplicator.RepLayout->DiffProperties(&(ReceivingRepState->RepNotifies), ShadowData, RepObjectData, EDiffPropertiesFlags::Sync);
+		}
+	};
+
+	// Make sure we diff Actor first
+	const FObjectReplicator& ActorReplicator = ActorChannel->GetActorReplicationData();
+	DiffObjectProperties(ActorReplicator);
+
+	// Diff any Components and SubObjects
+	for (const auto& ReplicatorPair : ActorChannel->ReplicationMap)
+	{
+		const FObjectReplicator& ObjectReplicator = ReplicatorPair.Value.Get();
+		// We don't need to diff Actor again
+		if (ActorReplicator.GetObject() == ObjectReplicator.GetObject())
+		{
+			continue;
+		}
+		DiffObjectProperties(ObjectReplicator);
+	}
+}
+
 bool UDemoNetDriver::ContinueListen(FURL& ListenURL)
 {
 	if (IsRecording() && ensure(IsRecordingPaused()))
@@ -2923,16 +2959,9 @@ void UDemoNetDriver::FinalizeFastForward(const double StartTime)
 				continue;
 			}
 
-			if (const FObjectReplicator* const ActorReplicator = ActorChannel->ActorReplicator.Get())
+			if (Actor->IsNetStartupActor())
 			{
-				if (Actor->IsNetStartupActor())
-				{
-					FReceivingRepState* ReceivingRepState = ActorReplicator->RepState->GetReceivingRepState();
-					FRepShadowDataBuffer ShadowData(ReceivingRepState->StaticBuffer.GetData());
-					FConstRepObjectDataBuffer ActorData(Actor);
-
-					ActorReplicator->RepLayout->DiffProperties(&(ReceivingRepState->RepNotifies), ShadowData, ActorData, EDiffPropertiesFlags::Sync);
-				}
+				DiffActorProperties(ActorChannel);
 			}
 		}
 	}
@@ -3790,14 +3819,7 @@ bool UDemoNetDriver::FastForwardLevels(const FGotoResult& GotoResult)
 				{
 					ChannelsToUpdate.Add(ActorChannel);
 
-					if (const FObjectReplicator* const ActorReplicator = ActorChannel->ActorReplicator.Get())
-					{
-						FReceivingRepState* ReceivingRepState = ActorReplicator->RepState->GetReceivingRepState();
-						FRepShadowDataBuffer ShadowData(ReceivingRepState->StaticBuffer.GetData());
-						FConstRepObjectDataBuffer ActorData(Actor);
-
-						ActorReplicator->RepLayout->DiffProperties(&(ReceivingRepState->RepNotifies), ShadowData, ActorData, EDiffPropertiesFlags::Sync);
-					}
+					DiffActorProperties(ActorChannel);
 				}
 			}
 		}
