Strata - fix for compil;ation crash in some case when BSDF is used(or registred) multiple time in a graph.
MAde errors visible from the mateiral window instead of the log.

diff --git a/Engine/Source/Runtime/Engine/Private/Materials/HLSLMaterialTranslator.cpp b/Engine/Source/Runtime/Engine/Private/Materials/HLSLMaterialTranslator.cpp
--- a/Engine/Source/Runtime/Engine/Private/Materials/HLSLMaterialTranslator.cpp
+++ b/Engine/Source/Runtime/Engine/Private/Materials/HLSLMaterialTranslator.cpp
@@ -2141,7 +2141,7 @@ void FHLSLMaterialTranslator::GetMaterialEnvironment(EShaderPlatform InPlatform,
 
 			if (StrataMaterialAnalysis.ClampedLayerCount == 0)
 			{
-				UE_LOG(LogMaterial, Error, TEXT("Material %s cannot have any layers rendered due to its complexity (asset: %s).\r\n"), *Material->GetDebugName(), *Material->GetAssetPath().ToString());
+				Errorf(TEXT("Material %s cannot have any layers rendered due to its complexity (asset: %s).\r\n"), *Material->GetDebugName(), *Material->GetAssetPath().ToString());
 			}
 			else if (StrataMaterialAnalysis.RequestedLayerCount > StrataMaterialAnalysis.ClampedLayerCount)
 			{
@@ -9624,24 +9624,26 @@ FStrataOperator& FHLSLMaterialTranslator::StrataCompilationRegisterOperator(int3
 		OperatorType = STRATA_OPERATOR_BSDF;
 	}
 
+	static FStrataOperator DefaultOperatorOnError = FStrataOperator();
+
 	if (StrataMaterialExpressionToOperatorIndex.Find(Expression))
 	{
-		UE_LOG(LogMaterial, Error, TEXT("Material %s tries to register a Strata operator twice (asset: %s).\r\n"), *Material->GetDebugName(), *Material->GetAssetPath().ToString());
-		return StrataMaterialExpressionRegisteredOperators[0];
+		Errorf(TEXT("Material %s tries to register a Strata operator twice (asset: %s).\r\n"), *Material->GetDebugName(), *Material->GetAssetPath().ToString());
+		return DefaultOperatorOnError;
 	}
 
 	const uint32 NewOperatorIndex = StrataMaterialExpressionToOperatorIndex.Num();
 	if (NewOperatorIndex >= STRATA_MAX_OPERATOR_COUNT)
 	{
-		UE_LOG(LogMaterial, Error, TEXT("Material %s have too many Strata Operators (asset: %s).\r\n"), *Material->GetDebugName(), *Material->GetAssetPath().ToString());
-		return StrataMaterialExpressionRegisteredOperators[0];
+		Errorf(TEXT("Material %s have too many Strata Operators (asset: %s).\r\n"), *Material->GetDebugName(), *Material->GetAssetPath().ToString());
+		return DefaultOperatorOnError;
 	}
 
 	int32* ParentOperatorIndex = StrataMaterialExpressionToOperatorIndex.Find(Parent);
 	if (Parent!=nullptr && ParentOperatorIndex == nullptr)
 	{
-		UE_LOG(LogMaterial, Error, TEXT("Material %s tries to register unknown operator parents (asset: %s).\r\n"), *Material->GetDebugName(), *Material->GetAssetPath().ToString());
-		return StrataMaterialExpressionRegisteredOperators[0];
+		Errorf(TEXT("Material %s tries to register unknown operator parents (asset: %s).\r\n"), *Material->GetDebugName(), *Material->GetAssetPath().ToString());
+		return DefaultOperatorOnError;
 	}
 
 	StrataMaterialExpressionToOperatorIndex.Add(Expression, NewOperatorIndex);
@@ -9669,9 +9671,8 @@ FStrataOperator& FHLSLMaterialTranslator::StrataCompilationGetOperator(UMaterial
 	auto* OperatorIndex = StrataMaterialExpressionToOperatorIndex.Find(Expression);
 	if (!(OperatorIndex && *OperatorIndex >= 0 && *OperatorIndex < STRATA_MAX_OPERATOR_COUNT))
 	{
-		UE_LOG(LogMaterial, Error, TEXT("Could not find operator for a BSDF in Material %s (asset: %s).\r\n"), *Material->GetDebugName(), *Material->GetAssetPath().ToString());
-		static FStrataOperator DefaultOperator = FStrataOperator();
-		return DefaultOperator;
+		static FStrataOperator DefaultOperatorOnError = FStrataOperator();
+		return DefaultOperatorOnError;
 	};
 	return StrataMaterialExpressionRegisteredOperators[*OperatorIndex];
 }
@@ -9680,7 +9681,7 @@ bool FHLSLMaterialTranslator::StrataGenerateDerivedMaterialOperatorData()
 {
 	if (StrataMaterialExpressionRegisteredOperators.IsEmpty())
 	{
-		UE_LOG(LogMaterial, Error, TEXT("Could not find any strata operators or BSDFs in Material %s (asset: %s).\r\n"), *Material->GetDebugName(), *Material->GetAssetPath().ToString());
+		Errorf(TEXT("Could not find any strata operators or BSDFs in Material %s (asset: %s).\r\n"), *Material->GetDebugName(), *Material->GetAssetPath().ToString());
 		return false;
 	}
 
@@ -9699,7 +9700,7 @@ bool FHLSLMaterialTranslator::StrataGenerateDerivedMaterialOperatorData()
 	StrataMaterialRootOperator = &StrataMaterialExpressionRegisteredOperators[RootIndex];
 	if (!StrataMaterialRootOperator)
 	{
-		UE_LOG(LogMaterial, Error, TEXT("Cannot find the root of the Strata Tree for Material %s (asset: %s).\r\n"), *Material->GetDebugName(), *Material->GetAssetPath().ToString());
+		Errorf(TEXT("Cannot find the root of the Strata Tree for Material %s (asset: %s).\r\n"), *Material->GetDebugName(), *Material->GetAssetPath().ToString());
 		return false;
 	}
 
@@ -9742,12 +9743,12 @@ bool FHLSLMaterialTranslator::StrataGenerateDerivedMaterialOperatorData()
 
 		if (bMustHaveLeftChild && It.LeftIndex == INDEX_NONE)
 		{
-			UE_LOG(LogMaterial, Error, TEXT("A Strata Operator %s node is missing its first input from material %s (asset: %s).\r\n"), GetStrataOperatorStr(It.OperatorType), *Material->GetDebugName(), *Material->GetAssetPath().ToString());
+			Errorf(TEXT("A Strata Operator %s node is missing its first input from material %s (asset: %s).\r\n"), GetStrataOperatorStr(It.OperatorType), *Material->GetDebugName(), *Material->GetAssetPath().ToString());
 			return false;
 		}
 		if (bMustHaveRightChild && It.RightIndex == INDEX_NONE)
 		{
-			UE_LOG(LogMaterial, Error, TEXT("A Strata Operator %s node is missing its second input from material %s (asset: %s).\r\n"), GetStrataOperatorStr(It.OperatorType), *Material->GetDebugName(), *Material->GetAssetPath().ToString());
+			Errorf(TEXT("A Strata Operator %s node is missing its second input from material %s (asset: %s).\r\n"), GetStrataOperatorStr(It.OperatorType), *Material->GetDebugName(), *Material->GetAssetPath().ToString());
 			return false;
 		}
 	}
@@ -10130,7 +10131,7 @@ int32 FHLSLMaterialTranslator::StrataSlabBSDF(
 	{
 		if (PromoteToOperator->Index == INDEX_NONE || PromoteToOperator->BSDFIndex == INDEX_NONE)
 		{
-			UE_LOG(LogMaterial, Error, TEXT("Invalid StrataSlabBSDF operator and BSDF indices during promotion in Material %s (asset: %s).\r\n"), *Material->GetDebugName(), *Material->GetAssetPath().ToString());
+			Errorf(TEXT("Invalid StrataSlabBSDF operator and BSDF indices during promotion in Material %s (asset: %s).\r\n"), *Material->GetDebugName(), *Material->GetAssetPath().ToString());
 			return INDEX_NONE;
 		}
 		return AddCodeChunk(
