Strata - fix for crash in editor windows when an error occur in the material graph.

diff --git a/Engine/Source/Runtime/Engine/Private/Materials/MaterialShared.cpp b/Engine/Source/Runtime/Engine/Private/Materials/MaterialShared.cpp
--- a/Engine/Source/Runtime/Engine/Private/Materials/MaterialShared.cpp
+++ b/Engine/Source/Runtime/Engine/Private/Materials/MaterialShared.cpp
@@ -2529,36 +2529,40 @@ bool FMaterial::CacheShaders(const FMaterialShaderMapId& ShaderMapId, EShaderPla
 #if WITH_EDITOR
 void FMaterial::CacheGivenTypes(EShaderPlatform Platform, const TArray<FVertexFactoryType*>& VFTypes, const TArray<FShaderType*>& ShaderTypes, const ITargetPlatform* TargetPlatform)
 {
-	TRACE_CPUPROFILER_EVENT_SCOPE(FMaterial::CacheGivenTypes);
-	check(IsInGameThread());
-	checkf(ShaderTypes.Num() == VFTypes.Num(), TEXT("The size of the shader type array and vertex factory type array must match."));
-	checkf(GameThreadShaderMap, TEXT("Shader map is not initialized.  Please call CacheShaders first."));
-	checkf(GetGameThreadCompilingShaderMapId() != 0, TEXT("Material is not prepared to compile yet.  Please call CacheShaders first."));
-
-	TArray<FShaderCommonCompileJobPtr> CompileJobs;
-	for (int i = 0; i < VFTypes.Num(); ++i)
+	if (GameThreadShaderMap)
 	{
-		FVertexFactoryType* VFType = VFTypes[i];
-		check(VFType);
-
-		FShaderType* ShaderType = ShaderTypes[i];
-		check(ShaderType);
+		TRACE_CPUPROFILER_EVENT_SCOPE(FMaterial::CacheGivenTypes);
+		check(IsInGameThread());
+		checkf(ShaderTypes.Num() == VFTypes.Num(), TEXT("The size of the shader type array and vertex factory type array must match."));
+		checkf(GameThreadShaderMap, TEXT("Shader map is not initialized.  Please call CacheShaders first."));
+		checkf(GetGameThreadCompilingShaderMapId() != 0, TEXT("Material is not prepared to compile yet.  Please call CacheShaders first."));
+
+		TArray<FShaderCommonCompileJobPtr> CompileJobs;
+		for (int i = 0; i < VFTypes.Num(); ++i)
+		{
+			FVertexFactoryType* VFType = VFTypes[i];
+			check(VFType);
+
+			FShaderType* ShaderType = ShaderTypes[i];
+			check(ShaderType);
+			check(ShaderType->AsMeshMaterialShaderType());
+
+			ShaderType->AsMeshMaterialShaderType()->BeginCompileShader(
+				EShaderCompileJobPriority::ForceLocal,
+				GetGameThreadCompilingShaderMapId(),
+				0,
+				Platform,
+				GameThreadShaderMap->GetPermutationFlags(),
+				this,
+				GameThreadPendingCompilerEnvironment,
+				VFType,
+				CompileJobs,
+				nullptr,
+				nullptr);
+		}
 
-		ShaderType->AsMeshMaterialShaderType()->BeginCompileShader(
-			EShaderCompileJobPriority::ForceLocal,
-			GetGameThreadCompilingShaderMapId(),
-			0,
-			Platform,
-			GameThreadShaderMap->GetPermutationFlags(),
-			this,
-			GameThreadPendingCompilerEnvironment,
-			VFType,
-			CompileJobs,
-			nullptr,
-			nullptr);
+		GShaderCompilingManager->SubmitJobs(CompileJobs, GetBaseMaterialPathName(), GameThreadShaderMap->GetDebugDescription());
 	}
-
-	GShaderCompilingManager->SubmitJobs(CompileJobs, GetBaseMaterialPathName(), GameThreadShaderMap->GetDebugDescription());
 }
 #endif // WITH_EDITOR
 
