Fixed argument parsing for HighResShot command.
  - Removed left-over code from the previous argument parsing algorithm.

PR #9006: Fix HighResShot command arguments interpretation (Contributed by [REDACTED_PERSON])

diff --git a/Engine/Source/Runtime/Engine/Private/UnrealClient.cpp b/Engine/Source/Runtime/Engine/Private/UnrealClient.cpp
--- a/Engine/Source/Runtime/Engine/Private/UnrealClient.cpp
+++ b/Engine/Source/Runtime/Engine/Private/UnrealClient.cpp
@@ -2216,7 +2216,6 @@ extern bool ParseResolution( const TCHAR* InResolution, uint32& OutX, uint32& Ou
 
 ENGINE_API bool GetHighResScreenShotInput(const TCHAR* Cmd, FOutputDevice& Ar, uint32& OutXRes, uint32& OutYRes, float& OutResMult, FIntRect& OutCaptureRegion, bool& OutShouldEnableMask, bool& OutDumpBufferVisualizationTargets, bool& OutCaptureHDR, FString& OutFilenameOverride, bool& OutUseDateTimeAsFileName)
 {
-	FString CmdString = Cmd;
 	TArray<FString> Arguments;
 	const FString FilenameSearchString = TEXT("filename=");
 
@@ -2236,11 +2235,6 @@ ENGINE_API bool GetHighResScreenShotInput(const TCHAR* Cmd, FOutputDevice& Ar, u
 		}
 	}
 
-	if (CmdString.Len() > 0)
-	{
-		Arguments.Add(CmdString);
-	}
-
 	int32 NumArguments = Arguments.Num();
 
 	if (NumArguments >= 1)
