clean up texture IsLongLatCubemap and GetVolumeSizeZ
do allow pow2 pad on volumes, it works fine

diff --git a/Engine/Source/Runtime/Engine/Classes/Engine/Texture.h b/Engine/Source/Runtime/Engine/Classes/Engine/Texture.h
--- a/Engine/Source/Runtime/Engine/Classes/Engine/Texture.h
+++ b/Engine/Source/Runtime/Engine/Classes/Engine/Texture.h
@@ -363,11 +363,13 @@ struct FTextureSource
 	UE_DEPRECATED(5.1, "Use GetSourceCompression instead")
 	FORCEINLINE bool IsPNGCompressed() const { return GetSourceCompression() == ETextureSourceCompressionFormat::TSCF_PNG; }
 	
-	// Warning: IsLongLatCubemap() is not correct.  LongLat Cubemaps have bLongLatCubemap == false
-	// the correct check is to see if the texture class is CubeMap and (NumSlices%6) == 0
+	// Warning: bLongLatCubemap is not correct.  LongLat Cubemaps often have bLongLatCubemap == false
 	// bLongLatCubemap is sometimes set to true for cube arrays to disambiguate the case of 6 longlat cubemaps in an array
-	// @@!! consider putting the NumSlices check in IsLongLatCubemap()
-	FORCEINLINE bool IsLongLatCubemap() const { return bLongLatCubemap; }
+	ENGINE_API bool IsCubeOrCubeArray() const;
+	ENGINE_API bool IsVolume() const;
+	ENGINE_API bool IsLongLatCubemap() const;
+	// returns volume depth, or 1 if not a volume
+	ENGINE_API int GetVolumeSizeZ() const;
 
 	FORCEINLINE int64 GetSizeOnDisk() const { return BulkData.GetPayloadSize(); }
 	inline bool HasPayloadData() const { return BulkData.HasPayloadData(); }
diff --git a/Engine/Source/Runtime/Engine/Private/Texture.cpp b/Engine/Source/Runtime/Engine/Private/Texture.cpp
--- a/Engine/Source/Runtime/Engine/Private/Texture.cpp
+++ b/Engine/Source/Runtime/Engine/Private/Texture.cpp
@@ -373,17 +373,11 @@ void UTexture::ValidateSettingsAfterImportOrEdit()
 		
 #if WITH_EDITORONLY_DATA
 
+	// IsPowerOfTwo only checks XY :
 	bool IsPowerOfTwo = Source.IsPowerOfTwo();
-	if ( GetTextureClass() == ETextureClass::Volume && ! FMath::IsPowerOfTwo(Source.NumSlices) )
+	if ( ! FMath::IsPowerOfTwo(Source.GetVolumeSizeZ()) )
 	{
 		IsPowerOfTwo = false;
-
-		if ( PowerOfTwoMode != ETexturePowerOfTwoSetting::None )
-		{
-			// @@!! yes it is? fix and remove this?
-			UE_LOG(LogTexture, Warning, TEXT("PowerOfTwoMode not supported for volume Z"));
-			PowerOfTwoMode = ETexturePowerOfTwoSetting::None;
-		}
 	}
 	
 	if (VirtualTextureStreaming)
@@ -1347,6 +1341,58 @@ int32 FTextureSource::GetBytesPerPixel(ETextureSourceFormat Format)
 
 #if WITH_EDITOR
 
+bool FTextureSource::IsCubeOrCubeArray() const
+{
+	check( Owner != nullptr );
+	ETextureClass TextureClass = Owner->GetTextureClass();
+	return TextureClass == ETextureClass::Cube || TextureClass == ETextureClass::CubeArray;
+}
+
+bool FTextureSource::IsVolume() const
+{
+	check( Owner != nullptr );
+	ETextureClass TextureClass = Owner->GetTextureClass();
+	return TextureClass == ETextureClass::Volume;
+}
+
+bool FTextureSource::IsLongLatCubemap() const
+{
+	if ( IsCubeOrCubeArray() )
+	{
+		check( NumLayers == 1 );
+
+		// bLongLatCubemap is sometimes set for LongLat Cube Arrays but not always
+		if ( bLongLatCubemap )
+		{
+			return true;
+		}
+		else
+		{
+			// if NumSlices is not a multiple of 6, must be longlat !?
+			return (NumSlices % 6) != 0;
+		}
+	}
+	else
+	{
+		check( ! bLongLatCubemap );
+		return false;
+	}
+}
+
+// returns volume depth, or 1 if not a volume
+int FTextureSource::GetVolumeSizeZ() const
+{
+	if ( IsVolume() )
+	{
+		check( NumLayers == 1 );
+		return NumSlices;
+	}
+	else
+	{
+		return 1;
+	}
+}
+
 void FTextureSource::InitBlocked(const ETextureSourceFormat* InLayerFormats,
 	const FTextureSourceBlock* InBlocks,
 	int32 InNumLayers,
diff --git a/Engine/Source/Runtime/Engine/Private/TextureDerivedData.cpp b/Engine/Source/Runtime/Engine/Private/TextureDerivedData.cpp
--- a/Engine/Source/Runtime/Engine/Private/TextureDerivedData.cpp
+++ b/Engine/Source/Runtime/Engine/Private/TextureDerivedData.cpp
@@ -753,9 +753,8 @@ static void GetTextureBuildSettings(
 	{
 		OutBuildSettings.bCubemap = true;
 		OutBuildSettings.DiffuseConvolveMipLevel = GDiffuseConvolveMipLevel;
-		// beware IsLongLatCubemap is pretty useless , GetNumSlices() is doing the work here
 		check( Texture.Source.GetNumSlices() == 1 || Texture.Source.GetNumSlices() == 6 );
-		OutBuildSettings.bLongLatSource = (Texture.Source.GetNumSlices() == 1) || Texture.Source.IsLongLatCubemap();
+		OutBuildSettings.bLongLatSource = Texture.Source.IsLongLatCubemap();
 	}
 	else if ( TextureClass == ETextureClass::Array )
 	{
