World Partition: force disabling streaming when the world settings doesn't support streaming.

diff --git a/Engine/Source/Editor/WorldPartitionEditor/Private/WorldPartition/Customizations/WorldPartitionDetailsCustomization.cpp b/Engine/Source/Editor/WorldPartitionEditor/Private/WorldPartition/Customizations/WorldPartitionDetailsCustomization.cpp
--- a/Engine/Source/Editor/WorldPartitionEditor/Private/WorldPartition/Customizations/WorldPartitionDetailsCustomization.cpp
+++ b/Engine/Source/Editor/WorldPartitionEditor/Private/WorldPartition/Customizations/WorldPartitionDetailsCustomization.cpp
@@ -6,6 +6,7 @@
 #include "DetailLayoutBuilder.h"
 #include "DetailCategoryBuilder.h"
 #include "DetailWidgetRow.h"
+#include "Misc/MessageDialog.h"
 #include "Widgets/Layout/SBorder.h"
 #include "Widgets/Text/STextBlock.h"
 #include "Widgets/Layout/SBox.h"
@@ -30,6 +31,22 @@ void FWorldPartitionDetails::CustomizeDetails(IDetailLayoutBuilder& DetailBuilde
 
 	IDetailCategoryBuilder& WorldPartitionCategory = DetailBuilder.EditCategory("WorldPartition");
 
+	WorldPartitionCategory.AddCustomRow(LOCTEXT("EnableStreaming", "Enable Streaming"), false)
+		.NameContent()
+		[
+			SNew(STextBlock)
+			.Text(LOCTEXT("WorldPartitionEnableStreaming", "Enable Streaming"))
+			.ToolTipText(LOCTEXT("WorldPartitionEnableStreaming_ToolTip", "Set the world partition enable streaming state."))
+			.Font(IDetailLayoutBuilder::GetDetailFont())
+		]
+		.ValueContent()
+		[
+			SNew(SCheckBox)
+			.IsChecked(MakeAttributeLambda([this]() { return WorldPartition->IsStreamingEnabled() ? ECheckBoxState::Checked : ECheckBoxState::Unchecked; }))
+			.OnCheckStateChanged(this, &FWorldPartitionDetails::HandleWorldPartitionEnableStreamingChanged)
+		]
+		.Visibility(TAttribute<EVisibility>::CreateLambda([this]() { return WorldPartition->SupportsStreaming() ? EVisibility::Visible : EVisibility::Hidden; }));
+
 	if (WorldPartition->RuntimeHash)
 	{
 		FAddPropertyParams Params;
@@ -38,18 +55,18 @@ void FWorldPartitionDetails::CustomizeDetails(IDetailLayoutBuilder& DetailBuilde
 
 		if (IDetailPropertyRow* RuntimeHashRow = WorldPartitionCategory.AddExternalObjects({ WorldPartition->RuntimeHash }, EPropertyLocation::Default, Params))
 		{
-			// use TAttribute::CreateLambda?
-			RuntimeHashRow->ShouldAutoExpand(true)
+			(*RuntimeHashRow)
+				.ShouldAutoExpand(true)
 				.DisplayName(LOCTEXT("RuntimeHash", "Runtime Hash"))
-				.Visibility(TAttribute<EVisibility>::Create(TAttribute<EVisibility>::FGetter::CreateLambda([this]() { return WorldPartition->bEnableStreaming ? EVisibility::Visible : EVisibility::Hidden; })));
+				.Visibility(TAttribute<EVisibility>::CreateLambda([this]() { return WorldPartition->IsStreamingEnabled() ? EVisibility::Visible : EVisibility::Hidden; }));
 		}
 	}
 
-	WorldPartitionCategory.AddCustomRow(LOCTEXT("WorldPartitionEditorCellSizeRow", "WorldPartitionEditorCellSize"), true)
+	WorldPartitionCategory.AddCustomRow(LOCTEXT("EditorCellSizeRow", "Editor Cell Size"), true)
 		.NameContent()
 		[
 			SNew(STextBlock)
-			.Text(LOCTEXT("WorldPartitionCellSize", "World Partition Editor Cell Size"))
+			.Text(LOCTEXT("WorldPartitionCellSize", "Editor Cell Size"))
 			.ToolTipText(LOCTEXT("WorldPartitionEditorCellSize_ToolTip", "Set the world partition editor cell size, will take effect on the next world reload."))
 			.Font(IDetailLayoutBuilder::GetDetailFont())
 		]
@@ -67,7 +84,33 @@ void FWorldPartitionDetails::CustomizeDetails(IDetailLayoutBuilder& DetailBuilde
 				.OnValueChanged(this, &FWorldPartitionDetails::HandleWorldPartitionEditorCellSizeChanged)
 				.Value(this, &FWorldPartitionDetails::HandleWorldPartitionEditorCellSizeValue)
 			]
-		];
+		];		
+}
+
+void FWorldPartitionDetails::HandleWorldPartitionEnableStreamingChanged(ECheckBoxState CheckState)
+{
+	if (CheckState == ECheckBoxState::Checked)
+	{
+		if (!WorldPartition->bStreamingWasEnabled)
+		{
+			if (FMessageDialog::Open(EAppMsgType::YesNo, LOCTEXT("WorldPartitionConfirmEnableStreaming", "You are about to enable streaming for the first time, the world will be setup to stream. Continue?")) == EAppReturnType::No)
+			{
+				return;
+			}
+			
+			FMessageDialog::Open(EAppMsgType::Ok, LOCTEXT("WorldPartitionEnableStreamingDialolg", "Please refer to https://docs.unrealengine.com/5.0/en-US/building-virtual-worlds/world-partition for how to set up streaming."));
+			WorldPartition->bStreamingWasEnabled = true;
+		}
+
+		WorldPartition->SetEnableStreaming(true);
+	}
+	else
+	{
+		if (FMessageDialog::Open(EAppMsgType::YesNo, LOCTEXT("WorldPartitionConfirmDisableStreaming", "You are about to disable streaming, all actors in the world will be always loaded. Continue?")) == EAppReturnType::Yes)
+		{
+			WorldPartition->SetEnableStreaming(false);
+		}
+	}
 }
 
 void FWorldPartitionDetails::HandleWorldPartitionEditorCellSizeChanged(uint32 NewValue)
diff --git a/Engine/Source/Editor/WorldPartitionEditor/Private/WorldPartition/Customizations/WorldPartitionDetailsCustomization.h b/Engine/Source/Editor/WorldPartitionEditor/Private/WorldPartition/Customizations/WorldPartitionDetailsCustomization.h
--- a/Engine/Source/Editor/WorldPartitionEditor/Private/WorldPartition/Customizations/WorldPartitionDetailsCustomization.h
+++ b/Engine/Source/Editor/WorldPartitionEditor/Private/WorldPartition/Customizations/WorldPartitionDetailsCustomization.h
@@ -23,6 +23,9 @@ private:
 	/** IDetailCustomization interface */
 	virtual void CustomizeDetails(IDetailLayoutBuilder& DetailBuilder) override;
 
+	// Callback for changes in the world partition enable streaming.
+	void HandleWorldPartitionEnableStreamingChanged(ECheckBoxState CheckState);
+
 	// Callback for changes in the world partition editor cell size.
 	void HandleWorldPartitionEditorCellSizeChanged(uint32 NewValue);
 
diff --git a/Engine/Source/Editor/WorldPartitionEditor/Private/WorldPartition/SWorldPartitionEditorGrid.cpp b/Engine/Source/Editor/WorldPartitionEditor/Private/WorldPartition/SWorldPartitionEditorGrid.cpp
--- a/Engine/Source/Editor/WorldPartitionEditor/Private/WorldPartition/SWorldPartitionEditorGrid.cpp
+++ b/Engine/Source/Editor/WorldPartitionEditor/Private/WorldPartition/SWorldPartitionEditorGrid.cpp
@@ -18,7 +18,7 @@ void SWorldPartitionEditorGrid::Construct(const FArguments& InArgs)
 	World = InArgs._InWorld;
 	WorldPartition = World ? World->GetWorldPartition() : nullptr;
 
-	if (!WorldPartition || !WorldPartition->bEnableStreaming)
+	if (!WorldPartition || !WorldPartition->IsStreamingEnabled())
 	{
 		const FText Message = WorldPartition ? 
 			LOCTEXT("WorldPartitionHasStreamingDisabled", "World Partition streaming is not enabled for this map") : 
@@ -110,7 +110,7 @@ void SWorldPartitionEditorGrid::Refresh()
 
 bool SWorldPartitionEditorGrid::IsDisabled() const
 {
-	return !WorldPartition || !WorldPartition->bEnableStreaming;
+	return !WorldPartition || !WorldPartition->IsStreamingEnabled();
 }
 
 #undef LOCTEXT_NAMESPACE
\ No newline at end of file
diff --git a/Engine/Source/Runtime/Engine/Classes/GameFramework/WorldSettings.h b/Engine/Source/Runtime/Engine/Classes/GameFramework/WorldSettings.h
--- a/Engine/Source/Runtime/Engine/Classes/GameFramework/WorldSettings.h
+++ b/Engine/Source/Runtime/Engine/Classes/GameFramework/WorldSettings.h
@@ -794,6 +794,11 @@ public:
 	void SetWorldPartition(UWorldPartition* InWorldPartition);
 	void ApplyWorldPartitionForcedSettings();
 
+#if WITH_EDITOR
+	virtual bool SupportsWorldPartitionStreaming() const { return true; }
+	void SupportsWorldPartitionStreamingChanged();
+#endif
+
 	/**
 	 * Returns the Z component of the current world gravity and initializes it to the default
 	 * gravity if called for the first time.
diff --git a/Engine/Source/Runtime/Engine/Private/WorldPartition/WorldPartition.cpp b/Engine/Source/Runtime/Engine/Private/WorldPartition/WorldPartition.cpp
--- a/Engine/Source/Runtime/Engine/Private/WorldPartition/WorldPartition.cpp
+++ b/Engine/Source/Runtime/Engine/Private/WorldPartition/WorldPartition.cpp
@@ -354,7 +354,7 @@ void UWorldPartition::OnEndPlay()
 
 FName UWorldPartition::GetWorldPartitionEditorName() const
 {
-	if (bEnableStreaming)
+	if (SupportsStreaming())
 	{
 		return EditorHash->GetWorldPartitionEditorName();
 	}
@@ -404,6 +404,13 @@ void UWorldPartition::Initialize(UWorld* InWorld, const FTransform& InTransform)
 
 	UE_LOG(LogWorldPartition, Log, TEXT("UWorldPartition::Initialize(IsEditor=%d, bPIEWorldTravel=%d IsGame=%d, IsCooking=%d)"), bIsEditor ? 1 : 0, bPIEWorldTravel ? 1 : 0, bIsGame ? 1 : 0, bIsCooking ? 1 : 0);
 
+	FWorldPartitionActorDesc::bForceAlwaysLoaded = !IsStreamingEnabled();
+
+	if (bEnableStreaming)
+	{
+		bStreamingWasEnabled = true;
+	}
+
 	if (bIsGame || bIsCooking)
 	{
 		// Don't rely on the editor hash for cooking or -game
@@ -1069,7 +1076,7 @@ void UWorldPartition::OnActorDescAdded(FWorldPartitionActorDesc* NewActorDesc)
 	Super::OnActorDescAdded(NewActorDesc);
 
 	HashActorDesc(NewActorDesc);
-		
+
 	if (WorldPartitionEditor)
 	{
 		WorldPartitionEditor->Refresh();
@@ -1144,111 +1151,73 @@ bool UWorldPartition::GetInstancingContext(const FLinkerInstancingContext*& OutI
 	return false;
 }
 
-void UWorldPartition::OnEnableStreaming()
+bool UWorldPartition::SupportsStreaming() const
 {
-	if (!bStreamingWasEnabled)
-	{
-		bStreamingWasEnabled = true;
-		FMessageDialog::Open(EAppMsgType::Ok, LOCTEXT("WorldPartitionEnableStreamingDialolg", "Please refer to https://docs.unrealengine.com/5.0/en-US/building-virtual-worlds/world-partition for how to set up streaming."));
-	}
-}
-
-void UWorldPartition::HashActorDesc(FWorldPartitionActorDesc* ActorDesc)
-{
-	check(ActorDesc);
-	check(EditorHash);
-	FWorldPartitionHandle ActorHandle(this, ActorDesc->GetGuid());
-	EditorHash->HashActor(ActorHandle);
-
-	bShouldCheckEnableStreamingWarning = true;
+	return World->GetWorldSettings()->SupportsWorldPartitionStreaming();
 }
 
-void UWorldPartition::UnhashActorDesc(FWorldPartitionActorDesc* ActorDesc)
+bool UWorldPartition::IsStreamingEnabled() const
 {
-	check(ActorDesc);
-	check(EditorHash);
-	FWorldPartitionHandle ActorHandle(this, ActorDesc->GetGuid());
-	EditorHash->UnhashActor(ActorHandle);
+	return bEnableStreaming && SupportsStreaming();
 }
 
-void UWorldPartition::PostLoad()
+void UWorldPartition::SetEnableStreaming(bool bInEnableStreaming)
 {
-	FWorldPartitionActorDesc::bForceAlwaysLoaded = !bEnableStreaming;
-
-	if (bEnableStreaming)
+	if (bEnableStreaming != bInEnableStreaming)
 	{
-		bStreamingWasEnabled = true;
+		bEnableStreaming = bInEnableStreaming;
+		OnEnableStreamingChanged();
 	}
-
-	Super::PostLoad();
 }
 
-void UWorldPartition::PostEditChangeProperty(struct FPropertyChangedEvent& PropertyChangedEvent)
+void UWorldPartition::OnEnableStreamingChanged()
 {
-	if (PropertyChangedEvent.GetPropertyName() == GET_MEMBER_NAME_CHECKED(UWorldPartition, bEnableStreaming))
-	{
-		bool bApplyStreamingChange = true;
+	// Pin the actor handles on the actor to prevent unloading it when unhashing
+	TArray<FWorldPartitionHandlePinRefScope> PinRefScopes;
+	PinRefScopes.Reserve(ActorDescList.Num());
 
-		if (bEnableStreaming)
-		{
-			if (!bStreamingWasEnabled)
-			{
-				if (FMessageDialog::Open(EAppMsgType::YesNo, LOCTEXT("WorldPartitionConfirmEnableStreaming", "You are about to enable streaming, the world will be setup to stream. Continue?")) == EAppReturnType::No)
-				{
-					bEnableStreaming = false;
-					bApplyStreamingChange = false;
-				}
-				else
-				{
-					bStreamingWasEnabled = true;
-
-					FMessageDialog::Open(EAppMsgType::Ok, LOCTEXT("WorldPartitionEnableStreamingDialolg", "Please refer to https://docs.unrealengine.com/5.0/en-US/building-virtual-worlds/world-partition for how to set up streaming."));
-				}
-			}
-		}
-		else
-		{
-			if (FMessageDialog::Open(EAppMsgType::YesNo, LOCTEXT("WorldPartitionConfirmDisableStreaming", "You are about to disable streaming, all actors in the world will be always loaded. Continue?")) == EAppReturnType::No)
-			{
-				bEnableStreaming = true;
-				bApplyStreamingChange = false;
-			}
-		}
+	for (UActorDescContainer::TIterator<> ActorDescIterator(this); ActorDescIterator; ++ActorDescIterator)
+	{
+		FWorldPartitionHandle ExistingActorHandle(this, ActorDescIterator->GetGuid());
+		PinRefScopes.Emplace(ExistingActorHandle);
 
-		if (bApplyStreamingChange)
-		{
-			// Pin the actor handles on the actor to prevent unloading it when unhashing
-			TArray<FWorldPartitionHandlePinRefScope> PinRefScopes;
-			PinRefScopes.Reserve(ActorDescList.Num());
+		UnhashActorDesc(*ActorDescIterator);
+	}
 
-			for (UActorDescContainer::TIterator<> ActorDescIterator(this); ActorDescIterator; ++ActorDescIterator)
-			{
-				FWorldPartitionHandle ExistingActorHandle(this, ActorDescIterator->GetGuid());
-				PinRefScopes.Emplace(ExistingActorHandle);
+	FWorldPartitionActorDesc::bForceAlwaysLoaded = !IsStreamingEnabled();
 
-				UnhashActorDesc(*ActorDescIterator);
-			}
+	for (UActorDescContainer::TIterator<> ActorDescIterator(this); ActorDescIterator; ++ActorDescIterator)
+	{
+		HashActorDesc(*ActorDescIterator);
+	}
 
-			FWorldPartitionActorDesc::bForceAlwaysLoaded = !bEnableStreaming;
+	if (!IsStreamingEnabled())
+	{
+		UpdateLoadingEditorCell(EditorHash->GetAlwaysLoadedCell(), true, false);
+	}
 
-			for (UActorDescContainer::TIterator<> ActorDescIterator(this); ActorDescIterator; ++ActorDescIterator)
-			{
-				HashActorDesc(*ActorDescIterator);
-			}
+	if (WorldPartitionEditor)
+	{
+		WorldPartitionEditor->Reconstruct();
+	}
+}
 
-			if (!bEnableStreaming)
-			{
-				UpdateLoadingEditorCell(EditorHash->GetAlwaysLoadedCell(), true, false);
-			}
+void UWorldPartition::HashActorDesc(FWorldPartitionActorDesc* ActorDesc)
+{
+	check(ActorDesc);
+	check(EditorHash);
+	FWorldPartitionHandle ActorHandle(this, ActorDesc->GetGuid());
+	EditorHash->HashActor(ActorHandle);
 
-			if (WorldPartitionEditor)
-			{
-				WorldPartitionEditor->Reconstruct();
-			}
-		}
-	}
+	bShouldCheckEnableStreamingWarning = true;
+}
 
-	Super::PostEditChangeProperty(PropertyChangedEvent);
+void UWorldPartition::UnhashActorDesc(FWorldPartitionActorDesc* ActorDesc)
+{
+	check(ActorDesc);
+	check(EditorHash);
+	FWorldPartitionHandle ActorHandle(this, ActorDesc->GetGuid());
+	EditorHash->UnhashActor(ActorHandle);
 }
 #endif
 
@@ -1319,7 +1288,7 @@ void UWorldPartition::Tick(float DeltaSeconds)
 	{
 		bShouldCheckEnableStreamingWarning = false;
 
-		if (!bEnableStreaming && !bStreamingWasEnabled && !bShouldEnableStreamingWarned)
+		if (!IsStreamingEnabled() && SupportsStreaming() && !bStreamingWasEnabled && !bShouldEnableStreamingWarned)
 		{
 			FBox AllActorsBounds(ForceInit);
 			for (UActorDescContainer::TConstIterator<> ActorDescIterator(this); ActorDescIterator; ++ActorDescIterator)
@@ -1330,9 +1299,12 @@ void UWorldPartition::Tick(float DeltaSeconds)
 				if (AllActorsBounds.GetSize().GetMax() >= 400000.0f)
 				{
 					bShouldEnableStreamingWarned = true;
+
 					if (FMessageDialog::Open(EAppMsgType::YesNo, LOCTEXT("WorldPartitionShouldEnableStreaming", "The size of your world has grown enough to justify enabling streaming. Enable streaming now?")) == EAppReturnType::Yes)
 					{
-						OnEnableStreaming();
+						bEnableStreaming = true;
+						bStreamingWasEnabled = true;
+						FMessageDialog::Open(EAppMsgType::Ok, LOCTEXT("WorldPartitionEnableStreamingDialolg", "Please refer to https://docs.unrealengine.com/5.0/en-US/building-virtual-worlds/world-partition for how to set up streaming."));			
 					}
 
 					break;
diff --git a/Engine/Source/Runtime/Engine/Private/WorldSettings.cpp b/Engine/Source/Runtime/Engine/Private/WorldSettings.cpp
--- a/Engine/Source/Runtime/Engine/Private/WorldSettings.cpp
+++ b/Engine/Source/Runtime/Engine/Private/WorldSettings.cpp
@@ -249,6 +249,16 @@ float AWorldSettings::GetGravityZ() const
 	return WorldGravityZ;
 }
 
+#if WITH_EDITOR
+void AWorldSettings::SupportsWorldPartitionStreamingChanged()
+{
+	if (WorldPartition)
+	{
+		WorldPartition->OnEnableStreamingChanged();
+	}
+}
+#endif
+
 void AWorldSettings::OnRep_WorldGravityZ()
 {
 	bWorldGravitySet = true;
diff --git a/Engine/Source/Runtime/Engine/Public/WorldPartition/WorldPartition.h b/Engine/Source/Runtime/Engine/Public/WorldPartition/WorldPartition.h
--- a/Engine/Source/Runtime/Engine/Public/WorldPartition/WorldPartition.h
+++ b/Engine/Source/Runtime/Engine/Public/WorldPartition/WorldPartition.h
@@ -74,8 +74,11 @@ class ENGINE_API UWorldPartition final : public UActorDescContainer
 	GENERATED_UCLASS_BODY()
 
 	friend class FWorldPartitionActorDesc;
+	friend class FWorldPartitionConverter;
 	friend class UWorldPartitionEditorCell;
+	friend class UWorldPartitionConvertCommandlet;
 	friend class FWorldPartitionEditorModule;
+	friend class FWorldPartitionDetails;
 	friend class FUnrealEdMisc;
 
 public:
@@ -95,8 +98,14 @@ public:
 
 #if WITH_EDITOR
 	TArray<FName> GetUserLoadedEditorGridCells() const;
-private:
 
+public:
+	bool SupportsStreaming() const;
+	bool IsStreamingEnabled() const;
+	void SetEnableStreaming(bool bInEnableStreaming);
+	void OnEnableStreamingChanged();
+
+private:
 	void SavePerUserSettings();
 		
 	void OnGCPostReachabilityAnalysis();
@@ -121,16 +130,10 @@ private:
 
 	virtual bool GetInstancingContext(const FLinkerInstancingContext*& OutInstancingContext, FSoftObjectPathFixupArchive*& OutSoftObjectPathFixupArchive) const override;
 	//~ End UActorDescContainer Interface
-
-	void OnEnableStreaming();
 #endif
 
 public:
 	//~ Begin UObject Interface
-#if WITH_EDITOR
-	virtual void PostLoad() override;
-	virtual void PostEditChangeProperty(struct FPropertyChangedEvent& PropertyChangedEvent) override;
-#endif
 	virtual bool ResolveSubobject(const TCHAR* SubObjectPath, UObject*& OutObject, bool bLoadIfExists) override;
 	//~ End UObject Interface
 
@@ -208,8 +211,9 @@ public:
 	UPROPERTY()
 	TSubclassOf<UWorldPartitionStreamingPolicy> WorldPartitionStreamingPolicyClass;
 
+private:
 	/** Enables streaming for this world. */
-	UPROPERTY(EditAnywhere, Category=WorldPartition)
+	UPROPERTY()
 	bool bEnableStreaming;
 
 	/** Used to know if it's the first time streaming is enabled on this world. */
@@ -217,22 +221,26 @@ public:
 	bool bStreamingWasEnabled;
 
 	/** Used to know if the user has already been warned about that it should enable streaming based on world size. */
+	UPROPERTY()
 	bool bShouldEnableStreamingWarned;
 
 	/** Used to know if we need to recheck if the user should enable streaming based on world size. */
 	bool bShouldCheckEnableStreamingWarning;
 #endif
 
+public:
 	UPROPERTY()
 	TObjectPtr<UWorldPartitionRuntimeHash> RuntimeHash;
 
 #if WITH_EDITOR
+private:
 	bool bForceGarbageCollection;
 	bool bForceGarbageCollectionPurge;
 	bool bIsPIE;
 #endif
 
 #if WITH_EDITORONLY_DATA
+public:
 	// Default HLOD layer
 	UPROPERTY(EditAnywhere, Category=WorldPartition, meta = (DisplayName = "Default HLOD Layer", EditCondition="bEnableStreaming", EditConditionHides, HideEditConditionToggle))
 	TObjectPtr<class UHLODLayer> DefaultHLODLayer;
