Fix potential crash when creating new partitioned worlds.

diff --git a/Engine/Source/Runtime/Engine/Private/WorldPartition/WorldPartition.cpp b/Engine/Source/Runtime/Engine/Private/WorldPartition/WorldPartition.cpp
--- a/Engine/Source/Runtime/Engine/Private/WorldPartition/WorldPartition.cpp
+++ b/Engine/Source/Runtime/Engine/Private/WorldPartition/WorldPartition.cpp
@@ -1153,7 +1153,7 @@ bool UWorldPartition::GetInstancingContext(const FLinkerInstancingContext*& OutI
 
 bool UWorldPartition::SupportsStreaming() const
 {
-	return World->GetWorldSettings()->SupportsWorldPartitionStreaming();
+	return World ? World->GetWorldSettings()->SupportsWorldPartitionStreaming() : false;
 }
 
 bool UWorldPartition::IsStreamingEnabled() const
