Export destructor to fix Linux editor build
#rnx

diff --git a/Engine/Source/Runtime/Engine/Classes/Components/SkeletalMeshComponent.h b/Engine/Source/Runtime/Engine/Classes/Components/SkeletalMeshComponent.h
--- a/Engine/Source/Runtime/Engine/Classes/Components/SkeletalMeshComponent.h
+++ b/Engine/Source/Runtime/Engine/Classes/Components/SkeletalMeshComponent.h
@@ -306,7 +306,7 @@ struct ENGINE_API FClosestPointOnPhysicsAsset
 struct FClothCollisionSource
 {
 	FClothCollisionSource(USkeletalMeshComponent* InSourceComponent, UPhysicsAsset* InSourcePhysicsAsset, const FOnBoneTransformsFinalizedMultiCast::FDelegate& InOnBoneTransformsFinalizedDelegate);
-	~FClothCollisionSource();
+	ENGINE_API ~FClothCollisionSource();
 
 	/** Component that collision data will be copied from */
 	TWeakObjectPtr<USkeletalMeshComponent> SourceComponent;
