Expose SetPriorityAndAffinityOnGameThread

#rnx

diff --git a/Engine/Source/Runtime/Engine/Classes/Engine/Engine.h b/Engine/Source/Runtime/Engine/Classes/Engine/Engine.h
--- a/Engine/Source/Runtime/Engine/Classes/Engine/Engine.h
+++ b/Engine/Source/Runtime/Engine/Classes/Engine/Engine.h
@@ -3653,6 +3653,10 @@ private:
 #endif
 
 	FDelegateHandle HandleScreenshotCapturedDelegateHandle;
+
+public:
+	/** Set priority and affinity on game thread either from ini file or from FPlatformAffinity::GetGameThreadPriority()*/
+	void SetPriorityAndAffinityOnGameThread();
 };
 
 /** Global engine pointer. Can be 0 so don't use without checking. */
diff --git a/Engine/Source/Runtime/Engine/Private/UnrealEngine.cpp b/Engine/Source/Runtime/Engine/Private/UnrealEngine.cpp
--- a/Engine/Source/Runtime/Engine/Private/UnrealEngine.cpp
+++ b/Engine/Source/Runtime/Engine/Private/UnrealEngine.cpp
@@ -17403,6 +17403,8 @@ static FAutoConsoleCommand SetupThreadConfigCmd(
 
 void InitThreadConfig()
 {
+	GetDefaultThreadConfigs();
+
 	if (GConfig)
 	{
 		TArray<FString> ThreadConfigEntries;
@@ -17414,6 +17416,11 @@ void InitThreadConfig()
 	}
 }
 
+void UEngine::SetPriorityAndAffinityOnGameThread()
+{
+	::SetPriorityAndAffinityOnGameThread();
+}
+
 #if !UE_BUILD_SHIPPING
 
 static void PakFileTest(const TArray<FString>& Args)
