VT use MaxTextureResolutionDefault instead of hard coding

diff --git a/Engine/Source/Runtime/Engine/Private/VT/VirtualTextureDataBuilder.cpp b/Engine/Source/Runtime/Engine/Private/VT/VirtualTextureDataBuilder.cpp
--- a/Engine/Source/Runtime/Engine/Private/VT/VirtualTextureDataBuilder.cpp
+++ b/Engine/Source/Runtime/Engine/Private/VT/VirtualTextureDataBuilder.cpp
@@ -574,7 +574,7 @@ void FVirtualTextureDataBuilder::BuildTiles(const TArray<FVTSourceTileEntry>& Ti
 		// TBSettings starts default constructed (no processing)
 		//	then we copy over just the compression options without the color-change processing
 		FTextureBuildSettings TBSettings;
-		TBSettings.MaxTextureResolution = TNumericLimits<uint32>::Max();
+		TBSettings.MaxTextureResolution = FTextureBuildSettings::MaxTextureResolutionDefault;
 		TBSettings.TextureFormatName = LayerData.TextureFormatName;
 		TBSettings.bSRGB = BuildSettingsForLayer.bSRGB;
 		TBSettings.bUseLegacyGamma = BuildSettingsForLayer.bUseLegacyGamma;
@@ -906,7 +906,7 @@ void FVirtualTextureDataBuilder::BuildSourcePixels(const FTextureSourceData& Sou
 			// TBSettings starts with the full Texture settings, so we get all options
 			//  then we change FormatName to be == Source format, so no Compression is done
 			FTextureBuildSettings TBSettings = SettingsPerLayer[0];
-			//TBSettings.MaxTextureResolution = TNumericLimits<uint32>::Max();
+			//TBSettings.MaxTextureResolution = FTextureBuildSettings::MaxTextureResolutionDefault;
 			TBSettings.TextureFormatName = LayerData.FormatName;
 			
 			if ( LayerIndex != 0 )
@@ -1100,7 +1100,7 @@ void FVirtualTextureDataBuilder::BuildSourcePixels(const FTextureSourceData& Sou
 			// Adjust the build settings to generate an uncompressed texture with mips but leave other settings
 			// like color correction, ... in place
 			FTextureBuildSettings TBSettings = SettingsPerLayer[0];
-			TBSettings.MaxTextureResolution = TNumericLimits<uint32>::Max(); // don't limit the size of the mip-tail, this limit only applies to each source block
+			TBSettings.MaxTextureResolution = FTextureBuildSettings::MaxTextureResolutionDefault;; // don't limit the size of the mip-tail, this limit only applies to each source block
 			TBSettings.TextureFormatName = LayerData.FormatName;
 			TBSettings.bSRGB = BuildSettingsForLayer.bSRGB;
 			TBSettings.bUseLegacyGamma = BuildSettingsForLayer.bUseLegacyGamma;
@@ -1223,7 +1223,7 @@ void FVirtualTextureDataBuilder::BuildMipTails()
 		// The pixels we have already include things like color correction, mip blurring, ... so we just start
 		// from pristine build settings here
 		FTextureBuildSettings TBSettings;
-		TBSettings.MaxTextureResolution = TNumericLimits<uint32>::Max();
+		TBSettings.MaxTextureResolution = FTextureBuildSettings::MaxTextureResolutionDefault;;
 		TBSettings.TextureFormatName = "BGRA8";
 		TBSettings.bSRGB = Settings.Layers[Layer].SourceBuildSettings.bSRGB;
 		TBSettings.bUseLegacyGamma = Settings.Layers[Layer].SourceBuildSettings.bUseLegacyGamma;
