Unshelved from pending changelist '19434316':

[Engine] Enable color Calibration Material Visualization for shipping
-Enabled the calibration flags
-Removed the no-shipping guards around associated logic

diff --git a/Engine/Source/Runtime/Engine/Public/ShowFlagsValues.inl b/Engine/Source/Runtime/Engine/Public/ShowFlagsValues.inl
--- a/Engine/Source/Runtime/Engine/Public/ShowFlagsValues.inl
+++ b/Engine/Source/Runtime/Engine/Public/ShowFlagsValues.inl
@@ -380,11 +380,11 @@ SHOWFLAG_FIXED_IN_SHIPPING(0, RayTracingDebug, SFG_Hidden, NSLOCTEXT("UnrealEd",
 /** Enable the SkyAtmosphere visualization to be drawn on screen */
 SHOWFLAG_FIXED_IN_SHIPPING(0, VisualizeSkyAtmosphere, SFG_Visualize, NSLOCTEXT("UnrealEd", "VisualizeSkyAtmosphereSF", "Sky Atmosphere"))
 /** Shows a full-screen calibration color post process material defined in renderer config  */
-SHOWFLAG_FIXED_IN_SHIPPING(0, VisualizeCalibrationColor, SFG_Developer, NSLOCTEXT("UnrealEd", "VisualizeCalibrationColorSF", "Visualize Calibration Color"))
+SHOWFLAG_ALWAYS_ACCESSIBLE(VisualizeCalibrationColor, SFG_Developer, NSLOCTEXT("UnrealEd", "VisualizeCalibrationColorSF", "Visualize Calibration Color"))
 /** Shows a full-screen calibration grayscale post process material defined in renderer config  */
-SHOWFLAG_FIXED_IN_SHIPPING(0, VisualizeCalibrationGrayscale, SFG_Developer, NSLOCTEXT("UnrealEd", "VisualizeCalibrationGrayscaleSF", "Visualize Calibration Grayscale"))
+SHOWFLAG_ALWAYS_ACCESSIBLE(VisualizeCalibrationGrayscale, SFG_Developer, NSLOCTEXT("UnrealEd", "VisualizeCalibrationGrayscaleSF", "Visualize Calibration Grayscale"))
 /** Shows a full-screen calibration custom post process material defined in renderer config  */
-SHOWFLAG_FIXED_IN_SHIPPING(0, VisualizeCalibrationCustom, SFG_Developer, NSLOCTEXT("UnrealEd", "VisualizeCalibrationCustomSF", "Visualize Calibration Custom"))
+SHOWFLAG_ALWAYS_ACCESSIBLE(VisualizeCalibrationCustom, SFG_Developer, NSLOCTEXT("UnrealEd", "VisualizeCalibrationCustomSF", "Visualize Calibration Custom"))
 /** Visualise information about the post processes stacked on screen */
 SHOWFLAG_FIXED_IN_SHIPPING(0, VisualizePostProcessStack, SFG_Visualize, NSLOCTEXT("UnrealEd", "VisualizePostProcessStackSF", "Visualize Post Process Stack"))
 
diff --git a/Engine/Source/Runtime/Renderer/Private/DeferredShadingRenderer.cpp b/Engine/Source/Runtime/Renderer/Private/DeferredShadingRenderer.cpp
--- a/Engine/Source/Runtime/Renderer/Private/DeferredShadingRenderer.cpp
+++ b/Engine/Source/Runtime/Renderer/Private/DeferredShadingRenderer.cpp
@@ -3317,7 +3317,6 @@ void FDeferredShadingSceneRenderer::Render(FRDGBuilder& GraphBuilder)
 
 				PostProcessingInputs.TranslucencyViewResourcesMap = FTranslucencyViewResourcesMap(TranslucencyResourceMap, ViewIndex);
 
-#if !(UE_BUILD_SHIPPING)
 				if (IsPostProcessVisualizeCalibrationMaterialEnabled(View))
 				{
 					const UMaterialInterface* DebugMaterialInterface = GetPostProcessVisualizeCalibrationMaterialInterface(View);
@@ -3326,7 +3325,6 @@ void FDeferredShadingSceneRenderer::Render(FRDGBuilder& GraphBuilder)
 					AddVisualizeCalibrationMaterialPostProcessingPasses(GraphBuilder, View, PostProcessingInputs, DebugMaterialInterface);
 				}
 				else
-#endif
 				{
 					const FPerViewPipelineState& ViewPipelineState = GetViewPipelineState(View);
 					const bool bAnyLumenActive = ViewPipelineState.DiffuseIndirectMethod == EDiffuseIndirectMethod::Lumen || ViewPipelineState.ReflectionsMethod == EReflectionsMethod::Lumen;
diff --git a/Engine/Source/Runtime/Renderer/Private/PostProcess/PostProcessing.cpp b/Engine/Source/Runtime/Renderer/Private/PostProcess/PostProcessing.cpp
--- a/Engine/Source/Runtime/Renderer/Private/PostProcess/PostProcessing.cpp
+++ b/Engine/Source/Runtime/Renderer/Private/PostProcess/PostProcessing.cpp
@@ -1615,8 +1615,6 @@ void AddDebugViewPostProcessingPasses(FRDGBuilder& GraphBuilder, const FViewInfo
 	}
 }
 
-#if !(UE_BUILD_SHIPPING)
-
 void AddVisualizeCalibrationMaterialPostProcessingPasses(FRDGBuilder& GraphBuilder, const FViewInfo& View, const FPostProcessingInputs& Inputs, const UMaterialInterface* InMaterialInterface)
 {
 	RDG_CSV_STAT_EXCLUSIVE_SCOPE(GraphBuilder, RenderPostProcessing);
@@ -1656,8 +1654,6 @@ void AddVisualizeCalibrationMaterialPostProcessingPasses(FRDGBuilder& GraphBuild
 	SceneColor = AddDeviceEncodingOnlyPass(GraphBuilder, View, PassInputs);
 }
 
-#endif
-
 ///////////////////////////////////////////////////////////////////////////
 // Mobile Post Processing
 //////////////////////////////////////////////////////////////////////////
diff --git a/Engine/Source/Runtime/Renderer/Private/PostProcess/PostProcessing.h b/Engine/Source/Runtime/Renderer/Private/PostProcess/PostProcessing.h
--- a/Engine/Source/Runtime/Renderer/Private/PostProcess/PostProcessing.h
+++ b/Engine/Source/Runtime/Renderer/Private/PostProcess/PostProcessing.h
@@ -53,13 +53,8 @@ void AddPostProcessingPasses(
 
 void AddDebugViewPostProcessingPasses(FRDGBuilder& GraphBuilder, const FViewInfo& View, const FPostProcessingInputs& Inputs, const Nanite::FRasterResults* NaniteRasterResults);
 
-#if !(UE_BUILD_SHIPPING)
-
 void AddVisualizeCalibrationMaterialPostProcessingPasses(FRDGBuilder& GraphBuilder, const FViewInfo& View, const FPostProcessingInputs& Inputs, const UMaterialInterface* InMaterialInterface);
 
-#endif
-
-
 struct FMobilePostProcessingInputs
 {
 	TRDGUniformBufferRef<FMobileSceneTextureUniformParameters> SceneTextures = nullptr;
