Strata - refactor and clean up of FStrataMaterialCompilationInfo into the strata tree.

diff --git a/Engine/Shaders/Private/Strata/Strata.ush b/Engine/Shaders/Private/Strata/Strata.ush
--- a/Engine/Shaders/Private/Strata/Strata.ush
+++ b/Engine/Shaders/Private/Strata/Strata.ush
@@ -84,14 +84,6 @@ FRWStrataMaterialContainer InitialiseRWStrataMaterialContainer(RWTexture2DArray<
 #define STRATA_SSS_MATERIAL_OVERRIDE 1
 #endif
 
-#ifndef STRATA_CLAMPED_LAYER_COUNT
-#define STRATA_CLAMPED_LAYER_COUNT 1
-#endif
-
-#ifndef STRATA_CLAMPED_MAX_BSDF_COUNT_PER_LAYER
-#define STRATA_CLAMPED_MAX_BSDF_COUNT_PER_LAYER STRATA_MAX_BSDF_COUNT_PER_LAYER
-#endif
-
 #ifndef MATERIAL_SHADINGMODEL_SINGLELAYERWATER
 #define MATERIAL_SHADINGMODEL_SINGLELAYERWATER 0 
 #endif
diff --git a/Engine/Shaders/Shared/StrataDefinitions.h b/Engine/Shaders/Shared/StrataDefinitions.h
--- a/Engine/Shaders/Shared/StrataDefinitions.h
+++ b/Engine/Shaders/Shared/StrataDefinitions.h
@@ -8,7 +8,7 @@
 #pragma once
 
 // Change this to force recompilation of all strata dependent shaders (use https://www.random.org/cgi-bin/randbyte?nbytes=4&format=h)
-#define STRATA_SHADER_VERSION 0xa8970d72
+#define STRATA_SHADER_VERSION 0x3ec8acbb 
 
 
 
diff --git a/Engine/Source/Runtime/Engine/Private/Materials/HLSLMaterialTranslator.cpp b/Engine/Source/Runtime/Engine/Private/Materials/HLSLMaterialTranslator.cpp
--- a/Engine/Source/Runtime/Engine/Private/Materials/HLSLMaterialTranslator.cpp
+++ b/Engine/Source/Runtime/Engine/Private/Materials/HLSLMaterialTranslator.cpp
@@ -246,7 +246,6 @@ FHLSLMaterialTranslator::FHLSLMaterialTranslator(FMaterial* InMaterial,
 ,	AllocatedUserTexCoords()
 ,	AllocatedUserVertexTexCoords()
 ,	DynamicParticleParameterMask(0)
-,	StrataValidFrontMaterialCodeChunkPostTranslate(INDEX_NONE)
 ,	NextFreeStrataShaderNormalIndex(0)
 ,	FinalUsedSharedLocalBasesCount(0)
 ,	NumVtSamples(0)
@@ -1328,73 +1327,19 @@ bool FHLSLMaterialTranslator::Translate()
 			TextureExpression->GetExternalTextureParameterInfo(MaterialCompilationOutput.UniformExpressionSet.UniformExternalTextureParameters.AddDefaulted_GetRef());
 		}
 
-		const bool bStrataFrontMaterialProvided = Chunk[MP_FrontMaterial] != INDEX_NONE && CodeChunkToStrataCompilationInfoMap.Contains(Chunk[MP_FrontMaterial]) == true;
+
+		const bool bStrataFrontMaterialProvided = Chunk[MP_FrontMaterial] != INDEX_NONE;
 		bool bStrataFrontMaterialIsValid = bStrataFrontMaterialProvided;
 		if (bStrataFrontMaterialIsValid)
 		{
-			StrataValidFrontMaterialCodeChunkPostTranslate = Chunk[MP_FrontMaterial];
-			const FStrataMaterialCompilationInfo& StrataCompilationInfo = CodeChunkToStrataCompilationInfoMap[StrataValidFrontMaterialCodeChunkPostTranslate];
-
 			// The material can be null when some entries are automatically generated, for instance in the material layer blending system
-			const bool bIsNullMaterial = (StrataCompilationInfo.LayerCount == 0 && StrataCompilationInfo.TotalBSDFCount == 0);
-			bStrataFrontMaterialIsValid &= !bIsNullMaterial;
+			bStrataFrontMaterialIsValid &= StrataMaterialBSDFCount > 0;
 		}
 
 		if (bStrataFrontMaterialIsValid)
 		{
 			EStrataBlendMode StrataBlendMode = Material->GetStrataBlendMode();
 			bMaterialIsStrata = true;
-			const FStrataMaterialCompilationInfo& StrataCompilationInfo = CodeChunkToStrataCompilationInfoMap[StrataValidFrontMaterialCodeChunkPostTranslate];
-
-			// VolumetricFogCloud must be used in isolation
-			if (StrataMaterialContainsAnyBSDF(this, StrataCompilationInfo, STRATA_BSDF_TYPE_VOLUMETRICFOGCLOUD))
-			{
-				// If the unlit node is used, it must be the only one used
-				if (!StrataIsVolumetricFogCloudOnly(this, StrataCompilationInfo))
-				{
-					FString ErrorMsg = FString::Printf(TEXT("Material %s contains Unlit BSDF but it is not the single BSDF representing the material asset: %s. It must be the single BSDF.\r\n"), *Material->GetDebugName(), *Material->GetAssetPath().ToString());
-					Error(*ErrorMsg);
-				}
-			}
-
-			// Unlit must be used in isolation
-			if (StrataMaterialContainsAnyBSDF(this, StrataCompilationInfo, STRATA_BSDF_TYPE_UNLIT))
-			{
-				FString MaterialDomain = MaterialDomainString(Material->GetMaterialDomain());
-				// If the unlit node is used, it must be the only one used. This is true for all domain (surface, light function, post process).
-				if (!StrataIsUnlitOnly(this, StrataCompilationInfo))
-				{
-					FString ErrorMsg = FString::Printf(TEXT("%s Material %s contains Unlit BSDF but it is not the single BSDF representing the material asset: %s. It must be the single BSDF.\r\n"), *MaterialDomain, *Material->GetDebugName(), *Material->GetAssetPath().ToString());
-					Error(*ErrorMsg);
-				}
-			}
-
-			// Hair must be used in isolation
-			if (StrataMaterialContainsAnyBSDF(this, StrataCompilationInfo, STRATA_BSDF_TYPE_HAIR))
-			{
-				// If the unlit node is used, it must be the only one used
-				if (!StrataIsHairOnly(this, StrataCompilationInfo))
-				{
-					FString ErrorMsg = FString::Printf(TEXT("Material %s contains hair BSDF but it is not the single BSDF representing the material asset: %s. It must be the single BSDF.\r\n"), *Material->GetDebugName(), *Material->GetAssetPath().ToString());
-					Error(*ErrorMsg);
-				}
-			}
-
-			// SingleLayerWater must be used in isolation
-			if (StrataMaterialContainsAnyBSDF(this, StrataCompilationInfo, STRATA_BSDF_TYPE_SINGLELAYERWATER))
-			{
-				// If the unlit node is used, it must be the only one used
-				if (!StrataIsSingleLayerWaterOnly(this, StrataCompilationInfo))
-				{
-					FString ErrorMsg = FString::Printf(TEXT("Material %s contains water BSDF but it is not the single BSDF representing the material asset: %s. It must be the single BSDF.\r\n"), *Material->GetDebugName(), *Material->GetAssetPath().ToString());
-					Error(*ErrorMsg);
-				}
-
-				if (StrataBlendMode != SBM_Opaque && StrataBlendMode != SBM_Masked)
-				{
-					Errorf(TEXT("SingleLayerWater materials must be opaque or masked."));
-				}
-			}
 
 			if (!bStrataUsesConversionFromLegacy && StrataMaterialRootOperator)
 			{
@@ -1982,65 +1927,58 @@ void FHLSLMaterialTranslator::GetMaterialEnvironment(EShaderPlatform InPlatform,
 		* #endif
 		*/
 
-
-		// We go over all output BSDF and reserve actual space for there shared tangent basis. Only up to STRATA_MAX_SHAREDLOCALBASES_REGISTERS.
-		// When we know we will need a basis, it is allocated as FinalSharedLocalBasisIndex. A macro is created with proper output index for the correct value to be output ot the correct slot by the BSDF.
-		// If too many output are required, we fail the compilation.
-		// STRATA_TODO: parameter blend in case it is not possible to output according to per pixel cost restriction.
-		const FStrataMaterialCompilationInfo& StrataCompilationInfo = CodeChunkToStrataCompilationInfoMap[StrataValidFrontMaterialCodeChunkPostTranslate];
-		for (uint32 LayerIt = 0; LayerIt < StrataCompilationInfo.LayerCount; ++LayerIt)
+		for(int32 OpIt=0; OpIt < StrataMaterialExpressionRegisteredOperators.Num(); ++OpIt)
 		{
-			const FStrataMaterialCompilationInfo::FLayer& Layer = StrataCompilationInfo.Layers[LayerIt];
-			for (uint32 BSDFIt = 0; BSDFIt < Layer.BSDFCount; ++BSDFIt)
+			FStrataOperator& BSDFOperator = StrataMaterialExpressionRegisteredOperators[OpIt];
+			if (BSDFOperator.BSDFIndex == INDEX_NONE && !BSDFOperator.bRootOfParameterBlendingSubTree)
 			{
-				const FStrataMaterialCompilationInfo::FBSDF& BSDF = Layer.BSDFs[BSDFIt];
-				if (BSDF.RegisteredSharedLocalBasis.NormalCodeChunk == INDEX_NONE && BSDF.RegisteredSharedLocalBasis.TangentCodeChunk == INDEX_NONE)
-				{
-					// We skip null normal on certain BSDF, for instance unlit.
-					continue;
-				}
+				continue;	// not a bsdf or not the root of a parameter blending subtree so there is not local basis to register
+			}
 
-				const FStrataSharedLocalBasesInfo StrataSharedLocalBasesInfo = StrataCompilationInfoGetMatchingSharedLocalBasisInfo(BSDF.RegisteredSharedLocalBasis);
+			if (BSDFOperator.BSDFRegisteredSharedLocalBasis.NormalCodeChunk == INDEX_NONE && BSDFOperator.BSDFRegisteredSharedLocalBasis.TangentCodeChunk == INDEX_NONE)
+			{
+				continue;	// We skip null normal on certain BSDF, for instance unlit.
+			}
+			const FStrataSharedLocalBasesInfo& StrataSharedLocalBasesInfo = StrataCompilationInfoGetMatchingSharedLocalBasisInfo(BSDFOperator.BSDFRegisteredSharedLocalBasis);
 
-				// First, we check that the normal/tangent has not already written out (avoid 2 BSDFs sharing the same normal to note generate the same code twice)
-				bool bAlreadyProcessed = false;
-				for (uint8 i = 0; i < FinalUsedSharedLocalBasesCount; ++i)
-				{
-					if (UsedSharedLocalBasesInfo[i].NormalCodeChunkHash == StrataSharedLocalBasesInfo.SharedData.NormalCodeChunkHash &&
-						(UsedSharedLocalBasesInfo[i].TangentCodeChunkHash == StrataSharedLocalBasesInfo.SharedData.TangentCodeChunkHash || BSDF.RegisteredSharedLocalBasis.TangentCodeChunk == INDEX_NONE))
-					{
-						bAlreadyProcessed = true;
-						break;
-					}
-				}
-				if (bAlreadyProcessed)
+			// First, we check that the normal/tangent has not already written out (avoid 2 BSDFs sharing the same normal to note generate the same code twice)
+			bool bAlreadyProcessed = false;
+			for (uint8 i = 0; i < FinalUsedSharedLocalBasesCount; ++i)
+			{
+				if (UsedSharedLocalBasesInfo[i].NormalCodeChunkHash == StrataSharedLocalBasesInfo.SharedData.NormalCodeChunkHash &&
+					(UsedSharedLocalBasesInfo[i].TangentCodeChunkHash == StrataSharedLocalBasesInfo.SharedData.TangentCodeChunkHash || BSDFOperator.BSDFRegisteredSharedLocalBasis.TangentCodeChunk == INDEX_NONE))
 				{
-					continue;
+					bAlreadyProcessed = true;
+					break;
 				}
+			}
+			if (bAlreadyProcessed)
+			{
+				continue;
+			}
 
-				check(FinalUsedSharedLocalBasesCount < STRATA_MAX_SHAREDLOCALBASES_REGISTERS);
-				const uint8 FinalSharedLocalBasisIndex = FinalUsedSharedLocalBasesCount++;
-				UsedSharedLocalBasesInfo[FinalSharedLocalBasisIndex] = StrataSharedLocalBasesInfo.SharedData;
+			check(FinalUsedSharedLocalBasesCount < STRATA_MAX_SHAREDLOCALBASES_REGISTERS);
+			const uint8 FinalSharedLocalBasisIndex = FinalUsedSharedLocalBasesCount++;
+			UsedSharedLocalBasesInfo[FinalSharedLocalBasisIndex] = StrataSharedLocalBasesInfo.SharedData;
 
-				// Write out normals
-				StrataPixelNormalInitializerValues += FString::Printf(TEXT("\tParameters.SharedLocalBases.Normals[%u] = %s;\n"), FinalSharedLocalBasisIndex, *StrataSharedLocalBasesInfo.NormalCode);
-				StrataPixelNormalInitializerValues += FString::Printf(TEXT("\t#if MATERIAL_TANGENTSPACENORMAL\n"));
-				StrataPixelNormalInitializerValues += FString::Printf(TEXT("\tParameters.SharedLocalBases.Normals[%u] *= Parameters.TwoSidedSign;\n"), FinalSharedLocalBasisIndex);
-				StrataPixelNormalInitializerValues += FString::Printf(TEXT("\t#endif\n"));
+			// Write out normals
+			StrataPixelNormalInitializerValues += FString::Printf(TEXT("\tParameters.SharedLocalBases.Normals[%u] = %s;\n"), FinalSharedLocalBasisIndex, *StrataSharedLocalBasesInfo.NormalCode);
+			StrataPixelNormalInitializerValues += FString::Printf(TEXT("\t#if MATERIAL_TANGENTSPACENORMAL\n"));
+			StrataPixelNormalInitializerValues += FString::Printf(TEXT("\tParameters.SharedLocalBases.Normals[%u] *= Parameters.TwoSidedSign;\n"), FinalSharedLocalBasisIndex);
+			StrataPixelNormalInitializerValues += FString::Printf(TEXT("\t#endif\n"));
 
-				// Write out tangents
-				if (StrataSharedLocalBasesInfo.SharedData.TangentCodeChunk != INDEX_NONE)
-				{
-					StrataPixelNormalInitializerValues += FString::Printf(TEXT("\tParameters.SharedLocalBases.Tangents[%u] = %s;\n"), FinalSharedLocalBasisIndex, *StrataSharedLocalBasesInfo.TangentCode);
-				}
-				else
-				{
-					StrataPixelNormalInitializerValues += FString::Printf(TEXT("\tParameters.SharedLocalBases.Tangents[%u] = Parameters.TangentToWorld[0];\n"), FinalSharedLocalBasisIndex);
-				}
-				StrataPixelNormalInitializerValues += FString::Printf(TEXT("\t#if MATERIAL_TANGENTSPACENORMAL\n"));
-				StrataPixelNormalInitializerValues += FString::Printf(TEXT("\tParameters.SharedLocalBases.Tangents[%u] *= Parameters.TwoSidedSign;\n"), FinalSharedLocalBasisIndex);
-				StrataPixelNormalInitializerValues += FString::Printf(TEXT("\t#endif\n"));
+			// Write out tangents
+			if (StrataSharedLocalBasesInfo.SharedData.TangentCodeChunk != INDEX_NONE)
+			{
+				StrataPixelNormalInitializerValues += FString::Printf(TEXT("\tParameters.SharedLocalBases.Tangents[%u] = %s;\n"), FinalSharedLocalBasisIndex, *StrataSharedLocalBasesInfo.TangentCode);
+			}
+			else
+			{
+				StrataPixelNormalInitializerValues += FString::Printf(TEXT("\tParameters.SharedLocalBases.Tangents[%u] = Parameters.TangentToWorld[0];\n"), FinalSharedLocalBasisIndex);
 			}
+			StrataPixelNormalInitializerValues += FString::Printf(TEXT("\t#if MATERIAL_TANGENTSPACENORMAL\n"));
+			StrataPixelNormalInitializerValues += FString::Printf(TEXT("\tParameters.SharedLocalBases.Tangents[%u] *= Parameters.TwoSidedSign;\n"), FinalSharedLocalBasisIndex);
+			StrataPixelNormalInitializerValues += FString::Printf(TEXT("\t#endif\n"));
 		}
 		StrataPixelNormalInitializerValues += FString::Printf(TEXT("\tParameters.SharedLocalBases.Count = %u;\n"), FinalUsedSharedLocalBasesCount);
 
@@ -2068,17 +2006,6 @@ void FHLSLMaterialTranslator::GetMaterialEnvironment(EShaderPlatform InPlatform,
 			static const auto CVarStrataBytePerPixel = IConsoleManager::Get().FindTConsoleVariableDataInt(TEXT("r.Strata.BytesPerPixel"));
 			check(CVarStrataBytePerPixel);
 			const uint32 StrataBytePerPixel = CVarStrataBytePerPixel ? CVarStrataBytePerPixel->GetValueOnAnyThread() : 0;
-			FStrataMaterialAnalysisResult StrataMaterialAnalysis = StrataCompilationInfoMaterialAnalysis(this, StrataCompilationInfo, StrataBytePerPixel);
-
-			if (StrataMaterialAnalysis.RequestedBSDFCount > 0)
-			{
-				OutEnvironment.SetDefine(TEXT("STRATA_CLAMPED_LAYER_COUNT"), StrataMaterialAnalysis.ClampedLayerCount);
-			}
-
-			if (StrataMaterialAnalysis.RequestedMaxBSDFCountPerLayer > 0)
-			{
-				OutEnvironment.SetDefine(TEXT("STRATA_CLAMPED_MAX_BSDF_COUNT_PER_LAYER"), StrataMaterialAnalysis.RequestedMaxBSDFCountPerLayer);
-			}
 
 			OutEnvironment.SetDefine(TEXT("STRATA_CLAMPED_BSDF_COUNT"), StrataMaterialBSDFCount);
 
@@ -2086,18 +2013,23 @@ void FHLSLMaterialTranslator::GetMaterialEnvironment(EShaderPlatform InPlatform,
 
 			StrataMaterialDescription += FString::Printf(TEXT("----- STRATA -----\r\n"));
 			StrataMaterialDescription += FString::Printf(TEXT("StrataCompilationInfo -\r\n"));
-			StrataMaterialDescription += FString::Printf(TEXT(" - TotalBSDFCount    = %i\r\n"), StrataCompilationInfo.TotalBSDFCount);
-			StrataMaterialDescription += FString::Printf(TEXT(" - SharedLocalBasesCount = %i\r\n"), StrataCompilationInfoGetSharedLocalBasesCount());
+			StrataMaterialDescription += FString::Printf(TEXT(" - TotalBSDFCount    = %i\r\n"), StrataMaterialBSDFCount);
+			StrataMaterialDescription += FString::Printf(TEXT(" - SharedLocalBasesCount = %i\r\n"), FinalUsedSharedLocalBasesCount);
 
-			for (uint32 LayerIt = 0; LayerIt < StrataCompilationInfo.LayerCount; ++LayerIt)
+			for (int32 OpIt = 0; OpIt < StrataMaterialExpressionRegisteredOperators.Num(); ++OpIt)
 			{
-				StrataMaterialDescription += FString::Printf(TEXT(" Layer %i BSDFs:\r\n"), LayerIt);
-
-				const FStrataMaterialCompilationInfo::FLayer& Layer = StrataCompilationInfo.Layers[LayerIt];
-				for (uint32 BSDFIt = 0; BSDFIt < Layer.BSDFCount; ++BSDFIt)
+				FStrataOperator& BSDFOperator = StrataMaterialExpressionRegisteredOperators[OpIt];
+				if (BSDFOperator.BSDFIndex == INDEX_NONE && !BSDFOperator.bRootOfParameterBlendingSubTree)
+				{
+					continue;	// not a bsdf or not the root of a parameter blending subtree so there is not local basis to register
+				}
+				if (BSDFOperator.BSDFRegisteredSharedLocalBasis.NormalCodeChunk == INDEX_NONE && BSDFOperator.BSDFRegisteredSharedLocalBasis.TangentCodeChunk == INDEX_NONE)
 				{
-					StrataMaterialDescription += FString::Printf(TEXT("     - %s - SharedLocalBasisIndexMacro = %s \r\n"), *GetStrataBSDFName(Layer.BSDFs[BSDFIt].Type), *GetStrataSharedLocalBasisIndexMacro(Layer.BSDFs[BSDFIt].RegisteredSharedLocalBasis));
+					continue;	// We skip null normal on certain BSDF, for instance unlit.
 				}
+				const FStrataSharedLocalBasesInfo& StrataSharedLocalBasesInfo = StrataCompilationInfoGetMatchingSharedLocalBasisInfo(BSDFOperator.BSDFRegisteredSharedLocalBasis);
+
+				StrataMaterialDescription += FString::Printf(TEXT("     - %s - SharedLocalBasisIndexMacro = %s \r\n"), *GetStrataBSDFName(BSDFOperator.BSDFType), *GetStrataSharedLocalBasisIndexMacro(BSDFOperator.BSDFRegisteredSharedLocalBasis));
 			}
 
 			StrataMaterialDescription += FString::Printf(TEXT("----------- STRATA TREE -----------\r\n"));
@@ -2121,16 +2053,7 @@ void FHLSLMaterialTranslator::GetMaterialEnvironment(EShaderPlatform InPlatform,
 			StrataMaterialDescription += FString::Printf(TEXT("-----------------------------------\r\n"));
 
 			StrataMaterialDescription += FString::Printf(TEXT("Byte Per Pixel Budget      %u\r\n"), StrataBytePerPixel);
-			StrataMaterialDescription += FString::Printf(TEXT("Result.bFitInMemoryBudget  %s\r\n"), StrataMaterialAnalysis.bFitInMemoryBudget ? TEXT("YES") : TEXT("NO"));
-			StrataMaterialDescription += FString::Printf(TEXT("Result.RequestedLayerCount %u\r\n"), StrataMaterialAnalysis.RequestedLayerCount);
-			StrataMaterialDescription += FString::Printf(TEXT("Result.RequestedBSDFCount  %u\r\n"), StrataMaterialAnalysis.RequestedBSDFCount);
-			StrataMaterialDescription += FString::Printf(TEXT("Result.RequestedByteCount  %u\r\n"), StrataMaterialAnalysis.RequestedByteCount);
-			if (!StrataMaterialAnalysis.bFitInMemoryBudget)
-			{
-				StrataMaterialDescription += FString::Printf(TEXT("Result.ClampedLayerCount   %u\r\n"), StrataMaterialAnalysis.ClampedLayerCount);	// STRATA_TODO this is no longer needed
-				StrataMaterialDescription += FString::Printf(TEXT("Result.ClampedBSDFCount    %u\r\n"), StrataMaterialAnalysis.ClampedBSDFCount);
-				StrataMaterialDescription += FString::Printf(TEXT("Result.UsedByteCount       %u\r\n"), StrataMaterialAnalysis.UsedByteCount);
-			}
+			// STRATA_TODO: display RequestedByteCount (fit in budget or not)
 			StrataMaterialDescription += FString::Printf(TEXT("------------------\r\n"));
 
 			ResourcesString += TEXT("/*");
@@ -2138,15 +2061,6 @@ void FHLSLMaterialTranslator::GetMaterialEnvironment(EShaderPlatform InPlatform,
 			ResourcesString += TEXT("*/");
 
 			MaterialCompilationOutput.StrataMaterialDescription = StrataMaterialDescription;
-
-			if (StrataMaterialAnalysis.ClampedLayerCount == 0)
-			{
-				Errorf(TEXT("Material %s cannot have any layers rendered due to its complexity (asset: %s).\r\n"), *Material->GetDebugName(), *Material->GetAssetPath().ToString());
-			}
-			else if (StrataMaterialAnalysis.RequestedLayerCount > StrataMaterialAnalysis.ClampedLayerCount)
-			{
-				UE_LOG(LogMaterial, Warning, TEXT("Material %s is not fitting the allocated byte per pixel budget for starta. Layers will be removed (asset: %s).\r\n"), *Material->GetDebugName(), *Material->GetAssetPath().ToString());
-			}
 		}
 
 	}
@@ -9678,6 +9592,17 @@ FStrataOperator& FHLSLMaterialTranslator::StrataCompilationGetOperator(UMaterial
 	return StrataMaterialExpressionRegisteredOperators[*OperatorIndex];
 }
 
+FStrataOperator* FHLSLMaterialTranslator::StrataCompilationGetOperatorFromIndex(int32 OperatorIndex)
+{
+	int32 OperatorCount = StrataMaterialExpressionRegisteredOperators.Num();
+	if (OperatorIndex < 0 || OperatorIndex >= OperatorCount)
+	{
+		Errorf(TEXT("StrataCompilationGetOperatorFromIndex - OperatorIndex out of range %s (asset: %s).\r\n"), *Material->GetDebugName(), *Material->GetAssetPath().ToString());
+		return nullptr;
+	};
+	return &StrataMaterialExpressionRegisteredOperators[OperatorIndex];
+}
+
 bool FHLSLMaterialTranslator::StrataGenerateDerivedMaterialOperatorData()
 {
 	if (StrataMaterialExpressionRegisteredOperators.IsEmpty())
@@ -9757,8 +9682,15 @@ bool FHLSLMaterialTranslator::StrataGenerateDerivedMaterialOperatorData()
 	//
 	// Parse the tree and mark nodes that are the root of a subtree using parameter blending, while other nodes in that tree are forced to use parameter blending.
 	// Allocate BSDFIndex at the same time.
+	// Check BSDF that should be uinique
 	//
 	{
+		bool bHasUnlit = false;
+		bool bHasVFogCloud = false;
+		bool bHasHair = false;
+		bool bHasSLW = false;
+		bool bHasSlab = false;
+
 		int VOpTopBranchCountTaken = 0;
 		int VOpBottomBranchCountTaken = 0;
 		bool bStrataUsesVerticalLayering = false;
@@ -9799,6 +9731,12 @@ bool FHLSLMaterialTranslator::StrataGenerateDerivedMaterialOperatorData()
 				{
 					CurrentOperator.BSDFIndex = StrataMaterialBSDFCount++;
 				}
+
+				bHasUnlit |= CurrentOperator.BSDFType == STRATA_BSDF_TYPE_UNLIT;
+				bHasVFogCloud |= CurrentOperator.BSDFType == STRATA_BSDF_TYPE_VOLUMETRICFOGCLOUD;
+				bHasHair |= CurrentOperator.BSDFType == STRATA_BSDF_TYPE_HAIR;
+				bHasSLW |= CurrentOperator.BSDFType == STRATA_BSDF_TYPE_SINGLELAYERWATER;
+				bHasSlab |= CurrentOperator.BSDFType == STRATA_BSDF_TYPE_SLAB;
 				break;
 			}
 			}
@@ -9809,8 +9747,7 @@ bool FHLSLMaterialTranslator::StrataGenerateDerivedMaterialOperatorData()
 			{
 				CurrentOperator.OperatorType = STRATA_OPERATOR_BSDF;
 				CurrentOperator.BSDFIndex = StrataMaterialBSDFCount++;
-				CurrentOperator.LeftIndex = INDEX_NONE;
-				CurrentOperator.RightIndex = INDEX_NONE;
+				// We do not reset LeftIndex and RightIndex because those are needed to recover local tangent basis information needed with parameter blending.
 			}
 
 			// When at least one vertical operator exists that is not parameter blending, we can enabled writing to opaque rough refraction buffer.
@@ -9822,6 +9759,12 @@ bool FHLSLMaterialTranslator::StrataGenerateDerivedMaterialOperatorData()
 		const EStrataBlendMode StrataBlendMode = Material->GetStrataBlendMode();
 		const bool bIsOpaqueOrMasked = StrataBlendMode == EStrataBlendMode::SBM_Opaque || StrataBlendMode == EStrataBlendMode::SBM_Masked;
 		bStrataOutputsOpaqueRoughRefractions = bStrataUsesVerticalLayering && bIsOpaqueOrMasked;
+
+		if ((bHasUnlit || bHasVFogCloud || bHasHair || bHasSLW) && StrataMaterialBSDFCount > 1)
+		{
+			Errorf(TEXT("Unlit, Fog/Cloud, Hair or SingleLayerWater must be used in isolation %s (asset: %s).\r\n"), *Material->GetDebugName(), *Material->GetAssetPath().ToString());
+			// Even though we could support Unlit with slab.
+		}
 	}
 
 	// STRATA_TODO: operation using parameter blending are not actually discarded so they still occupy a spot in the operation array in the compiler and in the shader also.
@@ -9845,7 +9788,11 @@ bool FHLSLMaterialTranslator::StrataGenerateDerivedMaterialOperatorData()
 		// Operators without any child
 		case STRATA_OPERATOR_BSDF:
 		{
-			check(It.LeftIndex == INDEX_NONE && It.RightIndex == INDEX_NONE && It.BSDFIndex != INDEX_NONE);
+			if (!It.bUseParameterBlending)
+			{
+				// When using parameter blending, we need to keep indices to be able to recover local basis information for normal blending.
+				check(It.LeftIndex == INDEX_NONE && It.RightIndex == INDEX_NONE && It.BSDFIndex != INDEX_NONE);
+			}
 			break;
 		}
 
@@ -9970,33 +9917,6 @@ bool FHLSLMaterialTranslator::StrataGenerateDerivedMaterialOperatorData()
 	return true; // Success
 }
 
-void FHLSLMaterialTranslator::StrataCompilationInfoRegisterCodeChunk(int32 CodeChunk, FStrataMaterialCompilationInfo& StrataMaterialCompilationInfo)
-{
-	if (CodeChunk != INDEX_NONE)
-	{
-		CodeChunkToStrataCompilationInfoMap.Add(CodeChunk, StrataMaterialCompilationInfo);
-	}
-	else
-	{
-		UE_LOG(LogMaterial, Error, TEXT("Invalid Strata CodeChunk registered in Material %s (asset: %s).\r\n"), *Material->GetDebugName(), *Material->GetAssetPath().ToString());
-	}
-}
-
-bool FHLSLMaterialTranslator::StrataCompilationInfoContainsCodeChunk(int32 CodeChunk)
-{
-	return CodeChunk != INDEX_NONE ? CodeChunkToStrataCompilationInfoMap.Contains(CodeChunk) : false;
-}
-
-const FStrataMaterialCompilationInfo& FHLSLMaterialTranslator::GetStrataCompilationInfo(int32 CodeChunk)
-{
-	if (CodeChunk != INDEX_NONE)
-	{
-		return CodeChunkToStrataCompilationInfoMap[CodeChunk];
-	}
-	static FStrataMaterialCompilationInfo InvalidInfo;
-	return InvalidInfo;
-}
-
 FStrataRegisteredSharedLocalBasis FHLSLMaterialTranslator::StrataCompilationInfoRegisterSharedLocalBasis(int32 NormalCodeChunk)
 {
 	check(NormalCodeChunk != INDEX_NONE);
@@ -10085,11 +10005,6 @@ FHLSLMaterialTranslator::FStrataSharedLocalBasesInfo FHLSLMaterialTranslator::St
 	return FStrataSharedLocalBasesInfo();
 }
 
-uint8 FHLSLMaterialTranslator::StrataCompilationInfoGetSharedLocalBasesCount()
-{
-	return FinalUsedSharedLocalBasesCount;
-}
-
 int32 FHLSLMaterialTranslator::StrataAddParameterBlendingBSDFCoverageToNormalMixCodeChunk(int32 ACodeChunk, int32 BCodeChunk)
 {
 	return AddCodeChunk(MCT_Float, TEXT("AddParameterBlendingBSDFCoverageToNormalMix(%s, %s)"), *GetParameterCode(ACodeChunk), *GetParameterCode(BCodeChunk));
@@ -10108,7 +10023,6 @@ int32 FHLSLMaterialTranslator::StrataHorizontalMixingParameterBlendingBSDFCovera
 int32 FHLSLMaterialTranslator::StrataCreateAndRegisterNullMaterial()
 {
 	int32 OutputCodeChunk = AddInlinedCodeChunk(MCT_Strata, TEXT("GetInitialisedStrataData()"));
-	StrataCompilationInfoCreateNullBSDF(this, OutputCodeChunk);
 	return OutputCodeChunk;
 }
 
diff --git a/Engine/Source/Runtime/Engine/Private/Materials/HLSLMaterialTranslator.h b/Engine/Source/Runtime/Engine/Private/Materials/HLSLMaterialTranslator.h
--- a/Engine/Source/Runtime/Engine/Private/Materials/HLSLMaterialTranslator.h
+++ b/Engine/Source/Runtime/Engine/Private/Materials/HLSLMaterialTranslator.h
@@ -406,10 +406,6 @@ protected:
 	/** Will contain all the shading models picked up from the material expression graph */
 	FMaterialShadingModelField ShadingModelsFromCompilation;
 
-	/** Maps a code chunk to a Strata material topology, if the code chunk represent a StrataData. */
-	TMap<int32, FStrataMaterialCompilationInfo> CodeChunkToStrataCompilationInfoMap;
-	/** Only valid after Translate() is called. This is the code chunk representing the front material. It can be used to recover the material topology. */
-	int32 StrataValidFrontMaterialCodeChunkPostTranslate;
 	/** The code initializing the array of shared local bases. */
 	FString StrataPixelNormalInitializerValues;
 	/** The next free index that can be used to represent a unique macros pointing to the position in the array of shared local bases written to memory once the shader is executed. */
@@ -1058,13 +1054,10 @@ protected:
 	
 	virtual FStrataOperator& StrataCompilationRegisterOperator(int32 OperatorType, UMaterialExpression* Expression, UMaterialExpression* Parent, bool bUseParameterBlending = false) override;
 	virtual FStrataOperator& StrataCompilationGetOperator(UMaterialExpression* Expression) override;
+	virtual FStrataOperator* StrataCompilationGetOperatorFromIndex(int32 OperatorIndex) override;
 
-	virtual void StrataCompilationInfoRegisterCodeChunk(int32 CodeChunk, FStrataMaterialCompilationInfo& StrataMaterialCompilationInfo) override;
-	virtual bool StrataCompilationInfoContainsCodeChunk(int32 CodeChunk) override;
-	virtual const FStrataMaterialCompilationInfo& GetStrataCompilationInfo(int32 CodeChunk) override;
 	virtual FStrataRegisteredSharedLocalBasis StrataCompilationInfoRegisterSharedLocalBasis(int32 NormalCodeChunk) override;
 	virtual FStrataRegisteredSharedLocalBasis StrataCompilationInfoRegisterSharedLocalBasis(int32 NormalCodeChunk, int32 TangentCodeChunk) override;
-	virtual uint8 StrataCompilationInfoGetSharedLocalBasesCount() override;
 	virtual int32 StrataAddParameterBlendingBSDFCoverageToNormalMixCodeChunk(int32 ACodeChunk, int32 BCodeChunk) override;
 	virtual int32 StrataVerticalLayeringParameterBlendingBSDFCoverageToNormalMixCodeChunk(int32 TopCodeChunk) override;
 	virtual int32 StrataHorizontalMixingParameterBlendingBSDFCoverageToNormalMixCodeChunk(int32 BackgroundCodeChunk, int32 ForegroundCodeChunk, int32 HorizontalMixCodeChunk) override;
diff --git a/Engine/Source/Runtime/Engine/Private/Materials/MaterialExpressions.cpp b/Engine/Source/Runtime/Engine/Private/Materials/MaterialExpressions.cpp
--- a/Engine/Source/Runtime/Engine/Private/Materials/MaterialExpressions.cpp
+++ b/Engine/Source/Runtime/Engine/Private/Materials/MaterialExpressions.cpp
@@ -20746,85 +20746,44 @@ int32 UMaterialExpressionStrataLegacyConversion::Compile(class FMaterialCompiler
 		ClearCoat_TangentCodeChunk,
 		ClearCoat_BasisIndexMacro);
 
-	// In case we don't know in advance what the conversion output looks like, we setup the info for the wost case (clear coat / eye)
-	auto AddDefaultWorstCase = [Compiler, OutputCodeChunk, &NewRegisteredSharedLocalBasis, &ClearCoat_NewRegisteredSharedLocalBasis](bool bSSS)
-	{
-		// Worst case for ClearCoat / Eye
-		FStrataMaterialCompilationInfo StrataInfo;
-		StrataInfo.LayerCount = 2;
-		StrataInfo.TotalBSDFCount = 2;
-
-		StrataInfo.Layers[0].BSDFCount = 1;
-		StrataInfo.Layers[0].BSDFs[0].Type = STRATA_BSDF_TYPE_SLAB;
-		StrataInfo.Layers[0].BSDFs[0].RegisteredSharedLocalBasis = NewRegisteredSharedLocalBasis;
-		StrataInfo.Layers[0].BSDFs[0].bHasSSS = bSSS;
-		StrataInfo.Layers[0].BSDFs[0].bHasMFPPluggedIn = true;
-		StrataInfo.Layers[0].BSDFs[0].bHasEdgeColor = false;
-		StrataInfo.Layers[0].BSDFs[0].bHasThinFilm = false;
-		StrataInfo.Layers[0].BSDFs[0].bHasFuzz = false;
-		StrataInfo.Layers[0].BSDFs[0].bHasHaziness = false;
-
-		StrataInfo.Layers[1].BSDFCount = 1;
-		StrataInfo.Layers[1].BSDFs[0].Type = STRATA_BSDF_TYPE_SLAB;
-		StrataInfo.Layers[1].BSDFs[0].RegisteredSharedLocalBasis = ClearCoat_NewRegisteredSharedLocalBasis;
-		StrataInfo.Layers[1].BSDFs[0].bHasSSS = bSSS;
-		StrataInfo.Layers[1].BSDFs[0].bHasMFPPluggedIn = false;
-		StrataInfo.Layers[1].BSDFs[0].bHasEdgeColor = false;
-		StrataInfo.Layers[1].BSDFs[0].bHasThinFilm = false;
-		StrataInfo.Layers[1].BSDFs[0].bHasFuzz = true;
-		StrataInfo.Layers[1].BSDFs[0].bHasHaziness = false;
-
-		Compiler->StrataCompilationInfoRegisterCodeChunk(OutputCodeChunk, StrataInfo);
-	};
+	// Now update the local shared basis
 
-	// Fill in Strata layer info
+	// The legacy node can only be added once and we know the worst shape it can take so we hard code it here, 
+	// with fake addresses for BSDF defined in shader code.
+	UMaterialExpression* ChildTopExpression = reinterpret_cast<UMaterialExpression*>(1);
+	UMaterialExpression* ChildBottomExpression = reinterpret_cast<UMaterialExpression*>(2);
 	if (ConvertedStrataMaterialInfo.CountShadingModels() > 1 || ConvertedStrataMaterialInfo.HasShadingModelFromExpression())
 	{
-		AddDefaultWorstCase(true);
+		Compiler->StrataCompilationGetOperator(ChildTopExpression).BSDFRegisteredSharedLocalBasis = ClearCoat_NewRegisteredSharedLocalBasis;
+		Compiler->StrataCompilationGetOperator(ChildBottomExpression).BSDFRegisteredSharedLocalBasis = NewRegisteredSharedLocalBasis;
 	}
 	else
 	{
 		check(ConvertedStrataMaterialInfo.CountShadingModels() == 1);
 
-		if (ConvertedStrataMaterialInfo.HasShadingModel(SSM_Unlit)) 
-		{ 
-			StrataCompilationInfoCreateSingleBSDFMaterial(Compiler, OutputCodeChunk, NewRegisteredSharedLocalBasis,
-				STRATA_BSDF_TYPE_UNLIT, 
-				false /*bHasSSS*/, 
-				false /*bHasMFPPluggedIn*/, 
-				false /*bHasEdgeColor*/, 
-				false /*bHasFuzz*/, 
-				false /*bHasHaziness*/);
+		if (ConvertedStrataMaterialInfo.HasShadingModel(SSM_Unlit))
+		{
+			Compiler->StrataCompilationGetOperator(this).BSDFRegisteredSharedLocalBasis = NewRegisteredSharedLocalBasis;
 		}
 		else if (ConvertedStrataMaterialInfo.HasShadingModel(SSM_DefaultLit))
 		{
-			AddDefaultWorstCase(false);
+			Compiler->StrataCompilationGetOperator(ChildTopExpression).BSDFRegisteredSharedLocalBasis = ClearCoat_NewRegisteredSharedLocalBasis;
+			Compiler->StrataCompilationGetOperator(ChildBottomExpression).BSDFRegisteredSharedLocalBasis = NewRegisteredSharedLocalBasis;
 		}
 		else if (ConvertedStrataMaterialInfo.HasShadingModel(SSM_SubsurfaceLit))
 		{
-			AddDefaultWorstCase(true);
+			Compiler->StrataCompilationGetOperator(ChildTopExpression).BSDFRegisteredSharedLocalBasis = ClearCoat_NewRegisteredSharedLocalBasis;
+			Compiler->StrataCompilationGetOperator(ChildBottomExpression).BSDFRegisteredSharedLocalBasis = NewRegisteredSharedLocalBasis;
 		}
 		else if (ConvertedStrataMaterialInfo.HasShadingModel(SSM_Hair))
 		{
-			StrataCompilationInfoCreateSingleBSDFMaterial(Compiler, OutputCodeChunk, NewRegisteredSharedLocalBasis,
-				STRATA_BSDF_TYPE_HAIR,
-				false /*bHasSSS*/,
-				false /*bHasMFPPluggedIn*/,
-				false /*bHasEdgeColor*/,
-				false /*bHasFuzz*/,
-				false /*bHasHaziness*/);
+			Compiler->StrataCompilationGetOperator(this).BSDFRegisteredSharedLocalBasis = NewRegisteredSharedLocalBasis;
 		}
 		else if (ConvertedStrataMaterialInfo.HasShadingModel(SSM_SingleLayerWater))
 		{
-			StrataCompilationInfoCreateSingleBSDFMaterial(Compiler, OutputCodeChunk, NewRegisteredSharedLocalBasis,
-				STRATA_BSDF_TYPE_SINGLELAYERWATER,
-				false /*bHasSSS*/,
-				false /*bHasMFPPluggedIn*/,
-				false /*bHasEdgeColor*/,
-				false /*bHasFuzz*/,
-				false /*bHasHaziness*/);
+			Compiler->StrataCompilationGetOperator(this).BSDFRegisteredSharedLocalBasis = NewRegisteredSharedLocalBasis;
 		}
-	}	
+	}
 
 	return OutputCodeChunk;
 }
@@ -20967,17 +20926,72 @@ void UMaterialExpressionStrataLegacyConversion::GatherStrataMaterialInfo(FStrata
 
 FStrataOperator* UMaterialExpressionStrataLegacyConversion::StrataGenerateMaterialTopologyTree(class FMaterialCompiler* Compiler, class UMaterialExpression* Parent, int32 OutputIndex)
 {
-	FStrataOperator& StrataOperator = Compiler->StrataCompilationRegisterOperator(STRATA_OPERATOR_VERTICAL, this, Parent);
-
 	// The legacy node can only be added once and we know the worst shape it can take so we hard code it here, 
 	// with fake addresses for BSDF defined in shader code.
-	UMaterialExpression* ChildAExpression = reinterpret_cast<UMaterialExpression*>(0);
-	UMaterialExpression* ChildBExpression = reinterpret_cast<UMaterialExpression*>(1);
+	UMaterialExpression* ChildTopExpression = reinterpret_cast<UMaterialExpression*>(1);
+	UMaterialExpression* ChildBottomExpression = reinterpret_cast<UMaterialExpression*>(2);
 
-	AssignOperatorIndexIfNotNull(StrataOperator.LeftIndex, &Compiler->StrataCompilationRegisterOperator(STRATA_OPERATOR_BSDF_LEGACY, ChildAExpression, this));
-	AssignOperatorIndexIfNotNull(StrataOperator.RightIndex, &Compiler->StrataCompilationRegisterOperator(STRATA_OPERATOR_BSDF_LEGACY, ChildBExpression, this));
+	auto AddDefaultWorstCase = [&](bool bSSS)
+	{
+		FStrataOperator& StrataOperator = Compiler->StrataCompilationRegisterOperator(STRATA_OPERATOR_VERTICAL, this, Parent);
 
-	return &StrataOperator;
+		FStrataOperator& TopOperator = Compiler->StrataCompilationRegisterOperator(STRATA_OPERATOR_BSDF_LEGACY, ChildTopExpression, this);
+		TopOperator.BSDFType = STRATA_BSDF_TYPE_SLAB;
+		TopOperator.bBSDFHasSSS = false;
+		TopOperator.bBSDFHasMFPPluggedIn = true;
+		TopOperator.bBSDFHasFuzz = true;
+		FStrataOperator& BottomOperator = Compiler->StrataCompilationRegisterOperator(STRATA_OPERATOR_BSDF_LEGACY, ChildBottomExpression, this);
+		BottomOperator.BSDFType = STRATA_BSDF_TYPE_SLAB;
+		BottomOperator.bBSDFHasSSS = true;
+		BottomOperator.bBSDFHasMFPPluggedIn = true;
+		TopOperator.bBSDFHasFuzz = false;
+
+		AssignOperatorIndexIfNotNull(StrataOperator.LeftIndex, &TopOperator);
+		AssignOperatorIndexIfNotNull(StrataOperator.RightIndex, &BottomOperator);
+
+		return &StrataOperator;
+	};
+
+	// Logic about shading models and complexity should match UMaterialExpressionStrataLegacyConversion::Compile.
+	if (ConvertedStrataMaterialInfo.CountShadingModels() > 1 || ConvertedStrataMaterialInfo.HasShadingModelFromExpression())
+	{
+		return AddDefaultWorstCase(true);
+	}
+	// else
+	{
+		check(ConvertedStrataMaterialInfo.CountShadingModels() == 1);
+
+		if (ConvertedStrataMaterialInfo.HasShadingModel(SSM_Unlit))
+		{
+			FStrataOperator& Operator = Compiler->StrataCompilationRegisterOperator(STRATA_OPERATOR_BSDF_LEGACY, this, Parent);
+			Operator.BSDFType = STRATA_BSDF_TYPE_UNLIT;
+			return &Operator;
+		}
+		else if (ConvertedStrataMaterialInfo.HasShadingModel(SSM_DefaultLit))
+		{
+			return AddDefaultWorstCase(false);
+		}
+		else if (ConvertedStrataMaterialInfo.HasShadingModel(SSM_SubsurfaceLit))
+		{
+			return AddDefaultWorstCase(true);
+		}
+		else if (ConvertedStrataMaterialInfo.HasShadingModel(SSM_Hair))
+		{
+			FStrataOperator& Operator = Compiler->StrataCompilationRegisterOperator(STRATA_OPERATOR_BSDF_LEGACY, this, Parent);
+			Operator.BSDFType = STRATA_BSDF_TYPE_HAIR;
+			return &Operator;
+		}
+		else if (ConvertedStrataMaterialInfo.HasShadingModel(SSM_SingleLayerWater))
+		{
+			FStrataOperator& Operator = Compiler->StrataCompilationRegisterOperator(STRATA_OPERATOR_BSDF_LEGACY, this, Parent);
+			Operator.BSDFType = STRATA_BSDF_TYPE_SINGLELAYERWATER;
+			return &Operator;
+		}
+
+		check(false);
+		static FStrataOperator DefaultOperatorOnError;
+		return &DefaultOperatorOnError;
+	}
 }
 
 bool UMaterialExpressionStrataLegacyConversion::HasSSS() const
@@ -21057,6 +21071,12 @@ int32 UMaterialExpressionStrataSlabBSDF::Compile(class FMaterialCompiler* Compil
 	const float DefaultF0 = DielectricSpecularToF0(DefaultSpecular);
 
 	FStrataOperator& StrataOperator = Compiler->StrataCompilationGetOperator(this);
+	StrataOperator.BSDFRegisteredSharedLocalBasis = NewRegisteredSharedLocalBasis;
+	StrataOperator.bBSDFHasEdgeColor = bHasEdgeColor;
+	StrataOperator.bBSDFHasFuzz = bHasFuzz;
+	StrataOperator.bBSDFHasHaziness = bHasHaziness;
+	StrataOperator.bBSDFHasSSS = bHasSSS;
+	StrataOperator.bBSDFHasMFPPluggedIn = bHasMFPPluggedIn;
 
 	int32 OutputCodeChunk = Compiler->StrataSlabBSDF(
 		bUseMetalness ? Compiler->Constant(1.0f) : Compiler->Constant(0.0f),
@@ -21086,7 +21106,6 @@ int32 UMaterialExpressionStrataSlabBSDF::Compile(class FMaterialCompiler* Compil
 		TangentCodeChunk,
 		Compiler->GetStrataSharedLocalBasisIndexMacro(NewRegisteredSharedLocalBasis),
 		!StrataOperator.bUseParameterBlending || (StrataOperator.bUseParameterBlending && StrataOperator.bRootOfParameterBlendingSubTree) ? &StrataOperator : nullptr);
-	StrataCompilationInfoCreateSingleBSDFMaterial(Compiler, OutputCodeChunk, NewRegisteredSharedLocalBasis, STRATA_BSDF_TYPE_SLAB, bHasSSS, bHasMFPPluggedIn, bHasEdgeColor, bHasFuzz, bHasHaziness);
 
 	return OutputCodeChunk;
 }
@@ -21493,7 +21512,9 @@ void UMaterialExpressionStrataSlabBSDF::GatherStrataMaterialInfo(FStrataMaterial
 
 FStrataOperator* UMaterialExpressionStrataSlabBSDF::StrataGenerateMaterialTopologyTree(class FMaterialCompiler* Compiler, class UMaterialExpression* Parent, int32 OutputIndex)
 {
-	return &Compiler->StrataCompilationRegisterOperator(STRATA_OPERATOR_BSDF, this, Parent);
+	FStrataOperator& StrataOperator = Compiler->StrataCompilationRegisterOperator(STRATA_OPERATOR_BSDF, this, Parent);
+	StrataOperator.BSDFType = STRATA_BSDF_TYPE_SLAB;
+	return &StrataOperator;
 }
 
 bool UMaterialExpressionStrataSlabBSDF::HasSSS() const
@@ -21557,8 +21578,6 @@ int32 UMaterialExpressionStrataVolumetricFogCloudBSDF::Compile(class FMaterialCo
 		CompileWithDefaultFloat3(Compiler, EmissiveColor, 0.0f, 0.0f, 0.0f),
 		CompileWithDefaultFloat1(Compiler, AmbientOcclusion, 1.0f));
 
-	StrataCompilationInfoCreateSingleBSDFMaterial(Compiler, OutputCodeChunk, StrataCompilationInfoCreateNullSharedLocalBasis(), STRATA_BSDF_TYPE_VOLUMETRICFOGCLOUD);
-
 	return OutputCodeChunk;
 }
 
@@ -21606,7 +21625,9 @@ void UMaterialExpressionStrataVolumetricFogCloudBSDF::GatherStrataMaterialInfo(F
 
 FStrataOperator* UMaterialExpressionStrataVolumetricFogCloudBSDF::StrataGenerateMaterialTopologyTree(class FMaterialCompiler* Compiler, class UMaterialExpression* Parent, int32 OutputIndex)
 {
-	return &Compiler->StrataCompilationRegisterOperator(STRATA_OPERATOR_BSDF, this, Parent);
+	FStrataOperator& StrataOperator = Compiler->StrataCompilationRegisterOperator(STRATA_OPERATOR_BSDF, this, Parent);
+	StrataOperator.BSDFType = STRATA_BSDF_TYPE_VOLUMETRICFOGCLOUD;
+	return &StrataOperator;
 }
 #endif // WITH_EDITOR
 
@@ -21633,8 +21654,6 @@ int32 UMaterialExpressionStrataLightFunction::Compile(class FMaterialCompiler* C
 		CompileWithDefaultFloat3(Compiler, Color, 0.0f, 0.0f, 0.0f),
 		Compiler->Constant(1.0f));	// Opacity / Transmittance is ignored by light functions.
 
-	StrataCompilationInfoCreateSingleBSDFMaterial(Compiler, OutputCodeChunk, StrataCompilationInfoCreateNullSharedLocalBasis(), STRATA_BSDF_TYPE_UNLIT);
-
 	return OutputCodeChunk;
 }
 
@@ -21673,7 +21692,9 @@ void UMaterialExpressionStrataLightFunction::GatherStrataMaterialInfo(FStrataMat
 
 FStrataOperator* UMaterialExpressionStrataLightFunction::StrataGenerateMaterialTopologyTree(class FMaterialCompiler* Compiler, class UMaterialExpression* Parent, int32 OutputIndex)
 {
-	return &Compiler->StrataCompilationRegisterOperator(STRATA_OPERATOR_BSDF, this, Parent);
+	FStrataOperator& StrataOperator = Compiler->StrataCompilationRegisterOperator(STRATA_OPERATOR_BSDF, this, Parent);
+	StrataOperator.BSDFType = STRATA_BSDF_TYPE_UNLIT;
+	return &StrataOperator;
 }
 #endif // WITH_EDITOR
 
@@ -21703,8 +21724,6 @@ int32 UMaterialExpressionStrataPostProcess::Compile(class FMaterialCompiler* Com
 		CompileWithDefaultFloat3(Compiler, Color, 0.0f, 0.0f, 0.0f),
 		TransmittanceCodeChunk);
 
-	StrataCompilationInfoCreateSingleBSDFMaterial(Compiler, OutputCodeChunk, StrataCompilationInfoCreateNullSharedLocalBasis(), STRATA_BSDF_TYPE_UNLIT);
-
 	return OutputCodeChunk;
 }
 
@@ -21746,7 +21765,9 @@ void UMaterialExpressionStrataPostProcess::GatherStrataMaterialInfo(FStrataMater
 
 FStrataOperator* UMaterialExpressionStrataPostProcess::StrataGenerateMaterialTopologyTree(class FMaterialCompiler* Compiler, class UMaterialExpression* Parent, int32 OutputIndex)
 {
-	return &Compiler->StrataCompilationRegisterOperator(STRATA_OPERATOR_BSDF, this, Parent);
+	FStrataOperator& StrataOperator = Compiler->StrataCompilationRegisterOperator(STRATA_OPERATOR_BSDF, this, Parent);
+	StrataOperator.BSDFType = STRATA_BSDF_TYPE_UNLIT;
+	return &StrataOperator;
 }
 #endif // WITH_EDITOR
 
@@ -21794,10 +21815,6 @@ int32 UMaterialExpressionStrataConvertToDecal::Compile(class FMaterialCompiler*
 		OutputCodeChunk = Compiler->StrataWeightParameterBlending(
 			DecalMaterialCodeChunk, WeightOfOneCodeChunk,
 			StrataOperator.bRootOfParameterBlendingSubTree ? &StrataOperator : nullptr);
-
-		FStrataMaterialCompilationInfo StrataInfo = StrataCompilationInfoWeightParamBlend(Compiler, Compiler->GetStrataCompilationInfo(DecalMaterialCodeChunk));
-
-		Compiler->StrataCompilationInfoRegisterCodeChunk(OutputCodeChunk, StrataInfo);
 	}
 	return OutputCodeChunk;
 }
@@ -21868,8 +21885,6 @@ int32 UMaterialExpressionStrataUnlitBSDF::Compile(class FMaterialCompiler* Compi
 		CompileWithDefaultFloat3(Compiler, EmissiveColor, 0.0f, 0.0f, 0.0f),
 		CompileWithDefaultFloat3(Compiler, TransmittanceColor, 1.0f, 1.0f, 1.0f));
 
-	StrataCompilationInfoCreateSingleBSDFMaterial(Compiler, OutputCodeChunk, StrataCompilationInfoCreateNullSharedLocalBasis(), STRATA_BSDF_TYPE_UNLIT);
-
 	return OutputCodeChunk;
 }
 
@@ -21912,7 +21927,9 @@ void UMaterialExpressionStrataUnlitBSDF::GatherStrataMaterialInfo(FStrataMateria
 
 FStrataOperator* UMaterialExpressionStrataUnlitBSDF::StrataGenerateMaterialTopologyTree(class FMaterialCompiler* Compiler, class UMaterialExpression* Parent, int32 OutputIndex)
 {
-	return &Compiler->StrataCompilationRegisterOperator(STRATA_OPERATOR_BSDF, this, Parent);
+	FStrataOperator& StrataOperator = Compiler->StrataCompilationRegisterOperator(STRATA_OPERATOR_BSDF, this, Parent);
+	StrataOperator.BSDFType = STRATA_BSDF_TYPE_UNLIT;
+	return &StrataOperator;
 }
 #endif // WITH_EDITOR
 
@@ -21949,8 +21966,6 @@ int32 UMaterialExpressionStrataHairBSDF::Compile(class FMaterialCompiler* Compil
 		TangentCodeChunk,
 		Compiler->GetStrataSharedLocalBasisIndexMacro(NewRegisteredSharedLocalBasis));
 
-	StrataCompilationInfoCreateSingleBSDFMaterial(Compiler, OutputCodeChunk, NewRegisteredSharedLocalBasis, STRATA_BSDF_TYPE_HAIR);
-
 	return OutputCodeChunk;
 }
 
@@ -22014,7 +22029,9 @@ void UMaterialExpressionStrataHairBSDF::GatherStrataMaterialInfo(FStrataMaterial
 
 FStrataOperator* UMaterialExpressionStrataHairBSDF::StrataGenerateMaterialTopologyTree(class FMaterialCompiler* Compiler, class UMaterialExpression* Parent, int32 OutputIndex)
 {
-	return &Compiler->StrataCompilationRegisterOperator(STRATA_OPERATOR_BSDF, this, Parent);
+	FStrataOperator& StrataOperator = Compiler->StrataCompilationRegisterOperator(STRATA_OPERATOR_BSDF, this, Parent);
+	StrataOperator.BSDFType = STRATA_BSDF_TYPE_HAIR;
+	return &StrataOperator;
 }
 #endif // WITH_EDITOR
 
@@ -22054,8 +22071,6 @@ int32 UMaterialExpressionStrataSingleLayerWaterBSDF::Compile(class FMaterialComp
 		NormalCodeChunk,
 		Compiler->GetStrataSharedLocalBasisIndexMacro(NewRegisteredSharedLocalBasis));
 
-	StrataCompilationInfoCreateSingleBSDFMaterial(Compiler, OutputCodeChunk, NewRegisteredSharedLocalBasis, STRATA_BSDF_TYPE_SINGLELAYERWATER);
-
 	return OutputCodeChunk;
 }
 
@@ -22132,7 +22147,9 @@ void UMaterialExpressionStrataSingleLayerWaterBSDF::GatherStrataMaterialInfo(FSt
 
 FStrataOperator* UMaterialExpressionStrataSingleLayerWaterBSDF::StrataGenerateMaterialTopologyTree(class FMaterialCompiler* Compiler, class UMaterialExpression* Parent, int32 OutputIndex)
 {
-	return &Compiler->StrataCompilationRegisterOperator(STRATA_OPERATOR_BSDF, this, Parent);
+	FStrataOperator& StrataOperator = Compiler->StrataCompilationRegisterOperator(STRATA_OPERATOR_BSDF, this, Parent);
+	StrataOperator.BSDFType = STRATA_BSDF_TYPE_SINGLELAYERWATER;
+	return &StrataOperator;
 }
 #endif // WITH_EDITOR
 
@@ -22167,12 +22184,6 @@ int32 UMaterialExpressionStrataHorizontalMixing::Compile(class FMaterialCompiler
 
 	int32 ForegroundCodeChunk = Foreground.Compile(Compiler);
 	int32 BackgroundCodeChunk = Background.Compile(Compiler);
-	if (!Compiler->StrataCompilationInfoContainsCodeChunk(ForegroundCodeChunk) || !Compiler->StrataCompilationInfoContainsCodeChunk(BackgroundCodeChunk))
-	{
-		return Compiler->Errorf(TEXT("Could not find ForegroundCodeChunk or BackgroundCodeChunk to mix"));
-	}
-	const FStrataMaterialCompilationInfo& ForegroundStrataData = Compiler->GetStrataCompilationInfo(ForegroundCodeChunk);
-	const FStrataMaterialCompilationInfo& BackgroundStrataData = Compiler->GetStrataCompilationInfo(BackgroundCodeChunk);
 
 	const int32 HorizontalMixCodeChunk = CompileWithDefaultFloat1(Compiler, Mix, 0.5f);
 
@@ -22183,41 +22194,28 @@ int32 UMaterialExpressionStrataHorizontalMixing::Compile(class FMaterialCompiler
 	{
 		const int32 NormalMixCodeChunk = Compiler->StrataHorizontalMixingParameterBlendingBSDFCoverageToNormalMixCodeChunk(BackgroundCodeChunk, ForegroundCodeChunk, HorizontalMixCodeChunk);
 
-		if (ForegroundStrataData.TotalBSDFCount != 1)
-		{
-			return Compiler->Errorf(TEXT("Foreground: cannot parameter horizontal blend complex material topology using parameter blending. Only a one-to-one blending is supported today."));
-		}
-		if (BackgroundStrataData.TotalBSDFCount != 1)
+		FStrataOperator* BackgroundBSDFOperator = Compiler->StrataCompilationGetOperatorFromIndex(StrataOperator.LeftIndex);
+		FStrataOperator* ForegroundBSDFOperator = Compiler->StrataCompilationGetOperatorFromIndex(StrataOperator.RightIndex);
+		if (!BackgroundBSDFOperator || !ForegroundBSDFOperator)
 		{
-			return Compiler->Errorf(TEXT("Background: cannot parameter horizontal blend complex material topology using parameter blending. Only a one-to-one blending is supported today."));
-		}
-
-		const FStrataMaterialCompilationInfo::FBSDF& ForegroundBSDF = ForegroundStrataData.Layers[0].BSDFs[0];
-		const FStrataMaterialCompilationInfo::FBSDF& BackgroundBSDF = BackgroundStrataData.Layers[0].BSDFs[0];
-		if (ForegroundBSDF.Type != STRATA_BSDF_TYPE_SLAB)
-		{
-			return Compiler->Errorf(TEXT("Foreground: we can only apply parameter blending to Slab BSDF."));
-		}
-		if (BackgroundBSDF.Type != STRATA_BSDF_TYPE_SLAB)
-		{
-			return Compiler->Errorf(TEXT("Background: we can only apply parameter blending to Slab BSDF."));
+			return Compiler->Errorf(TEXT("Missing input on horizontal blending node."));
 		}
 
 		// Compute the new Normal and Tangent resulting from the blending using code chunk
-		const int32 NewNormalCodeChunk = StrataBlendNormal(Compiler, BackgroundBSDF.RegisteredSharedLocalBasis.NormalCodeChunk, ForegroundBSDF.RegisteredSharedLocalBasis.NormalCodeChunk, NormalMixCodeChunk);
+		const int32 NewNormalCodeChunk = StrataBlendNormal(Compiler, BackgroundBSDFOperator->BSDFRegisteredSharedLocalBasis.NormalCodeChunk, ForegroundBSDFOperator->BSDFRegisteredSharedLocalBasis.NormalCodeChunk, NormalMixCodeChunk);
 		// The tangent is optional so we treat it differently if INDEX_NONE is specified
 		int32 NewTangentCodeChunk = INDEX_NONE;
-		if (ForegroundBSDF.RegisteredSharedLocalBasis.TangentCodeChunk != INDEX_NONE && BackgroundBSDF.RegisteredSharedLocalBasis.TangentCodeChunk != INDEX_NONE)
+		if (ForegroundBSDFOperator->BSDFRegisteredSharedLocalBasis.TangentCodeChunk != INDEX_NONE && BackgroundBSDFOperator->BSDFRegisteredSharedLocalBasis.TangentCodeChunk != INDEX_NONE)
 		{
-			NewTangentCodeChunk = StrataBlendNormal(Compiler, BackgroundBSDF.RegisteredSharedLocalBasis.TangentCodeChunk, ForegroundBSDF.RegisteredSharedLocalBasis.TangentCodeChunk, NormalMixCodeChunk);
+			NewTangentCodeChunk = StrataBlendNormal(Compiler, BackgroundBSDFOperator->BSDFRegisteredSharedLocalBasis.TangentCodeChunk, ForegroundBSDFOperator->BSDFRegisteredSharedLocalBasis.TangentCodeChunk, NormalMixCodeChunk);
 		}
-		else if (ForegroundBSDF.RegisteredSharedLocalBasis.TangentCodeChunk != INDEX_NONE)
+		else if (ForegroundBSDFOperator->BSDFRegisteredSharedLocalBasis.TangentCodeChunk != INDEX_NONE)
 		{
-			NewTangentCodeChunk = ForegroundBSDF.RegisteredSharedLocalBasis.TangentCodeChunk;
+			NewTangentCodeChunk = ForegroundBSDFOperator->BSDFRegisteredSharedLocalBasis.TangentCodeChunk;
 		}
-		else if (BackgroundBSDF.RegisteredSharedLocalBasis.TangentCodeChunk != INDEX_NONE)
+		else if (BackgroundBSDFOperator->BSDFRegisteredSharedLocalBasis.TangentCodeChunk != INDEX_NONE)
 		{
-			NewTangentCodeChunk = BackgroundBSDF.RegisteredSharedLocalBasis.TangentCodeChunk;
+			NewTangentCodeChunk = BackgroundBSDFOperator->BSDFRegisteredSharedLocalBasis.TangentCodeChunk;
 		}
 		const FStrataRegisteredSharedLocalBasis NewRegisteredSharedLocalBasis = StrataCompilationInfoCreateSharedLocalBasis(Compiler, NewNormalCodeChunk, NewTangentCodeChunk);
 
@@ -22225,10 +22223,9 @@ int32 UMaterialExpressionStrataHorizontalMixing::Compile(class FMaterialCompiler
 			BackgroundCodeChunk, ForegroundCodeChunk, HorizontalMixCodeChunk, NormalMixCodeChunk, Compiler->GetStrataSharedLocalBasisIndexMacro(NewRegisteredSharedLocalBasis),
 			StrataOperator.bRootOfParameterBlendingSubTree ? &StrataOperator : nullptr);
 
-		StrataCompilationInfoCreateSingleBSDFMaterial(Compiler, OutputCodeChunk, NewRegisteredSharedLocalBasis, STRATA_BSDF_TYPE_SLAB);
-		FStrataMaterialCompilationInfo StrataInfo = StrataCompilationInfoHorizontalMixingParamBlend(Compiler, BackgroundStrataData, ForegroundStrataData, NewRegisteredSharedLocalBasis);
-
-		Compiler->StrataCompilationInfoRegisterCodeChunk(OutputCodeChunk, StrataInfo);
+		// Propagate the parameter blended normal
+		StrataOperator.BSDFRegisteredSharedLocalBasis = NewRegisteredSharedLocalBasis;
+		StrataOperator.CombineFlagsForParameterBlending(*BackgroundBSDFOperator, *ForegroundBSDFOperator);
 	}
 	else
 	{
@@ -22238,10 +22235,6 @@ int32 UMaterialExpressionStrataHorizontalMixing::Compile(class FMaterialCompiler
 			HorizontalMixCodeChunk,
 			StrataOperator.Index,
 			StrataOperator.MaxDistanceFromLeaves);
-
-		FStrataMaterialCompilationInfo StrataInfo = StrataCompilationInfoHorizontalMixing(Compiler, BackgroundStrataData, ForegroundStrataData);
-
-		Compiler->StrataCompilationInfoRegisterCodeChunk(OutputCodeChunk, StrataInfo);
 	}
 
 	return OutputCodeChunk;
@@ -22336,12 +22329,6 @@ int32 UMaterialExpressionStrataVerticalLayering::Compile(class FMaterialCompiler
 
 	int32 TopCodeChunk = Top.Compile(Compiler);
 	int32 BaseCodeChunk = Base.Compile(Compiler);
-	if (!Compiler->StrataCompilationInfoContainsCodeChunk(TopCodeChunk) || !Compiler->StrataCompilationInfoContainsCodeChunk(BaseCodeChunk))
-	{
-		return Compiler->Errorf(TEXT("Could not find Top CodeChunk or BaseCodeChunk to VerticalLayer"));
-	}
-	const FStrataMaterialCompilationInfo& TopStrataData = Compiler->GetStrataCompilationInfo(TopCodeChunk);
-	const FStrataMaterialCompilationInfo& BaseStrataData = Compiler->GetStrataCompilationInfo(BaseCodeChunk);
 
 	const int32 TopNormalMixCodeChunk = Compiler->StrataVerticalLayeringParameterBlendingBSDFCoverageToNormalMixCodeChunk(TopCodeChunk);
 
@@ -22349,60 +22336,42 @@ int32 UMaterialExpressionStrataVerticalLayering::Compile(class FMaterialCompiler
 	FStrataOperator& StrataOperator = Compiler->StrataCompilationGetOperator(this);
 	if (StrataOperator.bUseParameterBlending)
 	{
-		if (TopStrataData.TotalBSDFCount != 1)
-		{
-			return Compiler->Errorf(TEXT("Top input: cannot vertical layer complex material topology using parameter blending. Only a one-to-one blending is supported today."));
-		}
-		if (BaseStrataData.TotalBSDFCount != 1)
-		{
-			return Compiler->Errorf(TEXT("Base input: cannot vertical layer complex material topology using parameter blending. Only a one-to-one blending is supported today."));
-		}
-
-		const FStrataMaterialCompilationInfo::FBSDF& TopBSDF = TopStrataData.Layers[0].BSDFs[0];
-		const FStrataMaterialCompilationInfo::FBSDF& BaseBSDF = BaseStrataData.Layers[0].BSDFs[0];
-		if (TopBSDF.Type != STRATA_BSDF_TYPE_SLAB)
+		FStrataOperator* TopBSDFOperator = Compiler->StrataCompilationGetOperatorFromIndex(StrataOperator.LeftIndex);
+		FStrataOperator* BaseBSDFOperator = Compiler->StrataCompilationGetOperatorFromIndex(StrataOperator.RightIndex);
+		if (!TopBSDFOperator || !BaseBSDFOperator)
 		{
-			return Compiler->Errorf(TEXT("Top input: we can only apply parameter blending to Slab BSDF."));
-		}
-		if (BaseBSDF.Type != STRATA_BSDF_TYPE_SLAB)
-		{
-			return Compiler->Errorf(TEXT("Base input: we can only apply parameter blending to Slab BSDF."));
+			return Compiler->Errorf(TEXT("Missing input on vertical layering node."));
 		}
 
 		// Compute the new Normal and Tangent resulting from the blending using code chunk
-		const int32 NewNormalCodeChunk = StrataBlendNormal(Compiler, BaseBSDF.RegisteredSharedLocalBasis.NormalCodeChunk, TopBSDF.RegisteredSharedLocalBasis.NormalCodeChunk, TopNormalMixCodeChunk);
+		const int32 NewNormalCodeChunk = StrataBlendNormal(Compiler, BaseBSDFOperator->BSDFRegisteredSharedLocalBasis.NormalCodeChunk, TopBSDFOperator->BSDFRegisteredSharedLocalBasis.NormalCodeChunk, TopNormalMixCodeChunk);
 		// The tangent is optional so we treat it differently if INDEX_NONE is specified
 		int32 NewTangentCodeChunk = INDEX_NONE;
-		if (TopBSDF.RegisteredSharedLocalBasis.TangentCodeChunk != INDEX_NONE && BaseBSDF.RegisteredSharedLocalBasis.TangentCodeChunk != INDEX_NONE)
+		if (TopBSDFOperator->BSDFRegisteredSharedLocalBasis.TangentCodeChunk != INDEX_NONE && BaseBSDFOperator->BSDFRegisteredSharedLocalBasis.TangentCodeChunk != INDEX_NONE)
 		{
-			NewTangentCodeChunk = StrataBlendNormal(Compiler, BaseBSDF.RegisteredSharedLocalBasis.TangentCodeChunk, TopBSDF.RegisteredSharedLocalBasis.TangentCodeChunk, TopNormalMixCodeChunk);
+			NewTangentCodeChunk = StrataBlendNormal(Compiler, BaseBSDFOperator->BSDFRegisteredSharedLocalBasis.TangentCodeChunk, TopBSDFOperator->BSDFRegisteredSharedLocalBasis.TangentCodeChunk, TopNormalMixCodeChunk);
 		}
-		else if (TopBSDF.RegisteredSharedLocalBasis.TangentCodeChunk != INDEX_NONE)
+		else if (TopBSDFOperator->BSDFRegisteredSharedLocalBasis.TangentCodeChunk != INDEX_NONE)
 		{
-			NewTangentCodeChunk = TopBSDF.RegisteredSharedLocalBasis.TangentCodeChunk;
+			NewTangentCodeChunk = TopBSDFOperator->BSDFRegisteredSharedLocalBasis.TangentCodeChunk;
 		}
-		else if (BaseBSDF.RegisteredSharedLocalBasis.TangentCodeChunk != INDEX_NONE)
+		else if (BaseBSDFOperator->BSDFRegisteredSharedLocalBasis.TangentCodeChunk != INDEX_NONE)
 		{
-			NewTangentCodeChunk = BaseBSDF.RegisteredSharedLocalBasis.TangentCodeChunk;
+			NewTangentCodeChunk = BaseBSDFOperator->BSDFRegisteredSharedLocalBasis.TangentCodeChunk;
 		}
 		const FStrataRegisteredSharedLocalBasis NewRegisteredSharedLocalBasis = StrataCompilationInfoCreateSharedLocalBasis(Compiler, NewNormalCodeChunk, NewTangentCodeChunk);
 
 		OutputCodeChunk = Compiler->StrataVerticalLayeringParameterBlending(
-			TopCodeChunk, BaseCodeChunk, Compiler->GetStrataSharedLocalBasisIndexMacro(NewRegisteredSharedLocalBasis), TopBSDF.RegisteredSharedLocalBasis.NormalCodeChunk,
+			TopCodeChunk, BaseCodeChunk, Compiler->GetStrataSharedLocalBasisIndexMacro(NewRegisteredSharedLocalBasis), TopBSDFOperator->BSDFRegisteredSharedLocalBasis.NormalCodeChunk,
 			StrataOperator.bRootOfParameterBlendingSubTree ? &StrataOperator : nullptr);
 
-		StrataCompilationInfoCreateSingleBSDFMaterial(Compiler, OutputCodeChunk, NewRegisteredSharedLocalBasis, STRATA_BSDF_TYPE_SLAB);
-		FStrataMaterialCompilationInfo StrataInfo = StrataCompilationInfoVerticalLayeringParamBlend(Compiler, TopStrataData, BaseStrataData, NewRegisteredSharedLocalBasis);
-
-		Compiler->StrataCompilationInfoRegisterCodeChunk(OutputCodeChunk, StrataInfo);
+		// Propagate the parameter blended normal
+		StrataOperator.BSDFRegisteredSharedLocalBasis = NewRegisteredSharedLocalBasis;
+		StrataOperator.CombineFlagsForParameterBlending(*TopBSDFOperator, *BaseBSDFOperator);
 	}
 	else
 	{
 		OutputCodeChunk = Compiler->StrataVerticalLayering(TopCodeChunk, BaseCodeChunk, StrataOperator.Index, StrataOperator.MaxDistanceFromLeaves);
-
-		FStrataMaterialCompilationInfo StrataInfo = StrataCompilationInfoVerticalLayering(Compiler, Compiler->GetStrataCompilationInfo(TopCodeChunk), Compiler->GetStrataCompilationInfo(BaseCodeChunk));
-
-		Compiler->StrataCompilationInfoRegisterCodeChunk(OutputCodeChunk, StrataInfo);
 	}
 
 
@@ -22498,12 +22467,6 @@ int32 UMaterialExpressionStrataAdd::Compile(class FMaterialCompiler* Compiler, i
 
 	int32 ACodeChunk = A.Compile(Compiler);
 	int32 BCodeChunk = B.Compile(Compiler);
-	if (!Compiler->StrataCompilationInfoContainsCodeChunk(ACodeChunk) || !Compiler->StrataCompilationInfoContainsCodeChunk(BCodeChunk))
-	{
-		return Compiler->Errorf(TEXT("Could not find ACodeChunk or BCodeChunk to add"));
-	}
-	const FStrataMaterialCompilationInfo& AStrataData = Compiler->GetStrataCompilationInfo(ACodeChunk);
-	const FStrataMaterialCompilationInfo& BStrataData = Compiler->GetStrataCompilationInfo(BCodeChunk);
 
 	const int32 ANormalMixCodeChunk = Compiler->StrataAddParameterBlendingBSDFCoverageToNormalMixCodeChunk(ACodeChunk, BCodeChunk);
 
@@ -22511,41 +22474,28 @@ int32 UMaterialExpressionStrataAdd::Compile(class FMaterialCompiler* Compiler, i
 	FStrataOperator& StrataOperator = Compiler->StrataCompilationGetOperator(this);
 	if (StrataOperator.bUseParameterBlending)
 	{
-		if (AStrataData.TotalBSDFCount != 1)
-		{
-			return Compiler->Errorf(TEXT("A input: cannot add complex material topology using parameter blending. Only a one-to-one blending is supported today."));
-		}
-		if (BStrataData.TotalBSDFCount != 1)
+		FStrataOperator* ABSDFOperator = Compiler->StrataCompilationGetOperatorFromIndex(StrataOperator.LeftIndex);
+		FStrataOperator* BBSDFOperator = Compiler->StrataCompilationGetOperatorFromIndex(StrataOperator.RightIndex);
+		if (!ABSDFOperator || !BBSDFOperator)
 		{
-			return Compiler->Errorf(TEXT("B input: cannot add complex material topology using parameter blending. Only a one-to-one blending is supported today."));
-		}
-
-		const FStrataMaterialCompilationInfo::FBSDF& ABSDF = AStrataData.Layers[0].BSDFs[0];
-		const FStrataMaterialCompilationInfo::FBSDF& BBSDF = BStrataData.Layers[0].BSDFs[0];
-		if (ABSDF.Type != STRATA_BSDF_TYPE_SLAB)
-		{
-			return Compiler->Errorf(TEXT("A input: we can only apply parameter blending to Slab BSDF."));
-		}
-		if (BBSDF.Type != STRATA_BSDF_TYPE_SLAB)
-		{
-			return Compiler->Errorf(TEXT("B input: we can only apply parameter blending to Slab BSDF."));
+			return Compiler->Errorf(TEXT("Missing input on add node."));
 		}
 
 		// Compute the new Normal and Tangent resulting from the blending using code chunk
-		const int32 NewNormalCodeChunk = StrataBlendNormal(Compiler, BBSDF.RegisteredSharedLocalBasis.NormalCodeChunk, ABSDF.RegisteredSharedLocalBasis.NormalCodeChunk, ANormalMixCodeChunk);
+		const int32 NewNormalCodeChunk = StrataBlendNormal(Compiler, BBSDFOperator->BSDFRegisteredSharedLocalBasis.NormalCodeChunk, ABSDFOperator->BSDFRegisteredSharedLocalBasis.NormalCodeChunk, ANormalMixCodeChunk);
 		// The tangent is optional so we treat it differently if INDEX_NONE is specified
 		int32 NewTangentCodeChunk = INDEX_NONE;
-		if (ABSDF.RegisteredSharedLocalBasis.TangentCodeChunk != INDEX_NONE && BBSDF.RegisteredSharedLocalBasis.TangentCodeChunk != INDEX_NONE)
+		if (ABSDFOperator->BSDFRegisteredSharedLocalBasis.TangentCodeChunk != INDEX_NONE && BBSDFOperator->BSDFRegisteredSharedLocalBasis.TangentCodeChunk != INDEX_NONE)
 		{
-			NewTangentCodeChunk = StrataBlendNormal(Compiler, BBSDF.RegisteredSharedLocalBasis.TangentCodeChunk, ABSDF.RegisteredSharedLocalBasis.TangentCodeChunk, ANormalMixCodeChunk);
+			NewTangentCodeChunk = StrataBlendNormal(Compiler, BBSDFOperator->BSDFRegisteredSharedLocalBasis.TangentCodeChunk, ABSDFOperator->BSDFRegisteredSharedLocalBasis.TangentCodeChunk, ANormalMixCodeChunk);
 		}
-		else if (ABSDF.RegisteredSharedLocalBasis.TangentCodeChunk != INDEX_NONE)
+		else if (ABSDFOperator->BSDFRegisteredSharedLocalBasis.TangentCodeChunk != INDEX_NONE)
 		{
-			NewTangentCodeChunk = ABSDF.RegisteredSharedLocalBasis.TangentCodeChunk;
+			NewTangentCodeChunk = ABSDFOperator->BSDFRegisteredSharedLocalBasis.TangentCodeChunk;
 		}
-		else if (BBSDF.RegisteredSharedLocalBasis.TangentCodeChunk != INDEX_NONE)
+		else if (BBSDFOperator->BSDFRegisteredSharedLocalBasis.TangentCodeChunk != INDEX_NONE)
 		{
-			NewTangentCodeChunk = BBSDF.RegisteredSharedLocalBasis.TangentCodeChunk;
+			NewTangentCodeChunk = BBSDFOperator->BSDFRegisteredSharedLocalBasis.TangentCodeChunk;
 		}
 		const FStrataRegisteredSharedLocalBasis NewRegisteredSharedLocalBasis = StrataCompilationInfoCreateSharedLocalBasis(Compiler, NewNormalCodeChunk, NewTangentCodeChunk);
 
@@ -22553,18 +22503,13 @@ int32 UMaterialExpressionStrataAdd::Compile(class FMaterialCompiler* Compiler, i
 			ACodeChunk, BCodeChunk, ANormalMixCodeChunk, Compiler->GetStrataSharedLocalBasisIndexMacro(NewRegisteredSharedLocalBasis),
 			StrataOperator.bRootOfParameterBlendingSubTree ? &StrataOperator : nullptr);
 
-		StrataCompilationInfoCreateSingleBSDFMaterial(Compiler, OutputCodeChunk, NewRegisteredSharedLocalBasis, STRATA_BSDF_TYPE_SLAB);
-		FStrataMaterialCompilationInfo StrataInfo = StrataCompilationInfoAddParamBlend(Compiler, AStrataData, BStrataData, NewRegisteredSharedLocalBasis);
-
-		Compiler->StrataCompilationInfoRegisterCodeChunk(OutputCodeChunk, StrataInfo);
+		// Propagate the parameter blended normal
+		StrataOperator.BSDFRegisteredSharedLocalBasis = NewRegisteredSharedLocalBasis;
+		StrataOperator.CombineFlagsForParameterBlending(*ABSDFOperator, *BBSDFOperator);
 	}
 	else
 	{
 		OutputCodeChunk = Compiler->StrataAdd(ACodeChunk, BCodeChunk, StrataOperator.Index, StrataOperator.MaxDistanceFromLeaves);
-
-		FStrataMaterialCompilationInfo StrataInfo = StrataCompilationInfoAdd(Compiler, Compiler->GetStrataCompilationInfo(ACodeChunk), Compiler->GetStrataCompilationInfo(BCodeChunk));
-
-		Compiler->StrataCompilationInfoRegisterCodeChunk(OutputCodeChunk, StrataInfo);
 	}
 
 	return OutputCodeChunk;
@@ -22659,24 +22604,22 @@ int32 UMaterialExpressionStrataWeight::Compile(class FMaterialCompiler* Compiler
 	FStrataOperator& StrataOperator = Compiler->StrataCompilationGetOperator(this);
 	if (StrataOperator.bUseParameterBlending)
 	{
+		// Propagate the parameter blended normal
+		FStrataOperator* Operator = Compiler->StrataCompilationGetOperatorFromIndex(StrataOperator.LeftIndex);
+		if (!Operator)
+		{
+			return Compiler->Errorf(TEXT("Missing input on weight node."));
+		}
+		StrataOperator.BSDFRegisteredSharedLocalBasis = Operator->BSDFRegisteredSharedLocalBasis;
+		StrataOperator.CopyFlagsForParameterBlending(*Operator);
+
 		OutputCodeChunk = Compiler->StrataWeightParameterBlending(
 			ACodeChunk, WeightCodeChunk, 
 			StrataOperator.bRootOfParameterBlendingSubTree ? &StrataOperator : nullptr);
-
-		FStrataMaterialCompilationInfo StrataInfo = StrataCompilationInfoWeightParamBlend(Compiler, Compiler->GetStrataCompilationInfo(ACodeChunk));
-
-		Compiler->StrataCompilationInfoRegisterCodeChunk(OutputCodeChunk, StrataInfo);
 	}
 	else
 	{
 		OutputCodeChunk = Compiler->StrataWeight(ACodeChunk, WeightCodeChunk, StrataOperator.Index, StrataOperator.MaxDistanceFromLeaves);
-
-		if (!Compiler->StrataCompilationInfoContainsCodeChunk(ACodeChunk))
-		{
-			return Compiler->Errorf(TEXT("Could not find ACodeChunk to multiply"));
-		}
-		FStrataMaterialCompilationInfo StrataInfo = StrataCompilationInfoWeight(Compiler, Compiler->GetStrataCompilationInfo(ACodeChunk));
-		Compiler->StrataCompilationInfoRegisterCodeChunk(OutputCodeChunk, StrataInfo);
 	}
 
 	return OutputCodeChunk;
@@ -22760,26 +22703,21 @@ int32 UMaterialExpressionStrataThinFilm::Compile(class FMaterialCompiler* Compil
 	FStrataOperator& StrataOperator = Compiler->StrataCompilationGetOperator(this);
 	if (StrataOperator.bUseParameterBlending)
 	{
-		// Extra normal
-		const FStrataMaterialCompilationInfo& AStrataData = Compiler->GetStrataCompilationInfo(ACodeChunk);
-		const FStrataMaterialCompilationInfo::FBSDF& ABSDF = AStrataData.Layers[0].BSDFs[0];
-		int32 NormalCodeChunk = ABSDF.RegisteredSharedLocalBasis.NormalCodeChunk;
+		// Propagate the parameter blended normal
+		FStrataOperator* Operator = Compiler->StrataCompilationGetOperatorFromIndex(StrataOperator.LeftIndex);
+		if (!Operator)
+		{
+			return Compiler->Errorf(TEXT("Missing input on ThinFilm node."));
+		}
+		StrataOperator.BSDFRegisteredSharedLocalBasis = Operator->BSDFRegisteredSharedLocalBasis;
+		StrataOperator.CopyFlagsForParameterBlending(*Operator);
 
+		int32 NormalCodeChunk = Operator->BSDFRegisteredSharedLocalBasis.NormalCodeChunk;
 		OutputCodeChunk = Compiler->StrataThinFilmParameterBlending(ACodeChunk, ThicknessCodeChunk, IORCodeChunk, NormalCodeChunk, StrataOperator.bRootOfParameterBlendingSubTree ? &StrataOperator : nullptr);
-
-		FStrataMaterialCompilationInfo StrataInfo = StrataCompilationInfoThinFilmParamBlend(Compiler, Compiler->GetStrataCompilationInfo(ACodeChunk));
-		Compiler->StrataCompilationInfoRegisterCodeChunk(OutputCodeChunk, StrataInfo);
 	}
 	else
 	{
 		OutputCodeChunk = Compiler->StrataThinFilm(ACodeChunk, ThicknessCodeChunk, IORCodeChunk, StrataOperator.Index, StrataOperator.MaxDistanceFromLeaves);
-
-		if (!Compiler->StrataCompilationInfoContainsCodeChunk(ACodeChunk))
-		{
-			return Compiler->Errorf(TEXT("Could not find ACodeChunk to coat with thin-film"));
-		}
-		FStrataMaterialCompilationInfo StrataInfo = StrataCompilationInfoThinFilm(Compiler, Compiler->GetStrataCompilationInfo(ACodeChunk));
-		Compiler->StrataCompilationInfoRegisterCodeChunk(OutputCodeChunk, StrataInfo);
 	}
 
 	return OutputCodeChunk;
diff --git a/Engine/Source/Runtime/Engine/Private/Materials/StrataMaterial.cpp b/Engine/Source/Runtime/Engine/Private/Materials/StrataMaterial.cpp
--- a/Engine/Source/Runtime/Engine/Private/Materials/StrataMaterial.cpp
+++ b/Engine/Source/Runtime/Engine/Private/Materials/StrataMaterial.cpp
@@ -29,29 +29,9 @@ FString GetStrataBSDFName(uint8 BSDFType)
 	return "";
 }
 
-static void UpdateTotalBSDFCount(FMaterialCompiler* Compiler, FStrataMaterialCompilationInfo& StrataInfo)
-{
-	StrataInfo.TotalBSDFCount = 0;
-	for (uint32 LayerIt = 0; LayerIt < StrataInfo.LayerCount; ++LayerIt)
-	{
-		StrataInfo.TotalBSDFCount += StrataInfo.Layers[LayerIt].BSDFCount;
-	}
-
-	if ((StrataInfo.TotalBSDFCount) > STRATA_MAX_BSDF_COUNT)
-	{
-		Compiler->Error(TEXT("This material contains too many BSDFs"));
-	}
-}
-
 FStrataRegisteredSharedLocalBasis StrataCompilationInfoCreateNullSharedLocalBasis()
 {
-	FStrataRegisteredSharedLocalBasis FakeSharedLocalBasis;
-	FakeSharedLocalBasis.NormalCodeChunk = INDEX_NONE;
-	FakeSharedLocalBasis.TangentCodeChunk = INDEX_NONE;
-	FakeSharedLocalBasis.NormalCodeChunkHash = 0;
-	FakeSharedLocalBasis.TangentCodeChunkHash = 0;
-	FakeSharedLocalBasis.GraphSharedLocalBasisIndex = 0;
-	return FakeSharedLocalBasis;
+	return FStrataRegisteredSharedLocalBasis();
 }
 
 FStrataRegisteredSharedLocalBasis StrataCompilationInfoCreateSharedLocalBasis(FMaterialCompiler* Compiler, int32 NormalCodeChunk, int32 TangentCodeChunk)
@@ -63,390 +43,115 @@ FStrataRegisteredSharedLocalBasis StrataCompilationInfoCreateSharedLocalBasis(FM
 	return Compiler->StrataCompilationInfoRegisterSharedLocalBasis(NormalCodeChunk, TangentCodeChunk);
 }
 
-void StrataCompilationInfoCreateSingleBSDFMaterial(FMaterialCompiler* Compiler, int32 CodeChunk,
-	const FStrataRegisteredSharedLocalBasis& RegisteredSharedLocalBasis,
-	uint8 BSDFType, bool bHasSSS, bool bHasMFPPluggedIn, bool bHasEdgeColor, bool bHasFuzz, bool bHasHaziness)
-{
-	FStrataMaterialCompilationInfo StrataInfo;
-	StrataInfo.LayerCount = 1;
-	StrataInfo.Layers[0].BSDFCount = 1;
-	StrataInfo.Layers[0].BSDFs[0].Type = BSDFType;
-	StrataInfo.Layers[0].BSDFs[0].RegisteredSharedLocalBasis = RegisteredSharedLocalBasis;
-	StrataInfo.Layers[0].BSDFs[0].bHasSSS = bHasSSS;
-	StrataInfo.Layers[0].BSDFs[0].bHasMFPPluggedIn = bHasMFPPluggedIn;
-	StrataInfo.Layers[0].BSDFs[0].bHasEdgeColor = bHasEdgeColor;
-	StrataInfo.Layers[0].BSDFs[0].bHasThinFilm = false;
-	StrataInfo.Layers[0].BSDFs[0].bHasFuzz = bHasFuzz;
-	StrataInfo.Layers[0].BSDFs[0].bHasHaziness = bHasHaziness;
-	UpdateTotalBSDFCount(Compiler, StrataInfo);
-	Compiler->StrataCompilationInfoRegisterCodeChunk(CodeChunk, StrataInfo);
-}
-
-void  StrataCompilationInfoCreateNullBSDF(FMaterialCompiler* Compiler, int32 CodeChunk)
-{
-	FStrataMaterialCompilationInfo StrataInfo;
-	StrataInfo.LayerCount = 0;
-	StrataInfo.TotalBSDFCount = 0;
-	Compiler->StrataCompilationInfoRegisterCodeChunk(CodeChunk, StrataInfo);
-}
-
-
-FStrataMaterialCompilationInfo StrataCompilationInfoWeight(FMaterialCompiler* Compiler, const FStrataMaterialCompilationInfo& A)
-{
-	FStrataMaterialCompilationInfo StrataInfo = A;
-	return StrataInfo;
-}
-
-FStrataMaterialCompilationInfo StrataCompilationInfoWeightParamBlend(FMaterialCompiler* Compiler, const FStrataMaterialCompilationInfo& A)
-{
-	FStrataMaterialCompilationInfo StrataInfo = A;
-	return StrataInfo;
-}
-
-FStrataMaterialCompilationInfo StrataCompilationInfoThinFilm(FMaterialCompiler* Compiler, const FStrataMaterialCompilationInfo& A)
-{
-	// STRATA_TODO: This is not correct. FStrataMaterialCompilationInfo is meant to be removed and replaced with Strata tree / metadata
-	FStrataMaterialCompilationInfo StrataInfo = A;
-	StrataInfo.Layers[0].BSDFs[0].bHasThinFilm = true;
-	return StrataInfo;
-}
-
-FStrataMaterialCompilationInfo StrataCompilationInfoThinFilmParamBlend(FMaterialCompiler* Compiler, const FStrataMaterialCompilationInfo& A)
-{
-	FStrataMaterialCompilationInfo StrataInfo = A;
-	return StrataInfo;
-}
-
-FStrataMaterialCompilationInfo StrataCompilationInfoAdd(FMaterialCompiler* Compiler, const FStrataMaterialCompilationInfo& A, const FStrataMaterialCompilationInfo& B)
-{
-	FStrataMaterialCompilationInfo StrataInfo = A;
-
-	// Append each BSDF from B to A, with same layer position
-	for (uint32 LayerIt = 0; LayerIt < B.LayerCount; ++LayerIt)
-	{
-		const FStrataMaterialCompilationInfo::FLayer& ALayer = A.Layers[LayerIt];
-		const FStrataMaterialCompilationInfo::FLayer& BLayer = B.Layers[LayerIt];
-
-		if ((ALayer.BSDFCount + BLayer.BSDFCount) > STRATA_MAX_BSDF_COUNT_PER_LAYER)
-		{
-			Compiler->Error(TEXT("Adding would result in too many BSDFs in a Layer"));
-			return A;
-		}
-
-		for (uint32 BSDF = 0; BSDF < BLayer.BSDFCount; BSDF++)
-		{
-			StrataInfo.Layers[LayerIt].BSDFs[ALayer.BSDFCount + BSDF] = BLayer.BSDFs[BSDF];
-		}
-
-		StrataInfo.Layers[LayerIt].BSDFCount = ALayer.BSDFCount + BLayer.BSDFCount;
-	}
-	StrataInfo.LayerCount = FMath::Max(A.LayerCount, B.LayerCount);
-
-	UpdateTotalBSDFCount(Compiler, StrataInfo);
-	return StrataInfo;
-}
-
-
-// ==> NOTE: Always pair with the shader behavior in StrataAddParameterBlending.
-FStrataMaterialCompilationInfo StrataCompilationInfoAddParamBlend(FMaterialCompiler* Compiler, const FStrataMaterialCompilationInfo& A, const FStrataMaterialCompilationInfo& B, const FStrataRegisteredSharedLocalBasis& RegisteredSharedLocalBasis)
-{
-	check(A.TotalBSDFCount == 1);
-	check(B.TotalBSDFCount == 1);
-
-	FStrataMaterialCompilationInfo StrataInfo = A;
-	FStrataMaterialCompilationInfo::FBSDF& NewBSDF = StrataInfo.Layers[0].BSDFs[0];
-	const FStrataMaterialCompilationInfo::FBSDF& OtherBSDF = B.Layers[0].BSDFs[0];
-
-	NewBSDF.RegisteredSharedLocalBasis = RegisteredSharedLocalBasis;
-
-	// When parameter blending is used, we take the union of all the features activated by input BSDFs.
-	NewBSDF.bHasSSS				|=	OtherBSDF.bHasSSS;
-	NewBSDF.bHasMFPPluggedIn	|=	OtherBSDF.bHasMFPPluggedIn;
-	NewBSDF.bHasEdgeColor		|=	OtherBSDF.bHasEdgeColor;
-	NewBSDF.bHasThinFilm		|=	OtherBSDF.bHasThinFilm;
-	NewBSDF.bHasFuzz			|=	OtherBSDF.bHasFuzz;
-	NewBSDF.bHasHaziness		|=	OtherBSDF.bHasHaziness;
-	UpdateTotalBSDFCount(Compiler, StrataInfo);
-	return StrataInfo;
-}
-
-
-FStrataMaterialCompilationInfo StrataCompilationInfoHorizontalMixing(FMaterialCompiler* Compiler, const FStrataMaterialCompilationInfo& Background, const FStrataMaterialCompilationInfo& Foreground)
-{
-	return StrataCompilationInfoAdd(Compiler, Background, Foreground); // Mixing is a similar operation to Add when it comes to bsdf count
-}
-
-
-// ==> NOTE: Always pair with the shader behavior in StrataHorizontalMixingParameterBlending
-FStrataMaterialCompilationInfo StrataCompilationInfoHorizontalMixingParamBlend(FMaterialCompiler* Compiler, const FStrataMaterialCompilationInfo& Background, const FStrataMaterialCompilationInfo& Foreground, const FStrataRegisteredSharedLocalBasis& RegisteredSharedLocalBasis)
-{
-	check(Background.TotalBSDFCount == 1);
-	check(Foreground.TotalBSDFCount == 1);
-
-	FStrataMaterialCompilationInfo StrataInfo = Background;
-	FStrataMaterialCompilationInfo::FBSDF& NewBSDF = StrataInfo.Layers[0].BSDFs[0];
-	const FStrataMaterialCompilationInfo::FBSDF& OtherBSDF = Foreground.Layers[0].BSDFs[0];
-
-	NewBSDF.RegisteredSharedLocalBasis = RegisteredSharedLocalBasis;
-
-	// When parameter blending is used, we take the union of all the features activated by input BSDFs.
-	NewBSDF.bHasSSS				|=	OtherBSDF.bHasSSS;
-	NewBSDF.bHasMFPPluggedIn	|=	OtherBSDF.bHasMFPPluggedIn;
-	NewBSDF.bHasEdgeColor		|=	OtherBSDF.bHasEdgeColor;
-	NewBSDF.bHasThinFilm		|=	OtherBSDF.bHasThinFilm;
-	NewBSDF.bHasFuzz			|=	OtherBSDF.bHasFuzz;
-	NewBSDF.bHasHaziness		|=	OtherBSDF.bHasHaziness;
-	UpdateTotalBSDFCount(Compiler, StrataInfo);
-	return StrataInfo;
-}
-
-
-FStrataMaterialCompilationInfo StrataCompilationInfoVerticalLayering(FMaterialCompiler* Compiler, const FStrataMaterialCompilationInfo& Top, const FStrataMaterialCompilationInfo& Base)
-{
-	if ((Top.LayerCount + Base.LayerCount) > STRATA_MAX_LAYER_COUNT)
-	{
-		Compiler->Error(TEXT("Layering would result in too many Layers"));
-		return Base;
-	}
-
-	FStrataMaterialCompilationInfo StrataInfo = Top;
-
-	// Add each layer from Base under Top
-	const uint32 TopLayerCount = Top.LayerCount;
-	for (uint32 LayerIt = 0; LayerIt < Base.LayerCount; ++LayerIt)
-	{
-		StrataInfo.Layers[TopLayerCount + LayerIt] = Base.Layers[LayerIt];
-	}
-	StrataInfo.LayerCount += Base.LayerCount;
-
-	UpdateTotalBSDFCount(Compiler, StrataInfo);
-	return StrataInfo;
-}
-
-
-// ==> NOTE: Always pair with the shader behavior in StrataVerticalLayeringParameterBlending
-FStrataMaterialCompilationInfo StrataCompilationInfoVerticalLayeringParamBlend(FMaterialCompiler* Compiler, const FStrataMaterialCompilationInfo& Top, const FStrataMaterialCompilationInfo& Base, const FStrataRegisteredSharedLocalBasis& RegisteredSharedLocalBasis)
-{
-	check(Top.TotalBSDFCount == 1);
-	check(Base.TotalBSDFCount == 1);
-
-	FStrataMaterialCompilationInfo StrataInfo = Base;
-	FStrataMaterialCompilationInfo::FBSDF& NewBSDF = StrataInfo.Layers[0].BSDFs[0];
-	const FStrataMaterialCompilationInfo::FBSDF& TopBSDF = Top.Layers[0].BSDFs[0];
-
-	NewBSDF.RegisteredSharedLocalBasis = RegisteredSharedLocalBasis;
-
-//	NewBSDF.bHasSSS				= TopBSDF.bHasSSS;				// We keep SSS only if the base layer has it. Otherwise it will be simple volume and the throughput wll be applied on the parameters.
-//	NewBSDF.bHasMFPPluggedIn	= TopBSDF.bHasMFPPluggedIn;		// Idem
-	NewBSDF.bHasEdgeColor		|= TopBSDF.bHasEdgeColor;		// We keep the union of both, even though it will be hard to get a perfect match.
-	NewBSDF.bHasFuzz			|= TopBSDF.bHasFuzz;			// Idem
-	NewBSDF.bHasHaziness		|= TopBSDF.bHasHaziness;		// Idem
-	NewBSDF.bHasThinFilm		 = TopBSDF.bHasThinFilm;		// We only keep thin film from the top layer, because its color is otherwise not controlable
-
-	UpdateTotalBSDFCount(Compiler, StrataInfo);
-	return StrataInfo;
-}
-
-static bool StrataIsSingleBSDF(FMaterialCompiler* Compiler, const FStrataMaterialCompilationInfo& Material)
-{
-	if (Material.TotalBSDFCount == 0 || Material.LayerCount == 0)
-	{
-		Compiler->Error(TEXT("There is no layer or BSDF plugged in, but one of the BSDF in the graph wants to enforce one and only one BSDF to be used."));
-		return false;
-	}
-	if (Material.TotalBSDFCount > 1 || Material.LayerCount > 1)
-	{
-		Compiler->Error(TEXT("There is more than one layer or BSDF, but one of the BSDF in the graph wants to enforce one and only one BSDF to be used."));
-		return false;
-	}
-	return true;
-}
-
-bool StrataIsVolumetricFogCloudOnly(FMaterialCompiler* Compiler, const FStrataMaterialCompilationInfo& Material)
-{
-	if (!StrataIsSingleBSDF(Compiler, Material))
-	{
-		return false;
-	}
-	if (Material.Layers[0].BSDFs[0].Type != STRATA_BSDF_TYPE_VOLUMETRICFOGCLOUD)
-	{
-		Compiler->Error(TEXT("The single BSDF resulting from the graph is not of type Volume."));
-		return false;
-	}
-
-	return true;
-}
-
-bool StrataIsUnlitOnly(FMaterialCompiler* Compiler, const FStrataMaterialCompilationInfo& Material)
-{
-	if (!StrataIsSingleBSDF(Compiler, Material))
-	{
-		return false;
-	}
-	if (Material.Layers[0].BSDFs[0].Type != STRATA_BSDF_TYPE_UNLIT)
-	{
-		Compiler->Error(TEXT("The single BSDF resulting from the graph is not of type Unlit."));
-		return false;
-	}
-
-	return true;
-}
-
-bool StrataIsHairOnly(FMaterialCompiler* Compiler, const FStrataMaterialCompilationInfo& Material)
-{
-	if (!StrataIsSingleBSDF(Compiler, Material))
-	{
-		return false;
-	}
-	if (Material.Layers[0].BSDFs[0].Type != STRATA_BSDF_TYPE_HAIR)
-	{
-		Compiler->Error(TEXT("The single BSDF resulting from the graph is not of type Hair."));
-		return false;
-	}
-
-	return true;
-}
-
-bool StrataIsSingleLayerWaterOnly(FMaterialCompiler* Compiler, const FStrataMaterialCompilationInfo& Material)
-{
-	if (!StrataIsSingleBSDF(Compiler, Material))
-	{
-		return false;
-	}
-	if (Material.Layers[0].BSDFs[0].Type != STRATA_BSDF_TYPE_SINGLELAYERWATER)
-	{
-		Compiler->Error(TEXT("The single BSDF resulting from the graph is not of type SingleLayerWater."));
-		return false;
-	}
-
-	return true;
-}
-
-bool StrataMaterialContainsAnyBSDF(FMaterialCompiler* Compiler, const FStrataMaterialCompilationInfo& Material, uint8 BSDFType)
-{
-	for (uint32 LayerIt = 0; LayerIt < Material.LayerCount; ++LayerIt)
-	{
-		const FStrataMaterialCompilationInfo::FLayer& Layer = Material.Layers[LayerIt];
-		for (uint32 BSDF = 0; BSDF < Layer.BSDFCount; BSDF++)
-		{
-			if (Layer.BSDFs[BSDF].Type == BSDFType)
-			{
-				return true;
-			}
-		}
-	}
-	return false;
-}
-
-FStrataMaterialAnalysisResult::FStrataMaterialAnalysisResult()
-{
-	bFitInMemoryBudget = true;
-	RequestedLayerCount = 0;
-	RequestedBSDFCount = 0;
-	RequestedByteCount = 0;
-	ClampedLayerCount = 0;
-	ClampedBSDFCount = 0;
-	UsedByteCount = 0;
-}
-
-FStrataMaterialAnalysisResult StrataCompilationInfoMaterialAnalysis(FMaterialCompiler* Compiler, const FStrataMaterialCompilationInfo& Material, const uint32 StrataBytePerPixel)
-{
-	const uint32 UintByteSize = sizeof(uint32);
-
-	FStrataMaterialAnalysisResult Result;
-
-	// 1. Header
-
-	// SharedLocalBases_BSDFCount
-	Result.RequestedByteCount += UintByteSize;
-
-	// shared local bases between BSDFs
-	Result.RequestedByteCount += Compiler->StrataCompilationInfoGetSharedLocalBasesCount() * STRATA_PACKED_SHAREDLOCALBASIS_STRIDE_BYTES;
-	
-	// 2. The list of BSDFs
-	Result.RequestedMaxBSDFCountPerLayer = 0;
-	// We process layers from top to bottom to cull the bottom ones in case we run out of pixel bytes
-	for (uint32 LayerIt = 0; LayerIt < Material.LayerCount; LayerIt++)
-	{
-		const FStrataMaterialCompilationInfo::FLayer& Layer = Material.Layers[LayerIt];
-
-		const bool bTopLayer = LayerIt == 0;
-		const bool bBottomLayer = LayerIt == (Material.LayerCount - 1);
-
-		for (uint32 BSDFIt = 0; BSDFIt < Layer.BSDFCount; BSDFIt++)
-		{
-			const FStrataMaterialCompilationInfo::FBSDF& BSDF = Layer.BSDFs[BSDFIt];
-
-			// BSDF state
-			Result.RequestedByteCount += UintByteSize;
-
-			// From the compiler side, we can only assume the top layer has grey scale weight/throughput
-			const bool bMayBeColoredWeight = LayerIt > 0;
-			if (bMayBeColoredWeight)
-			{
-				Result.RequestedByteCount += UintByteSize;
-			}
-
-			switch (BSDF.Type)
-			{
-			case STRATA_BSDF_TYPE_SLAB:
-			{
-				// Compute values closer to the reality for HasSSS and IsSimpleVolume, now that we know that we know the topology of the material.
-				const bool bIsSimpleVolume = !bBottomLayer && BSDF.bHasMFPPluggedIn;
-				const bool bHasSSS = bBottomLayer && BSDF.bHasSSS && !bIsSimpleVolume;
-
-				Result.RequestedByteCount += UintByteSize;
-				Result.RequestedByteCount += UintByteSize;
-				if (BSDF.bHasEdgeColor || BSDF.bHasThinFilm || BSDF.bHasHaziness)
-				{
-					Result.RequestedByteCount += UintByteSize;
-				}
-				if (bHasSSS || bIsSimpleVolume)
-				{
-					Result.RequestedByteCount += UintByteSize;
-				}
-				if (BSDF.bHasFuzz)
-				{
-					Result.RequestedByteCount += UintByteSize;
-				}
-				break;
-			}
-			case STRATA_BSDF_TYPE_HAIR:
-			{
-				Result.RequestedByteCount += UintByteSize;
-				Result.RequestedByteCount += UintByteSize;
-				break;
-			}
-			case STRATA_BSDF_TYPE_SINGLELAYERWATER:
-			{
-				Result.RequestedByteCount += UintByteSize;
-				Result.RequestedByteCount += UintByteSize;
-				break;
-			}
-			}
-		}
-		Result.RequestedMaxBSDFCountPerLayer = FMath::Max(Result.RequestedMaxBSDFCountPerLayer, Layer.BSDFCount);
-
-		Result.RequestedLayerCount++;
-		Result.RequestedBSDFCount += Layer.BSDFCount;
-
-		if (Result.RequestedByteCount <= StrataBytePerPixel && Result.bFitInMemoryBudget)
-		{
-			// We only validate all the BSDF of a layer if it remains within budget and we are not already out of budget
-			Result.ClampedBSDFCount += Layer.BSDFCount;
-			Result.ClampedLayerCount++;
-
-			// Set the current used bytes
-			Result.UsedByteCount = Result.RequestedByteCount;
-		}
-		else
-		{
-			// Used byte count remains unchanged and we notify that we start peeling off top layers now
-			Result.bFitInMemoryBudget = false;
-		}
-	}
-	return Result;
-
-}
+//FStrataMaterialAnalysisResult::FStrataMaterialAnalysisResult()
+//{
+//	bFitInMemoryBudget = true;
+//	RequestedLayerCount = 0;
+//	RequestedBSDFCount = 0;
+//	RequestedByteCount = 0;
+//	ClampedLayerCount = 0;
+//	ClampedBSDFCount = 0;
+//	UsedByteCount = 0;
+//}
+//
+//FStrataMaterialAnalysisResult StrataCompilationInfoMaterialAnalysis(FMaterialCompiler* Compiler, const FStrataMaterialCompilationInfo& Material, const uint32 StrataBytePerPixel)
+//{
+//	const uint32 UintByteSize = sizeof(uint32);
+//
+//	FStrataMaterialAnalysisResult Result;
+//
+//	// 1. Header
+//
+//	// SharedLocalBases_BSDFCount
+//	Result.RequestedByteCount += UintByteSize;
+//
+//	// shared local bases between BSDFs
+//	Result.RequestedByteCount += Compiler->StrataCompilationInfoGetSharedLocalBasesCount() * STRATA_PACKED_SHAREDLOCALBASIS_STRIDE_BYTES;
+//	
+//	// 2. The list of BSDFs
+//	Result.RequestedMaxBSDFCountPerLayer = 0;
+//	// We process layers from top to bottom to cull the bottom ones in case we run out of pixel bytes
+//	for (uint32 LayerIt = 0; LayerIt < Material.LayerCount; LayerIt++)
+//	{
+//		const FStrataMaterialCompilationInfo::FLayer& Layer = Material.Layers[LayerIt];
+//
+//		const bool bTopLayer = LayerIt == 0;
+//		const bool bBottomLayer = LayerIt == (Material.LayerCount - 1);
+//
+//		for (uint32 BSDFIt = 0; BSDFIt < Layer.BSDFCount; BSDFIt++)
+//		{
+//			const FStrataMaterialCompilationInfo::FBSDF& BSDF = Layer.BSDFs[BSDFIt];
+//
+//			// BSDF state
+//			Result.RequestedByteCount += UintByteSize;
+//
+//			// From the compiler side, we can only assume the top layer has grey scale weight/throughput
+//			const bool bMayBeColoredWeight = LayerIt > 0;
+//			if (bMayBeColoredWeight)
+//			{
+//				Result.RequestedByteCount += UintByteSize;
+//			}
+//
+//			switch (BSDF.Type)
+//			{
+//			case STRATA_BSDF_TYPE_SLAB:
+//			{
+//				// Compute values closer to the reality for HasSSS and IsSimpleVolume, now that we know that we know the topology of the material.
+//				const bool bIsSimpleVolume = !bBottomLayer && BSDF.bHasMFPPluggedIn;
+//				const bool bHasSSS = bBottomLayer && BSDF.bHasSSS && !bIsSimpleVolume;
+//
+//				Result.RequestedByteCount += UintByteSize;
+//				Result.RequestedByteCount += UintByteSize;
+//				if (BSDF.bHasEdgeColor || BSDF.bHasThinFilm || BSDF.bHasHaziness)
+//				{
+//					Result.RequestedByteCount += UintByteSize;
+//				}
+//				if (bHasSSS || bIsSimpleVolume)
+//				{
+//					Result.RequestedByteCount += UintByteSize;
+//				}
+//				if (BSDF.bHasFuzz)
+//				{
+//					Result.RequestedByteCount += UintByteSize;
+//				}
+//				break;
+//			}
+//			case STRATA_BSDF_TYPE_HAIR:
+//			{
+//				Result.RequestedByteCount += UintByteSize;
+//				Result.RequestedByteCount += UintByteSize;
+//				break;
+//			}
+//			case STRATA_BSDF_TYPE_SINGLELAYERWATER:
+//			{
+//				Result.RequestedByteCount += UintByteSize;
+//				Result.RequestedByteCount += UintByteSize;
+//				break;
+//			}
+//			}
+//		}
+//		Result.RequestedMaxBSDFCountPerLayer = FMath::Max(Result.RequestedMaxBSDFCountPerLayer, Layer.BSDFCount);
+//
+//		Result.RequestedLayerCount++;
+//		Result.RequestedBSDFCount += Layer.BSDFCount;
+//
+//		if (Result.RequestedByteCount <= StrataBytePerPixel && Result.bFitInMemoryBudget)
+//		{
+//			// We only validate all the BSDF of a layer if it remains within budget and we are not already out of budget
+//			Result.ClampedBSDFCount += Layer.BSDFCount;
+//			Result.ClampedLayerCount++;
+//
+//			// Set the current used bytes
+//			Result.UsedByteCount = Result.RequestedByteCount;
+//		}
+//		else
+//		{
+//			// Used byte count remains unchanged and we notify that we start peeling off top layers now
+//			Result.bFitInMemoryBudget = false;
+//		}
+//	}
+//	return Result;
+//
+//}
 	
 
diff --git a/Engine/Source/Runtime/Engine/Private/Materials/StrataMaterial.h b/Engine/Source/Runtime/Engine/Private/Materials/StrataMaterial.h
--- a/Engine/Source/Runtime/Engine/Private/Materials/StrataMaterial.h
+++ b/Engine/Source/Runtime/Engine/Private/Materials/StrataMaterial.h
@@ -8,93 +8,10 @@
 #include "MaterialCompiler.h"
 
 
-
 class FMaterialCompiler;
 
-
-
-struct FStrataMaterialCompilationInfo
-{
-	struct FBSDF
-	{
-		uint8 Type;
-		FStrataRegisteredSharedLocalBasis RegisteredSharedLocalBasis;
-
-		// Those properties are centered around the SLAB node. Maybe this can be abstracted at some point.
-		bool  bHasSSS;
-		bool  bHasMFPPluggedIn;
-		bool  bHasEdgeColor;
-		bool  bHasThinFilm;
-		bool  bHasFuzz;
-		bool  bHasHaziness;
-	};
-
-	struct FLayer
-	{
-		uint32 BSDFCount;
-		FBSDF BSDFs[STRATA_MAX_BSDF_COUNT_PER_LAYER];
-	};
-
-	uint32 LayerCount;
-	FLayer Layers[STRATA_MAX_LAYER_COUNT];
-
-	uint32 TotalBSDFCount;
-
-	FStrataMaterialCompilationInfo()
-	{
-		TotalBSDFCount = 0;
-		LayerCount = 0;
-		memset(Layers, 0, sizeof(Layers));
-	}
-
-	uint32 GetPerLayerMaxBSDFCount() const 
-	{
-		uint32 Count = 0;
-		for (uint32 LayerIt = 0; LayerIt < LayerCount; ++LayerIt)
-		{
-			Count = FMath::Max(Count, Layers[LayerIt].BSDFCount);
-		}
-		return Count;
-	}
-};
-
 FString GetStrataBSDFName(uint8 BSDFType);
 
 FStrataRegisteredSharedLocalBasis StrataCompilationInfoCreateNullSharedLocalBasis();
 FStrataRegisteredSharedLocalBasis StrataCompilationInfoCreateSharedLocalBasis(FMaterialCompiler* Compiler, int32 NormalCodeChunk, int32 TangentCodeChunk = INDEX_NONE);
 
-void  StrataCompilationInfoCreateNullBSDF(FMaterialCompiler* Compiler, int32 CodeChunk);
-void  StrataCompilationInfoCreateSingleBSDFMaterial(FMaterialCompiler* Compiler, int32 CodeChunk, 
-	const FStrataRegisteredSharedLocalBasis& RegisteredSharedLocalBasis,
-	uint8 BSDFType, bool bHasSSS = false, bool bHasMFPPluggedIn = false, bool bHasEdgeColor = false, bool bHasFuzz = false, bool bHasHaziness = false);
-
-FStrataMaterialCompilationInfo StrataCompilationInfoWeight(FMaterialCompiler* Compiler, const FStrataMaterialCompilationInfo& A);
-FStrataMaterialCompilationInfo StrataCompilationInfoWeightParamBlend(FMaterialCompiler* Compiler, const FStrataMaterialCompilationInfo& A);
-FStrataMaterialCompilationInfo StrataCompilationInfoThinFilm(FMaterialCompiler* Compiler, const FStrataMaterialCompilationInfo& A);
-FStrataMaterialCompilationInfo StrataCompilationInfoThinFilmParamBlend(FMaterialCompiler* Compiler, const FStrataMaterialCompilationInfo& A);
-FStrataMaterialCompilationInfo StrataCompilationInfoAdd(FMaterialCompiler* Compiler, const FStrataMaterialCompilationInfo& A, const FStrataMaterialCompilationInfo& B);
-FStrataMaterialCompilationInfo StrataCompilationInfoAddParamBlend(FMaterialCompiler* Compiler, const FStrataMaterialCompilationInfo& A, const FStrataMaterialCompilationInfo& B, const FStrataRegisteredSharedLocalBasis& RegisteredSharedLocalBasis);
-FStrataMaterialCompilationInfo StrataCompilationInfoHorizontalMixing(FMaterialCompiler* Compiler, const FStrataMaterialCompilationInfo& Background, const FStrataMaterialCompilationInfo& Foreground);
-FStrataMaterialCompilationInfo StrataCompilationInfoHorizontalMixingParamBlend(FMaterialCompiler* Compiler, const FStrataMaterialCompilationInfo& Background, const FStrataMaterialCompilationInfo& Foreground, const FStrataRegisteredSharedLocalBasis& RegisteredSharedLocalBasis);
-FStrataMaterialCompilationInfo StrataCompilationInfoVerticalLayering(FMaterialCompiler* Compiler, const FStrataMaterialCompilationInfo& Top, const FStrataMaterialCompilationInfo& Base);
-FStrataMaterialCompilationInfo StrataCompilationInfoVerticalLayeringParamBlend(FMaterialCompiler* Compiler, const FStrataMaterialCompilationInfo& Top, const FStrataMaterialCompilationInfo& Base, const FStrataRegisteredSharedLocalBasis& RegisteredSharedLocalBasis);
-
-bool StrataIsVolumetricFogCloudOnly(FMaterialCompiler* Compiler, const FStrataMaterialCompilationInfo& Material);
-bool StrataIsUnlitOnly(FMaterialCompiler* Compiler, const FStrataMaterialCompilationInfo& Material);
-bool StrataIsHairOnly(FMaterialCompiler* Compiler, const FStrataMaterialCompilationInfo& Material);
-bool StrataIsSingleLayerWaterOnly(FMaterialCompiler* Compiler, const FStrataMaterialCompilationInfo& Material);
-bool StrataMaterialContainsAnyBSDF(FMaterialCompiler* Compiler, const FStrataMaterialCompilationInfo& Material, uint8 BSDFType);
-
-struct FStrataMaterialAnalysisResult
-{
-	FStrataMaterialAnalysisResult();
-	bool bFitInMemoryBudget;
-	uint32 RequestedLayerCount;
-	uint32 RequestedMaxBSDFCountPerLayer;
-	uint32 RequestedBSDFCount;
-	uint32 RequestedByteCount;
-	uint32 ClampedLayerCount;
-	uint32 ClampedBSDFCount;
-	uint32 UsedByteCount;
-};
-FStrataMaterialAnalysisResult StrataCompilationInfoMaterialAnalysis(FMaterialCompiler* Compiler, const FStrataMaterialCompilationInfo& Material, const uint32 StrataBytePerPixel);
diff --git a/Engine/Source/Runtime/Engine/Public/MaterialCompiler.h b/Engine/Source/Runtime/Engine/Public/MaterialCompiler.h
--- a/Engine/Source/Runtime/Engine/Public/MaterialCompiler.h
+++ b/Engine/Source/Runtime/Engine/Public/MaterialCompiler.h
@@ -29,7 +29,6 @@ class UMaterialParameterCollection;
 class URuntimeVirtualTexture;
 class UTexture;
 struct FMaterialParameterInfo;
-struct FStrataMaterialCompilationInfo;
 
 enum EMaterialForceCastFlags
 {
@@ -66,6 +65,15 @@ struct FStrataRegisteredSharedLocalBasis
 	uint64 NormalCodeChunkHash;
 	uint64 TangentCodeChunkHash;
 	uint8 GraphSharedLocalBasisIndex;
+
+	FStrataRegisteredSharedLocalBasis()
+	{
+		NormalCodeChunk = INDEX_NONE;
+		TangentCodeChunk = INDEX_NONE;
+		NormalCodeChunkHash = 0;
+		TangentCodeChunkHash = 0;
+		GraphSharedLocalBasisIndex = 0;
+	}
 };
 
 struct FStrataOperator
@@ -78,7 +86,15 @@ struct FStrataOperator
 	int32 LeftIndex;	// Left child operator index
 	int32 RightIndex;	// Right child operator index
 
+	// Data used for BSDF type nodes only
 	int32 BSDFIndex;	// Index in the array of BSDF if a BSDF operator
+	uint8 BSDFType;
+	FStrataRegisteredSharedLocalBasis BSDFRegisteredSharedLocalBasis;
+	bool  bBSDFHasSSS;
+	bool  bBSDFHasMFPPluggedIn;
+	bool  bBSDFHasEdgeColor;
+	bool  bBSDFHasFuzz;
+	bool  bBSDFHasHaziness;
 
 	// Data derived after the tree has been built.
 	int32 MaxDistanceFromLeaves;
@@ -96,7 +112,15 @@ struct FStrataOperator
 		ParentIndex = INDEX_NONE;
 		LeftIndex = INDEX_NONE;
 		RightIndex = INDEX_NONE;
+
 		BSDFIndex = INDEX_NONE;
+		BSDFType = 0;
+		bBSDFHasSSS = false;
+		bBSDFHasMFPPluggedIn = false;
+		bBSDFHasEdgeColor = false;
+		bBSDFHasFuzz = false;
+		bBSDFHasHaziness = false;
+
 		MaxDistanceFromLeaves = 0;
 		LayerDepth = 0;
 		bIsTop = false;
@@ -105,6 +129,24 @@ struct FStrataOperator
 		bRootOfParameterBlendingSubTree = false;
 	}
 
+	void CombineFlagsForParameterBlending(FStrataOperator& A, FStrataOperator& B)
+	{
+		bBSDFHasSSS = A.bBSDFHasSSS || B.bBSDFHasSSS;
+		bBSDFHasMFPPluggedIn = A.bBSDFHasMFPPluggedIn || B.bBSDFHasMFPPluggedIn;
+		bBSDFHasEdgeColor = A.bBSDFHasEdgeColor || B.bBSDFHasEdgeColor;
+		bBSDFHasFuzz = A.bBSDFHasFuzz || B.bBSDFHasFuzz;
+		bBSDFHasHaziness = A.bBSDFHasHaziness || B.bBSDFHasHaziness;
+	}
+
+	void CopyFlagsForParameterBlending(FStrataOperator& A)
+	{
+		bBSDFHasSSS = A.bBSDFHasSSS;
+		bBSDFHasMFPPluggedIn = A.bBSDFHasMFPPluggedIn;
+		bBSDFHasEdgeColor = A.bBSDFHasEdgeColor;
+		bBSDFHasFuzz = A.bBSDFHasFuzz;
+		bBSDFHasHaziness = A.bBSDFHasHaziness;
+	}
+
 	bool IsDiscarded()
 	{
 		return bUseParameterBlending && !bRootOfParameterBlendingSubTree;
@@ -539,13 +581,13 @@ public:
 	 * Return the operator information for a given expression.
 	 */
 	virtual FStrataOperator& StrataCompilationGetOperator(UMaterialExpression* Expression) = 0;
+	/**
+	 * Return the operator information for a given index.
+	 */
+	virtual FStrataOperator* StrataCompilationGetOperatorFromIndex(int32 OperatorIndex) = 0;
 
-	virtual void StrataCompilationInfoRegisterCodeChunk(int32 CodeChunk, FStrataMaterialCompilationInfo& StrataMaterialCompilationInfo) = 0;
-	virtual bool StrataCompilationInfoContainsCodeChunk(int32 CodeChunk) = 0;
-	virtual const FStrataMaterialCompilationInfo& GetStrataCompilationInfo(int32 CodeChunk) = 0;
 	virtual FStrataRegisteredSharedLocalBasis StrataCompilationInfoRegisterSharedLocalBasis(int32 NormalCodeChunk) = 0;
 	virtual FStrataRegisteredSharedLocalBasis StrataCompilationInfoRegisterSharedLocalBasis(int32 NormalCodeChunk, int32 TangentCodeChunk) = 0;
-	virtual uint8 StrataCompilationInfoGetSharedLocalBasesCount() = 0;
 	FString GetStrataSharedLocalBasisIndexMacro(const FStrataRegisteredSharedLocalBasis& SharedLocalBasis)
 	{
 		return FString::Printf(TEXT("SHAREDLOCALBASIS_INDEX_%u"), SharedLocalBasis.GraphSharedLocalBasisIndex);
@@ -1169,11 +1211,6 @@ public:
 		return Compiler->StrataMetalnessToDiffuseAlbedoF0(BaseColor, Specular, Metallic, OutputIndex);
 	}
 
-	virtual void StrataCompilationInfoRegisterCodeChunk(int32 CodeChunk, FStrataMaterialCompilationInfo& StrataMaterialCompilationInfo) override
-	{
-		Compiler->StrataCompilationInfoRegisterCodeChunk(CodeChunk, StrataMaterialCompilationInfo);
-	}
-
 	virtual FStrataOperator& StrataCompilationRegisterOperator(int32 OperatorType, UMaterialExpression* Expression, UMaterialExpression* Parent, bool bUseParameterBlending = false) override
 	{
 		return Compiler->StrataCompilationRegisterOperator(OperatorType, Expression, Parent, bUseParameterBlending);
@@ -1184,11 +1221,12 @@ public:
 		return Compiler->StrataCompilationGetOperator(Expression);
 	}
 
-	virtual bool StrataCompilationInfoContainsCodeChunk(int32 CodeChunk) override
+	virtual FStrataOperator* StrataCompilationGetOperatorFromIndex(int32 OperatorIndex) override
 	{
-		return Compiler->StrataCompilationInfoContainsCodeChunk(CodeChunk);
+		return Compiler->StrataCompilationGetOperatorFromIndex(OperatorIndex);
 	}
 
+
 	virtual int32 StrataAddParameterBlendingBSDFCoverageToNormalMixCodeChunk(int32 ACodeChunk, int32 BCodeChunk) override
 	{
 		return Compiler->StrataAddParameterBlendingBSDFCoverageToNormalMixCodeChunk(ACodeChunk, BCodeChunk);
@@ -1203,11 +1241,6 @@ public:
 	{
 		return Compiler->StrataHorizontalMixingParameterBlendingBSDFCoverageToNormalMixCodeChunk(BackgroundCodeChunk, ForegroundCodeChunk, HorizontalMixCodeChunk);
 	}
-	
-	virtual const FStrataMaterialCompilationInfo& GetStrataCompilationInfo(int32 CodeChunk) override
-	{
-		return Compiler->GetStrataCompilationInfo(CodeChunk);
-	}
 
 	virtual FStrataRegisteredSharedLocalBasis StrataCompilationInfoRegisterSharedLocalBasis(int32 NormalCodeChunk) override
 	{
@@ -1219,11 +1252,6 @@ public:
 		return Compiler->StrataCompilationInfoRegisterSharedLocalBasis(NormalCodeChunk, TangentCodeChunk);
 	}
 
-	virtual uint8 StrataCompilationInfoGetSharedLocalBasesCount() override
-	{
-		return Compiler->StrataCompilationInfoGetSharedLocalBasesCount();
-	}
-
 protected:
 		
 	FMaterialCompiler* Compiler;
