Strata more defines clean up

diff --git a/Engine/Shaders/Shared/StrataDefinitions.h b/Engine/Shaders/Shared/StrataDefinitions.h
--- a/Engine/Shaders/Shared/StrataDefinitions.h
+++ b/Engine/Shaders/Shared/StrataDefinitions.h
@@ -12,8 +12,6 @@
 
 
 
-#define STRATA_MAX_BSDF_COUNT_PER_LAYER		4
-#define STRATA_MAX_LAYER_COUNT				4
 #define STRATA_MAX_BSDF_COUNT				15
 #define STRATA_MAX_OPERATOR_COUNT			15
 
diff --git a/Engine/Source/Runtime/Engine/Private/Materials/HLSLMaterialTranslator.cpp b/Engine/Source/Runtime/Engine/Private/Materials/HLSLMaterialTranslator.cpp
--- a/Engine/Source/Runtime/Engine/Private/Materials/HLSLMaterialTranslator.cpp
+++ b/Engine/Source/Runtime/Engine/Private/Materials/HLSLMaterialTranslator.cpp
@@ -9762,9 +9762,14 @@ bool FHLSLMaterialTranslator::StrataGenerateDerivedMaterialOperatorData()
 
 		if ((bHasUnlit || bHasVFogCloud || bHasHair || bHasSLW) && StrataMaterialBSDFCount > 1)
 		{
-			Errorf(TEXT("Unlit, Fog/Cloud, Hair or SingleLayerWater must be used in isolation %s (asset: %s).\r\n"), *Material->GetDebugName(), *Material->GetAssetPath().ToString());
+			Errorf(TEXT("Unlit, Fog/Cloud, Hair or SingleLayerWater must be used in isolation. See %s (asset: %s).\r\n"), *Material->GetDebugName(), *Material->GetAssetPath().ToString());
 			// Even though we could support Unlit with slab.
 		}
+
+		if (StrataMaterialBSDFCount > STRATA_MAX_BSDF_COUNT)
+		{
+			Errorf(TEXT("Material tries to register more BSDF than can be supproted (%d > %d). See %s (asset: %s).\r\n"), StrataMaterialBSDFCount, STRATA_MAX_BSDF_COUNT, *Material->GetDebugName(), *Material->GetAssetPath().ToString());
+		}
 	}
 
 	// STRATA_TODO: operation using parameter blending are not actually discarded so they still occupy a spot in the operation array in the compiler and in the shader also.
