Fix some typos in unit test comments.

#rnx

diff --git a/Engine/Source/Runtime/Engine/Private/Tests/Texture2DTests.cpp b/Engine/Source/Runtime/Engine/Private/Tests/Texture2DTests.cpp
--- a/Engine/Source/Runtime/Engine/Private/Tests/Texture2DTests.cpp
+++ b/Engine/Source/Runtime/Engine/Private/Tests/Texture2DTests.cpp
@@ -10,7 +10,7 @@ namespace Texture2DTest
 
 constexpr const uint32 TestFlags = EAutomationTestFlags::ApplicationContextMask | EAutomationTestFlags::EngineFilter;
 
-// A simple test to make sure that basica functionality in UTexture2D::CreateTransient works as it seems to be a 
+// A simple test to make sure that basic functionality in UTexture2D::CreateTransient works as it seems to be a 
 // fairly uncommon code path in our samples/test games etc.
 IMPLEMENT_SIMPLE_AUTOMATION_TEST(FTexture2DTestTransient, "System.Engine.Texture2D.CreateTransient", TestFlags)
 bool FTexture2DTestTransient::RunTest(const FString& Parameters)
@@ -51,7 +51,7 @@ bool FTexture2DTestLockingWhenEmpty::RunTest(const FString& Parameters)
 		EmptyTexture->Source.UnlockMip(0);
 
 		// If a lock on the primary mip (0) fails because the bulkdata is empty we still
-		// get a successful lock and so need to unlocok it
+		// get a successful lock and so need to unlock it
 		uint8* SecondLockPtr = EmptyTexture->Source.LockMip(0);
 		TestNull(TEXT("Locking an empty texture a second time"), SecondLockPtr);
 		EmptyTexture->Source.UnlockMip(0);
@@ -94,6 +94,6 @@ bool FTexture2DTestLockingWhenEmpty::RunTest(const FString& Parameters)
 
 #endif //WITH_EDITORONLY_DATA
 
-}
+} // namespace Texture2DTest
 
-#endif 
+#endif //WITH_DEV_AUTOMATION_TESTS
