Strata - Added BSDF per pixel size back report.

diff --git a/Engine/Source/Runtime/Engine/Private/Materials/HLSLMaterialTranslator.cpp b/Engine/Source/Runtime/Engine/Private/Materials/HLSLMaterialTranslator.cpp
--- a/Engine/Source/Runtime/Engine/Private/Materials/HLSLMaterialTranslator.cpp
+++ b/Engine/Source/Runtime/Engine/Private/Materials/HLSLMaterialTranslator.cpp
@@ -299,6 +299,7 @@ FHLSLMaterialTranslator::FHLSLMaterialTranslator(FMaterial* InMaterial,
 	StrataMaterialExpressionRegisteredOperators.Reserve(STRATA_MAX_OPERATOR_COUNT);
 	StrataMaterialExpressionToOperatorIndex.Reserve(STRATA_MAX_OPERATOR_COUNT);
 	StrataMaterialBSDFCount = 0;
+	StrataMaterialRequestedSizeByte = 0;
 	bStrataUsesConversionFromLegacy = false;
 	bStrataOutputsOpaqueRoughRefractions = false;
 }
@@ -2013,8 +2014,10 @@ void FHLSLMaterialTranslator::GetMaterialEnvironment(EShaderPlatform InPlatform,
 
 			StrataMaterialDescription += FString::Printf(TEXT("----- STRATA -----\r\n"));
 			StrataMaterialDescription += FString::Printf(TEXT("StrataCompilationInfo -\r\n"));
-			StrataMaterialDescription += FString::Printf(TEXT(" - TotalBSDFCount    = %i\r\n"), StrataMaterialBSDFCount);
-			StrataMaterialDescription += FString::Printf(TEXT(" - SharedLocalBasesCount = %i\r\n"), FinalUsedSharedLocalBasesCount);
+			StrataMaterialDescription += FString::Printf(TEXT(" - Byte Per Pixel Budget      %u\r\n"), StrataBytePerPixel);
+			StrataMaterialDescription += FString::Printf(TEXT(" - Requested Byte Size        %u (%d UINT32)\r\n"), StrataMaterialRequestedSizeByte, StrataMaterialRequestedSizeByte / 4);
+			StrataMaterialDescription += FString::Printf(TEXT(" - TotalBSDFCount             %i\r\n"), StrataMaterialBSDFCount);
+			StrataMaterialDescription += FString::Printf(TEXT(" - SharedLocalBasesCount      %i\r\n"), FinalUsedSharedLocalBasesCount);
 
 			for (int32 OpIt = 0; OpIt < StrataMaterialExpressionRegisteredOperators.Num(); ++OpIt)
 			{
@@ -2044,17 +2047,13 @@ void FHLSLMaterialTranslator::GetMaterialEnvironment(EShaderPlatform InPlatform,
 					{
 						if (!It.IsDiscarded() && It.MaxDistanceFromLeaves == DistanceToLeaves)
 						{
-							StrataMaterialDescription += FString::Printf(TEXT("\tIdx=%d Op=%s ParentIdx=%d LeftIndex=%d RightIndex=%d BSDFIdx=%d LayerDepth=%d IsTop=%d IsBot=%d\r\n"),
-								It.Index, GetStrataOperatorStr(It.OperatorType), It.ParentIndex, It.LeftIndex, It.RightIndex, It.BSDFIndex, It.LayerDepth, It.bIsTop, It.bIsBottom);
+							StrataMaterialDescription += FString::Printf(TEXT("\tIdx=%d Op=%s ParentIdx=%d LeftIndex=%d RightIndex=%d BSDFIdx=%d LayerDepth=%d IsTop=%d IsBot=%d BSDFType=%s SSS=%d MFP=%d F90=%d Hazy=%d Fuzz=%d \r\n"),
+								It.Index, GetStrataOperatorStr(It.OperatorType), It.ParentIndex, It.LeftIndex, It.RightIndex, It.BSDFIndex, It.LayerDepth, It.bIsTop, It.bIsBottom, 
+								*GetStrataBSDFName(It.BSDFType), It.bBSDFHasSSS, It.bBSDFHasMFPPluggedIn, It.bBSDFHasEdgeColor, It.bBSDFHasHaziness, It.bBSDFHasFuzz);
 						}
 					}
 				}
 			}
-			StrataMaterialDescription += FString::Printf(TEXT("-----------------------------------\r\n"));
-
-			StrataMaterialDescription += FString::Printf(TEXT("Byte Per Pixel Budget      %u\r\n"), StrataBytePerPixel);
-			// STRATA_TODO: display RequestedByteCount (fit in budget or not)
-			StrataMaterialDescription += FString::Printf(TEXT("------------------\r\n"));
 
 			ResourcesString += TEXT("/*");
 			ResourcesString += StrataMaterialDescription;
@@ -9631,7 +9630,7 @@ bool FHLSLMaterialTranslator::StrataGenerateDerivedMaterialOperatorData()
 	}
 
 	//
-	// Make sure each and every path of the strata tree ends up on a valid BSDF.
+	// Make sure each and every path of the strata tree have valid children and path.
 	//
 	for (auto& It : StrataMaterialExpressionRegisteredOperators)
 	{
@@ -9682,7 +9681,7 @@ bool FHLSLMaterialTranslator::StrataGenerateDerivedMaterialOperatorData()
 	//
 	// Parse the tree and mark nodes that are the root of a subtree using parameter blending, while other nodes in that tree are forced to use parameter blending.
 	// Allocate BSDFIndex at the same time.
-	// Check BSDF that should be uinique
+	// Check BSDF that should be unique.
 	//
 	{
 		bool bHasUnlit = false;
@@ -9919,6 +9918,97 @@ bool FHLSLMaterialTranslator::StrataGenerateDerivedMaterialOperatorData()
 		WalkOperators(*StrataMaterialRootOperator);
 	}
 
+	//
+	// Compute the size of the material
+	//
+	{
+		const uint32 UintByteSize = sizeof(uint32);
+		StrataMaterialRequestedSizeByte = 0;
+
+		// 1. Header
+
+		// Packed Header
+		StrataMaterialRequestedSizeByte += UintByteSize;
+
+		// Shared local bases between BSDFs
+		StrataMaterialRequestedSizeByte += FinalUsedSharedLocalBasesCount * STRATA_PACKED_SHAREDLOCALBASIS_STRIDE_BYTES;
+
+		// 2. The list of BSDFs
+		for (auto& It : StrataMaterialExpressionRegisteredOperators)
+		{
+			if (It.IsDiscarded())
+			{
+				continue; // ignore discarded operations in sub tree using parameter blending
+			}
+
+			switch (It.OperatorType)
+			{
+			case STRATA_OPERATOR_BSDF:
+			{
+				// BSDF state
+				StrataMaterialRequestedSizeByte += UintByteSize;
+
+				// From the compiler side, we can only assume the top layer has gray scale luminance weight.
+				const bool bMayHaveColoredWeight = !It.bIsTop;
+				if (bMayHaveColoredWeight)
+				{
+					StrataMaterialRequestedSizeByte += UintByteSize;
+				}
+
+				switch (It.BSDFType)
+				{
+				case STRATA_BSDF_TYPE_SLAB:
+				{
+					// Compute values closer to the reality for HasSSS and IsSimpleVolume, now that we know that we know the topology of the material.
+					const bool bIsSimpleVolume = !It.bIsBottom && It.bBSDFHasMFPPluggedIn;
+					const bool bHasSSS = It.bIsBottom && It.bBSDFHasSSS;
+				
+					StrataMaterialRequestedSizeByte += UintByteSize;
+					StrataMaterialRequestedSizeByte += UintByteSize;
+
+					if (It.bBSDFHasEdgeColor || It.bBSDFHasHaziness)
+					{
+						StrataMaterialRequestedSizeByte += UintByteSize;
+					}
+					if (bHasSSS || bIsSimpleVolume)
+					{
+						StrataMaterialRequestedSizeByte += UintByteSize;
+					}
+					if (It.bBSDFHasFuzz)
+					{
+						StrataMaterialRequestedSizeByte += UintByteSize;
+					}
+					break;
+				}
+				case STRATA_BSDF_TYPE_HAIR:
+				{
+					StrataMaterialRequestedSizeByte += UintByteSize;
+					StrataMaterialRequestedSizeByte += UintByteSize;
+					break;
+				}
+				case STRATA_BSDF_TYPE_SINGLELAYERWATER:
+				{
+					StrataMaterialRequestedSizeByte += UintByteSize;
+					StrataMaterialRequestedSizeByte += UintByteSize;
+					break;
+				}
+				case STRATA_BSDF_TYPE_UNLIT:
+				{
+					// Never stored, it goes directly into the scene as emitted luminance.
+					break;
+				}
+				default:
+				{
+					Errorf(TEXT("Unkownd BSDF type encountered in %s (asset: %s).\r\n"), *Material->GetDebugName(), *Material->GetAssetPath().ToString());
+					break;
+				}
+				}
+				break;
+			}
+			}
+		}
+	}
+
 	return true; // Success
 }
 
diff --git a/Engine/Source/Runtime/Engine/Private/Materials/HLSLMaterialTranslator.h b/Engine/Source/Runtime/Engine/Private/Materials/HLSLMaterialTranslator.h
--- a/Engine/Source/Runtime/Engine/Private/Materials/HLSLMaterialTranslator.h
+++ b/Engine/Source/Runtime/Engine/Private/Materials/HLSLMaterialTranslator.h
@@ -427,6 +427,7 @@ protected:
 	TArray<FStrataOperator> StrataMaterialExpressionRegisteredOperators;
 	FStrataOperator* StrataMaterialRootOperator;
 	uint32 StrataMaterialBSDFCount;
+	uint32 StrataMaterialRequestedSizeByte;
 	bool bStrataUsesConversionFromLegacy;
 	bool bStrataOutputsOpaqueRoughRefractions;
 
diff --git a/Engine/Source/Runtime/Engine/Private/Materials/StrataMaterial.cpp b/Engine/Source/Runtime/Engine/Private/Materials/StrataMaterial.cpp
--- a/Engine/Source/Runtime/Engine/Private/Materials/StrataMaterial.cpp
+++ b/Engine/Source/Runtime/Engine/Private/Materials/StrataMaterial.cpp
@@ -43,115 +43,4 @@ FStrataRegisteredSharedLocalBasis StrataCompilationInfoCreateSharedLocalBasis(FM
 	return Compiler->StrataCompilationInfoRegisterSharedLocalBasis(NormalCodeChunk, TangentCodeChunk);
 }
 
-//FStrataMaterialAnalysisResult::FStrataMaterialAnalysisResult()
-//{
-//	bFitInMemoryBudget = true;
-//	RequestedLayerCount = 0;
-//	RequestedBSDFCount = 0;
-//	RequestedByteCount = 0;
-//	ClampedLayerCount = 0;
-//	ClampedBSDFCount = 0;
-//	UsedByteCount = 0;
-//}
-//
-//FStrataMaterialAnalysisResult StrataCompilationInfoMaterialAnalysis(FMaterialCompiler* Compiler, const FStrataMaterialCompilationInfo& Material, const uint32 StrataBytePerPixel)
-//{
-//	const uint32 UintByteSize = sizeof(uint32);
-//
-//	FStrataMaterialAnalysisResult Result;
-//
-//	// 1. Header
-//
-//	// SharedLocalBases_BSDFCount
-//	Result.RequestedByteCount += UintByteSize;
-//
-//	// shared local bases between BSDFs
-//	Result.RequestedByteCount += Compiler->StrataCompilationInfoGetSharedLocalBasesCount() * STRATA_PACKED_SHAREDLOCALBASIS_STRIDE_BYTES;
-//	
-//	// 2. The list of BSDFs
-//	Result.RequestedMaxBSDFCountPerLayer = 0;
-//	// We process layers from top to bottom to cull the bottom ones in case we run out of pixel bytes
-//	for (uint32 LayerIt = 0; LayerIt < Material.LayerCount; LayerIt++)
-//	{
-//		const FStrataMaterialCompilationInfo::FLayer& Layer = Material.Layers[LayerIt];
-//
-//		const bool bTopLayer = LayerIt == 0;
-//		const bool bBottomLayer = LayerIt == (Material.LayerCount - 1);
-//
-//		for (uint32 BSDFIt = 0; BSDFIt < Layer.BSDFCount; BSDFIt++)
-//		{
-//			const FStrataMaterialCompilationInfo::FBSDF& BSDF = Layer.BSDFs[BSDFIt];
-//
-//			// BSDF state
-//			Result.RequestedByteCount += UintByteSize;
-//
-//			// From the compiler side, we can only assume the top layer has grey scale weight/throughput
-//			const bool bMayBeColoredWeight = LayerIt > 0;
-//			if (bMayBeColoredWeight)
-//			{
-//				Result.RequestedByteCount += UintByteSize;
-//			}
-//
-//			switch (BSDF.Type)
-//			{
-//			case STRATA_BSDF_TYPE_SLAB:
-//			{
-//				// Compute values closer to the reality for HasSSS and IsSimpleVolume, now that we know that we know the topology of the material.
-//				const bool bIsSimpleVolume = !bBottomLayer && BSDF.bHasMFPPluggedIn;
-//				const bool bHasSSS = bBottomLayer && BSDF.bHasSSS && !bIsSimpleVolume;
-//
-//				Result.RequestedByteCount += UintByteSize;
-//				Result.RequestedByteCount += UintByteSize;
-//				if (BSDF.bHasEdgeColor || BSDF.bHasThinFilm || BSDF.bHasHaziness)
-//				{
-//					Result.RequestedByteCount += UintByteSize;
-//				}
-//				if (bHasSSS || bIsSimpleVolume)
-//				{
-//					Result.RequestedByteCount += UintByteSize;
-//				}
-//				if (BSDF.bHasFuzz)
-//				{
-//					Result.RequestedByteCount += UintByteSize;
-//				}
-//				break;
-//			}
-//			case STRATA_BSDF_TYPE_HAIR:
-//			{
-//				Result.RequestedByteCount += UintByteSize;
-//				Result.RequestedByteCount += UintByteSize;
-//				break;
-//			}
-//			case STRATA_BSDF_TYPE_SINGLELAYERWATER:
-//			{
-//				Result.RequestedByteCount += UintByteSize;
-//				Result.RequestedByteCount += UintByteSize;
-//				break;
-//			}
-//			}
-//		}
-//		Result.RequestedMaxBSDFCountPerLayer = FMath::Max(Result.RequestedMaxBSDFCountPerLayer, Layer.BSDFCount);
-//
-//		Result.RequestedLayerCount++;
-//		Result.RequestedBSDFCount += Layer.BSDFCount;
-//
-//		if (Result.RequestedByteCount <= StrataBytePerPixel && Result.bFitInMemoryBudget)
-//		{
-//			// We only validate all the BSDF of a layer if it remains within budget and we are not already out of budget
-//			Result.ClampedBSDFCount += Layer.BSDFCount;
-//			Result.ClampedLayerCount++;
-//
-//			// Set the current used bytes
-//			Result.UsedByteCount = Result.RequestedByteCount;
-//		}
-//		else
-//		{
-//			// Used byte count remains unchanged and we notify that we start peeling off top layers now
-//			Result.bFitInMemoryBudget = false;
-//		}
-//	}
-//	return Result;
-//
-//}
-	
 
