Strata - evalauted byte count per pixel for material no matches better.

diff --git a/Engine/Source/Runtime/Engine/Private/Materials/HLSLMaterialTranslator.cpp b/Engine/Source/Runtime/Engine/Private/Materials/HLSLMaterialTranslator.cpp
--- a/Engine/Source/Runtime/Engine/Private/Materials/HLSLMaterialTranslator.cpp
+++ b/Engine/Source/Runtime/Engine/Private/Materials/HLSLMaterialTranslator.cpp
@@ -300,6 +300,8 @@ FHLSLMaterialTranslator::FHLSLMaterialTranslator(FMaterial* InMaterial,
 	StrataMaterialExpressionToOperatorIndex.Reserve(STRATA_MAX_OPERATOR_COUNT);
 	StrataMaterialBSDFCount = 0;
 	StrataMaterialRequestedSizeByte = 0;
+	bStrataMaterialIsSimple = false;
+	bStrataMaterialIsSingle = false;
 	bStrataUsesConversionFromLegacy = false;
 	bStrataOutputsOpaqueRoughRefractions = false;
 }
@@ -1897,10 +1899,6 @@ void FHLSLMaterialTranslator::GetMaterialEnvironment(EShaderPlatform InPlatform,
 
 	if (bMaterialIsStrata)
 	{
-		StrataPixelNormalInitializerValues += FString::Printf(TEXT("\n\n\n\t// Strata normal and tangent\n"));
-		FinalUsedSharedLocalBasesCount = 0;
-		FStrataRegisteredSharedLocalBasis UsedSharedLocalBasesInfo[STRATA_MAX_SHAREDLOCALBASES_REGISTERS];
-
 		check(StrataMaterialRootOperator);
 		uint32 RootMaximumDistanceToLeaves = StrataMaterialRootOperator->MaxDistanceFromLeaves;
 		OutEnvironment.SetDefine(TEXT("MATERIAL_TOPOLOGY_TREE_MAX_DEPTH"), RootMaximumDistanceToLeaves);
@@ -1909,97 +1907,10 @@ void FHLSLMaterialTranslator::GetMaterialEnvironment(EShaderPlatform InPlatform,
 
 		OutEnvironment.SetDefine(TEXT("STRATA_MATERIAL_OUTPUT_OPAQUE_ROUGH_REFRACTIONS"), bStrataOutputsOpaqueRoughRefractions ? 1 : 0);
 
-		/*
-		* The final output code/workflow for shared tangent basis should look like
-		*
-		* #define SHAREDLOCALBASIS_INDEX_0 0		// default, unused
-		* #define SHAREDLOCALBASIS_INDEX_1 0		// default, unused
-		* #define SHAREDLOCALBASIS_INDEX_2 0
-		*
-		* FStrataData BSDF0 = GetStrataSlabBSDF(... SHAREDLOCALBASIS_0, NormalCode0 ...)
-		* FStrataData BSDF1 = GetStrataSlabBSDF(... SHAREDLOCALBASIS_1, NormalCode1 ...)
-		*
-		* float3 NormalCode2 = lerp(NormalCode0, NormalCode1, mix)
-		* FStrataData BSDF2 = StrataHorizontalMixingParameterBlending(BSDF0, BSDF1, mix, NormalCode2, SHAREDLOCALBASIS_INDEX_2) // will internally create NormalCode2
-		*
-		* tParameters.SharedLocalBases.Normals[SHAREDLOCALBASIS_INDEX_2] = NormalCode2;
-		* #if MATERIAL_TANGENTSPACENORMAL
-		* Parameters.SharedLocalBases.Normals[SHAREDLOCALBASIS_INDEX_2] *= Parameters.TwoSidedSign;
-		* #endif
-		*/
-
-		for(int32 OpIt=0; OpIt < StrataMaterialExpressionRegisteredOperators.Num(); ++OpIt)
-		{
-			FStrataOperator& BSDFOperator = StrataMaterialExpressionRegisteredOperators[OpIt];
-			if (BSDFOperator.BSDFIndex == INDEX_NONE && !BSDFOperator.bRootOfParameterBlendingSubTree)
-			{
-				continue;	// not a bsdf or not the root of a parameter blending subtree so there is not local basis to register
-			}
-
-			if (BSDFOperator.BSDFRegisteredSharedLocalBasis.NormalCodeChunk == INDEX_NONE && BSDFOperator.BSDFRegisteredSharedLocalBasis.TangentCodeChunk == INDEX_NONE)
-			{
-				continue;	// We skip null normal on certain BSDF, for instance unlit.
-			}
-			const FStrataSharedLocalBasesInfo& StrataSharedLocalBasesInfo = StrataCompilationInfoGetMatchingSharedLocalBasisInfo(BSDFOperator.BSDFRegisteredSharedLocalBasis);
-
-			// First, we check that the normal/tangent has not already written out (avoid 2 BSDFs sharing the same normal to note generate the same code twice)
-			bool bAlreadyProcessed = false;
-			for (uint8 i = 0; i < FinalUsedSharedLocalBasesCount; ++i)
-			{
-				if (UsedSharedLocalBasesInfo[i].NormalCodeChunkHash == StrataSharedLocalBasesInfo.SharedData.NormalCodeChunkHash &&
-					(UsedSharedLocalBasesInfo[i].TangentCodeChunkHash == StrataSharedLocalBasesInfo.SharedData.TangentCodeChunkHash || BSDFOperator.BSDFRegisteredSharedLocalBasis.TangentCodeChunk == INDEX_NONE))
-				{
-					bAlreadyProcessed = true;
-					break;
-				}
-			}
-			if (bAlreadyProcessed)
-			{
-				continue;
-			}
-
-			check(FinalUsedSharedLocalBasesCount < STRATA_MAX_SHAREDLOCALBASES_REGISTERS);
-			const uint8 FinalSharedLocalBasisIndex = FinalUsedSharedLocalBasesCount++;
-			UsedSharedLocalBasesInfo[FinalSharedLocalBasisIndex] = StrataSharedLocalBasesInfo.SharedData;
-
-			// Write out normals
-			StrataPixelNormalInitializerValues += FString::Printf(TEXT("\tParameters.SharedLocalBases.Normals[%u] = %s;\n"), FinalSharedLocalBasisIndex, *StrataSharedLocalBasesInfo.NormalCode);
-			StrataPixelNormalInitializerValues += FString::Printf(TEXT("\t#if MATERIAL_TANGENTSPACENORMAL\n"));
-			StrataPixelNormalInitializerValues += FString::Printf(TEXT("\tParameters.SharedLocalBases.Normals[%u] *= Parameters.TwoSidedSign;\n"), FinalSharedLocalBasisIndex);
-			StrataPixelNormalInitializerValues += FString::Printf(TEXT("\t#endif\n"));
-
-			// Write out tangents
-			if (StrataSharedLocalBasesInfo.SharedData.TangentCodeChunk != INDEX_NONE)
-			{
-				StrataPixelNormalInitializerValues += FString::Printf(TEXT("\tParameters.SharedLocalBases.Tangents[%u] = %s;\n"), FinalSharedLocalBasisIndex, *StrataSharedLocalBasesInfo.TangentCode);
-			}
-			else
-			{
-				StrataPixelNormalInitializerValues += FString::Printf(TEXT("\tParameters.SharedLocalBases.Tangents[%u] = Parameters.TangentToWorld[0];\n"), FinalSharedLocalBasisIndex);
-			}
-			StrataPixelNormalInitializerValues += FString::Printf(TEXT("\t#if MATERIAL_TANGENTSPACENORMAL\n"));
-			StrataPixelNormalInitializerValues += FString::Printf(TEXT("\tParameters.SharedLocalBases.Tangents[%u] *= Parameters.TwoSidedSign;\n"), FinalSharedLocalBasisIndex);
-			StrataPixelNormalInitializerValues += FString::Printf(TEXT("\t#endif\n"));
-		}
-		StrataPixelNormalInitializerValues += FString::Printf(TEXT("\tParameters.SharedLocalBases.Count = %u;\n"), FinalUsedSharedLocalBasesCount);
-
-		// Now write out all the macros, them mapping from the BSDF to the effective position/index in the shared local basis array they should write to.
-		for (TMultiMap<uint64, FStrataSharedLocalBasesInfo>::TIterator It(CodeChunkToStrataSharedLocalBasis); It; ++It)
-		{
-			// The default linear output index will be 0 by default, and different if in fact the shared local basis points to one that is effectively in used in the array of shared local bases.
-			uint8 LinearIndex = 0;
-			for (uint8 i = 0; i < FinalUsedSharedLocalBasesCount; ++i)
-			{
-				if (UsedSharedLocalBasesInfo[i].NormalCodeChunkHash == It->Value.SharedData.NormalCodeChunkHash &&
-					(UsedSharedLocalBasesInfo[i].TangentCodeChunkHash == It->Value.SharedData.TangentCodeChunkHash || It->Value.SharedData.TangentCodeChunk == INDEX_NONE))
-				{
-					LinearIndex = i;
-					break;
-				}
-			}
-
-			OutEnvironment.SetDefine(*GetStrataSharedLocalBasisIndexMacro(It->Value.SharedData), LinearIndex);
-		}
+		// Compute the shared local basis count and generate the hlsl shader code for it.
+		StrataPixelNormalInitializerValues = FString::Printf(TEXT("\n\n\n\t// Strata normal and tangent\n"));
+		FinalUsedSharedLocalBasesCount = 0;
+		StrataEvaluateSharedLocalBases(FinalUsedSharedLocalBasesCount, &StrataPixelNormalInitializerValues, &OutEnvironment);
 
 		// Now write some feedback to the user
 		{
@@ -2015,6 +1926,7 @@ void FHLSLMaterialTranslator::GetMaterialEnvironment(EShaderPlatform InPlatform,
 			StrataMaterialDescription += FString::Printf(TEXT("----- STRATA -----\r\n"));
 			StrataMaterialDescription += FString::Printf(TEXT("StrataCompilationInfo -\r\n"));
 			StrataMaterialDescription += FString::Printf(TEXT(" - Byte Per Pixel Budget      %u\r\n"), StrataBytePerPixel);
+			StrataMaterialDescription += FString::Printf(TEXT(" - Material complexity        %s\r\n"), bStrataMaterialIsSingle ? TEXT("SINGLE") : (bStrataMaterialIsSimple ? TEXT("SIMPLE") : TEXT("COMPLEX")));
 			StrataMaterialDescription += FString::Printf(TEXT(" - Requested Byte Size        %u (%d UINT32)\r\n"), StrataMaterialRequestedSizeByte, StrataMaterialRequestedSizeByte / 4);
 			StrataMaterialDescription += FString::Printf(TEXT(" - TotalBSDFCount             %i\r\n"), StrataMaterialBSDFCount);
 			StrataMaterialDescription += FString::Printf(TEXT(" - SharedLocalBasesCount      %i\r\n"), FinalUsedSharedLocalBasesCount);
@@ -2047,9 +1959,9 @@ void FHLSLMaterialTranslator::GetMaterialEnvironment(EShaderPlatform InPlatform,
 					{
 						if (!It.IsDiscarded() && It.MaxDistanceFromLeaves == DistanceToLeaves)
 						{
-							StrataMaterialDescription += FString::Printf(TEXT("\tIdx=%d Op=%s ParentIdx=%d LeftIndex=%d RightIndex=%d BSDFIdx=%d LayerDepth=%d IsTop=%d IsBot=%d BSDFType=%s SSS=%d MFP=%d F90=%d Hazy=%d Fuzz=%d \r\n"),
+							StrataMaterialDescription += FString::Printf(TEXT("\tIdx=%d Op=%s ParentIdx=%d LeftIndex=%d RightIndex=%d BSDFIdx=%d LayerDepth=%d IsTop=%d IsBot=%d BSDFType=%s SSS=%d MFP=%d F90=%d Hazy=%d Fuzz=%d  Aniso=%d \r\n"),
 								It.Index, GetStrataOperatorStr(It.OperatorType), It.ParentIndex, It.LeftIndex, It.RightIndex, It.BSDFIndex, It.LayerDepth, It.bIsTop, It.bIsBottom, 
-								*GetStrataBSDFName(It.BSDFType), It.bBSDFHasSSS, It.bBSDFHasMFPPluggedIn, It.bBSDFHasEdgeColor, It.bBSDFHasHaziness, It.bBSDFHasFuzz);
+								*GetStrataBSDFName(It.BSDFType), It.bBSDFHasSSS, It.bBSDFHasMFPPluggedIn, It.bBSDFHasEdgeColor, It.bBSDFHasHaziness, It.bBSDFHasFuzz, It.bBSDFHasAnisotropy);
 						}
 					}
 				}
@@ -9602,6 +9514,116 @@ FStrataOperator* FHLSLMaterialTranslator::StrataCompilationGetOperatorFromIndex(
 	return &StrataMaterialExpressionRegisteredOperators[OperatorIndex];
 }
 
+void FHLSLMaterialTranslator::StrataEvaluateSharedLocalBases(
+	uint8& UsedSharedLocalBasesCount,
+	FString* OutStrataPixelNormalInitializerValues, 
+	FShaderCompilerEnvironment* OutEnvironment)
+{
+	/*
+	* The final output code/workflow for shared tangent basis should look like
+	*
+	* #define SHAREDLOCALBASIS_INDEX_0 0		// default, unused
+	* #define SHAREDLOCALBASIS_INDEX_1 0		// default, unused
+	* #define SHAREDLOCALBASIS_INDEX_2 0
+	*
+	* FStrataData BSDF0 = GetStrataSlabBSDF(... SHAREDLOCALBASIS_0, NormalCode0 ...)
+	* FStrataData BSDF1 = GetStrataSlabBSDF(... SHAREDLOCALBASIS_1, NormalCode1 ...)
+	*
+	* float3 NormalCode2 = lerp(NormalCode0, NormalCode1, mix)
+	* FStrataData BSDF2 = StrataHorizontalMixingParameterBlending(BSDF0, BSDF1, mix, NormalCode2, SHAREDLOCALBASIS_INDEX_2) // will internally create NormalCode2
+	*
+	* tParameters.SharedLocalBases.Normals[SHAREDLOCALBASIS_INDEX_2] = NormalCode2;
+	* #if MATERIAL_TANGENTSPACENORMAL
+	* Parameters.SharedLocalBases.Normals[SHAREDLOCALBASIS_INDEX_2] *= Parameters.TwoSidedSign;
+	* #endif
+	*/
+
+	FStrataRegisteredSharedLocalBasis UsedSharedLocalBasesInfo[STRATA_MAX_SHAREDLOCALBASES_REGISTERS];
+	UsedSharedLocalBasesCount = 0;
+
+	for (int32 OpIt = 0; OpIt < StrataMaterialExpressionRegisteredOperators.Num(); ++OpIt)
+	{
+		FStrataOperator& BSDFOperator = StrataMaterialExpressionRegisteredOperators[OpIt];
+		if (BSDFOperator.BSDFIndex == INDEX_NONE && !BSDFOperator.bRootOfParameterBlendingSubTree)
+		{
+			continue;	// not a bsdf or not the root of a parameter blending subtree so there is not local basis to register
+		}
+
+		if (BSDFOperator.BSDFRegisteredSharedLocalBasis.NormalCodeChunk == INDEX_NONE && BSDFOperator.BSDFRegisteredSharedLocalBasis.TangentCodeChunk == INDEX_NONE)
+		{
+			continue;	// We skip null normal on certain BSDF, for instance unlit.
+		}
+		const FStrataSharedLocalBasesInfo& StrataSharedLocalBasesInfo = StrataCompilationInfoGetMatchingSharedLocalBasisInfo(BSDFOperator.BSDFRegisteredSharedLocalBasis);
+
+		// First, we check that the normal/tangent has not already written out (avoid 2 BSDFs sharing the same normal to note generate the same code twice)
+		bool bAlreadyProcessed = false;
+		for (uint8 i = 0; i < UsedSharedLocalBasesCount; ++i)
+		{
+			if (UsedSharedLocalBasesInfo[i].NormalCodeChunkHash == StrataSharedLocalBasesInfo.SharedData.NormalCodeChunkHash &&
+				(UsedSharedLocalBasesInfo[i].TangentCodeChunkHash == StrataSharedLocalBasesInfo.SharedData.TangentCodeChunkHash || BSDFOperator.BSDFRegisteredSharedLocalBasis.TangentCodeChunk == INDEX_NONE))
+			{
+				bAlreadyProcessed = true;
+				break;
+			}
+		}
+		if (bAlreadyProcessed)
+		{
+			continue;
+		}
+
+		check(UsedSharedLocalBasesCount < STRATA_MAX_SHAREDLOCALBASES_REGISTERS);
+		const uint8 FinalSharedLocalBasisIndex = UsedSharedLocalBasesCount++;
+		UsedSharedLocalBasesInfo[FinalSharedLocalBasisIndex] = StrataSharedLocalBasesInfo.SharedData;
+
+		if (OutStrataPixelNormalInitializerValues)
+		{
+			// Write out normals
+			*OutStrataPixelNormalInitializerValues += FString::Printf(TEXT("\tParameters.SharedLocalBases.Normals[%u] = %s;\n"), FinalSharedLocalBasisIndex, *StrataSharedLocalBasesInfo.NormalCode);
+			*OutStrataPixelNormalInitializerValues += FString::Printf(TEXT("\t#if MATERIAL_TANGENTSPACENORMAL\n"));
+			*OutStrataPixelNormalInitializerValues += FString::Printf(TEXT("\tParameters.SharedLocalBases.Normals[%u] *= Parameters.TwoSidedSign;\n"), FinalSharedLocalBasisIndex);
+			*OutStrataPixelNormalInitializerValues += FString::Printf(TEXT("\t#endif\n"));
+
+			// Write out tangents
+			if (StrataSharedLocalBasesInfo.SharedData.TangentCodeChunk != INDEX_NONE)
+			{
+				*OutStrataPixelNormalInitializerValues += FString::Printf(TEXT("\tParameters.SharedLocalBases.Tangents[%u] = %s;\n"), FinalSharedLocalBasisIndex, *StrataSharedLocalBasesInfo.TangentCode);
+			}
+			else
+			{
+				*OutStrataPixelNormalInitializerValues += FString::Printf(TEXT("\tParameters.SharedLocalBases.Tangents[%u] = Parameters.TangentToWorld[0];\n"), FinalSharedLocalBasisIndex);
+			}
+			*OutStrataPixelNormalInitializerValues += FString::Printf(TEXT("\t#if MATERIAL_TANGENTSPACENORMAL\n"));
+			*OutStrataPixelNormalInitializerValues += FString::Printf(TEXT("\tParameters.SharedLocalBases.Tangents[%u] *= Parameters.TwoSidedSign;\n"), FinalSharedLocalBasisIndex);
+			*OutStrataPixelNormalInitializerValues += FString::Printf(TEXT("\t#endif\n"));
+		}
+	}
+	if (OutStrataPixelNormalInitializerValues)
+	{
+		*OutStrataPixelNormalInitializerValues += FString::Printf(TEXT("\tParameters.SharedLocalBases.Count = %u;\n"), FinalUsedSharedLocalBasesCount);
+	}
+
+	if (OutEnvironment)
+	{
+		// Now write out all the macros, them mapping from the BSDF to the effective position/index in the shared local basis array they should write to.
+		for (TMultiMap<uint64, FStrataSharedLocalBasesInfo>::TIterator It(CodeChunkToStrataSharedLocalBasis); It; ++It)
+		{
+			// The default linear output index will be 0 by default, and different if in fact the shared local basis points to one that is effectively in used in the array of shared local bases.
+			uint8 LinearIndex = 0;
+			for (uint8 i = 0; i < FinalUsedSharedLocalBasesCount; ++i)
+			{
+				if (UsedSharedLocalBasesInfo[i].NormalCodeChunkHash == It->Value.SharedData.NormalCodeChunkHash &&
+					(UsedSharedLocalBasesInfo[i].TangentCodeChunkHash == It->Value.SharedData.TangentCodeChunkHash || It->Value.SharedData.TangentCodeChunk == INDEX_NONE))
+				{
+					LinearIndex = i;
+					break;
+				}
+			}
+
+			OutEnvironment->SetDefine(*GetStrataSharedLocalBasisIndexMacro(It->Value.SharedData), LinearIndex);
+		}
+	}
+}
+
 bool FHLSLMaterialTranslator::StrataGenerateDerivedMaterialOperatorData()
 {
 	if (StrataMaterialExpressionRegisteredOperators.IsEmpty())
@@ -9922,16 +9944,77 @@ bool FHLSLMaterialTranslator::StrataGenerateDerivedMaterialOperatorData()
 	// Compute the size of the material
 	//
 	{
+		// Compute the shared local basis count only
+		// But we cannot use StrataEvaluateSharedLocalBases here, becaise the material has not been compiled yet so al lthe bases would just default to the same.
+		// STRATA_TODO: can we do that material generation in two passes? 
+		//		1- A first one to evaluate the normal/tangent code
+		//		2- Operators are processed and simplification computed based on memory budget
+		//		3- Material is finally compiled for with operator updated to fit in memory budget.
+		uint8 UsedSharedLocalBasesCount = StrataMaterialBSDFCount;
+
 		const uint32 UintByteSize = sizeof(uint32);
 		StrataMaterialRequestedSizeByte = 0;
 
-		// 1. Header
+		// 1. Evaluate simple/single BSDF
+		bStrataMaterialIsSimple = StrataMaterialBSDFCount == 1;
+		bStrataMaterialIsSingle = StrataMaterialBSDFCount == 1;
+		for (auto& It : StrataMaterialExpressionRegisteredOperators)
+		{
+			if (It.IsDiscarded())
+			{
+				continue; // ignore discarded operations in sub tree using parameter blending
+			}
+
+			switch (It.OperatorType)
+			{
+			case STRATA_OPERATOR_BSDF:
+			{
+				// From the compiler side, we can only assume the top layer has gray scale luminance weight.
+				const bool bMayHaveColoredWeight = !It.bIsTop;
+
+				switch (It.BSDFType)
+				{
+				case STRATA_BSDF_TYPE_SLAB:
+				{
+					bStrataMaterialIsSimple = bStrataMaterialIsSimple && !bMayHaveColoredWeight && !It.bBSDFHasAnisotropy && !It.bBSDFHasEdgeColor && !It.bBSDFHasFuzz && !It.bBSDFHasHaziness && !It.bBSDFHasMFPPluggedIn && !It.bBSDFHasSSS;
+					bStrataMaterialIsSingle = bStrataMaterialIsSingle && !bMayHaveColoredWeight && !It.bBSDFHasAnisotropy;
+					break;
+				}
+				case STRATA_BSDF_TYPE_HAIR:
+				{
+					bStrataMaterialIsSimple = false;
+					bStrataMaterialIsSingle = false;
+					break;
+				}
+				case STRATA_BSDF_TYPE_SINGLELAYERWATER:
+				{
+					bStrataMaterialIsSimple = false;
+					bStrataMaterialIsSingle = false;
+					break;
+				}
+				}
+				break;
+			}
+			}
+		}
+		bStrataMaterialIsSingle = bStrataMaterialIsSingle && !bStrataMaterialIsSimple;
 
-		// Packed Header
-		StrataMaterialRequestedSizeByte += UintByteSize;
+		// 2. Header
+
+		if (!bStrataMaterialIsSimple && !bStrataMaterialIsSingle) // header written later, 
+		{
+			// Packed Header
+			StrataMaterialRequestedSizeByte += UintByteSize;
+
+			// Shared local bases between BSDFs
+			StrataMaterialRequestedSizeByte += UsedSharedLocalBasesCount * STRATA_PACKED_SHAREDLOCALBASIS_STRIDE_BYTES;
+		}
+		else
+		{
+			// Top normal texture is used to read the data
+			StrataMaterialRequestedSizeByte += UintByteSize;
+		}
 
-		// Shared local bases between BSDFs
-		StrataMaterialRequestedSizeByte += FinalUsedSharedLocalBasesCount * STRATA_PACKED_SHAREDLOCALBASIS_STRIDE_BYTES;
 
 		// 2. The list of BSDFs
 		for (auto& It : StrataMaterialExpressionRegisteredOperators)
@@ -9945,13 +10028,32 @@ bool FHLSLMaterialTranslator::StrataGenerateDerivedMaterialOperatorData()
 			{
 			case STRATA_OPERATOR_BSDF:
 			{
-				// BSDF state
-				StrataMaterialRequestedSizeByte += UintByteSize;
-
 				// From the compiler side, we can only assume the top layer has gray scale luminance weight.
 				const bool bMayHaveColoredWeight = !It.bIsTop;
-				if (bMayHaveColoredWeight)
+
+				if (bStrataMaterialIsSimple)
+				{
+					// Header
+					StrataMaterialRequestedSizeByte += UintByteSize;
+					// Disney material
+					StrataMaterialRequestedSizeByte += UintByteSize;
+					break; // Stop here
+				}
+				else if (bStrataMaterialIsSingle)
+				{
+					// Header
+					StrataMaterialRequestedSizeByte += UintByteSize;
+				}
+				else if (bMayHaveColoredWeight)
+				{
+					// BSDF state
+					StrataMaterialRequestedSizeByte += UintByteSize;
+					// Color weight
+					StrataMaterialRequestedSizeByte += UintByteSize;
+				}
+				else
 				{
+					// BSDF state with gray scale weight
 					StrataMaterialRequestedSizeByte += UintByteSize;
 				}
 
diff --git a/Engine/Source/Runtime/Engine/Private/Materials/HLSLMaterialTranslator.h b/Engine/Source/Runtime/Engine/Private/Materials/HLSLMaterialTranslator.h
--- a/Engine/Source/Runtime/Engine/Private/Materials/HLSLMaterialTranslator.h
+++ b/Engine/Source/Runtime/Engine/Private/Materials/HLSLMaterialTranslator.h
@@ -428,6 +428,8 @@ protected:
 	FStrataOperator* StrataMaterialRootOperator;
 	uint32 StrataMaterialBSDFCount;
 	uint32 StrataMaterialRequestedSizeByte;
+	bool bStrataMaterialIsSimple;
+	bool bStrataMaterialIsSingle;
 	bool bStrataUsesConversionFromLegacy;
 	bool bStrataOutputsOpaqueRoughRefractions;
 
@@ -642,6 +644,7 @@ protected:
 	int32 GenericSwitch(const TCHAR* Function, int32 IfTrue, int32 IfFalse);
 
 	bool StrataGenerateDerivedMaterialOperatorData();
+	void StrataEvaluateSharedLocalBases(uint8& UsedSharedLocalBasesCount, FString* OutStrataPixelNormalInitializerValues, FShaderCompilerEnvironment* OutEnvironment);
 
 	// FMaterialCompiler interface.
 
diff --git a/Engine/Source/Runtime/Engine/Private/Materials/MaterialExpressions.cpp b/Engine/Source/Runtime/Engine/Private/Materials/MaterialExpressions.cpp
--- a/Engine/Source/Runtime/Engine/Private/Materials/MaterialExpressions.cpp
+++ b/Engine/Source/Runtime/Engine/Private/Materials/MaterialExpressions.cpp
@@ -21041,21 +21041,16 @@ int32 UMaterialExpressionStrataSlabBSDF::Compile(class FMaterialCompiler* Compil
 {
 	int32 RoughnessCodeChunk = CompileWithDefaultFloat1(Compiler, Roughness, 0.5f);
 	int32 AnisotropyCodeChunk = CompileWithDefaultFloat1(Compiler, Anisotropy, 0.0f);
+
 	// As long as both roughness are potentially different, we must take it into account in our encoding.
 	// We also cannot ignore the tangent when using the default Tangent because GetTangentBasis
 	// used in StrataGetBSDFSharedBasis cannot be relied on for smooth tangent used for lighting on any mesh.
 	const bool bHasAnisotropy = HasAnisotropy();
 
-	int32 NormalCodeChunk = CompileWithDefaultNormalWS(Compiler, Normal);
-
-	int32 TangentCodeChunk = bHasAnisotropy ? CompileWithDefaultTangentWS(Compiler, Tangent) : INDEX_NONE;
-	const FStrataRegisteredSharedLocalBasis NewRegisteredSharedLocalBasis = StrataCompilationInfoCreateSharedLocalBasis(Compiler, NormalCodeChunk, TangentCodeChunk);
-
-	const bool bHasEdgeColor = HasEdgeColor(); // This accounts for EdgeColor and also F90 when the non metalness worfklow is selected.
+	const bool bHasEdgeColor = HasEdgeColor(); // This accounts for EdgeColor and also F90 when the non metalness workflow is selected.
 	const bool bHasFuzz = HasFuzz();
 	const bool bHasHaziness = HasHaziness();
 	const bool bHasMFPPluggedIn = HasMFPPluggedIn();
-
 	int32 SSSProfileCodeChunk = INDEX_NONE;
 	const bool bHasSSS = HasSSS();
 	if (bHasSSS)
@@ -21070,13 +21065,12 @@ int32 UMaterialExpressionStrataSlabBSDF::Compile(class FMaterialCompiler* Compil
 	const float DefaultSpecular = 0.5f;
 	const float DefaultF0 = DielectricSpecularToF0(DefaultSpecular);
 
+	int32 NormalCodeChunk = CompileWithDefaultNormalWS(Compiler, Normal);
+	int32 TangentCodeChunk = bHasAnisotropy ? CompileWithDefaultTangentWS(Compiler, Tangent) : INDEX_NONE;
+	const FStrataRegisteredSharedLocalBasis NewRegisteredSharedLocalBasis = StrataCompilationInfoCreateSharedLocalBasis(Compiler, NormalCodeChunk, TangentCodeChunk);
+
 	FStrataOperator& StrataOperator = Compiler->StrataCompilationGetOperator(this);
 	StrataOperator.BSDFRegisteredSharedLocalBasis = NewRegisteredSharedLocalBasis;
-	StrataOperator.bBSDFHasEdgeColor = bHasEdgeColor;
-	StrataOperator.bBSDFHasFuzz = bHasFuzz;
-	StrataOperator.bBSDFHasHaziness = bHasHaziness;
-	StrataOperator.bBSDFHasSSS = bHasSSS;
-	StrataOperator.bBSDFHasMFPPluggedIn = bHasMFPPluggedIn;
 
 	int32 OutputCodeChunk = Compiler->StrataSlabBSDF(
 		bUseMetalness ? Compiler->Constant(1.0f) : Compiler->Constant(0.0f),
@@ -21514,6 +21508,12 @@ FStrataOperator* UMaterialExpressionStrataSlabBSDF::StrataGenerateMaterialTopolo
 {
 	FStrataOperator& StrataOperator = Compiler->StrataCompilationRegisterOperator(STRATA_OPERATOR_BSDF, this, Parent);
 	StrataOperator.BSDFType = STRATA_BSDF_TYPE_SLAB;
+	StrataOperator.bBSDFHasEdgeColor = HasEdgeColor();
+	StrataOperator.bBSDFHasFuzz = HasFuzz();
+	StrataOperator.bBSDFHasHaziness = HasHaziness();
+	StrataOperator.bBSDFHasSSS = HasSSS();
+	StrataOperator.bBSDFHasMFPPluggedIn = HasMFPPluggedIn();
+	StrataOperator.bBSDFHasAnisotropy = HasAnisotropy();
 	return &StrataOperator;
 }
 
diff --git a/Engine/Source/Runtime/Engine/Public/MaterialCompiler.h b/Engine/Source/Runtime/Engine/Public/MaterialCompiler.h
--- a/Engine/Source/Runtime/Engine/Public/MaterialCompiler.h
+++ b/Engine/Source/Runtime/Engine/Public/MaterialCompiler.h
@@ -95,6 +95,7 @@ struct FStrataOperator
 	bool  bBSDFHasEdgeColor;
 	bool  bBSDFHasFuzz;
 	bool  bBSDFHasHaziness;
+	bool  bBSDFHasAnisotropy;
 
 	// Data derived after the tree has been built.
 	int32 MaxDistanceFromLeaves;
@@ -120,6 +121,7 @@ struct FStrataOperator
 		bBSDFHasEdgeColor = false;
 		bBSDFHasFuzz = false;
 		bBSDFHasHaziness = false;
+		bBSDFHasAnisotropy = false;
 
 		MaxDistanceFromLeaves = 0;
 		LayerDepth = 0;
@@ -136,6 +138,7 @@ struct FStrataOperator
 		bBSDFHasEdgeColor = A.bBSDFHasEdgeColor || B.bBSDFHasEdgeColor;
 		bBSDFHasFuzz = A.bBSDFHasFuzz || B.bBSDFHasFuzz;
 		bBSDFHasHaziness = A.bBSDFHasHaziness || B.bBSDFHasHaziness;
+		bBSDFHasAnisotropy = A.bBSDFHasAnisotropy || B.bBSDFHasAnisotropy;
 	}
 
 	void CopyFlagsForParameterBlending(FStrataOperator& A)
@@ -145,6 +148,7 @@ struct FStrataOperator
 		bBSDFHasEdgeColor = A.bBSDFHasEdgeColor;
 		bBSDFHasFuzz = A.bBSDFHasFuzz;
 		bBSDFHasHaziness = A.bBSDFHasHaziness;
+		bBSDFHasAnisotropy = A.bBSDFHasAnisotropy;
 	}
 
 	bool IsDiscarded()
