World Partition: fix dirtying the world when changing enable streaming value.

#rnx

diff --git a/Engine/Source/Runtime/Engine/Private/WorldPartition/WorldPartition.cpp b/Engine/Source/Runtime/Engine/Private/WorldPartition/WorldPartition.cpp
--- a/Engine/Source/Runtime/Engine/Private/WorldPartition/WorldPartition.cpp
+++ b/Engine/Source/Runtime/Engine/Private/WorldPartition/WorldPartition.cpp
@@ -1127,7 +1127,6 @@ void UWorldPartition::OnActorDescRegistered(const FWorldPartitionActorDesc& Acto
 	check(Actor);
 	ApplyActorTransform(Actor, InstanceTransform);
 	Actor->GetLevel()->AddLoadedActor(Actor);
-
 }
 
 void UWorldPartition::OnActorDescUnregistered(const FWorldPartitionActorDesc& ActorDesc) 
@@ -1165,6 +1164,7 @@ void UWorldPartition::SetEnableStreaming(bool bInEnableStreaming)
 {
 	if (bEnableStreaming != bInEnableStreaming)
 	{
+		Modify();
 		bEnableStreaming = bInEnableStreaming;
 		OnEnableStreamingChanged();
 	}
