LevelInstance: Extract ILevelInstanceInterface, FLevelInstanceActorImpl from ALevelInstance to allow LevelInstance implementation in existing AActor class hierarchies

#rnx

diff --git a/Engine/Source/Editor/LevelEditor/Private/LevelEditorSubsystem.cpp b/Engine/Source/Editor/LevelEditor/Private/LevelEditorSubsystem.cpp
--- a/Engine/Source/Editor/LevelEditor/Private/LevelEditorSubsystem.cpp
+++ b/Engine/Source/Editor/LevelEditor/Private/LevelEditorSubsystem.cpp
@@ -21,7 +21,7 @@
 #include "Styling/SlateIconFinder.h"
 #include "Subsystems/ActorEditorContextSubsystem.h"
 #include "LevelInstance/LevelInstanceSubsystem.h"
-#include "LevelInstance/LevelInstanceActor.h"
+#include "LevelInstance/LevelInstanceInterface.h"
 #include "ClassIconFinder.h"
 
 DEFINE_LOG_CATEGORY_STATIC(LevelEditorSubsystem, Log, All);
@@ -725,7 +725,7 @@ public:
 						SNew(STextBlock)
 						.Text_Lambda([this]() 
 						{
-							ALevelInstance* LevelInstance = GetEditingLevelInstance();
+							AActor* LevelInstance = GetEditingLevelInstance();
 							return LevelInstance ? FText::FromString(LevelInstance->GetActorLabel()) : FText::GetEmpty();
 						})
 						.ColorAndOpacity(FAppStyle::Get().GetSlateColor("Colors.AccentGreen"))
@@ -762,10 +762,10 @@ public:
 	}
 
 private:
-	ALevelInstance* GetEditingLevelInstance() const
+	AActor* GetEditingLevelInstance() const
 	{
 		ULevelInstanceSubsystem* LevelInstanceSubsystem = GetWorld()->GetSubsystem<ULevelInstanceSubsystem>();
-		return LevelInstanceSubsystem ? LevelInstanceSubsystem->GetEditingLevelInstance() : nullptr;
+		return LevelInstanceSubsystem ? Cast<AActor>(LevelInstanceSubsystem->GetEditingLevelInstance()) : nullptr;
 	}
 
 	FText GetCurrentLevelText() const
diff --git a/Engine/Source/Editor/LevelInstanceEditor/Private/LevelInstanceActorDetails.cpp b/Engine/Source/Editor/LevelInstanceEditor/Private/LevelInstanceActorDetails.cpp
--- a/Engine/Source/Editor/LevelInstanceEditor/Private/LevelInstanceActorDetails.cpp
+++ b/Engine/Source/Editor/LevelInstanceEditor/Private/LevelInstanceActorDetails.cpp
@@ -19,7 +19,7 @@ namespace LevelInstanceActorDetailsCallbacks
 	{
 		if (ALevelInstance* LevelInstanceActor = LevelInstanceActorPtr.Get())
 		{
-			return LevelInstanceActor->CanEdit() || LevelInstanceActor->CanCommit();
+			return LevelInstanceActor->CanEnterEdit() || LevelInstanceActor->CanExitEdit();
 		}
 
 		return false;
@@ -29,7 +29,7 @@ namespace LevelInstanceActorDetailsCallbacks
 	{
 		if (ALevelInstance* LevelInstanceActor = LevelInstanceActorPtr.Get())
 		{
-			if (LevelInstanceActor->CanCommit())
+			if (LevelInstanceActor->CanExitEdit())
 			{
 				return LOCTEXT("CommitChanges", "Commit Changes");
 			}
@@ -45,11 +45,11 @@ namespace LevelInstanceActorDetailsCallbacks
 		{
 			if (!LevelInstanceActor->IsEditing())
 			{
-				LevelInstanceActor->CanEdit(&Reason);
+				LevelInstanceActor->CanEnterEdit(&Reason);
 				return Reason;
 			}
 
-			LevelInstanceActor->CanCommit(&Reason);
+			LevelInstanceActor->CanExitEdit(/*bDiscardEdits=*/false, &Reason);
 		}
 		return Reason;
 	}
@@ -68,13 +68,13 @@ namespace LevelInstanceActorDetailsCallbacks
 	{
 		if (ALevelInstance* LevelInstanceActor = LevelInstanceActorPtr.Get())
 		{
-			if (LevelInstanceActor->CanCommit())
+			if (LevelInstanceActor->CanExitEdit())
 			{
-				LevelInstanceActor->Commit();
+				LevelInstanceActor->ExitEdit();
 			}
-			else if (LevelInstanceActor->CanEdit())
+			else if (LevelInstanceActor->CanEnterEdit())
 			{
-				LevelInstanceActor->Edit();
+				LevelInstanceActor->EnterEdit();
 			}
 		}
 		return FReply::Handled();
diff --git a/Engine/Source/Editor/LevelInstanceEditor/Private/LevelInstanceEditorMode.cpp b/Engine/Source/Editor/LevelInstanceEditor/Private/LevelInstanceEditorMode.cpp
--- a/Engine/Source/Editor/LevelInstanceEditor/Private/LevelInstanceEditorMode.cpp
+++ b/Engine/Source/Editor/LevelInstanceEditor/Private/LevelInstanceEditorMode.cpp
@@ -7,7 +7,7 @@
 #include "EditorModes.h"
 #include "Engine/World.h"
 #include "LevelInstance/LevelInstanceSubsystem.h"
-#include "LevelInstance/LevelInstanceActor.h"
+#include "LevelInstance/LevelInstanceInterface.h"
 #include "LevelInstance/ILevelInstanceEditorModule.h"
 #include "LevelEditorViewport.h"
 #include "LevelEditorActions.h"
@@ -108,9 +108,9 @@ bool ULevelInstanceEditorMode::IsSelectionDisallowed(AActor* InActor, bool bInSe
 	{
 		if (UWorld* World = InActor->GetWorld())
 		{
-			if (ALevelInstance* LevelInstanceActor = Cast<ALevelInstance>(InActor))
+			if (ILevelInstanceInterface* LevelInstance = Cast<ILevelInstanceInterface>(InActor))
 			{
-				if (LevelInstanceActor->IsEditing())
+				if (LevelInstance->IsEditing())
 				{
 					return false;
 				}
@@ -118,8 +118,8 @@ bool ULevelInstanceEditorMode::IsSelectionDisallowed(AActor* InActor, bool bInSe
 
 			if (ULevelInstanceSubsystem* LevelInstanceSubsystem = World->GetSubsystem<ULevelInstanceSubsystem>())
 			{
-				ALevelInstance* EditingLevelInstance = LevelInstanceSubsystem->GetEditingLevelInstance();
-				ALevelInstance* LevelInstance = LevelInstanceSubsystem->GetParentLevelInstance(InActor);
+				ILevelInstanceInterface* EditingLevelInstance = LevelInstanceSubsystem->GetEditingLevelInstance();
+				ILevelInstanceInterface* LevelInstance = LevelInstanceSubsystem->GetParentLevelInstance(InActor);
 				// Allow selection on actors that are part of the currently edited Level Instance hierarchy because AActor::GetRootSelectionParent() will eventually
 				// Bubble up the selection to its parent.
 				while (LevelInstance != nullptr)
@@ -129,9 +129,9 @@ bool ULevelInstanceEditorMode::IsSelectionDisallowed(AActor* InActor, bool bInSe
 						return false;
 					}
 
-					LevelInstance = LevelInstanceSubsystem->GetParentLevelInstance(LevelInstance);
+					LevelInstance = LevelInstanceSubsystem->GetParentLevelInstance(CastChecked<AActor>(LevelInstance));
 				}
-								
+
 				return EditingLevelInstance != nullptr;
 			}
 		}
diff --git a/Engine/Source/Editor/LevelInstanceEditor/Private/LevelInstanceEditorModule.cpp b/Engine/Source/Editor/LevelInstanceEditor/Private/LevelInstanceEditorModule.cpp
--- a/Engine/Source/Editor/LevelInstanceEditor/Private/LevelInstanceEditorModule.cpp
+++ b/Engine/Source/Editor/LevelInstanceEditor/Private/LevelInstanceEditorModule.cpp
@@ -3,6 +3,7 @@
 #include "LevelInstanceActorDetails.h"
 #include "LevelInstancePivotDetails.h"
 #include "LevelInstance/LevelInstanceSubsystem.h"
+#include "LevelInstance/LevelInstanceInterface.h"
 #include "LevelInstance/LevelInstanceActor.h"
 #include "PackedLevelActor/PackedLevelActor.h"
 #include "LevelInstanceEditorSettings.h"
@@ -45,11 +46,14 @@ IMPLEMENT_MODULE( FLevelInstanceEditorModule, LevelInstanceEditor );
 
 namespace LevelInstanceMenuUtils
 {
-	bool IsExperimentalSettingEnabled(ALevelInstance* LevelInstance)
+	bool IsExperimentalSettingEnabled(ILevelInstanceInterface* LevelInstance)
 	{
-		if (LevelInstance->IsA<APackedLevelActor>() && !GetDefault<UEditorExperimentalSettings>()->bPackedLevelActor)
+		if (APackedLevelActor* Actor = Cast<APackedLevelActor>(LevelInstance))
 		{
-			return false;
+			if (!GetDefault<UEditorExperimentalSettings>()->bPackedLevelActor)
+			{
+				return false;
+			}
 		}
 
 		return GetDefault<UEditorExperimentalSettings>()->bLevelInstance;
@@ -67,15 +71,16 @@ namespace LevelInstanceMenuUtils
 		return Section;
 	}
 
-	void CreateEditMenuEntry(FToolMenuSection& Section, ALevelInstance* LevelInstance, AActor* ContextActor, bool bSingleEntry)
+	void CreateEditMenuEntry(FToolMenuSection& Section, ILevelInstanceInterface* LevelInstance, AActor* ContextActor, bool bSingleEntry)
 	{
 		FToolUIAction LevelInstanceEditAction;
 		FText EntryDesc;
-		const bool bCanEdit = LevelInstance->CanEdit(&EntryDesc);
+		AActor* LevelInstanceActor = CastChecked<AActor>(LevelInstance);
+		const bool bCanEdit = LevelInstance->CanEnterEdit(&EntryDesc);
 
 		LevelInstanceEditAction.ExecuteAction.BindLambda([LevelInstance, ContextActor](const FToolMenuContext&)
 		{
-			LevelInstance->Edit(ContextActor);
+			LevelInstance->EnterEdit(ContextActor);
 		});
 		LevelInstanceEditAction.CanExecuteAction.BindLambda([bCanEdit](const FToolMenuContext&)
 		{
@@ -85,44 +90,41 @@ namespace LevelInstanceMenuUtils
 		FText EntryLabel = bSingleEntry ? LOCTEXT("EditLevelInstances", "Edit") : FText::FromString(LevelInstance->GetWorldAsset().GetAssetName());
 		if (bCanEdit)
 		{
-			EntryDesc = FText::Format(LOCTEXT("LevelInstanceName", "{0}:{1}"), FText::FromString(LevelInstance->GetActorLabel()), FText::FromString(LevelInstance->GetWorldAssetPackage()));
+			EntryDesc = FText::Format(LOCTEXT("LevelInstanceName", "{0}:{1}"), FText::FromString(LevelInstanceActor->GetActorLabel()), FText::FromString(LevelInstance->GetWorldAssetPackage()));
 		}
 		Section.AddMenuEntry(NAME_None, EntryLabel, EntryDesc, FSlateIcon(), LevelInstanceEditAction);
 	}
 
-	void CreateEditSubMenu(UToolMenu* Menu, TArray<ALevelInstance*> LevelInstanceHierarchy, AActor* ContextActor)
+	void CreateEditSubMenu(UToolMenu* Menu, TArray<ILevelInstanceInterface*> LevelInstanceHierarchy, AActor* ContextActor)
 	{
 		FToolMenuSection& Section = Menu->AddSection(NAME_None, LOCTEXT("LevelInstanceContextEditSection", "Context"));
-		for (ALevelInstance* LevelInstance : LevelInstanceHierarchy)
+		for (ILevelInstanceInterface* LevelInstance : LevelInstanceHierarchy)
 		{
 			CreateEditMenuEntry(Section, LevelInstance, ContextActor, false);
 		}
 	}
 		
-	void MoveSelectionToLevelInstance(ALevelInstance* DestinationLevelInstance)
+	void MoveSelectionToLevelInstance(ILevelInstanceInterface* DestinationLevelInstance)
 	{
-		if (ULevelInstanceSubsystem* LevelInstanceSubsystem = DestinationLevelInstance->GetLevelInstanceSubsystem())
+		TArray<AActor*> ActorsToMove;
+		ActorsToMove.Reserve(GEditor->GetSelectedActorCount());
+		for (FSelectionIterator It(GEditor->GetSelectedActorIterator()); It; ++It)
 		{
-			TArray<AActor*> ActorsToMove;
-			ActorsToMove.Reserve(GEditor->GetSelectedActorCount());
-			for (FSelectionIterator It(GEditor->GetSelectedActorIterator()); It; ++It)
+			if (AActor* Actor = Cast<AActor>(*It))
 			{
-				if (AActor* Actor = Cast<AActor>(*It))
-				{
-					ActorsToMove.Add(Actor);
-				}
+				ActorsToMove.Add(Actor);
 			}
-			
-			LevelInstanceSubsystem->MoveActorsTo(DestinationLevelInstance, ActorsToMove);
 		}
+
+		DestinationLevelInstance->MoveActorsTo(ActorsToMove);
 	}
 		
 	void CreateEditMenu(UToolMenu* Menu, AActor* ContextActor)
 	{
 		if (ULevelInstanceSubsystem* LevelInstanceSubsystem = ContextActor->GetWorld()->GetSubsystem<ULevelInstanceSubsystem>())
 		{
-			TArray<ALevelInstance*> LevelInstanceHierarchy;
-			LevelInstanceSubsystem->ForEachLevelInstanceAncestorsAndSelf(ContextActor, [&LevelInstanceHierarchy](ALevelInstance* AncestorLevelInstance)
+			TArray<ILevelInstanceInterface*> LevelInstanceHierarchy;
+			LevelInstanceSubsystem->ForEachLevelInstanceAncestorsAndSelf(ContextActor, [&LevelInstanceHierarchy](ILevelInstanceInterface* AncestorLevelInstance)
 			{
 				if (IsExperimentalSettingEnabled(AncestorLevelInstance))
 				{
@@ -152,7 +154,7 @@ namespace LevelInstanceMenuUtils
 		
 	void CreateCommitDiscardMenu(UToolMenu* Menu, AActor* ContextActor)
 	{
-		ALevelInstance* LevelInstanceEdit = nullptr;
+		ILevelInstanceInterface* LevelInstanceEdit = nullptr;
 		if (ContextActor)
 		{
 			if (ULevelInstanceSubsystem* LevelInstanceSubsystem = ContextActor->GetWorld()->GetSubsystem<ULevelInstanceSubsystem>())
@@ -174,18 +176,18 @@ namespace LevelInstanceMenuUtils
 			FToolMenuSection& Section = CreateLevelSection(Menu);
 
 			FText CommitTooltip;
-			const bool bCanCommit = LevelInstanceEdit->CanCommit(&CommitTooltip);
+			const bool bCanCommit = LevelInstanceEdit->CanExitEdit(/*bDiscardEdits=*/false, &CommitTooltip);
 
 			FToolUIAction CommitAction;
-			CommitAction.ExecuteAction.BindLambda([LevelInstanceEdit](const FToolMenuContext&) { LevelInstanceEdit->Commit(); });
+			CommitAction.ExecuteAction.BindLambda([LevelInstanceEdit](const FToolMenuContext&) { LevelInstanceEdit->ExitEdit(/*bDiscardEdits=*/false); });
 			CommitAction.CanExecuteAction.BindLambda([bCanCommit](const FToolMenuContext&) { return bCanCommit; });
 			Section.AddMenuEntry(NAME_None, LOCTEXT("LevelInstanceCommitLabel", "Commit"), CommitTooltip, FSlateIcon(), CommitAction);
 
 			FText DiscardTooltip;
-			const bool bCanDiscard = LevelInstanceEdit->CanDiscard(&DiscardTooltip);
+			const bool bCanDiscard = LevelInstanceEdit->CanExitEdit(/*bDiscardEdits=*/true, &DiscardTooltip);
 
 			FToolUIAction DiscardAction;
-			DiscardAction.ExecuteAction.BindLambda([LevelInstanceEdit](const FToolMenuContext&) { LevelInstanceEdit->Discard(); });
+			DiscardAction.ExecuteAction.BindLambda([LevelInstanceEdit](const FToolMenuContext&) { LevelInstanceEdit->ExitEdit(/*bDiscardEdits=*/true); });
 			DiscardAction.CanExecuteAction.BindLambda([bCanDiscard](const FToolMenuContext&) { return bCanDiscard; });
 			Section.AddMenuEntry(NAME_None, LOCTEXT("LevelInstanceDiscardLabel", "Discard"), DiscardTooltip, FSlateIcon(), DiscardAction);
 		}
@@ -193,7 +195,7 @@ namespace LevelInstanceMenuUtils
 
 	void CreateSetCurrentMenu(UToolMenu* Menu, AActor* ContextActor)
 	{
-		ALevelInstance* LevelInstanceEdit = nullptr;
+		ILevelInstanceInterface* LevelInstanceEdit = nullptr;
 		if (ULevelInstanceSubsystem* LevelInstanceSubsystem = ContextActor->GetWorld()->GetSubsystem<ULevelInstanceSubsystem>())
 		{
 			LevelInstanceEdit = LevelInstanceSubsystem->GetEditingLevelInstance();
@@ -202,9 +204,9 @@ namespace LevelInstanceMenuUtils
 			{
 				FToolUIAction LevelInstanceSetCurrentAction;
 				LevelInstanceSetCurrentAction.ExecuteAction.BindLambda([LevelInstanceEdit](const FToolMenuContext&)
-					{
-						LevelInstanceEdit->SetCurrent();
-					});
+				{
+					LevelInstanceEdit->SetCurrent();
+				});
 
 				FToolMenuSection& Section = CreateLevelSection(Menu);
 				Section.AddMenuEntry(NAME_None, LOCTEXT("LevelInstanceSetCurrent", "Set Current Level"), TAttribute<FText>(), FSlateIcon(), LevelInstanceSetCurrentAction);
@@ -216,7 +218,7 @@ namespace LevelInstanceMenuUtils
 	{
 		if (GEditor->GetSelectedActorCount() > 0)
 		{
-			ALevelInstance* LevelInstanceEdit = nullptr;
+			ILevelInstanceInterface* LevelInstanceEdit = nullptr;
 			ULevelInstanceSubsystem* LevelInstanceSubsystem = GEditor->GetEditorWorldContext().World()->GetSubsystem<ULevelInstanceSubsystem>();
 			if (LevelInstanceSubsystem)
 			{
@@ -345,13 +347,13 @@ namespace LevelInstanceMenuUtils
 		}
 	}
 		
-	void CreateBreakSubMenu(UToolMenu* Menu, ALevelInstance* ContextLevelInstance)
+	void CreateBreakSubMenu(UToolMenu* Menu, ILevelInstanceInterface* ContextLevelInstance)
 	{
 		static int32 BreakLevels = 1;
 
 		check(ContextLevelInstance);
 
-		if (ULevelInstanceSubsystem* LevelInstanceSubsystem = ContextLevelInstance->GetWorld()->GetSubsystem<ULevelInstanceSubsystem>())
+		if (ULevelInstanceSubsystem* LevelInstanceSubsystem = ContextLevelInstance->GetLevelInstanceSubsystem())
 		{
 			FToolMenuSection& Section = Menu->AddSection(NAME_None, LOCTEXT("LevelInstanceBreakSection", "Break Level Instance"));
 			TSharedRef<SWidget> MenuWidget =
@@ -404,20 +406,20 @@ namespace LevelInstanceMenuUtils
 
 		if (ULevelInstanceSubsystem* LevelInstanceSubsystem = ContextActor->GetWorld()->GetSubsystem<ULevelInstanceSubsystem>())
 		{
-			ALevelInstance* ContextLevelInstance = nullptr;
+			ILevelInstanceInterface* ContextLevelInstance = nullptr;
 
 			// Find the top level LevelInstance
-			LevelInstanceSubsystem->ForEachLevelInstanceAncestorsAndSelf(ContextActor, [LevelInstanceSubsystem, ContextActor, &ContextLevelInstance](ALevelInstance* Ancestor)
+			LevelInstanceSubsystem->ForEachLevelInstanceAncestorsAndSelf(ContextActor, [LevelInstanceSubsystem, ContextActor, &ContextLevelInstance](ILevelInstanceInterface* Ancestor)
+			{
+				if (CastChecked<AActor>(Ancestor)->GetLevel() == ContextActor->GetWorld()->GetCurrentLevel())
 				{
-					if (Ancestor->GetLevel() == ContextActor->GetWorld()->GetCurrentLevel())
-					{
-						ContextLevelInstance = Ancestor;
-						return false;
-					}
-					return true;
-				});
+					ContextLevelInstance = Ancestor;
+					return false;
+				}
+				return true;
+			});
 
-			if (ContextLevelInstance && IsExperimentalSettingEnabled(ContextLevelInstance) && !ContextLevelInstance->IsEditing() && !LevelInstanceSubsystem->LevelInstanceHasLevelScriptBlueprint(ContextLevelInstance))
+			if (ContextLevelInstance && IsExperimentalSettingEnabled(ContextLevelInstance) && !LevelInstanceSubsystem->IsEditingLevelInstance(ContextLevelInstance) && !LevelInstanceSubsystem->LevelInstanceHasLevelScriptBlueprint(ContextLevelInstance))
 			{
 				FToolMenuSection& Section = CreateLevelSection(Menu);
 
@@ -435,12 +437,12 @@ namespace LevelInstanceMenuUtils
 	{
 		if (ULevelInstanceSubsystem* LevelInstanceSubsystem = ContextActor->GetWorld()->GetSubsystem<ULevelInstanceSubsystem>())
 		{
-			ALevelInstance* ContextLevelInstance = nullptr;
+			ILevelInstanceInterface* ContextLevelInstance = nullptr;
 
 			// Find the top level LevelInstance
-			LevelInstanceSubsystem->ForEachLevelInstanceAncestorsAndSelf(ContextActor, [LevelInstanceSubsystem, ContextActor, &ContextLevelInstance](ALevelInstance* Ancestor)
+			LevelInstanceSubsystem->ForEachLevelInstanceAncestorsAndSelf(ContextActor, [LevelInstanceSubsystem, ContextActor, &ContextLevelInstance](ILevelInstanceInterface* Ancestor)
 			{
-				if (Ancestor->GetLevel() == ContextActor->GetWorld()->GetCurrentLevel())
+				if (CastChecked<AActor>(Ancestor)->GetLevel() == ContextActor->GetWorld()->GetCurrentLevel())
 				{
 					ContextLevelInstance = Ancestor;
 					return false;
@@ -484,12 +486,12 @@ namespace LevelInstanceMenuUtils
 		
 		virtual bool IsClassAllowed(const FClassViewerInitializationOptions& InInitOptions, const UClass* InClass, TSharedRef< FClassViewerFilterFuncs > InFilterFuncs) override
 		{
-			return InClass && InClass->IsChildOf(ALevelInstance::StaticClass()) && !InClass->IsChildOf(APackedLevelActor::StaticClass()) && !InClass->HasAnyClassFlags(CLASS_Deprecated);
+			return InClass && InClass->ImplementsInterface(ULevelInstanceInterface::StaticClass()) && !InClass->IsChildOf(APackedLevelActor::StaticClass()) && !InClass->HasAnyClassFlags(CLASS_Deprecated);
 		}
 
 		virtual bool IsUnloadedClassAllowed(const FClassViewerInitializationOptions& InInitOptions, const TSharedRef< const IUnloadedBlueprintData > InUnloadedClassData, TSharedRef< FClassViewerFilterFuncs > InFilterFuncs) override
 		{
-			return InUnloadedClassData->IsChildOf(ALevelInstance::StaticClass()) && !InUnloadedClassData->IsChildOf(APackedLevelActor::StaticClass())  && !InUnloadedClassData->HasAnyClassFlags(CLASS_Deprecated);
+			return InUnloadedClassData->ImplementsInterface(ULevelInstanceInterface::StaticClass()) && !InUnloadedClassData->IsChildOf(APackedLevelActor::StaticClass())  && !InUnloadedClassData->HasAnyClassFlags(CLASS_Deprecated);
 		}
 	};
 
@@ -523,8 +525,9 @@ namespace LevelInstanceMenuUtils
 
 		if (UBlueprint* NewBlueprint = Cast<UBlueprint>(AssetTools.CreateAssetWithDialog(AssetName, PackagePath, UBlueprint::StaticClass(), BlueprintFactory, FName("Create LevelInstance Blueprint"))))
 		{
-			ALevelInstance* CDO = CastChecked<ALevelInstance>(NewBlueprint->GeneratedClass->GetDefaultObject());
-			CDO->SetWorldAsset(LevelInstancePtr);
+			AActor* CDO = NewBlueprint->GeneratedClass->GetDefaultObject<AActor>();
+			ILevelInstanceInterface* LevelInstanceCDO = CastChecked<ILevelInstanceInterface>(CDO);
+			LevelInstanceCDO->SetWorldAsset(LevelInstancePtr);
 			FBlueprintEditorUtils::MarkBlueprintAsModified(NewBlueprint);
 			
 			FContentBrowserModule& ContentBrowserModule = FModuleManager::LoadModuleChecked<FContentBrowserModule>("ContentBrowser");
diff --git a/Engine/Source/Editor/LevelInstanceEditor/Private/LevelInstancePivotDetails.cpp b/Engine/Source/Editor/LevelInstanceEditor/Private/LevelInstancePivotDetails.cpp
--- a/Engine/Source/Editor/LevelInstanceEditor/Private/LevelInstancePivotDetails.cpp
+++ b/Engine/Source/Editor/LevelInstanceEditor/Private/LevelInstancePivotDetails.cpp
@@ -10,7 +10,7 @@
 #include "ScopedTransaction.h"
 #include "Engine/Brush.h"
 #include "Engine/World.h"
-#include "LevelInstance/LevelInstanceActor.h"
+#include "LevelInstance/LevelInstanceInterface.h"
 #include "LevelInstance/LevelInstanceEditorPivotActor.h"
 
 #define LOCTEXT_NAMESPACE "FLevelInstancePivotDetails"
diff --git a/Engine/Source/Editor/SceneOutliner/Private/ActorBrowsingMode.cpp b/Engine/Source/Editor/SceneOutliner/Private/ActorBrowsingMode.cpp
--- a/Engine/Source/Editor/SceneOutliner/Private/ActorBrowsingMode.cpp
+++ b/Engine/Source/Editor/SceneOutliner/Private/ActorBrowsingMode.cpp
@@ -31,7 +31,7 @@
 #include "ActorBrowsingModeSettings.h"
 #include "ActorDescTreeItem.h"
 #include "ScopedTransaction.h"
-#include "LevelInstance/LevelInstanceActor.h"
+#include "LevelInstance/LevelInstanceInterface.h"
 #include "LevelInstance/LevelInstanceSubsystem.h"
 #include "LevelInstance/LevelInstanceEditorInstanceActor.h"
 #include "EditorLevelUtils.h"
@@ -152,7 +152,7 @@ FActorBrowsingMode::FActorBrowsingMode(SSceneOutliner* InSceneOutliner, TWeakObj
 				{
 					// if actor has a valid parent and the parent is not being edited,
 					// then the actor should not be selectable.
-					if (const ALevelInstance* ParentLevelInstance = LevelInstanceSubsystem->GetParentLevelInstance(Actor))
+					if (const ILevelInstanceInterface* ParentLevelInstance = LevelInstanceSubsystem->GetParentLevelInstance(Actor))
 					{
 						if (!LevelInstanceSubsystem->IsEditingLevelInstance(ParentLevelInstance))
 						{
@@ -337,7 +337,7 @@ TSharedRef<FSceneOutlinerFilter> FActorBrowsingMode::CreateHideLevelInstancesFil
 			// Check if actor belongs to a LevelInstance
 			if (const ULevelInstanceSubsystem* LevelInstanceSubsystem = Actor->GetWorld()->GetSubsystem<ULevelInstanceSubsystem>())
 			{
-				if (const ALevelInstance* ParentLevelInstance = LevelInstanceSubsystem->GetParentLevelInstance(Actor))
+				if (const ILevelInstanceInterface* ParentLevelInstance = LevelInstanceSubsystem->GetParentLevelInstance(Actor))
 				{
 					if (!LevelInstanceSubsystem->IsEditingLevelInstance(ParentLevelInstance))
 					{
@@ -765,16 +765,16 @@ void FActorBrowsingMode::OnItemDoubleClick(FSceneOutlinerTreeItemPtr Item)
 		AActor* Actor = ActorItem->Actor.Get();
 		check(Actor);
 
-		ALevelInstance* LevelInstanceActor = Cast<ALevelInstance>(Actor);
-		if (LevelInstanceActor && FSlateApplication::Get().GetModifierKeys().IsAltDown())
+		ILevelInstanceInterface* LevelInstance = Cast<ILevelInstanceInterface>(Actor);
+		if (LevelInstance && FSlateApplication::Get().GetModifierKeys().IsAltDown())
 		{
-			if (LevelInstanceActor->CanEdit())
+			if (LevelInstance->CanEnterEdit())
 			{
-				LevelInstanceActor->Edit();
+				LevelInstance->EnterEdit();
 			}
-			else if (LevelInstanceActor->CanCommit())
+			else if (LevelInstance->CanExitEdit())
 			{
-				LevelInstanceActor->Commit();
+				LevelInstance->ExitEdit();
 			}
 		}
 		else if (Item->CanInteract())
@@ -1054,7 +1054,7 @@ FSceneOutlinerDragValidationInfo FActorBrowsingMode::ValidateDrop(const ISceneOu
 			return FSceneOutlinerDragValidationInfo(ESceneOutlinerDropCompatibility::IncompatibleGeneric, FText());
 		}
 
-		const ALevelInstance* LevelInstanceTarget = Cast<ALevelInstance>(ActorTarget);
+		const ILevelInstanceInterface* LevelInstanceTarget = Cast<ILevelInstanceInterface>(ActorTarget);
 		const ULevelInstanceSubsystem* LevelInstanceSubsystem = RepresentingWorld->GetSubsystem<ULevelInstanceSubsystem>();
 		if (LevelInstanceTarget)
 		{
@@ -1091,7 +1091,7 @@ FSceneOutlinerDragValidationInfo FActorBrowsingMode::ValidateDrop(const ISceneOu
 					if (LevelInstanceSubsystem)
 					{
 						// Either all actors must be in a LevelInstance or none of them
-						if (const ALevelInstance* ParentLevelInstance = LevelInstanceSubsystem->GetParentLevelInstance(DragActor))
+						if (const ILevelInstanceInterface* ParentLevelInstance = LevelInstanceSubsystem->GetParentLevelInstance(DragActor))
 						{
 							if (!ParentLevelInstance->IsEditing())
 							{
@@ -1169,7 +1169,7 @@ FSceneOutlinerDragValidationInfo FActorBrowsingMode::ValidateDrop(const ISceneOu
 		// WorldTreeItem and LevelTreeItem are treated as root folders (path = none), with the difference that LevelTreeItem has a RootObject.
 		const FFolder DestinationPath = FolderItem ? FolderItem->GetFolder() : (LevelItem ? FFolder(FFolder::GetEmptyPath(), FFolder::GetOptionalFolderRootObject(LevelItem->Level.Get()).Get(FFolder::GetDefaultRootObject())) : GWorldRoot);
 		const FFolder::FRootObject& DestinationRootObject = DestinationPath.GetRootObject();
-		ALevelInstance* LevelInstanceTarget = Cast<ALevelInstance>(DestinationPath.GetRootObjectPtr());
+		ILevelInstanceInterface* LevelInstanceTarget = Cast<ILevelInstanceInterface>(DestinationPath.GetRootObjectPtr());
 		if (LevelInstanceTarget && !LevelInstanceTarget->IsEditing())
 		{
 			return FSceneOutlinerDragValidationInfo(ESceneOutlinerDropCompatibility::IncompatibleGeneric, LOCTEXT("Error_DragInNonEditingLevelInstance", "Cannot drag into a LevelInstance which is not being edited"));
@@ -1251,7 +1251,7 @@ FSceneOutlinerDragValidationInfo FActorBrowsingMode::ValidateDrop(const ISceneOu
 				bool bActorContainedInLevelInstance = false;
 				if (LevelInstanceSubsystem)
 				{
-					if (const ALevelInstance* ParentLevelInstance = LevelInstanceSubsystem->GetParentLevelInstance(Actor))
+					if (const ILevelInstanceInterface* ParentLevelInstance = LevelInstanceSubsystem->GetParentLevelInstance(Actor))
 					{
 						if (!ParentLevelInstance->IsEditing())
 						{
@@ -1260,10 +1260,10 @@ FSceneOutlinerDragValidationInfo FActorBrowsingMode::ValidateDrop(const ISceneOu
 						bActorContainedInLevelInstance = true;
 					}
 
-					if (const ALevelInstance* LevelInstanceActor = Cast<ALevelInstance>(Actor))
+					if (const ILevelInstanceInterface* LevelInstance = Cast<ILevelInstanceInterface>(Actor))
 					{
 						FText Reason;
-						if (!LevelInstanceSubsystem->CanMoveActorToLevel(LevelInstanceActor, &Reason))
+						if (!LevelInstanceSubsystem->CanMoveActorToLevel(Actor, &Reason))
 						{
 							return FSceneOutlinerDragValidationInfo(ESceneOutlinerDropCompatibility::IncompatibleGeneric, Reason);
 						}
@@ -1362,12 +1362,12 @@ void FActorBrowsingMode::OnDrop(ISceneOutlinerTreeItem& DropTarget, const FScene
 		{
 			// Show socket chooser if we have sockets to select
 
-			if (ALevelInstance* TargetLevelInstance = Cast<ALevelInstance>(DropActor))
+			if (ILevelInstanceInterface* TargetLevelInstance = Cast<ILevelInstanceInterface>(DropActor))
 			{
 				check(TargetLevelInstance->IsEditing());
 				const FScopedTransaction Transaction(LOCTEXT("UndoAction_MoveActorsToLevelInstance", "Move actors to LevelInstance"));
 
-				const FFolder DestinationPath(FFolder::GetEmptyPath(), FFolder::FRootObject(TargetLevelInstance));
+				const FFolder DestinationPath(FFolder::GetEmptyPath(), FFolder::FRootObject(DropActor));
 				auto MoveToDestination = [&DestinationPath](FFolderTreeItem& Item)
 				{
 					Item.MoveTo(DestinationPath);
@@ -1534,7 +1534,7 @@ void FActorBrowsingMode::OnDrop(ISceneOutlinerTreeItem& DropTarget, const FScene
 				Payload.ForEachItem<FActorTreeItem>([LevelInstanceSubsystem, &LevelInstanceActorsToMove, &ActorsToMoveToPersistentLevel](const FActorTreeItem& ActorItem)
 				{
 					AActor* Actor = ActorItem.Actor.Get();
-					if (const ALevelInstance* ParentLevelInstance = LevelInstanceSubsystem->GetParentLevelInstance(Actor))
+					if (const ILevelInstanceInterface* ParentLevelInstance = LevelInstanceSubsystem->GetParentLevelInstance(Actor))
 					{
 						check(ParentLevelInstance->IsEditing());
 						LevelInstanceActorsToMove.Add(Actor);
@@ -1563,7 +1563,7 @@ void FActorBrowsingMode::OnDrop(ISceneOutlinerTreeItem& DropTarget, const FScene
 			}
 			else if (MovingActorsToValidRootObject.Num())
 			{
-				if (ALevelInstance* TargetLevelInstance = Cast<ALevelInstance>(DestinationPath.GetRootObjectPtr()))
+				if (ILevelInstanceInterface* TargetLevelInstance = Cast<ILevelInstanceInterface>(DestinationPath.GetRootObjectPtr()))
 				{
 					// We are moving actors inside an editing level instance
 					check(TargetLevelInstance->IsEditing());
diff --git a/Engine/Source/Editor/SceneOutliner/Private/ActorFolderHierarchy.cpp b/Engine/Source/Editor/SceneOutliner/Private/ActorFolderHierarchy.cpp
--- a/Engine/Source/Editor/SceneOutliner/Private/ActorFolderHierarchy.cpp
+++ b/Engine/Source/Editor/SceneOutliner/Private/ActorFolderHierarchy.cpp
@@ -8,7 +8,7 @@
 #include "ActorFolderTreeItem.h"
 #include "EditorActorFolders.h"
 #include "EditorFolderUtils.h"
-#include "LevelInstance/LevelInstanceActor.h"
+#include "LevelInstance/LevelInstanceInterface.h"
 
 FActorFolderHierarchy::FActorFolderHierarchy(ISceneOutlinerMode* InMode, const TWeakObjectPtr<UWorld>& World, const FFolder::FRootObject& InRootObject)
 	: ISceneOutlinerHierarchy(InMode)
@@ -44,9 +44,9 @@ FSceneOutlinerTreeItemPtr FActorFolderHierarchy::FindOrCreateParentItem(const IS
 					if (Folder.GetRootObject() == RootObject)
 					{
 						// If item belongs to a LevelInstance
-						if (ALevelInstance* RootLevelInstance = Cast<ALevelInstance>(Folder.GetRootObjectPtr()))
+						if (ILevelInstanceInterface* RootLevelInstance = Cast<ILevelInstanceInterface>(Folder.GetRootObjectPtr()))
 						{
-							return Mode->CreateItemFor<FActorTreeItem>(RootLevelInstance, true);
+							return Mode->CreateItemFor<FActorTreeItem>(CastChecked<AActor>(RootLevelInstance), true);
 						}
 					}
 				}
@@ -93,9 +93,9 @@ void FActorFolderHierarchy::CreateWorldChildren(UWorld* World, TArray<FSceneOutl
 	if (FFolder::HasRootObject(RootObject))
 	{
 		UObject* RootObjectPtr = FFolder::GetRootObjectPtr(RootObject);
-		if (ALevelInstance* RootLevelInstance = Cast<ALevelInstance>(RootObjectPtr))
+		if (ILevelInstanceInterface* RootLevelInstance = Cast<ILevelInstanceInterface>(RootObjectPtr))
 		{
-			if (FSceneOutlinerTreeItemPtr ActorItem = Mode->CreateItemFor<FActorTreeItem>(RootLevelInstance, true))
+			if (FSceneOutlinerTreeItemPtr ActorItem = Mode->CreateItemFor<FActorTreeItem>(CastChecked<AActor>(RootLevelInstance), true))
 			{
 				OutItems.Add(ActorItem);
 			}
@@ -151,11 +151,11 @@ void FActorFolderHierarchy::CreateChildren(const FSceneOutlinerTreeItemPtr& Item
 	else if (FActorTreeItem* ParentActorItem = Item->CastTo<FActorTreeItem>())
 	{
 		AActor* ParentActor = ParentActorItem->Actor.Get();
-		if (const ALevelInstance* LevelInstanceParentActor = Cast<ALevelInstance>(ParentActor))
+		if (const ILevelInstanceInterface* LevelInstanceParent = Cast<ILevelInstanceInterface>(ParentActor))
 		{
 			check(ParentActor->GetWorld() == World);
-			FFolder ParentFolder = LevelInstanceParentActor->GetFolder();
-			CreateChildrenFolders(World, ParentFolder, LevelInstanceParentActor, OutChildren);
+			FFolder ParentFolder = ParentActor->GetFolder();
+			CreateChildrenFolders(World, ParentFolder, ParentActor, OutChildren);
 		}
 	}
 	else if (FActorFolderTreeItem* FolderItem = Item->CastTo<FActorFolderTreeItem>())
diff --git a/Engine/Source/Editor/SceneOutliner/Private/ActorFolderTreeItem.cpp b/Engine/Source/Editor/SceneOutliner/Private/ActorFolderTreeItem.cpp
--- a/Engine/Source/Editor/SceneOutliner/Private/ActorFolderTreeItem.cpp
+++ b/Engine/Source/Editor/SceneOutliner/Private/ActorFolderTreeItem.cpp
@@ -16,7 +16,7 @@
 #include "WorldPersistentFolders.h"
 #include "Engine/Level.h"
 #include "LevelUtils.h"
-#include "LevelInstance/LevelInstanceActor.h"
+#include "LevelInstance/LevelInstanceInterface.h"
 #include "LevelInstance/LevelInstanceSubsystem.h"
 
 #define LOCTEXT_NAMESPACE "SceneOutliner_ActorFolderTreeItem"
@@ -337,9 +337,9 @@ bool FActorFolderTreeItem::CanInteract() const
 
 	if (ULevelInstanceSubsystem* LevelInstanceSubsystem = World.IsValid() ? World->GetSubsystem<ULevelInstanceSubsystem>() : nullptr)
 	{
-		if (ALevelInstance* EditingLevelInstance = LevelInstanceSubsystem->GetEditingLevelInstance())
+		if (ILevelInstanceInterface* EditingLevelInstance = LevelInstanceSubsystem->GetEditingLevelInstance())
 		{
-			if (GetRootObject() != FFolder::FRootObject(EditingLevelInstance))
+			if (GetRootObject() != FFolder::FRootObject(CastChecked<AActor>(EditingLevelInstance)))
 			{
 				return false;
 			}
diff --git a/Engine/Source/Editor/SceneOutliner/Private/ActorHierarchy.cpp b/Engine/Source/Editor/SceneOutliner/Private/ActorHierarchy.cpp
--- a/Engine/Source/Editor/SceneOutliner/Private/ActorHierarchy.cpp
+++ b/Engine/Source/Editor/SceneOutliner/Private/ActorHierarchy.cpp
@@ -17,7 +17,7 @@
 #include "GameFramework/WorldSettings.h"
 #include "EditorActorFolders.h"
 #include "EditorFolderUtils.h"
-#include "LevelInstance/LevelInstanceActor.h"
+#include "LevelInstance/LevelInstanceInterface.h"
 #include "LevelInstance/LevelInstanceSubsystem.h"
 #include "Modules/ModuleManager.h"
 #include "WorldPartition/WorldPartitionHelpers.h"
@@ -167,20 +167,21 @@ FSceneOutlinerTreeItemPtr FActorHierarchy::FindOrCreateParentItem(const ISceneOu
 			}
 
 			// Parent Level Instance
-			if (ALevelInstance* OwningLevelInstance = Cast<ALevelInstance>(Folder.GetRootObjectPtr()))
+			if (ILevelInstanceInterface* OwningLevelInstance = Cast<ILevelInstanceInterface>(Folder.GetRootObjectPtr()))
 			{
-				const ALevelInstance* LevelInstanceActor = Cast<ALevelInstance>(Actor);
-				const bool bIsAnEditingLevelInstance = LevelInstanceActor ? LevelInstanceActor->IsEditing() : false;
+				const ILevelInstanceInterface* LevelInstance = Cast<ILevelInstanceInterface>(Actor);
+				const bool bIsAnEditingLevelInstance = LevelInstance ? LevelInstance->IsEditing() : false;
 				// Parent this to a LevelInstance if the parent LevelInstance is being edited or if this is a sub LevelInstance which is being edited
 				if (bShowingLevelInstances || (OwningLevelInstance->IsEditing() || bIsAnEditingLevelInstance))
 				{
-					if (const FSceneOutlinerTreeItemPtr* ParentItem = Items.Find(OwningLevelInstance))
+					AActor* OwningActor = CastChecked<AActor>(OwningLevelInstance);
+					if (const FSceneOutlinerTreeItemPtr* ParentItem = Items.Find(OwningActor))
 					{
 						return *ParentItem;
 					}
 					else
 					{
-						return bCreate ? Mode->CreateItemFor<FActorTreeItem>(OwningLevelInstance, true) : nullptr;
+						return bCreate ? Mode->CreateItemFor<FActorTreeItem>(OwningActor, true) : nullptr;
 					}
 				}
 			}
@@ -219,17 +220,18 @@ FSceneOutlinerTreeItemPtr FActorHierarchy::FindOrCreateParentItem(const ISceneOu
 			}
 		}
 		// Parent Level Instance
-		else if (ALevelInstance* OwningLevelInstance = Cast<ALevelInstance>(ParentPath.GetRootObjectPtr()))
+		else if (ILevelInstanceInterface* OwningLevelInstance = Cast<ILevelInstanceInterface>(ParentPath.GetRootObjectPtr()))
 		{
 			if (bShowingLevelInstances || OwningLevelInstance->IsEditing())
 			{
-				if (const FSceneOutlinerTreeItemPtr* ParentItem = Items.Find(OwningLevelInstance))
+				AActor* OwningActor = CastChecked<AActor>(OwningLevelInstance);
+				if (const FSceneOutlinerTreeItemPtr* ParentItem = Items.Find(OwningActor))
 				{
 					return *ParentItem;
 				}
 				else
 				{
-					return bCreate ? Mode->CreateItemFor<FActorTreeItem>(OwningLevelInstance, true) : nullptr;
+					return bCreate ? Mode->CreateItemFor<FActorTreeItem>(OwningActor, true) : nullptr;
 				}
 			}
 		}
@@ -371,7 +373,7 @@ void FActorHierarchy::CreateWorldChildren(UWorld* World, TArray<FSceneOutlinerTr
 		// If we are not showing LevelInstances, LevelInstance sub actor items should not be created unless they belong to a LevelInstance which is being edited
 		if (LevelInstanceSubsystem)
 		{
-			if (const ALevelInstance* ParentLevelInstance = LevelInstanceSubsystem->GetParentLevelInstance(Actor))
+			if (const ILevelInstanceInterface* ParentLevelInstance = LevelInstanceSubsystem->GetParentLevelInstance(Actor))
 			{
 				if (!bShowingLevelInstances && !ParentLevelInstance->IsEditing())
 				{
@@ -472,25 +474,25 @@ void FActorHierarchy::CreateChildren(const FSceneOutlinerTreeItemPtr& Item, TArr
 
 		TArray<AActor*> ChildActors;
 
-		if (const ALevelInstance* LevelInstanceParentActor = Cast<ALevelInstance>(ParentActor))
+		if (const ILevelInstanceInterface* LevelInstanceParent = Cast<ILevelInstanceInterface>(ParentActor))
 		{
 			const ULevelInstanceSubsystem* LevelInstanceSubsystem = RepresentingWorld->GetSubsystem<ULevelInstanceSubsystem>();
 			check(LevelInstanceSubsystem);
 
-			LevelInstanceSubsystem->ForEachActorInLevelInstance(LevelInstanceParentActor, [this, LevelInstanceParentActor, LevelInstanceSubsystem, &ChildActors](AActor* SubActor)
+			LevelInstanceSubsystem->ForEachActorInLevelInstance(LevelInstanceParent, [this, LevelInstanceParent, LevelInstanceSubsystem, &ChildActors](AActor* SubActor)
 			{
-				const ALevelInstance* LevelInstanceActor = Cast<ALevelInstance>(SubActor);
-				const bool bIsAnEditingLevelInstance = LevelInstanceActor ? LevelInstanceSubsystem->IsEditingLevelInstance(LevelInstanceActor) : false;
-				if (bShowingLevelInstances || (LevelInstanceSubsystem->IsEditingLevelInstance(LevelInstanceParentActor) || bIsAnEditingLevelInstance))
+				const ILevelInstanceInterface* LevelInstance = Cast<ILevelInstanceInterface>(SubActor);
+				const bool bIsAnEditingLevelInstance = LevelInstance ? LevelInstanceSubsystem->IsEditingLevelInstance(LevelInstance) : false;
+				if (bShowingLevelInstances || (LevelInstanceSubsystem->IsEditingLevelInstance(LevelInstanceParent) || bIsAnEditingLevelInstance))
 				{
 					ChildActors.Add(SubActor);
 				}
 				return true;
 			});
 
-			check(World == LevelInstanceParentActor->GetWorld());
-			FFolder ParentFolder = LevelInstanceParentActor->GetFolder();
-			CreateChildrenFolders(World, ParentFolder, LevelInstanceParentActor, OutChildren);
+			check(World == CastChecked<AActor>(LevelInstanceParent)->GetWorld());
+			FFolder ParentFolder = ParentActor->GetFolder();
+			CreateChildrenFolders(World, ParentFolder, ParentActor, OutChildren);
 		}
 		else
 		{
diff --git a/Engine/Source/Editor/SceneOutliner/Private/ActorMode.cpp b/Engine/Source/Editor/SceneOutliner/Private/ActorMode.cpp
--- a/Engine/Source/Editor/SceneOutliner/Private/ActorMode.cpp
+++ b/Engine/Source/Editor/SceneOutliner/Private/ActorMode.cpp
@@ -16,7 +16,7 @@
 #include "ComponentTreeItem.h"
 #include "ActorDescTreeItem.h"
 #include "WorldTreeItem.h"
-#include "LevelInstance/LevelInstanceActor.h"
+#include "LevelInstance/LevelInstanceInterface.h"
 #include "LevelInstance/LevelInstanceSubsystem.h"
 #include "LevelInstance/LevelInstanceEditorInstanceActor.h"
 
@@ -97,7 +97,7 @@ FActorMode::FActorMode(const FActorModeParams& Params)
 		{
 			return true;
 		}
-		if (ALevelInstance* LevelInstance = Cast<ALevelInstance>(InFolder.GetRootObjectPtr()))
+		if (ILevelInstanceInterface* LevelInstance = Cast<ILevelInstanceInterface>(InFolder.GetRootObjectPtr()))
 		{
 			if (LevelInstance->IsEditing())
 			{
diff --git a/Engine/Source/Editor/SceneOutliner/Private/ActorTreeItem.cpp b/Engine/Source/Editor/SceneOutliner/Private/ActorTreeItem.cpp
--- a/Engine/Source/Editor/SceneOutliner/Private/ActorTreeItem.cpp
+++ b/Engine/Source/Editor/SceneOutliner/Private/ActorTreeItem.cpp
@@ -17,7 +17,7 @@
 #include "ISceneOutlinerMode.h"
 #include "Logging/MessageLog.h"
 #include "SSocketChooser.h"
-#include "LevelInstance/LevelInstanceActor.h"
+#include "LevelInstance/LevelInstanceInterface.h"
 #include "WorldPartition/WorldPartition.h"
 #include "ToolMenu.h"
 #include "Engine/Level.h"
@@ -120,9 +120,9 @@ private:
 		if (const FSceneOutlinerTreeItemPtr TreeItem = TreeItemPtr.Pin())
 		{
 			const AActor* Actor = ActorPtr.Get();
-			if (const ALevelInstance* LevelInstanceActor = Cast<ALevelInstance>(Actor))
+			if (const ILevelInstanceInterface* LevelInstance = Cast<ILevelInstanceInterface>(Actor))
 			{
-				if (LevelInstanceActor->IsDirty() && !bInEditingMode)
+				if (LevelInstance->IsDirty() && !bInEditingMode)
 				{
 					FFormatNamedArguments Args;
 					Args.Add(TEXT("ActorLabel"), FText::FromString(TreeItem->GetDisplayString()));
@@ -255,9 +255,9 @@ private:
 		AActor* Actor = ActorPtr.Get();
 
 		// Color LevelInstances differently if they are being edited
-		if (const ALevelInstance* LevelInstanceActor = Cast<ALevelInstance>(Actor))
+		if (const ILevelInstanceInterface* LevelInstance = Cast<ILevelInstanceInterface>(Actor))
 		{
-			if (LevelInstanceActor->IsEditing())
+			if (LevelInstance->IsEditing())
 			{
 				return FAppStyle::Get().GetSlateColor("Colors.AccentGreen");
 			}
@@ -450,8 +450,8 @@ void FActorTreeItem::OnLabelChanged()
 void FActorTreeItem::GenerateContextMenu(UToolMenu* Menu, SSceneOutliner& Outliner)
 {
 	const AActor* ActorPtr = Actor.Get();
-	const ALevelInstance* LevelInstanceActor = Cast<ALevelInstance>(ActorPtr);
-	if (LevelInstanceActor && LevelInstanceActor->IsEditing())
+	const ILevelInstanceInterface* LevelInstance = Cast<ILevelInstanceInterface>(ActorPtr);
+	if (LevelInstance && LevelInstance->IsEditing())
 	{
 		auto SharedOutliner = StaticCastSharedRef<SSceneOutliner>(Outliner.AsShared());
 		const FSlateIcon NewFolderIcon(FEditorStyle::GetStyleSetName(), "SceneOutliner.NewFolderIcon");
diff --git a/Engine/Source/Editor/UnrealEd/Private/Commandlets/ContentCommandlets.cpp b/Engine/Source/Editor/UnrealEd/Private/Commandlets/ContentCommandlets.cpp
--- a/Engine/Source/Editor/UnrealEd/Private/Commandlets/ContentCommandlets.cpp
+++ b/Engine/Source/Editor/UnrealEd/Private/Commandlets/ContentCommandlets.cpp
@@ -50,7 +50,7 @@
 #include "CommandletSourceControlUtils.h"
 #include "WorldPartition/WorldPartition.h"
 #include "WorldPartition/WorldPartitionHelpers.h"
-#include "LevelInstance/LevelInstanceActor.h"
+#include "LevelInstance/LevelInstanceInterface.h"
 #include "AssetCompilingManager.h"
 #include "PackageHelperFunctions.h"
 #include "PackageTools.h"
@@ -2588,8 +2588,8 @@ int32 UWrangleContentCommandlet::Main( const FString& Params )
 					}
 
 					// Add any levels referenced by Level Instance actors to the list of levels to load
-					ALevelInstance* LevelInstance = Cast<ALevelInstance>(Object);
-					if (LevelInstance && !LevelInstance->IsTemplate())
+					ILevelInstanceInterface* LevelInstance = Cast<ILevelInstanceInterface>(Object);
+					if (LevelInstance && !CastChecked<AActor>(LevelInstance)->IsTemplate())
 					{
 						FString LevelPackageName = LevelInstance->GetWorldAssetPackage();
 						if (!LevelPackageName.IsEmpty() && PackagesToFullyLoad.FindKey(LevelPackageName) == nullptr)
diff --git a/Engine/Source/Editor/UnrealEd/Private/EditorActorFolders.cpp b/Engine/Source/Editor/UnrealEd/Private/EditorActorFolders.cpp
--- a/Engine/Source/Editor/UnrealEd/Private/EditorActorFolders.cpp
+++ b/Engine/Source/Editor/UnrealEd/Private/EditorActorFolders.cpp
@@ -15,7 +15,7 @@
 #include "EditorFolderUtils.h"
 #include "ScopedTransaction.h"
 #include "UObject/ObjectSaveContext.h"
-#include "LevelInstance/LevelInstanceActor.h"
+#include "LevelInstance/LevelInstanceInterface.h"
 #include "ActorFolder.h"
 #include "Subsystems/ActorEditorContextSubsystem.h"
 
@@ -267,11 +267,11 @@ FFolder FActorFolders::GetDefaultFolderForSelection(UWorld& InWorld, TArray<FFol
 
 		FFolder Folder = Actor->GetFolder();
 		// Special case for Level Instance, make root as level instance if editing
-		if (ALevelInstance* LevelInstance = Cast<ALevelInstance>(Actor))
+		if (ILevelInstanceInterface* LevelInstance = Cast<ILevelInstanceInterface>(Actor))
 		{
 			if (LevelInstance->IsEditing())
 			{
-				Folder = FFolder(FFolder::GetEmptyPath(), FFolder::FRootObject(LevelInstance));
+				Folder = FFolder(FFolder::GetEmptyPath(), FFolder::FRootObject(Actor));
 			}
 		}
 		if (!MergeFolders(Folder))
diff --git a/Engine/Source/Editor/UnrealEd/Private/WorldFolders.cpp b/Engine/Source/Editor/UnrealEd/Private/WorldFolders.cpp
--- a/Engine/Source/Editor/UnrealEd/Private/WorldFolders.cpp
+++ b/Engine/Source/Editor/UnrealEd/Private/WorldFolders.cpp
@@ -10,7 +10,7 @@
 #include "EngineUtils.h"
 #include "EditorActorFolders.h"
 #include "LevelInstance/LevelInstanceSubsystem.h"
-#include "LevelInstance/LevelInstanceActor.h"
+#include "LevelInstance/LevelInstanceInterface.h"
 
 #define LOCTEXT_NAMESPACE "UnrealEd.WorldFolders"
 
@@ -186,8 +186,9 @@ FFolder UWorldFolders::GetActorEditorContextFolder() const
 	if (ContainsFolder(Folder))
 	{
 		ULevelInstanceSubsystem* LevelInstanceSubsystem = GetWorld()->GetSubsystem<ULevelInstanceSubsystem>();
-		UObject* EditingLevelInstance = LevelInstanceSubsystem ? LevelInstanceSubsystem->GetEditingLevelInstance() : nullptr;
-		FFolder::FRootObject RootObject = FFolder::FRootObject(EditingLevelInstance ? EditingLevelInstance : (GetWorld()->GetCurrentLevel() == GetWorld()->PersistentLevel) ? nullptr : GetWorld()->GetCurrentLevel());
+		ILevelInstanceInterface* EditingLevelInstance = LevelInstanceSubsystem ? LevelInstanceSubsystem->GetEditingLevelInstance() : nullptr;
+		UObject* EditingLevelInstanceObject = EditingLevelInstance ? CastChecked<UObject>(EditingLevelInstance) : nullptr;
+		FFolder::FRootObject RootObject = FFolder::FRootObject(EditingLevelInstanceObject ? EditingLevelInstanceObject : (GetWorld()->GetCurrentLevel() == GetWorld()->PersistentLevel) ? nullptr : GetWorld()->GetCurrentLevel());
 		if (Folder.GetRootObject() == RootObject)
 		{
 			return Folder;
diff --git a/Engine/Source/Editor/UnrealEd/Private/WorldPersistentFolders.cpp b/Engine/Source/Editor/UnrealEd/Private/WorldPersistentFolders.cpp
--- a/Engine/Source/Editor/UnrealEd/Private/WorldPersistentFolders.cpp
+++ b/Engine/Source/Editor/UnrealEd/Private/WorldPersistentFolders.cpp
@@ -7,7 +7,7 @@
 #include "Engine/Level.h"
 #include "UObject/UObjectHash.h"
 #include "ExternalPackageHelper.h"
-#include "LevelInstance/LevelInstanceActor.h"
+#include "LevelInstance/LevelInstanceInterface.h"
 
 #define LOCTEXT_NAMESPACE "UnrealEd.WorldPersistentFolders"
 
@@ -148,7 +148,7 @@ ULevel* FWorldPersistentFolders::GetRootObjectContainer(const FFolder& InFolder,
 	{
 		return WorldPtr->PersistentLevel;
 	}
-	else if (const ALevelInstance* LevelInstance = Cast<ALevelInstance>(RootObjectPtr))
+	else if (const ILevelInstanceInterface* LevelInstance = Cast<ILevelInstanceInterface>(RootObjectPtr))
 	{
 		return LevelInstance->GetLoadedLevel();
 	}
diff --git a/Engine/Source/Runtime/Engine/Private/ActorEditor.cpp b/Engine/Source/Runtime/Engine/Private/ActorEditor.cpp
--- a/Engine/Source/Runtime/Engine/Private/ActorEditor.cpp
+++ b/Engine/Source/Runtime/Engine/Private/ActorEditor.cpp
@@ -762,7 +762,7 @@ bool AActor::SupportsLayers() const
 		// Actors part of Level Instance are not valid for layers
 		if (ULevelInstanceSubsystem* LevelInstanceSubsystem = GetWorld()->GetSubsystem<ULevelInstanceSubsystem>())
 		{
-			if (ALevelInstance* LevelInstance = LevelInstanceSubsystem->GetParentLevelInstance(this))
+			if (ILevelInstanceInterface* LevelInstance = LevelInstanceSubsystem->GetParentLevelInstance(this))
 			{
 				return false;
 			}
diff --git a/Engine/Source/Runtime/Engine/Private/LevelInstance/LevelInstanceActor.cpp b/Engine/Source/Runtime/Engine/Private/LevelInstance/LevelInstanceActor.cpp
--- a/Engine/Source/Runtime/Engine/Private/LevelInstance/LevelInstanceActor.cpp
+++ b/Engine/Source/Runtime/Engine/Private/LevelInstance/LevelInstanceActor.cpp
@@ -3,27 +3,22 @@
 #include "LevelInstance/LevelInstanceActor.h"
 #include "LevelInstance/LevelInstanceSubsystem.h"
 #include "LevelInstance/LevelInstanceComponent.h"
+#include "LevelInstance/LevelInstanceInterface.h"
 #include "LevelInstance/LevelInstanceEditorInstanceActor.h"
 #include "Engine/World.h"
 #include "Engine/Engine.h"
 #include "LevelInstancePrivate.h"
 
 #if WITH_EDITOR
-#include "Logging/MessageLog.h"
-#include "Logging/TokenizedMessage.h"
-#include "Misc/MapErrors.h"
-#include "Misc/UObjectToken.h"
 #include "WorldPartition/WorldPartitionActorDesc.h"
 #include "WorldPartition/LevelInstance/LevelInstanceActorDesc.h"
 #endif
 
 #define LOCTEXT_NAMESPACE "LevelInstanceActor"
 
-ALevelInstance::ALevelInstance(const FObjectInitializer& ObjectInitializer)
-	: Super(ObjectInitializer)
-#if WITH_EDITOR
-	, bGuardLoadUnload(false)
-#endif
+ALevelInstance::ALevelInstance()
+	: LevelInstanceActorGuid(this)
+	, LevelInstanceActorImpl(this)
 {
 	RootComponent = CreateDefaultSubobject<ULevelInstanceComponent>(TEXT("Root"));
 	RootComponent->Mobility = EComponentMobility::Static;
@@ -33,129 +28,34 @@ ALevelInstance::ALevelInstance(const FObjectInitializer& ObjectInitializer)
 #endif
 }
 
-ULevelInstanceSubsystem* ALevelInstance::GetLevelInstanceSubsystem() const
-{
-	if (UWorld* CurrentWorld = GetWorld())
-	{
-		return CurrentWorld->GetSubsystem<ULevelInstanceSubsystem>();
-	}
-
-	return nullptr;
-}
-
 void ALevelInstance::Serialize(FArchive& Ar)
 {
 	Super::Serialize(Ar);
-#if WITH_EDITOR
-	if (Ar.IsSaving() && Ar.IsCooking() && !IsTemplate())
-	{
-		FGuid Guid = GetLevelInstanceActorGuid();
-		Ar << Guid;
-	}
-#else
-	if (Ar.IsLoading())
-	{
-		if (IsTemplate())
-		{
-			check(!LevelInstanceActorGuid.IsValid());
-		}
-		else if (Ar.GetPortFlags() & PPF_Duplicate)
-		{
-			LevelInstanceActorGuid = FGuid::NewGuid();
-		}
-		else if (Ar.IsPersistent())
-		{
-			Ar << LevelInstanceActorGuid;
-			check(LevelInstanceActorGuid.IsValid());
-		}
-	}
-#endif
+	Ar << LevelInstanceActorGuid;
 }
 
 void ALevelInstance::PostRegisterAllComponents()
 {
 	Super::PostRegisterAllComponents();
 
-	if (ULevelInstanceSubsystem* LevelInstanceSubsystem = GetLevelInstanceSubsystem())
-	{
 #if !WITH_EDITOR
-		// If the level instance was spawned, not loaded
-		if (!LevelInstanceActorGuid.IsValid())
-		{
-			LevelInstanceActorGuid = FGuid::NewGuid();
-		}
+	// If the level instance was spawned, not loaded
+	LevelInstanceActorGuid.AssignIfInvalid();
 #endif
-		LevelInstanceID = LevelInstanceSubsystem->RegisterLevelInstance(this);
-
-		LoadLevelInstance();
 
-#if WITH_EDITOR
-		// Make sure transformation is up to date after registration as its possible LevelInstance actor can get unregistered when editing properties
-		// through Details panel. In this case the ULevelInstanceComponent might not be able to update the ALevelInstanceEditorInstanceActor transform.
-		Cast<ULevelInstanceComponent>(GetRootComponent())->UpdateEditorInstanceActor();
-#endif
-	}
+	LevelInstanceActorImpl.RegisterLevelInstance();
 }
 
 void ALevelInstance::PostUnregisterAllComponents()
 {
 	Super::PostUnregisterAllComponents();
 
-	if (ULevelInstanceSubsystem* LevelInstanceSubsystem = GetLevelInstanceSubsystem())
-	{
-		// If LevelInstance has already been unregistered it will have an Invalid LevelInstanceID. Avoid processing it.
-		if (!HasValidLevelInstanceID())
-		{
-			return;
-		}
-
-		LevelInstanceSubsystem->UnregisterLevelInstance(this);
-
-		UnloadLevelInstance();
-
-		// To avoid processing PostUnregisterAllComponents multiple times (BP Recompile is one use case)
-		LevelInstanceID = FLevelInstanceID();
-	}
-
+	LevelInstanceActorImpl.UnregisterLevelInstance();
 }
 
-bool ALevelInstance::SupportsLoading() const
+bool ALevelInstance::IsLoadingEnabled() const
 {
-#if WITH_EDITOR
-	return !bGuardLoadUnload && !bIsEditorPreviewActor;
-#else
-	return true;
-#endif
-}
-
-void ALevelInstance::LoadLevelInstance()
-{
-	if(SupportsLoading())
-	{
-		if (ULevelInstanceSubsystem* LevelInstanceSubsystem = GetLevelInstanceSubsystem())
-		{
-			bool bForce = false;
-#if WITH_EDITOR
-			// When reinstancing or world wasn't ticked between changes. Avoid reloading the level but if the package changed, force the load
-			bForce = IsLoaded() && LevelInstanceSubsystem->GetLevelInstanceLevel(this)->GetPackage()->GetLoadedPath() != FPackagePath::FromPackageNameChecked(GetWorldAssetPackage());
-#endif
-			LevelInstanceSubsystem->RequestLoadLevelInstance(this, bForce);
-		}
-	}
-}
-
-void ALevelInstance::UnloadLevelInstance()
-{
-	if(SupportsLoading())
-	{
-		if (ULevelInstanceSubsystem* LevelInstanceSubsystem = GetLevelInstanceSubsystem())
-		{
-#if WITH_EDITOR
-			check(!HasDirtyChildren());
-#endif
-			LevelInstanceSubsystem->RequestUnloadLevelInstance(this);
-		}
-	}
+	return LevelInstanceActorImpl.IsLoadingEnabled();
 }
 
 const TSoftObjectPtr<UWorld>& ALevelInstance::GetWorldAsset() const
@@ -167,61 +67,35 @@ const TSoftObjectPtr<UWorld>& ALevelInstance::GetWorldAsset() const
 #endif
 }
 
-bool ALevelInstance::IsLevelInstancePathValid() const
+const FLevelInstanceID& ALevelInstance::GetLevelInstanceID() const
 {
-	return GetWorldAsset().GetUniqueID().IsValid();
+	return LevelInstanceActorImpl.GetLevelInstanceID();
 }
 
 bool ALevelInstance::HasValidLevelInstanceID() const
 {
-	return LevelInstanceID.IsValid();
+	return LevelInstanceActorImpl.HasValidLevelInstanceID();
 }
 
-const FLevelInstanceID& ALevelInstance::GetLevelInstanceID() const
+const FGuid& ALevelInstance::GetLevelInstanceGuid() const
 {
-	check(HasValidLevelInstanceID());
-	return LevelInstanceID;
+	return LevelInstanceActorGuid.GetGuid();
 }
 
-const FGuid& ALevelInstance::GetLevelInstanceActorGuid() const
+void ALevelInstance::OnLevelInstanceLoaded()
 {
-#if WITH_EDITOR
-	const FGuid& Guid = GetActorGuid();
-#else
-	const FGuid& Guid = LevelInstanceActorGuid;
-#endif
-	check(IsTemplate() || Guid.IsValid());
-	return Guid;
+	LevelInstanceActorImpl.OnLevelInstanceLoaded();
 }
 
 #if WITH_EDITOR
-
-TUniquePtr<FWorldPartitionActorDesc> ALevelInstance::CreateClassActorDesc() const
+ULevelInstanceComponent* ALevelInstance::GetLevelInstanceComponent() const
 {
-	return TUniquePtr<FWorldPartitionActorDesc>(new FLevelInstanceActorDesc());
+	return Cast<ULevelInstanceComponent>(RootComponent);
 }
 
-AActor* ALevelInstance::FindEditorInstanceActor() const
+TUniquePtr<FWorldPartitionActorDesc> ALevelInstance::CreateClassActorDesc() const
 {
-	AActor* FoundActor = nullptr;
-	if (ULevelInstanceSubsystem* LevelInstanceSubsystem = GetLevelInstanceSubsystem())
-	{
-		if (LevelInstanceSubsystem->IsLoaded(this))
-		{
-			LevelInstanceSubsystem->ForEachActorInLevelInstance(this, [&FoundActor, this](AActor* LevelActor)
-			{
-				if (ALevelInstanceEditorInstanceActor* LevelInstanceEditorInstanceActor = Cast<ALevelInstanceEditorInstanceActor>(LevelActor))
-				{
-					check(LevelInstanceEditorInstanceActor->GetLevelInstanceID() == GetLevelInstanceID());
-					FoundActor = LevelInstanceEditorInstanceActor;
-					return false;
-				}
-				return true;
-			});
-		}
-	}
-
-	return FoundActor;
+	return TUniquePtr<FWorldPartitionActorDesc>(new FLevelInstanceActorDesc());
 }
 
 ALevelInstance::FOnLevelInstanceActorPostLoad ALevelInstance::OnLevelInstanceActorPostLoad;
@@ -239,7 +113,7 @@ void ALevelInstance::PostLoad()
 #endif
 		
 #if WITH_EDITORONLY_DATA
-	if (IsRunningCookCommandlet() && SupportsLoading())
+	if (IsRunningCookCommandlet() && IsLoadingEnabled())
 	{
 		CookedWorldAsset = WorldAsset;
 	}
@@ -250,219 +124,36 @@ void ALevelInstance::PostLoad()
 
 void ALevelInstance::PreEditUndo()
 {
-	CachedLevelInstanceID = LevelInstanceID;
-	CachedWorldAsset = WorldAsset;
-	bCachedIsTemporarilyHiddenInEditor = IsTemporarilyHiddenInEditor(false);
-
-	{
-		TGuardValue<bool> LoadUnloadGuard(bGuardLoadUnload, true);
-		Super::PreEditUndo();
-	}
+	LevelInstanceActorImpl.PreEditUndo([this]() { Super::PreEditUndo(); });
 }
 
 void ALevelInstance::PostEditUndo(TSharedPtr<ITransactionObjectAnnotation> TransactionAnnotation)
 {
-	{
-		TGuardValue<bool> LoadUnloadGuard(bGuardLoadUnload, true);
-		Super::PostEditUndo(TransactionAnnotation);
-	}
-
-	PostEditUndoInternal();
+	LevelInstanceActorImpl.PostEditUndo(TransactionAnnotation, [this](TSharedPtr<ITransactionObjectAnnotation> InTransactionAnnotation) { Super::PostEditUndo(InTransactionAnnotation); });
 }
 
 void ALevelInstance::PostEditUndo()
 {
-	{
-		TGuardValue<bool> LoadUnloadGuard(bGuardLoadUnload, true);
-		Super::PostEditUndo();
-	}
-
-	PostEditUndoInternal();
-}
-
-void ALevelInstance::PostEditUndoInternal()
-{
-	if (CachedWorldAsset != WorldAsset)
-	{
-		OnWorldAssetChanged();
-	}
-
-	if (bCachedIsTemporarilyHiddenInEditor != IsTemporarilyHiddenInEditor(false))
-	{
-		if (ULevelInstanceSubsystem* LevelInstanceSubsystem = GetLevelInstanceSubsystem())
-		{
-			LevelInstanceSubsystem->SetIsTemporarilyHiddenInEditor(this, !bCachedIsTemporarilyHiddenInEditor);
-		}
-	}
-
-	// Here we want to load or unload based on our current state
-	if (HasValidLevelInstanceID() && !IsLoaded())
-	{
-		LoadLevelInstance();
-	}
-	else if (!IsValidChecked(this))
-	{
-		// Temp restore the ID so that we can unload
-		TGuardValue<FLevelInstanceID> LevelInstanceIDGuard(LevelInstanceID, CachedLevelInstanceID);
-		if (IsLoaded())
-		{
-			UnloadLevelInstance();
-		}
-	}
-
-	CachedLevelInstanceID = FLevelInstanceID();
-	CachedWorldAsset.Reset();
-
-	if (ULevelInstanceComponent* LevelInstanceComponent = Cast<ULevelInstanceComponent>(GetRootComponent()))
-	{
-		// Order of operations when undoing may lead to the RootComponent being undone before our actor so we need to make sure we update here and in the component when undoing
-		LevelInstanceComponent->UpdateEditorInstanceActor();
-	}
-}
-
-FString ALevelInstance::GetWorldAssetPackage() const
-{
-	return GetWorldAsset().GetUniqueID().GetLongPackageName();
+	LevelInstanceActorImpl.PostEditUndo([this]() { Super::PostEditUndo(); });
 }
 
 void ALevelInstance::PreEditChange(FProperty* PropertyThatWillChange)
 {
-	{
-		TGuardValue<bool> LoadUnloadGuard(bGuardLoadUnload, true);
-		Super::PreEditChange(PropertyThatWillChange);
-	}
-
-	if (PropertyThatWillChange->GetFName() == GET_MEMBER_NAME_CHECKED(ALevelInstance, WorldAsset))
-	{
-		CachedWorldAsset = WorldAsset;
-	}
+	const bool bWorldAssetChange = PropertyThatWillChange->GetFName() == GET_MEMBER_NAME_CHECKED(ALevelInstance, WorldAsset);
+	LevelInstanceActorImpl.PreEditChange(PropertyThatWillChange, bWorldAssetChange, [this](FProperty* Property) { Super::PreEditChange(Property); });
 }
 
 void ALevelInstance::CheckForErrors()
 {
 	Super::CheckForErrors();
 
-	TArray<TPair<FText, TSoftObjectPtr<UWorld>>> LoopInfo;
-	const ALevelInstance* LoopStart = nullptr;
-	if (!CheckForLoop(GetWorldAsset(), &LoopInfo, &LoopStart))
-	{
-		TSharedRef<FTokenizedMessage> Error = FMessageLog("MapCheck").Error()->AddToken(FTextToken::Create(LOCTEXT("LevelInstanceActor_Loop_CheckForErrors", "LevelInstance level loop found!")));
-		TSoftObjectPtr<UWorld> LoopStartAsset(LoopStart->GetLevel()->GetTypedOuter<UWorld>());
-		Error->AddToken(FAssetNameToken::Create(LoopStartAsset.GetLongPackageName(), FText::FromString(LoopStartAsset.GetAssetName())));
-		Error->AddToken(FTextToken::Create(FText::FromString(TEXT(":"))));
-		Error->AddToken(FUObjectToken::Create(LoopStart));
-		for (int32 i = LoopInfo.Num() - 1; i >= 0; --i)
-		{
-			Error->AddToken(FTextToken::Create(LoopInfo[i].Key));
-			TSoftObjectPtr<UWorld> LevelInstancePtr = LoopInfo[i].Value;
-			Error->AddToken(FAssetNameToken::Create(LevelInstancePtr.GetLongPackageName(), FText::FromString(LevelInstancePtr.GetAssetName())));
-		}
-
-		Error->AddToken(FMapErrorToken::Create(FName(TEXT("LevelInstanceActor_Loop_CheckForErrors"))));
-	}
-
-	FPackagePath WorldAssetPath;
-	if (!FPackagePath::TryFromPackageName(GetWorldAssetPackage(), WorldAssetPath) || !FPackageName::DoesPackageExist(WorldAssetPath))
-	{
-		TSharedRef<FTokenizedMessage> Error = FMessageLog("MapCheck").Error()
-			->AddToken(FTextToken::Create(LOCTEXT("LevelInstanceActor_InvalidPackage", "LevelInstance actor")))
-			->AddToken(FUObjectToken::Create(this))
-			->AddToken(FTextToken::Create(FText::FromString(TEXT("refers to an invalid asset:"))))
-			->AddToken(FAssetNameToken::Create(GetWorldAsset().GetLongPackageName(), FText::FromString(GetWorldAsset().GetLongPackageName())))
-			->AddToken(FMapErrorToken::Create(FName(TEXT("LevelInstanceActor_InvalidPackage_CheckForErrors"))));
-	}
-}
-
-bool ALevelInstance::CheckForLoop(TSoftObjectPtr<UWorld> InLevelInstance, TArray<TPair<FText,TSoftObjectPtr<UWorld>>>* LoopInfo, const ALevelInstance** LoopStart) const
-{
-	bool bValid = true;
-	if (ULevelInstanceSubsystem* LevelInstanceSubsystem = GetLevelInstanceSubsystem())
-	{
-		LevelInstanceSubsystem->ForEachLevelInstanceAncestorsAndSelf(this, [&bValid, &InLevelInstance, LoopInfo, LoopStart, this](const ALevelInstance* LevelInstanceActor)
-		{
-			FName LongPackageName(*InLevelInstance.GetLongPackageName());
-			// Check to exclude NAME_None since Preview Levels are in the transient package
-			// Check the level we are spawned in to detect the loop (this will handle loops caused by LevelInstances and by regular level streaming)
-			if (LongPackageName != NAME_None && LevelInstanceActor->GetLevel()->GetPackage()->GetLoadedPath() == FPackagePath::FromPackageNameChecked(LongPackageName))
-			{
-				bValid = false;
-				if (LoopStart)
-				{
-					*LoopStart = LevelInstanceActor;
-				}
-			}
-
-			if (LoopInfo)
-			{
-				TSoftObjectPtr<UWorld> LevelInstancePtr = LevelInstanceActor == this ? InLevelInstance : LevelInstanceActor->GetWorldAsset();
-				FText LevelInstanceName = FText::FromString(LevelInstanceActor->GetPathName());
-				FText Description = FText::Format(LOCTEXT("LevelInstanceLoopLink", "-> Actor: {0} loads"), LevelInstanceName);
-				LoopInfo->Emplace(Description, LevelInstancePtr);
-			}
-				
-			return bValid;
-		});
-	}
-
-	return bValid;
-}
-
-bool ALevelInstance::CanSetValue(TSoftObjectPtr<UWorld> InLevelInstance, FString* Reason) const
-{
-	// Set to null is valid
-	if (InLevelInstance.IsNull())
-	{
-		return true;
-	}
-
-	FString PackageName;
-	if (!FPackageName::DoesPackageExist(InLevelInstance.GetLongPackageName()))
-	{
-		if (Reason)
-		{
-			*Reason = FString::Format(TEXT("Attempting to set Level Instance to package {0} which does not exist. Ensure the level was saved before attepting to set the level instance world asset."), { InLevelInstance.GetLongPackageName() });
-		}
-		return false;
-	}
-
-	TArray<TPair<FText, TSoftObjectPtr<UWorld>>> LoopInfo;
-	const ALevelInstance* LoopStart = nullptr;
-
-	if (ULevel::GetIsLevelPartitionedFromPackage(*InLevelInstance.GetLongPackageName()))
-	{
-		if (Reason)
-		{
-			*Reason = FString::Format(TEXT("LevelInstance doesn't support partitioned world {0}\n"), { InLevelInstance.GetLongPackageName() });
-		}
-
-		return false;
-	}
-
-	if (!CheckForLoop(InLevelInstance, Reason ? &LoopInfo : nullptr, Reason ? &LoopStart : nullptr))
-	{
-		if (Reason)
-		{
-			if (ensure(LoopStart))
-			{
-				TSoftObjectPtr<UWorld> LoopStartAsset(LoopStart->GetLevel()->GetTypedOuter<UWorld>());
-				*Reason = FString::Format(TEXT("Setting LevelInstance to {0} would cause loop {1}:{2}\n"), { InLevelInstance.GetLongPackageName(), LoopStart->GetName(), LoopStartAsset.GetLongPackageName() });
-				for (int32 i = LoopInfo.Num() - 1; i >= 0; --i)
-				{
-					Reason->Append(FString::Format(TEXT("{0} {1}\n"), { *LoopInfo[i].Key.ToString(), *LoopInfo[i].Value.GetLongPackageName() }));
-				}
-			}
-		}
-
-		return false;
-	}
-
-	return true;
+	LevelInstanceActorImpl.CheckForErrors();
 }
 
 bool ALevelInstance::SetWorldAsset(TSoftObjectPtr<UWorld> InWorldAsset)
 {
 	FString Reason;
-	if (!CanSetValue(InWorldAsset, &Reason))
+	if (!ULevelInstanceSubsystem::CanUseWorldAsset(this, InWorldAsset, &Reason))
 	{
 		UE_LOG(LogLevelInstance, Warning, TEXT("%s"), *Reason);
 		return false;
@@ -474,194 +165,48 @@ bool ALevelInstance::SetWorldAsset(TSoftObjectPtr<UWorld> InWorldAsset)
 
 void ALevelInstance::PostEditChangeProperty(FPropertyChangedEvent& PropertyChangedEvent)
 {
-	{
-		TGuardValue<bool> LoadUnloadGuard(bGuardLoadUnload, true);
-		Super::PostEditChangeProperty(PropertyChangedEvent);
-	}
+	const bool bWorldAssetChange = PropertyChangedEvent.Property && PropertyChangedEvent.Property->GetFName() == GET_MEMBER_NAME_CHECKED(ALevelInstance, WorldAsset);
 
-	if (FProperty* const PropertyThatChanged = PropertyChangedEvent.Property)
-	{
-		if (PropertyThatChanged->GetFName() == GET_MEMBER_NAME_CHECKED(ALevelInstance, WorldAsset))
-		{
-			if (ULevelInstanceSubsystem* LevelInstanceSubsystem = GetLevelInstanceSubsystem())
-			{
-				FString Reason;
-				if (!CanSetValue(GetWorldAsset(), &Reason))
-				{
-					UE_LOG(LogLevelInstance, Warning, TEXT("%s"), *Reason);
-					WorldAsset = CachedWorldAsset;
-				}
-				else
-				{
-					OnWorldAssetChanged();
-				}
-				CachedWorldAsset.Reset();
-			}
-		}
-	}
+	LevelInstanceActorImpl.PostEditChangeProperty(PropertyChangedEvent, bWorldAssetChange, [this](FPropertyChangedEvent& Event) { Super::PostEditChangeProperty(Event); });
 }
 
-bool ALevelInstance::CanEditChange(const FProperty* InProperty) const
+bool ALevelInstance::CanEditChange(const FProperty* Property) const
 {
-	if (!Super::CanEditChange(InProperty))
-	{
-		return false;
-	}
-
-	if (IsEditing())
-	{
-		return false;
-	}
-
-	if (HasDirtyChildren())
-	{
-		return false;
-	}
-
-	return true;
+	return Super::CanEditChange(Property) && LevelInstanceActorImpl.CanEditChange(Property);
 }
 
 void ALevelInstance::PostEditImport()
 {
-	{
-		TGuardValue<bool> LoadUnloadGuard(bGuardLoadUnload, true);
-		Super::PostEditImport();
-	}
-	UpdateFromLevel();
+	LevelInstanceActorImpl.PostEditImport([this]() { Super::PostEditImport(); });
 }
 
 bool ALevelInstance::CanDeleteSelectedActor(FText& OutReason) const
 {
-	if (!Super::CanDeleteSelectedActor(OutReason))
-	{
-		return false;
-	}
-
-	if (IsEditing())
-	{
-		OutReason = LOCTEXT("HasEditingLevel", "Can't delete LevelInstance because it is editing!");
-		return false;
-	}
-
-	if (HasChildEdit())
-	{
-		OutReason = LOCTEXT("HasEditingChildLevel", "Can't delete LevelInstance because it has editing child LevelInstances!");
-		return false;
-	}
-	return true;
+	return Super::CanDeleteSelectedActor(OutReason) && LevelInstanceActorImpl.CanDeleteSelectedActor(OutReason);
 }
 
 void ALevelInstance::SetIsTemporarilyHiddenInEditor(bool bIsHidden)
 {
-	{
-		TGuardValue<bool> LoadUnloadGuard(bGuardLoadUnload, true);
-		Super::SetIsTemporarilyHiddenInEditor(bIsHidden);
-	}
-
-	if (ULevelInstanceSubsystem* LevelInstanceSubsystem = GetLevelInstanceSubsystem())
-	{
-		LevelInstanceSubsystem->SetIsTemporarilyHiddenInEditor(this, bIsHidden);
-	}
+	LevelInstanceActorImpl.SetIsTemporarilyHiddenInEditor(bIsHidden, [this](bool bInHidden) { Super::SetIsTemporarilyHiddenInEditor(bInHidden); });
 }
 
 bool ALevelInstance::SetIsHiddenEdLayer(bool bIsHiddenEdLayer)
 {
-	bool bHasChanged = false;
-	{
-		TGuardValue<bool> LoadUnloadGuard(bGuardLoadUnload, true);
-		bHasChanged = Super::SetIsHiddenEdLayer(bIsHiddenEdLayer);
-	}
-	if (bHasChanged)
-	{
-		if (ULevelInstanceSubsystem* LevelInstanceSubsystem = GetLevelInstanceSubsystem())
-		{
-			LevelInstanceSubsystem->SetIsHiddenEdLayer(this, bIsHiddenEdLayer);
-		}
-	}
-	return bHasChanged;
+	return LevelInstanceActorImpl.SetIsHiddenEdLayer(bIsHiddenEdLayer, [this](bool bInHiddenEdLayer) { return Super::SetIsHiddenEdLayer(bInHiddenEdLayer); });
 }
 
 void ALevelInstance::EditorGetUnderlyingActors(TSet<AActor*>& OutUnderlyingActors) const
 {
 	Super::EditorGetUnderlyingActors(OutUnderlyingActors);
-
-	if (ULevelInstanceSubsystem* LevelInstanceSubsystem = GetLevelInstanceSubsystem())
-	{
-		LevelInstanceSubsystem->ForEachActorInLevelInstance(this, [&](AActor* LevelActor)
-		{
-			bool bAlreadySet = false;
-			OutUnderlyingActors.Add(LevelActor, &bAlreadySet);
-			if (!bAlreadySet)
-			{
-				LevelActor->EditorGetUnderlyingActors(OutUnderlyingActors);
-			}
-			return true;
-		});
-	}
-}
-
-void ALevelInstance::UpdateFromLevel()
-{
-	if (HasValidLevelInstanceID())
-	{
-		if (ULevelInstanceSubsystem* LevelInstanceSubsystem = GetLevelInstanceSubsystem())
-		{
-			if (IsLevelInstancePathValid() && SupportsLoading())
-			{
-				const bool bForceUpdate = true;
-				LevelInstanceSubsystem->RequestLoadLevelInstance(this, bForceUpdate);
-			}
-			else if(IsLoaded())
-			{
-				UnloadLevelInstance();
-			}
-		}
-	}
-}
-
-bool ALevelInstance::IsLoaded() const
-{
-	if (ULevelInstanceSubsystem* LevelInstanceSubsystem = GetLevelInstanceSubsystem())
-	{
-		return LevelInstanceSubsystem->IsLoaded(this);
-	}
-
-	return false;
-}
-
-void ALevelInstance::OnLevelInstanceLoaded()
-{
-	if (!GetWorld()->IsGameWorld())
-	{
-		// Propagate bounds dirtyness up and check if we need to hide our LevelInstance because self or ancestor is hidden
-		bool bHiddenInEditor = false;
-		GetLevelInstanceSubsystem()->ForEachLevelInstanceAncestorsAndSelf(this, [&bHiddenInEditor](const ALevelInstance* AncestorOrSelf)
-		{
-			AncestorOrSelf->GetLevel()->MarkLevelBoundsDirty();
-			bHiddenInEditor |= AncestorOrSelf->IsTemporarilyHiddenInEditor();
-			return true;
-		});
-
-		if (bHiddenInEditor)
-		{
-			SetIsTemporarilyHiddenInEditor(bHiddenInEditor);
-		}
-	}
+	LevelInstanceActorImpl.EditorGetUnderlyingActors(OutUnderlyingActors);
 }
 
 FBox ALevelInstance::GetStreamingBounds() const
 {
-	// Add Level Bounds
-	if (SupportsLoading())
+	FBox LevelInstanceBounds;
+	if (LevelInstanceActorImpl.GetBounds(LevelInstanceBounds))
 	{
-		if (ULevelInstanceSubsystem* LevelInstanceSubsystem = GetLevelInstanceSubsystem())
-		{
-			FBox LevelInstanceBounds;
-			if (LevelInstanceSubsystem->GetLevelInstanceBounds(this, LevelInstanceBounds))
-			{
-				return LevelInstanceBounds;
-			}
-		}
+		return LevelInstanceBounds;
 	}
 
 	return Super::GetStreamingBounds();
@@ -669,196 +214,34 @@ FBox ALevelInstance::GetStreamingBounds() const
 
 bool ALevelInstance::IsLockLocation() const
 {
-	return Super::IsLockLocation() || IsEditing() || HasChildEdit();
+	return Super::IsLockLocation() || LevelInstanceActorImpl.IsLockLocation();
 }
 
 FBox ALevelInstance::GetComponentsBoundingBox(bool bNonColliding, bool bIncludeFromChildActors) const
 {
 	FBox Box = Super::GetComponentsBoundingBox(bNonColliding, bIncludeFromChildActors);
 	
-	// Add Level Bounds
-	if (SupportsLoading())
+	FBox LevelInstanceBounds;
+	if (LevelInstanceActorImpl.GetBounds(LevelInstanceBounds))
 	{
-		if (ULevelInstanceSubsystem* LevelInstanceSubsystem = GetLevelInstanceSubsystem())
-		{
-			FBox LevelInstanceBounds;
-			if (LevelInstanceSubsystem->GetLevelInstanceBounds(this, LevelInstanceBounds))
-			{
-				Box += LevelInstanceBounds;
-			}
-		}
+		Box += LevelInstanceBounds;
 	}
-
+	
 	return Box;
 }
 
-bool ALevelInstance::CanEdit(FText* OutReason) const
-{
-	if (HasValidLevelInstanceID())
-	{
-		if (ULevelInstanceSubsystem* LevelInstanceSubsystem = GetLevelInstanceSubsystem())
-		{
-			return LevelInstanceSubsystem->CanEditLevelInstance(this, OutReason);
-		}
-	}
-
-	return false;
-}
-
-bool ALevelInstance::CanCommit(FText* OutReason) const
-{
-	if (HasValidLevelInstanceID())
-	{
-		if (ULevelInstanceSubsystem* LevelInstanceSubsystem = GetLevelInstanceSubsystem())
-		{
-			return LevelInstanceSubsystem->CanCommitLevelInstance(this, OutReason);
-		}
-	}
-
-	return false;
-}
-
-bool ALevelInstance::CanDiscard(FText* OutReason) const
-{
-	if (HasValidLevelInstanceID())
-	{
-		if (ULevelInstanceSubsystem* LevelInstanceSubsystem = GetLevelInstanceSubsystem())
-		{
-			return LevelInstanceSubsystem->CanDiscardLevelInstance(this, OutReason);
-		}
-	}
-
-	return false;
-}
-
-bool ALevelInstance::IsEditing() const
-{
-	if (HasValidLevelInstanceID())
-	{
-		if (ULevelInstanceSubsystem* LevelInstanceSubsystem = GetLevelInstanceSubsystem())
-		{
-			return LevelInstanceSubsystem->IsEditingLevelInstance(this);
-		}
-	}
-	return false;
-}
-
-ULevel* ALevelInstance::GetLoadedLevel() const
-{
-	if (ULevelInstanceSubsystem* LevelInstanceSubsystem = GetLevelInstanceSubsystem())
-	{
-		return LevelInstanceSubsystem->GetLevelInstanceLevel(this);
-	}
-	return nullptr;
-}
-
-bool ALevelInstance::HasChildEdit() const
-{
-	if (HasValidLevelInstanceID())
-	{
-		if (ULevelInstanceSubsystem* LevelInstanceSubsystem = GetLevelInstanceSubsystem())
-		{
-			return LevelInstanceSubsystem->HasChildEdit(this);
-		}
-	}
-
-	return false;
-}
-
-void ALevelInstance::Edit(AActor* ContextActor)
-{
-	ULevelInstanceSubsystem* LevelInstanceSubsystem = GetLevelInstanceSubsystem();
-	check(LevelInstanceSubsystem);
-	LevelInstanceSubsystem->EditLevelInstance(this, ContextActor);
-}
-
-void ALevelInstance::Commit()
-{
-	ULevelInstanceSubsystem* LevelInstanceSubsystem = GetLevelInstanceSubsystem();
-	check(LevelInstanceSubsystem);
-	LevelInstanceSubsystem->CommitLevelInstance(this);
-}
-
-void ALevelInstance::Discard()
-{
-	ULevelInstanceSubsystem* LevelInstanceSubsystem = GetLevelInstanceSubsystem();
-	check(LevelInstanceSubsystem);
-	const bool bDiscardEdits = true;
-	LevelInstanceSubsystem->CommitLevelInstance(this, bDiscardEdits);
-}
-
-bool ALevelInstance::HasDirtyChildren() const
-{
-	if (ULevelInstanceSubsystem* LevelInstanceSubsystem = GetLevelInstanceSubsystem())
-	{
-		return LevelInstanceSubsystem->HasDirtyChildrenLevelInstances(this);
-	}
-
-	return false;
-}
-
-bool ALevelInstance::IsDirty() const
-{
-	if (ULevelInstanceSubsystem* LevelInstanceSubsystem = GetLevelInstanceSubsystem())
-	{
-		return LevelInstanceSubsystem->IsEditingLevelInstanceDirty(this);
-	}
-
-	return false;
-}
-
-bool ALevelInstance::SetCurrent()
-{
-	if (ULevelInstanceSubsystem* LevelInstanceSubsystem = GetLevelInstanceSubsystem())
-	{
-		return LevelInstanceSubsystem->SetCurrent(this);
-	}
-
-	return false;
-}
-
-bool ALevelInstance::IsCurrent() const
-{
-	if (ULevelInstanceSubsystem* LevelInstanceSubsystem = GetLevelInstanceSubsystem())
-	{
-		return LevelInstanceSubsystem->IsCurrent(this);
-	}
-
-	return false;
-}
-
 void ALevelInstance::PushSelectionToProxies()
 {
 	Super::PushSelectionToProxies();
 
-	// Actors of the LevelInstance need to reflect the LevelInstance actor's selected state
-	if (ULevelInstanceSubsystem* LevelInstanceSubsystem = GetLevelInstanceSubsystem())
-	{
-		LevelInstanceSubsystem->ForEachActorInLevelInstance(this, [](AActor* LevelActor)
-		{
-			if(ALevelInstanceEditorInstanceActor* EditorInstanceActor = Cast<ALevelInstanceEditorInstanceActor>(LevelActor))
-			{
-				EditorInstanceActor->PushSelectionToProxies();
-				return false;
-			}
-			return true;
-		});
-	}
+	LevelInstanceActorImpl.PushSelectionToProxies();
 }
 
 void ALevelInstance::PushLevelInstanceEditingStateToProxies(bool bInEditingState)
 {
 	Super::PushLevelInstanceEditingStateToProxies(bInEditingState);
 
-	// Actors of the LevelInstance need to reflect the LevelInstance actor's Editing state
-	if (ULevelInstanceSubsystem* LevelInstanceSubsystem = GetLevelInstanceSubsystem())
-	{
-		LevelInstanceSubsystem->ForEachActorInLevelInstance(this, [bInEditingState](AActor* LevelActor)
-			{
-				LevelActor->PushLevelInstanceEditingStateToProxies(bInEditingState);
-				return true;
-			});
-	}
+	LevelInstanceActorImpl.PushLevelInstanceEditingStateToProxies(bInEditingState);
 }
 
 #endif
diff --git a/Engine/Source/Runtime/Engine/Private/LevelInstance/LevelInstanceActorGuid.cpp b/Engine/Source/Runtime/Engine/Private/LevelInstance/LevelInstanceActorGuid.cpp
new file mode 100644
--- /dev/null
+++ b/Engine/Source/Runtime/Engine/Private/LevelInstance/LevelInstanceActorGuid.cpp
@@ -0,0 +1,56 @@
+// Copyright Epic Games, Inc. All Rights Reserved.
+
+#include "LevelInstance/LevelInstanceActorGuid.h"
+#include "GameFramework/Actor.h"
+
+#if !WITH_EDITOR
+void FLevelInstanceActorGuid::AssignIfInvalid()
+{
+	if (!ActorGuid.IsValid())
+	{
+		ActorGuid = FGuid::NewGuid();
+	}
+}
+#endif
+
+const FGuid& FLevelInstanceActorGuid::GetGuid() const
+{
+	check(Actor);
+#if WITH_EDITOR
+	const FGuid& Guid = Actor->GetActorGuid();
+#else
+	const FGuid& Guid = ActorGuid;
+#endif
+	check(Actor->IsTemplate() || Guid.IsValid());
+	return Guid;
+}
+
+FArchive& operator<<(FArchive& Ar, FLevelInstanceActorGuid& LevelInstanceActorGuid)
+{
+	check(LevelInstanceActorGuid.Actor);
+#if WITH_EDITOR
+	if (Ar.IsSaving() && Ar.IsCooking() && !LevelInstanceActorGuid.Actor->IsTemplate())
+	{
+		FGuid Guid = LevelInstanceActorGuid.GetGuid();
+		Ar << Guid;
+	}
+#else
+	if (Ar.IsLoading())
+	{
+		if (LevelInstanceActorGuid.Actor->IsTemplate())
+		{
+			check(!LevelInstanceActorGuid.ActorGuid.IsValid());
+		}
+		else if (Ar.GetPortFlags() & PPF_Duplicate)
+		{
+			LevelInstanceActorGuid.ActorGuid = FGuid::NewGuid();
+		}
+		else if (Ar.IsPersistent())
+		{
+			Ar << LevelInstanceActorGuid.ActorGuid;
+			check(LevelInstanceActorGuid.ActorGuid.IsValid());
+		}
+	}
+#endif
+	return Ar;
+}
\ No newline at end of file
diff --git a/Engine/Source/Runtime/Engine/Private/LevelInstance/LevelInstanceActorImpl.cpp b/Engine/Source/Runtime/Engine/Private/LevelInstance/LevelInstanceActorImpl.cpp
new file mode 100644
--- /dev/null
+++ b/Engine/Source/Runtime/Engine/Private/LevelInstance/LevelInstanceActorImpl.cpp
@@ -0,0 +1,388 @@
+// Copyright Epic Games, Inc. All Rights Reserved.
+
+#include "LevelInstance/LevelInstanceActorImpl.h"
+
+#if WITH_EDITOR
+
+#include "LevelInstance/LevelInstancePrivate.h"
+#include "LevelInstance/LevelInstanceComponent.h"
+#include "LevelInstance/LevelInstanceEditorInstanceActor.h"
+#include "Engine/World.h"
+#include "Logging/MessageLog.h"
+#include "Logging/TokenizedMessage.h"
+#include "Misc/MapErrors.h"
+#include "Misc/UObjectToken.h"
+
+#endif
+
+#define LOCTEXT_NAMESPACE "LevelInstanceActor"
+
+void FLevelInstanceActorImpl::RegisterLevelInstance()
+{
+	if (ULevelInstanceSubsystem* LevelInstanceSubsystem = LevelInstance->GetLevelInstanceSubsystem())
+	{
+		LevelInstanceID = LevelInstanceSubsystem->RegisterLevelInstance(LevelInstance);
+
+		LevelInstance->LoadLevelInstance();
+
+#if WITH_EDITOR
+		// Make sure transformation is up to date after registration as its possible LevelInstance actor can get unregistered when editing properties
+		// through Details panel. In this case the ULevelInstanceComponent might not be able to update the ALevelInstanceEditorInstanceActor transform.
+		LevelInstance->GetLevelInstanceComponent()->UpdateEditorInstanceActor();
+#endif
+	}
+}
+
+void FLevelInstanceActorImpl::UnregisterLevelInstance()
+{
+	if (ULevelInstanceSubsystem* LevelInstanceSubsystem = LevelInstance->GetLevelInstanceSubsystem())
+	{
+		// If LevelInstance has already been unregistered it will have an Invalid LevelInstanceID. Avoid processing it.
+		if (!LevelInstanceID.IsValid())
+		{
+			return;
+		}
+
+		LevelInstanceSubsystem->UnregisterLevelInstance(LevelInstance);
+
+		LevelInstance->UnloadLevelInstance();
+
+		// To avoid processing PostUnregisterAllComponents multiple times (BP Recompile is one use case)
+		LevelInstanceID = FLevelInstanceID();
+	}
+}
+
+const FLevelInstanceID& FLevelInstanceActorImpl::GetLevelInstanceID() const
+{
+	check(HasValidLevelInstanceID());
+	return LevelInstanceID;
+}
+
+bool FLevelInstanceActorImpl::HasValidLevelInstanceID() const
+{
+	return LevelInstanceID.IsValid();
+}
+
+bool FLevelInstanceActorImpl::IsLoadingEnabled() const
+{
+#if WITH_EDITOR
+	return !bGuardLoadUnload && !CastChecked<AActor>(LevelInstance)->bIsEditorPreviewActor;
+#else
+	return true;
+#endif
+}
+
+void FLevelInstanceActorImpl::OnLevelInstanceLoaded()
+{
+#if WITH_EDITOR
+	AActor* Actor = CastChecked<AActor>(LevelInstance);
+	if (!Actor->GetWorld()->IsGameWorld())
+	{
+		// Propagate bounds dirtyness up and check if we need to hide our LevelInstance because self or ancestor is hidden
+		bool bHiddenInEditor = false;
+		LevelInstance->GetLevelInstanceSubsystem()->ForEachLevelInstanceAncestorsAndSelf(Actor, [&bHiddenInEditor](const ILevelInstanceInterface* AncestorOrSelf)
+		{
+			const AActor* AncestorOrSelfActor = CastChecked<AActor>(AncestorOrSelf);
+			AncestorOrSelfActor->GetLevel()->MarkLevelBoundsDirty();
+			bHiddenInEditor |= AncestorOrSelfActor->IsTemporarilyHiddenInEditor();
+			return true;
+		});
+
+		if (bHiddenInEditor)
+		{
+			Actor->SetIsTemporarilyHiddenInEditor(bHiddenInEditor);
+		}
+	}
+#endif
+}
+
+#if WITH_EDITOR
+
+void FLevelInstanceActorImpl::PreEditUndo(TFunctionRef<void()> SuperCall)
+{
+	CachedLevelInstanceID = LevelInstanceID;
+	CachedWorldAsset = LevelInstance->GetWorldAsset();
+
+	AActor* Actor = CastChecked<AActor>(LevelInstance);
+	bCachedIsTemporarilyHiddenInEditor = Actor->IsTemporarilyHiddenInEditor(false);
+
+	{
+		TGuardValue<bool> LoadUnloadGuard(bGuardLoadUnload, true);
+		SuperCall();
+	}
+}
+
+void FLevelInstanceActorImpl::PostEditUndo(TSharedPtr<ITransactionObjectAnnotation> TransactionAnnotation, TFunctionRef<void(TSharedPtr<ITransactionObjectAnnotation> TransactionAnnotation)> SuperCall)
+{
+	{
+		TGuardValue<bool> LoadUnloadGuard(bGuardLoadUnload, true);
+		SuperCall(TransactionAnnotation);
+	}
+
+	PostEditUndoInternal();
+}
+
+void FLevelInstanceActorImpl::PostEditUndo(TFunctionRef<void()> SuperCall)
+{
+	{
+		TGuardValue<bool> LoadUnloadGuard(bGuardLoadUnload, true);
+		SuperCall();
+	}
+
+	PostEditUndoInternal();
+}
+
+void FLevelInstanceActorImpl::PostEditUndoInternal()
+{
+	if (CachedWorldAsset != LevelInstance->GetWorldAsset())
+	{
+		LevelInstance->UpdateLevelInstanceFromWorldAsset();
+	}
+
+	AActor* Actor = CastChecked<AActor>(LevelInstance);
+	if (bCachedIsTemporarilyHiddenInEditor != Actor->IsTemporarilyHiddenInEditor(false))
+	{
+		if (ULevelInstanceSubsystem* LevelInstanceSubsystem = LevelInstance->GetLevelInstanceSubsystem())
+		{
+			LevelInstanceSubsystem->SetIsTemporarilyHiddenInEditor(LevelInstance, !bCachedIsTemporarilyHiddenInEditor);
+		}
+	}
+
+	// Here we want to load or unload based on our current state
+	if (LevelInstance->HasValidLevelInstanceID() && !LevelInstance->IsLoaded())
+	{
+		LevelInstance->LoadLevelInstance();
+	}
+	else if (!IsValidChecked(Actor))
+	{
+		// Temp restore the ID so that we can unload
+		TGuardValue<FLevelInstanceID> LevelInstanceIDGuard(LevelInstanceID, CachedLevelInstanceID);
+		if (LevelInstance->IsLoaded())
+		{
+			LevelInstance->UnloadLevelInstance();
+		}
+	}
+
+	CachedLevelInstanceID = FLevelInstanceID();
+	CachedWorldAsset.Reset();
+
+	if (ULevelInstanceComponent* LevelInstanceComponent = LevelInstance->GetLevelInstanceComponent())
+	{
+		// Order of operations when undoing may lead to the RootComponent being undone before our actor so we need to make sure we update here and in the component when undoing
+		LevelInstanceComponent->UpdateEditorInstanceActor();
+	}
+}
+
+void FLevelInstanceActorImpl::PostEditImport(TFunctionRef<void()> SuperCall)
+{
+	{
+		TGuardValue<bool> LoadUnloadGuard(bGuardLoadUnload, true);
+		SuperCall();
+	}
+	LevelInstance->UpdateLevelInstanceFromWorldAsset();
+}
+
+bool FLevelInstanceActorImpl::CanEditChange(const FProperty* Property) const
+{
+	if (LevelInstance->IsEditing())
+	{
+		return false;
+	}
+
+	if (LevelInstance->HasDirtyChildren())
+	{
+		return false;
+	}
+
+	return true;
+}
+
+void FLevelInstanceActorImpl::PreEditChange(FProperty* Property, bool bWorldAssetChange, TFunctionRef<void(FProperty*)> SuperCall)
+{
+	{
+		TGuardValue<bool> LoadUnloadGuard(bGuardLoadUnload, true);
+		SuperCall(Property);
+	}
+
+	if (bWorldAssetChange)
+	{
+		CachedWorldAsset = LevelInstance->GetWorldAsset();
+	}
+}
+
+void FLevelInstanceActorImpl::PostEditChangeProperty(FPropertyChangedEvent& PropertyChangedEvent, bool bWorldAssetChange, TFunction<void(FPropertyChangedEvent&)> SuperCall)
+{
+	{
+		TGuardValue<bool> LoadUnloadGuard(bGuardLoadUnload, true);
+		SuperCall(PropertyChangedEvent);
+	}
+		
+	if(bWorldAssetChange)
+	{
+		FString Reason;
+		if (!ULevelInstanceSubsystem::CanUseWorldAsset(LevelInstance, LevelInstance->GetWorldAsset(), &Reason))
+		{
+			UE_LOG(LogLevelInstance, Warning, TEXT("%s"), *Reason);
+			LevelInstance->SetWorldAsset(CachedWorldAsset);
+		}
+		else
+		{
+			LevelInstance->UpdateLevelInstanceFromWorldAsset();
+		}
+		CachedWorldAsset.Reset();
+	}
+}
+
+bool FLevelInstanceActorImpl::CanDeleteSelectedActor(FText& OutReason) const
+{
+	if (LevelInstance->IsEditing())
+	{
+		OutReason = LOCTEXT("HasEditingLevel", "Can't delete LevelInstance because it is editing!");
+		return false;
+	}
+
+	if (LevelInstance->HasChildEdit())
+	{
+		OutReason = LOCTEXT("HasEditingChildLevel", "Can't delete LevelInstance because it has editing child LevelInstances!");
+		return false;
+	}
+	return true;
+}
+
+void FLevelInstanceActorImpl::SetIsTemporarilyHiddenInEditor(bool bIsHidden, TFunctionRef<void(bool)> SuperCall)
+{
+	{
+		TGuardValue<bool> LoadUnloadGuard(bGuardLoadUnload, true);
+		SuperCall(bIsHidden);
+	}
+
+	if (ULevelInstanceSubsystem* LevelInstanceSubsystem = LevelInstance->GetLevelInstanceSubsystem())
+	{
+		LevelInstanceSubsystem->SetIsTemporarilyHiddenInEditor(LevelInstance, bIsHidden);
+	}
+}
+
+bool FLevelInstanceActorImpl::SetIsHiddenEdLayer(bool bIsHiddenEdLayer, TFunctionRef<bool(bool)> SuperCall)
+{
+	bool bHasChanged = false;
+	{
+		TGuardValue<bool> LoadUnloadGuard(bGuardLoadUnload, true);
+		bHasChanged = SuperCall(bIsHiddenEdLayer);
+	}
+
+	if (bHasChanged)
+	{
+		if (ULevelInstanceSubsystem* LevelInstanceSubsystem = LevelInstance->GetLevelInstanceSubsystem())
+		{
+			LevelInstanceSubsystem->SetIsHiddenEdLayer(LevelInstance, bIsHiddenEdLayer);
+		}
+	}
+	return bHasChanged;
+}
+
+void FLevelInstanceActorImpl::EditorGetUnderlyingActors(TSet<AActor*>& OutUnderlyingActors) const
+{
+	if (ULevelInstanceSubsystem* LevelInstanceSubsystem = LevelInstance->GetLevelInstanceSubsystem())
+	{
+		LevelInstanceSubsystem->ForEachActorInLevelInstance(LevelInstance, [&](AActor* LevelActor)
+		{
+			bool bAlreadySet = false;
+			OutUnderlyingActors.Add(LevelActor, &bAlreadySet);
+			if (!bAlreadySet)
+			{
+				LevelActor->EditorGetUnderlyingActors(OutUnderlyingActors);
+			}
+			return true;
+		});
+	}
+}
+
+bool FLevelInstanceActorImpl::IsLockLocation() const
+{
+	return LevelInstance->IsEditing() || LevelInstance->HasChildEdit();
+}
+
+bool FLevelInstanceActorImpl::GetBounds(FBox& OutBounds) const
+{
+	// Add Level Bounds
+	if (LevelInstance->IsLoadingEnabled())
+	{
+		if (ULevelInstanceSubsystem* LevelInstanceSubsystem = LevelInstance->GetLevelInstanceSubsystem())
+		{
+			if (LevelInstanceSubsystem->GetLevelInstanceBounds(LevelInstance, OutBounds))
+			{
+				return true;
+			}
+		}
+	}
+
+	return false;
+}
+
+void FLevelInstanceActorImpl::PushSelectionToProxies()
+{
+	// Actors of the LevelInstance need to reflect the LevelInstance actor's selected state
+	if (ULevelInstanceSubsystem* LevelInstanceSubsystem = LevelInstance->GetLevelInstanceSubsystem())
+	{
+		LevelInstanceSubsystem->ForEachActorInLevelInstance(LevelInstance, [](AActor* LevelActor)
+		{
+			if (ALevelInstanceEditorInstanceActor* EditorInstanceActor = Cast<ALevelInstanceEditorInstanceActor>(LevelActor))
+			{
+				EditorInstanceActor->PushSelectionToProxies();
+				return false;
+			}
+			return true;
+		});
+	}
+}
+
+void FLevelInstanceActorImpl::PushLevelInstanceEditingStateToProxies(bool bInEditingState)
+{
+	// Actors of the LevelInstance need to reflect the LevelInstance actor's Editing state
+	if (ULevelInstanceSubsystem* LevelInstanceSubsystem = LevelInstance->GetLevelInstanceSubsystem())
+	{
+		LevelInstanceSubsystem->ForEachActorInLevelInstance(LevelInstance, [bInEditingState](AActor* LevelActor)
+		{
+			LevelActor->PushLevelInstanceEditingStateToProxies(bInEditingState);
+			return true;
+		});
+	}
+}
+
+void FLevelInstanceActorImpl::CheckForErrors()
+{
+	TArray<TPair<FText, TSoftObjectPtr<UWorld>>> LoopInfo;
+	const ILevelInstanceInterface* LoopStart = nullptr;
+	if (!ULevelInstanceSubsystem::CheckForLoop(LevelInstance, LevelInstance->GetWorldAsset(), & LoopInfo, &LoopStart))
+	{
+		TSharedRef<FTokenizedMessage> Error = FMessageLog("MapCheck").Error()->AddToken(FTextToken::Create(LOCTEXT("LevelInstanceActor_Loop_CheckForErrors", "LevelInstance level loop found!")));
+		const AActor* LoopStartActor = CastChecked<AActor>(LoopStart);
+		TSoftObjectPtr<UWorld> LoopStartAsset(LoopStartActor->GetLevel()->GetTypedOuter<UWorld>());
+		Error->AddToken(FAssetNameToken::Create(LoopStartAsset.GetLongPackageName(), FText::FromString(LoopStartAsset.GetAssetName())));
+		Error->AddToken(FTextToken::Create(FText::FromString(TEXT(":"))));
+		Error->AddToken(FUObjectToken::Create(LoopStartActor));
+		for (int32 i = LoopInfo.Num() - 1; i >= 0; --i)
+		{
+			Error->AddToken(FTextToken::Create(LoopInfo[i].Key));
+			TSoftObjectPtr<UWorld> LevelInstancePtr = LoopInfo[i].Value;
+			Error->AddToken(FAssetNameToken::Create(LevelInstancePtr.GetLongPackageName(), FText::FromString(LevelInstancePtr.GetAssetName())));
+		}
+
+		Error->AddToken(FMapErrorToken::Create(FName(TEXT("LevelInstanceActor_Loop_CheckForErrors"))));
+	}
+
+	FPackagePath WorldAssetPath;
+	if (!FPackagePath::TryFromPackageName(LevelInstance->GetWorldAssetPackage(), WorldAssetPath) || !FPackageName::DoesPackageExist(WorldAssetPath))
+	{
+		TSharedRef<FTokenizedMessage> Error = FMessageLog("MapCheck").Error()
+			->AddToken(FTextToken::Create(LOCTEXT("LevelInstanceActor_InvalidPackage", "LevelInstance actor")))
+			->AddToken(FUObjectToken::Create(CastChecked<AActor>(LevelInstance)))
+			->AddToken(FTextToken::Create(FText::FromString(TEXT("refers to an invalid asset:"))))
+			->AddToken(FAssetNameToken::Create(LevelInstance->GetWorldAsset().GetLongPackageName(), FText::FromString(LevelInstance->GetWorldAsset().GetLongPackageName())))
+			->AddToken(FMapErrorToken::Create(FName(TEXT("LevelInstanceActor_InvalidPackage_CheckForErrors"))));
+	}
+}
+
+#endif
+
+#undef LOCTEXT_NAMESPACE
\ No newline at end of file
diff --git a/Engine/Source/Runtime/Engine/Private/LevelInstance/LevelInstanceComponent.cpp b/Engine/Source/Runtime/Engine/Private/LevelInstance/LevelInstanceComponent.cpp
--- a/Engine/Source/Runtime/Engine/Private/LevelInstance/LevelInstanceComponent.cpp
+++ b/Engine/Source/Runtime/Engine/Private/LevelInstance/LevelInstanceComponent.cpp
@@ -1,6 +1,8 @@
 // Copyright Epic Games, Inc. All Rights Reserved.
 #include "LevelInstance/LevelInstanceComponent.h"
-#include "LevelInstance/LevelInstanceActor.h"
+#include "LevelInstance/LevelInstanceInterface.h"
+#include "LevelInstance/LevelInstanceSubsystem.h"
+#include "LevelInstance/LevelInstanceEditorInstanceActor.h"
 #include "Engine/Texture2D.h"
 #include "Engine/World.h"
 
@@ -60,8 +62,25 @@ void ULevelInstanceComponent::UpdateEditorInstanceActor()
 {
 	if (!CachedEditorInstanceActorPtr.IsValid())
 	{
-		check(GetOuterALevelInstance());
-		CachedEditorInstanceActorPtr = GetOuterALevelInstance()->FindEditorInstanceActor();
+		if (ILevelInstanceInterface* LevelInstance = Cast<ILevelInstanceInterface>(GetOwner()))
+		{
+			if (ULevelInstanceSubsystem* LevelInstanceSubsystem = LevelInstance->GetLevelInstanceSubsystem())
+			{
+				if (LevelInstanceSubsystem->IsLoaded(LevelInstance))
+				{
+					LevelInstanceSubsystem->ForEachActorInLevelInstance(LevelInstance, [this, LevelInstance](AActor* LevelActor)
+					{
+						if (ALevelInstanceEditorInstanceActor* LevelInstanceEditorInstanceActor = Cast<ALevelInstanceEditorInstanceActor>(LevelActor))
+						{
+							check(LevelInstanceEditorInstanceActor->GetLevelInstanceID() == LevelInstance->GetLevelInstanceID());
+							CachedEditorInstanceActorPtr = LevelInstanceEditorInstanceActor;
+							return false;
+						}
+						return true;
+					});
+				}
+			}
+		}
 	}
 
 	if (AActor* EditorInstanceActor = CachedEditorInstanceActorPtr.Get())
diff --git a/Engine/Source/Runtime/Engine/Private/LevelInstance/LevelInstanceComponent.h b/Engine/Source/Runtime/Engine/Private/LevelInstance/LevelInstanceComponent.h
deleted file mode 100644
diff --git a/Engine/Source/Runtime/Engine/Private/LevelInstance/LevelInstanceEditorInstanceActor.cpp b/Engine/Source/Runtime/Engine/Private/LevelInstance/LevelInstanceEditorInstanceActor.cpp
--- a/Engine/Source/Runtime/Engine/Private/LevelInstance/LevelInstanceEditorInstanceActor.cpp
+++ b/Engine/Source/Runtime/Engine/Private/LevelInstance/LevelInstanceEditorInstanceActor.cpp
@@ -1,7 +1,7 @@
 // Copyright Epic Games, Inc. All Rights Reserved.
 
 #include "LevelInstance/LevelInstanceEditorInstanceActor.h"
-#include "LevelInstance/LevelInstanceActor.h"
+#include "LevelInstance/LevelInstanceInterface.h"
 #include "LevelInstance/LevelInstanceSubsystem.h"
 #include "Engine/World.h"
 
@@ -17,13 +17,16 @@ AActor* ALevelInstanceEditorInstanceActor::GetSelectionParent() const
 {
 	if (ULevelInstanceSubsystem* LevelInstanceSubsystem = GetWorld()->GetSubsystem<ULevelInstanceSubsystem>())
 	{
-		return LevelInstanceSubsystem->GetLevelInstance(LevelInstanceID);
+		if (ILevelInstanceInterface* LevelInstance = LevelInstanceSubsystem->GetLevelInstance(LevelInstanceID))
+		{
+			return CastChecked<AActor>(LevelInstance);
+		}
 	}
 
 	return nullptr;
 }
 
-ALevelInstanceEditorInstanceActor* ALevelInstanceEditorInstanceActor::Create(ALevelInstance* LevelInstanceActor, ULevel* LoadedLevel)
+ALevelInstanceEditorInstanceActor* ALevelInstanceEditorInstanceActor::Create(ILevelInstanceInterface* LevelInstance, ULevel* LoadedLevel)
 {
 	FActorSpawnParameters SpawnParams;
 	SpawnParams.OverrideLevel = LoadedLevel;
@@ -31,9 +34,11 @@ ALevelInstanceEditorInstanceActor* ALevelInstanceEditorInstanceActor::Create(ALe
 	SpawnParams.bCreateActorPackage = false;
 	SpawnParams.ObjectFlags = RF_Transient;
 	SpawnParams.bNoFail = true;
+
+	AActor* LevelInstanceActor = CastChecked<AActor>(LevelInstance);
 	ALevelInstanceEditorInstanceActor* InstanceActor = LevelInstanceActor->GetWorld()->SpawnActor<ALevelInstanceEditorInstanceActor>(LevelInstanceActor->GetActorLocation(), LevelInstanceActor->GetActorRotation(), SpawnParams);
 	InstanceActor->SetActorScale3D(LevelInstanceActor->GetActorScale3D());
-	InstanceActor->SetLevelInstanceID(LevelInstanceActor->GetLevelInstanceID());
+	InstanceActor->SetLevelInstanceID(LevelInstance->GetLevelInstanceID());
 	
 	for (AActor* LevelActor : LoadedLevel->Actors)
 	{
diff --git a/Engine/Source/Runtime/Engine/Private/LevelInstance/LevelInstanceEditorLevelStreaming.cpp b/Engine/Source/Runtime/Engine/Private/LevelInstance/LevelInstanceEditorLevelStreaming.cpp
--- a/Engine/Source/Runtime/Engine/Private/LevelInstance/LevelInstanceEditorLevelStreaming.cpp
+++ b/Engine/Source/Runtime/Engine/Private/LevelInstance/LevelInstanceEditorLevelStreaming.cpp
@@ -3,7 +3,7 @@
 #include "LevelInstance/LevelInstanceEditorLevelStreaming.h"
 
 #if WITH_EDITOR
-#include "LevelInstance/LevelInstanceActor.h"
+#include "LevelInstance/LevelInstanceInterface.h"
 #include "LevelInstance/LevelInstanceSubsystem.h"
 #include "LevelInstance/LevelInstanceEditorPivotActor.h"
 #include "EditorLevelUtils.h"
@@ -35,15 +35,18 @@ ULevelStreamingLevelInstanceEditor::ULevelStreamingLevelInstanceEditor(const FOb
 #if WITH_EDITOR
 TOptional<FFolder::FRootObject> ULevelStreamingLevelInstanceEditor::GetFolderRootObject() const
 {
-	if (ALevelInstance* LevelInstance = GetLevelInstanceActor())
+	if (ILevelInstanceInterface* LevelInstance = GetLevelInstance())
 	{
-		return FFolder::FRootObject(LevelInstance);
+		if (AActor* Actor = CastChecked<AActor>(LevelInstance))
+		{
+			return FFolder::FRootObject(Actor);
+		}
 	}
 
 	return TOptional<FFolder::FRootObject>();
 }
 
-ALevelInstance* ULevelStreamingLevelInstanceEditor::GetLevelInstanceActor() const
+ILevelInstanceInterface* ULevelStreamingLevelInstanceEditor::GetLevelInstance() const
 {
 	if (ULevelInstanceSubsystem* LevelInstanceSubsystem = GetWorld()->GetSubsystem<ULevelInstanceSubsystem>())
 	{
@@ -53,20 +56,21 @@ ALevelInstance* ULevelStreamingLevelInstanceEditor::GetLevelInstanceActor() cons
 	return nullptr;
 }
 
-ULevelStreamingLevelInstanceEditor* ULevelStreamingLevelInstanceEditor::Load(ALevelInstance* LevelInstanceActor)
+ULevelStreamingLevelInstanceEditor* ULevelStreamingLevelInstanceEditor::Load(ILevelInstanceInterface* LevelInstance)
 {
+	AActor* LevelInstanceActor = CastChecked<AActor>(LevelInstance);
 	UWorld* CurrentWorld = LevelInstanceActor->GetWorld();
 
-	TGuardValue<FLevelInstanceID> GuardEditLevelInstanceID(EditLevelInstanceID, LevelInstanceActor->GetLevelInstanceID());
-	if (ULevelStreamingLevelInstanceEditor* LevelStreaming = Cast<ULevelStreamingLevelInstanceEditor>(EditorLevelUtils::AddLevelToWorld(CurrentWorld, *LevelInstanceActor->GetWorldAssetPackage(), ULevelStreamingLevelInstanceEditor::StaticClass(), LevelInstanceActor->GetTransform())))
+	TGuardValue<FLevelInstanceID> GuardEditLevelInstanceID(EditLevelInstanceID, LevelInstance->GetLevelInstanceID());
+	if (ULevelStreamingLevelInstanceEditor* LevelStreaming = Cast<ULevelStreamingLevelInstanceEditor>(EditorLevelUtils::AddLevelToWorld(CurrentWorld, *LevelInstance->GetWorldAssetPackage(), ULevelStreamingLevelInstanceEditor::StaticClass(), LevelInstanceActor->GetTransform())))
 	{
 		check(LevelStreaming);
-		check(LevelStreaming->LevelInstanceID == LevelInstanceActor->GetLevelInstanceID());
+		check(LevelStreaming->LevelInstanceID == LevelInstance->GetLevelInstanceID());
 
 		GEngine->BlockTillLevelStreamingCompleted(LevelInstanceActor->GetWorld());
 
 		// Create special actor that will handle changing the pivot of this level
-		ALevelInstancePivot::Create(LevelInstanceActor, LevelStreaming);
+		ALevelInstancePivot::Create(LevelInstance, LevelStreaming);
 
 		return LevelStreaming;
 	}
diff --git a/Engine/Source/Runtime/Engine/Private/LevelInstance/LevelInstanceEditorPivotActor.cpp b/Engine/Source/Runtime/Engine/Private/LevelInstance/LevelInstanceEditorPivotActor.cpp
--- a/Engine/Source/Runtime/Engine/Private/LevelInstance/LevelInstanceEditorPivotActor.cpp
+++ b/Engine/Source/Runtime/Engine/Private/LevelInstance/LevelInstanceEditorPivotActor.cpp
@@ -1,7 +1,7 @@
 // Copyright Epic Games, Inc. All Rights Reserved.
 
 #include "LevelInstance/LevelInstanceEditorPivotActor.h"
-#include "LevelInstance/LevelInstanceActor.h"
+#include "LevelInstance/LevelInstanceInterface.h"
 #include "LevelInstance/LevelInstanceSubsystem.h"
 #include "GameFramework/WorldSettings.h"
 #include "LevelUtils.h"
@@ -23,7 +23,7 @@ ALevelInstancePivot::ALevelInstancePivot(const FObjectInitializer& ObjectInitial
 
 #if WITH_EDITOR
 
-ALevelInstancePivot* ALevelInstancePivot::Create(ALevelInstance* LevelInstanceActor, ULevelStreaming* LevelStreaming)
+ALevelInstancePivot* ALevelInstancePivot::Create(ILevelInstanceInterface* LevelInstance, ULevelStreaming* LevelStreaming)
 {
 	FActorSpawnParameters SpawnParams;
 	SpawnParams.OverrideLevel = LevelStreaming->GetLoadedLevel();
@@ -32,6 +32,7 @@ ALevelInstancePivot* ALevelInstancePivot::Create(ALevelInstance* LevelInstanceAc
 	SpawnParams.bNoFail = true;
 		
 	// We place the pivot actor at the LevelInstance Transform so that it makes sense to the user (the pivot being the zero)
+	AActor* LevelInstanceActor = CastChecked<AActor>(LevelInstance);
 	ALevelInstancePivot* PivotActor = LevelInstanceActor->GetWorld()->SpawnActor<ALevelInstancePivot>(LevelInstanceActor->GetActorLocation(), LevelInstanceActor->GetActorRotation(), SpawnParams);
 	
 	AWorldSettings* WorldSettings = LevelStreaming->GetLoadedLevel()->GetWorldSettings();
@@ -40,7 +41,7 @@ ALevelInstancePivot* ALevelInstancePivot::Create(ALevelInstance* LevelInstanceAc
 	// Keep Spawn World Transform in case Level Instance transform changes.
 	PivotActor->SpawnTransform = PivotActor->GetActorTransform();
 	PivotActor->OriginalPivotOffset = WorldSettings->LevelInstancePivotOffset;
-	PivotActor->SetLevelInstanceID(LevelInstanceActor->GetLevelInstanceID());
+	PivotActor->SetLevelInstanceID(LevelInstance->GetLevelInstanceID());
 	// Set Label last as this will call PostEditChangeProperty and update the offset so other fields need to be setup first.
 	PivotActor->SetActorLabel(TEXT("Pivot"));
 
@@ -83,7 +84,7 @@ void ALevelInstancePivot::SetPivot(ELevelInstancePivotType PivotType, AActor* Pi
 	else if(PivotType == ELevelInstancePivotType::Center || PivotType == ELevelInstancePivotType::CenterMinZ)
 	{
 		ULevelInstanceSubsystem* LevelInstanceSubsystem = UWorld::GetSubsystem<ULevelInstanceSubsystem>(GetWorld());
-		ALevelInstance* LevelInstance = LevelInstanceSubsystem->GetLevelInstance(GetLevelInstanceID());
+		ILevelInstanceInterface* LevelInstance = LevelInstanceSubsystem->GetLevelInstance(GetLevelInstanceID());
 		FBox OutBounds;
 		LevelInstanceSubsystem->GetLevelInstanceBounds(LevelInstance, OutBounds);
 		FVector Location = OutBounds.GetCenter();
diff --git a/Engine/Source/Runtime/Engine/Private/LevelInstance/LevelInstanceInterface.cpp b/Engine/Source/Runtime/Engine/Private/LevelInstance/LevelInstanceInterface.cpp
new file mode 100644
--- /dev/null
+++ b/Engine/Source/Runtime/Engine/Private/LevelInstance/LevelInstanceInterface.cpp
@@ -0,0 +1,206 @@
+// Copyright Epic Games, Inc. All Rights Reserved.
+
+#include "LevelInstance/LevelInstanceInterface.h"
+#include "LevelInstance/LevelInstanceSubsystem.h"
+#include "Engine/World.h"
+#include "GameFramework/Actor.h"
+
+#if WITH_EDITOR
+#include "LevelInstance/LevelInstanceEditorInstanceActor.h"
+#endif
+
+ULevelInstanceInterface::ULevelInstanceInterface(const FObjectInitializer& ObjectInitializer)
+	: Super(ObjectInitializer)
+{}
+
+ULevelInstanceSubsystem* ILevelInstanceInterface::GetLevelInstanceSubsystem() const
+{
+	const AActor* Actor = CastChecked<AActor>(this);
+	return Actor->GetWorld() ? Actor->GetWorld()->GetSubsystem<ULevelInstanceSubsystem>() : nullptr;
+}
+
+bool ILevelInstanceInterface::IsLoaded() const
+{
+	if (ULevelInstanceSubsystem* LevelInstanceSubsystem = GetLevelInstanceSubsystem())
+	{
+		return LevelInstanceSubsystem->IsLoaded(this);
+	}
+
+	return false;
+}
+
+void ILevelInstanceInterface::LoadLevelInstance()
+{
+	if (IsLoadingEnabled())
+	{
+		if (ULevelInstanceSubsystem* LevelInstanceSubsystem = GetLevelInstanceSubsystem())
+		{
+			bool bForce = false;
+#if WITH_EDITOR
+			// When reinstancing or world wasn't ticked between changes. Avoid reloading the level but if the package changed, force the load
+			bForce = IsLoaded() && LevelInstanceSubsystem->GetLevelInstanceLevel(this)->GetPackage()->GetLoadedPath() != FPackagePath::FromPackageNameChecked(GetWorldAssetPackage());
+#endif
+			LevelInstanceSubsystem->RequestLoadLevelInstance(this, bForce);
+		}
+	}
+}
+
+void ILevelInstanceInterface::UnloadLevelInstance()
+{
+	if (IsLoadingEnabled())
+	{
+		if (ULevelInstanceSubsystem* LevelInstanceSubsystem = GetLevelInstanceSubsystem())
+		{
+#if WITH_EDITOR
+			check(!HasDirtyChildren());
+#endif
+			LevelInstanceSubsystem->RequestUnloadLevelInstance(this);
+		}
+	}
+}
+
+#if WITH_EDITOR
+ULevel* ILevelInstanceInterface::GetLoadedLevel() const
+{
+	if (ULevelInstanceSubsystem* LevelInstanceSubsystem = GetLevelInstanceSubsystem())
+	{
+		LevelInstanceSubsystem->GetLevelInstanceLevel(this);
+	}
+
+	return nullptr;
+}
+
+void ILevelInstanceInterface::UpdateLevelInstanceFromWorldAsset()
+{
+	if (HasValidLevelInstanceID())
+	{
+		if (ULevelInstanceSubsystem* LevelInstanceSubsystem = GetLevelInstanceSubsystem())
+		{
+			if (IsWorldAssetValid() && IsLoadingEnabled())
+			{
+				const bool bForceUpdate = true;
+				LevelInstanceSubsystem->RequestLoadLevelInstance(this, bForceUpdate);
+			}
+			else if (IsLoaded())
+			{
+				UnloadLevelInstance();
+			}
+		}
+	}
+}
+
+bool ILevelInstanceInterface::IsEditing() const
+{
+	if (HasValidLevelInstanceID())
+	{
+		if (ULevelInstanceSubsystem* LevelInstanceSubsystem = GetLevelInstanceSubsystem())
+		{
+			return LevelInstanceSubsystem->IsEditingLevelInstance(this);
+		}
+	}
+
+	return false;
+}
+
+bool ILevelInstanceInterface::HasChildEdit() const
+{
+	if (HasValidLevelInstanceID())
+	{
+		if (ULevelInstanceSubsystem* LevelInstanceSubsystem = GetLevelInstanceSubsystem())
+		{
+			return LevelInstanceSubsystem->HasChildEdit(this);
+		}
+	}
+
+	return false;
+}
+
+bool ILevelInstanceInterface::IsDirty() const
+{
+	if (HasValidLevelInstanceID())
+	{
+		if (ULevelInstanceSubsystem* LevelInstanceSubsystem = GetLevelInstanceSubsystem())
+		{
+			return LevelInstanceSubsystem->IsEditingLevelInstanceDirty(this);
+		}
+	}
+
+	return false;
+}
+
+bool ILevelInstanceInterface::HasDirtyChildren() const
+{
+	if (HasValidLevelInstanceID())
+	{
+		if (ULevelInstanceSubsystem* LevelInstanceSubsystem = GetLevelInstanceSubsystem())
+		{
+			return LevelInstanceSubsystem->HasDirtyChildrenLevelInstances(this);
+		}
+	}
+
+	return false;
+}
+
+bool ILevelInstanceInterface::CanEnterEdit(FText* OutReason) const
+{
+	if (HasValidLevelInstanceID())
+	{
+		if (ULevelInstanceSubsystem* LevelInstanceSubsystem = GetLevelInstanceSubsystem())
+		{
+			return LevelInstanceSubsystem->CanEditLevelInstance(this, OutReason);
+		}
+	}
+
+	return false;
+}
+
+bool ILevelInstanceInterface::EnterEdit(AActor* ContextActor)
+{
+	ULevelInstanceSubsystem* LevelInstanceSubsystem = GetLevelInstanceSubsystem();
+	check(LevelInstanceSubsystem);
+	LevelInstanceSubsystem->EditLevelInstance(this, ContextActor);
+	return true;
+}
+
+bool ILevelInstanceInterface::CanExitEdit(bool bDiscardEdits, FText* OutReason) const
+{
+	if (HasValidLevelInstanceID())
+	{
+		if (ULevelInstanceSubsystem* LevelInstanceSubsystem = GetLevelInstanceSubsystem())
+		{
+			return LevelInstanceSubsystem->CanCommitLevelInstance(this, bDiscardEdits, OutReason);
+		}
+	}
+
+	return false;
+}
+
+bool ILevelInstanceInterface::ExitEdit(bool bDiscardEdits)
+{
+	ULevelInstanceSubsystem* LevelInstanceSubsystem = GetLevelInstanceSubsystem();
+	check(LevelInstanceSubsystem);
+	return LevelInstanceSubsystem->CommitLevelInstance(this, bDiscardEdits);
+}
+
+bool ILevelInstanceInterface::SetCurrent()
+{
+	if (ULevelInstanceSubsystem* LevelInstanceSubsystem = GetLevelInstanceSubsystem())
+	{
+		return LevelInstanceSubsystem->SetCurrent(this);
+	}
+
+	return false;
+}
+
+bool ILevelInstanceInterface::MoveActorsTo(const TArray<AActor*>& ActorsToMove)
+{
+	if (ULevelInstanceSubsystem* LevelInstanceSubsystem = GetLevelInstanceSubsystem())
+	{
+		return LevelInstanceSubsystem->MoveActorsTo(this, ActorsToMove);
+	}
+
+	return false;
+}
+
+#endif
+
diff --git a/Engine/Source/Runtime/Engine/Private/LevelInstance/LevelInstanceLevelStreaming.cpp b/Engine/Source/Runtime/Engine/Private/LevelInstance/LevelInstanceLevelStreaming.cpp
--- a/Engine/Source/Runtime/Engine/Private/LevelInstance/LevelInstanceLevelStreaming.cpp
+++ b/Engine/Source/Runtime/Engine/Private/LevelInstance/LevelInstanceLevelStreaming.cpp
@@ -1,7 +1,7 @@
 // Copyright Epic Games, Inc. All Rights Reserved.
 
 #include "LevelInstance/LevelInstanceLevelStreaming.h"
-#include "LevelInstance/LevelInstanceActor.h"
+#include "LevelInstance/LevelInstanceInterface.h"
 #include "LevelInstance/LevelInstancePrivate.h"
 #include "LevelInstance/LevelInstanceSubsystem.h"
 #include "Engine/LevelBounds.h"
@@ -25,7 +25,7 @@ ULevelStreamingLevelInstance::ULevelStreamingLevelInstance(const FObjectInitiali
 #endif
 }
 
-ALevelInstance* ULevelStreamingLevelInstance::GetLevelInstanceActor() const
+ILevelInstanceInterface* ULevelStreamingLevelInstance::GetLevelInstance() const
 {
 	if (ULevelInstanceSubsystem* LevelInstanceSubsystem = GetWorld()->GetSubsystem<ULevelInstanceSubsystem>())
 	{
@@ -38,9 +38,12 @@ ALevelInstance* ULevelStreamingLevelInstance::GetLevelInstanceActor() const
 #if WITH_EDITOR
 TOptional<FFolder::FRootObject> ULevelStreamingLevelInstance::GetFolderRootObject() const
 {
-	if (ALevelInstance* LevelInstance = GetLevelInstanceActor())
+	if (ILevelInstanceInterface* LevelInstance = GetLevelInstance())
 	{
-		return FFolder::FRootObject(LevelInstance);
+		if (AActor* Actor = CastChecked<AActor>(LevelInstance))
+		{
+			return FFolder::FRootObject(Actor);
+		}
 	}
 
 	return TOptional<FFolder::FRootObject>();
@@ -53,34 +56,35 @@ FBox ULevelStreamingLevelInstance::GetBounds() const
 }
 #endif
 
-ULevelStreamingLevelInstance* ULevelStreamingLevelInstance::LoadInstance(ALevelInstance* LevelInstanceActor)
+ULevelStreamingLevelInstance* ULevelStreamingLevelInstance::LoadInstance(ILevelInstanceInterface* LevelInstance)
 {
+	AActor* LevelInstanceActor = CastChecked<AActor>(LevelInstance);
 #if WITH_EDITOR
-	if (!LevelInstanceActor->CheckForLoop(LevelInstanceActor->GetWorldAsset()))
+	if (!ULevelInstanceSubsystem::CheckForLoop(LevelInstance))
 	{
 		UE_LOG(LogLevelInstance, Error, TEXT("Failed to load LevelInstance Actor '%s' because that would cause a loop. Run Map Check for more details."), *LevelInstanceActor->GetPathName());
 		return nullptr;
 	}
 
 	FPackagePath WorldAssetPath;
-	if (!FPackagePath::TryFromPackageName(LevelInstanceActor->GetWorldAssetPackage(), WorldAssetPath) || !FPackageName::DoesPackageExist(WorldAssetPath))
+	if (!FPackagePath::TryFromPackageName(LevelInstance->GetWorldAssetPackage(), WorldAssetPath) || !FPackageName::DoesPackageExist(WorldAssetPath))
 	{
-		UE_LOG(LogLevelInstance, Error, TEXT("Failed to load LevelInstance Actor '%s' because it refers to an invalid package ('%s'). Run Map Check for more details."), *LevelInstanceActor->GetPathName(), *LevelInstanceActor->GetWorldAsset().GetLongPackageName());
+		UE_LOG(LogLevelInstance, Error, TEXT("Failed to load LevelInstance Actor '%s' because it refers to an invalid package ('%s'). Run Map Check for more details."), *LevelInstanceActor->GetPathName(), *LevelInstance->GetWorldAsset().GetLongPackageName());
 		return nullptr;
 	}
 #endif
 
 	bool bOutSuccess = false;
 
-	const FString ShortPackageName = FPackageName::GetShortName(LevelInstanceActor->GetWorldAsset().GetLongPackageName());
+	const FString ShortPackageName = FPackageName::GetShortName(LevelInstance->GetWorldAsset().GetLongPackageName());
 	// Build a unique and deterministic LevelInstance level instance name by using LevelInstanceID. 
 	// Distinguish game from editor since we don't want to duplicate for PIE already loaded editor instances (not yet supported).
-	const FString Suffix = FString::Printf(TEXT("%s_LevelInstance_%016llx_%d"), *ShortPackageName, LevelInstanceActor->GetLevelInstanceID().GetHash(), LevelInstanceActor->GetWorld()->IsGameWorld() ? 1 : 0);
+	const FString Suffix = FString::Printf(TEXT("%s_LevelInstance_%016llx_%d"), *ShortPackageName, LevelInstance->GetLevelInstanceID().GetHash(), LevelInstanceActor->GetWorld()->IsGameWorld() ? 1 : 0);
 	const bool bLoadAsTempPackage = true;
-	ULevelStreamingLevelInstance* LevelStreaming = Cast<ULevelStreamingLevelInstance>(ULevelStreamingDynamic::LoadLevelInstanceBySoftObjectPtr(LevelInstanceActor->GetWorld(), LevelInstanceActor->GetWorldAsset(), LevelInstanceActor->GetActorTransform(), bOutSuccess, Suffix, ULevelStreamingLevelInstance::StaticClass(), bLoadAsTempPackage));
+	ULevelStreamingLevelInstance* LevelStreaming = Cast<ULevelStreamingLevelInstance>(ULevelStreamingDynamic::LoadLevelInstanceBySoftObjectPtr(LevelInstanceActor->GetWorld(), LevelInstance->GetWorldAsset(), LevelInstanceActor->GetActorTransform(), bOutSuccess, Suffix, ULevelStreamingLevelInstance::StaticClass(), bLoadAsTempPackage));
 	if (bOutSuccess)
 	{
-		LevelStreaming->LevelInstanceID = LevelInstanceActor->GetLevelInstanceID();
+		LevelStreaming->LevelInstanceID = LevelInstance->GetLevelInstanceID();
 		
 #if WITH_EDITOR
 		if (!LevelInstanceActor->GetWorld()->IsGameWorld())
@@ -120,7 +124,7 @@ ULevelStreamingLevelInstance* ULevelStreamingLevelInstance::LoadInstance(ALevelI
 					}
 
 					LevelActor->PushSelectionToProxies();
-					LevelActor->PushLevelInstanceEditingStateToProxies(LevelInstanceActor->IsInEditingLevelInstance());
+					LevelActor->PushLevelInstanceEditingStateToProxies(CastChecked<AActor>(LevelInstance)->IsInEditingLevelInstance());
 				}
 			}
 			Level->ForEachActorFolder([](UActorFolder* ActorFolder)
@@ -134,7 +138,7 @@ ULevelStreamingLevelInstance* ULevelStreamingLevelInstance::LoadInstance(ALevelI
 			});
 
 			// Create special actor that will handle selection and transform
-			ALevelInstanceEditorInstanceActor::Create(LevelInstanceActor, Level);
+			ALevelInstanceEditorInstanceActor::Create(LevelInstance, Level);
 		}
 #endif
 		return LevelStreaming;
diff --git a/Engine/Source/Runtime/Engine/Private/LevelInstance/LevelInstanceSubsystem.cpp b/Engine/Source/Runtime/Engine/Private/LevelInstance/LevelInstanceSubsystem.cpp
--- a/Engine/Source/Runtime/Engine/Private/LevelInstance/LevelInstanceSubsystem.cpp
+++ b/Engine/Source/Runtime/Engine/Private/LevelInstance/LevelInstanceSubsystem.cpp
@@ -84,58 +84,58 @@ bool ULevelInstanceSubsystem::DoesSupportWorldType(EWorldType::Type WorldType) c
 	return Super::DoesSupportWorldType(WorldType) || WorldType == EWorldType::EditorPreview;
 }
 
-ALevelInstance* ULevelInstanceSubsystem::GetLevelInstance(FLevelInstanceID LevelInstanceID) const
+ILevelInstanceInterface* ULevelInstanceSubsystem::GetLevelInstance(FLevelInstanceID LevelInstanceID) const
 {
-	if (ALevelInstance*const* LevelInstanceActor = RegisteredLevelInstances.Find(LevelInstanceID))
+	if (ILevelInstanceInterface*const* LevelInstance = RegisteredLevelInstances.Find(LevelInstanceID))
 	{
-		return *LevelInstanceActor;
+		return *LevelInstance;
 	}
 
 	return nullptr;
 }
 
-FLevelInstanceID::FLevelInstanceID(ULevelInstanceSubsystem* LevelInstanceSubsystem, ALevelInstance* LevelInstanceActor)
+FLevelInstanceID::FLevelInstanceID(ULevelInstanceSubsystem* LevelInstanceSubsystem, ILevelInstanceInterface* LevelInstance)
 {
-	LevelInstanceSubsystem->ForEachLevelInstanceAncestorsAndSelf(LevelInstanceActor, [this](const ALevelInstance* AncestorOrSelf)
-		{
-			Guids.Add(AncestorOrSelf->GetLevelInstanceActorGuid());
-			return true;
-		});
+	LevelInstanceSubsystem->ForEachLevelInstanceAncestorsAndSelf(CastChecked<AActor>(LevelInstance), [this](const ILevelInstanceInterface* AncestorOrSelf)
+	{
+		Guids.Add(AncestorOrSelf->GetLevelInstanceGuid());
+		return true;
+	});
 	check(!Guids.IsEmpty());
 	Hash = CityHash64((const char*)Guids.GetData(), Guids.Num() * sizeof(FGuid));
 }
 
-FLevelInstanceID ULevelInstanceSubsystem::RegisterLevelInstance(ALevelInstance* LevelInstanceActor)
+FLevelInstanceID ULevelInstanceSubsystem::RegisterLevelInstance(ILevelInstanceInterface* LevelInstance)
 {
-	FLevelInstanceID LevelInstanceID(this, LevelInstanceActor);
+	FLevelInstanceID LevelInstanceID(this, LevelInstance);
 	check(LevelInstanceID.IsValid());
-	ALevelInstance*& Value = RegisteredLevelInstances.FindOrAdd(LevelInstanceID);
-	check(GIsReinstancing || Value == nullptr || Value == LevelInstanceActor);
-	Value = LevelInstanceActor;
+	ILevelInstanceInterface*& Value = RegisteredLevelInstances.FindOrAdd(LevelInstanceID);
+	check(GIsReinstancing || Value == nullptr || Value == LevelInstance);
+	Value = LevelInstance;
 
 	return LevelInstanceID;
 }
 
-void ULevelInstanceSubsystem::UnregisterLevelInstance(ALevelInstance* LevelInstanceActor)
+void ULevelInstanceSubsystem::UnregisterLevelInstance(ILevelInstanceInterface* LevelInstance)
 {
-	RegisteredLevelInstances.Remove(LevelInstanceActor->GetLevelInstanceID());
+	RegisteredLevelInstances.Remove(LevelInstance->GetLevelInstanceID());
 }
 
-void ULevelInstanceSubsystem::RequestLoadLevelInstance(ALevelInstance* LevelInstanceActor, bool bForce /* = false */)
+void ULevelInstanceSubsystem::RequestLoadLevelInstance(ILevelInstanceInterface* LevelInstance, bool bForce /* = false */)
 {
-	check(LevelInstanceActor && IsValidChecked(LevelInstanceActor) && !LevelInstanceActor->IsUnreachable());
-	if (LevelInstanceActor->IsLevelInstancePathValid())
+	check(LevelInstance && IsValidChecked(CastChecked<AActor>(LevelInstance)) && !CastChecked<AActor>(LevelInstance)->IsUnreachable());
+	if (LevelInstance->IsWorldAssetValid())
 	{
 #if WITH_EDITOR
-		if (!IsEditingLevelInstance(LevelInstanceActor))
+		if (!IsEditingLevelInstance(LevelInstance))
 #endif
 		{
-			LevelInstancesToUnload.Remove(LevelInstanceActor->GetLevelInstanceID());
+			LevelInstancesToUnload.Remove(LevelInstance->GetLevelInstanceID());
 
-			bool* bForcePtr = LevelInstancesToLoadOrUpdate.Find(LevelInstanceActor);
+			bool* bForcePtr = LevelInstancesToLoadOrUpdate.Find(LevelInstance);
 
 			// Avoid loading if already loaded. Can happen if actor requests unload/load in same frame. Without the force it means its not necessary.
-			if (IsLoaded(LevelInstanceActor) && !bForce && (bForcePtr == nullptr || !(*bForcePtr)))
+			if (IsLoaded(LevelInstance) && !bForce && (bForcePtr == nullptr || !(*bForcePtr)))
 			{
 				return;
 			}
@@ -146,26 +146,26 @@ void ULevelInstanceSubsystem::RequestLoadLevelInstance(ALevelInstance* LevelInst
 			}
 			else
 			{
-				LevelInstancesToLoadOrUpdate.Add(LevelInstanceActor, bForce);
+				LevelInstancesToLoadOrUpdate.Add(LevelInstance, bForce);
 			}
 		}
 	}
 }
 
-void ULevelInstanceSubsystem::RequestUnloadLevelInstance(ALevelInstance* LevelInstanceActor)
+void ULevelInstanceSubsystem::RequestUnloadLevelInstance(ILevelInstanceInterface* LevelInstance)
 {
-	const FLevelInstanceID& LevelInstanceID = LevelInstanceActor->GetLevelInstanceID();
+	const FLevelInstanceID& LevelInstanceID = LevelInstance->GetLevelInstanceID();
 	if (LevelInstances.Contains(LevelInstanceID))
 	{
-		// LevelInstancesToUnload uses FLevelInstanceID because LevelInstanceActor* can be destroyed in later Tick and we don't need it.
+		// LevelInstancesToUnload uses FLevelInstanceID because LevelInstance* can be destroyed in later Tick and we don't need it.
 		LevelInstancesToUnload.Add(LevelInstanceID);
 	}
-	LevelInstancesToLoadOrUpdate.Remove(LevelInstanceActor);
+	LevelInstancesToLoadOrUpdate.Remove(LevelInstance);
 }
 
-bool ULevelInstanceSubsystem::IsLoaded(const ALevelInstance* LevelInstanceActor) const
+bool ULevelInstanceSubsystem::IsLoaded(const ILevelInstanceInterface* LevelInstance) const
 {
-	return LevelInstanceActor->HasValidLevelInstanceID() && LevelInstances.Contains(LevelInstanceActor->GetLevelInstanceID());
+	return LevelInstance->HasValidLevelInstanceID() && LevelInstances.Contains(LevelInstance->GetLevelInstanceID());
 }
 
 void ULevelInstanceSubsystem::UpdateStreamingState()
@@ -204,16 +204,16 @@ void ULevelInstanceSubsystem::UpdateStreamingState()
 	if (LevelInstancesToLoadOrUpdate.Num())
 	{
 		// Unload levels before doing any loading
-		TMap<ALevelInstance*, bool> LevelInstancesToLoadOrUpdateCopy(MoveTemp(LevelInstancesToLoadOrUpdate));
-		for (const TPair<ALevelInstance*, bool>& Pair : LevelInstancesToLoadOrUpdateCopy)
+		TMap<ILevelInstanceInterface*, bool> LevelInstancesToLoadOrUpdateCopy(MoveTemp(LevelInstancesToLoadOrUpdate));
+		for (const TPair<ILevelInstanceInterface*, bool>& Pair : LevelInstancesToLoadOrUpdateCopy)
 		{
 #if WITH_EDITOR
 			SlowTask.EnterProgressFrame(1.f, LOCTEXT("UnloadingLevelInstance", "Unloading Level Instance"));
 #endif
-			ALevelInstance* LevelInstanceActor = Pair.Key;
+			ILevelInstanceInterface* LevelInstance = Pair.Key;
 			if (Pair.Value)
 			{
-				UnloadLevelInstance(LevelInstanceActor->GetLevelInstanceID());
+				UnloadLevelInstance(LevelInstance->GetLevelInstanceID());
 			}
 		}
 
@@ -221,7 +221,7 @@ void ULevelInstanceSubsystem::UpdateStreamingState()
 		LevelsToRemoveScope.Reset();
 		double StartTime = FPlatformTime::Seconds();
 #endif
-		for (const TPair<ALevelInstance*, bool>& Pair : LevelInstancesToLoadOrUpdateCopy)
+		for (const TPair<ILevelInstanceInterface*, bool>& Pair : LevelInstancesToLoadOrUpdateCopy)
 		{
 #if WITH_EDITOR
 			SlowTask.EnterProgressFrame(1.f, LOCTEXT("LoadingLevelInstance", "Loading Level Instance"));
@@ -241,13 +241,13 @@ void ULevelInstanceSubsystem::UpdateStreamingState()
 
 void ULevelInstanceSubsystem::RegisterLoadedLevelStreamingLevelInstance(ULevelStreamingLevelInstance* LevelStreaming)
 {
-	ALevelInstance* LevelInstanceActor = LevelStreaming->GetLevelInstanceActor();
-	const FLevelInstanceID LevelInstanceID = LevelStreaming->GetLevelInstanceActor()->GetLevelInstanceID();
+	ILevelInstanceInterface* LevelInstance = LevelStreaming->GetLevelInstance();
+	const FLevelInstanceID LevelInstanceID = LevelInstance->GetLevelInstanceID();
 	check(!LevelInstances.Contains(LevelInstanceID));
-	FLevelInstance& LevelInstance = LevelInstances.Add(LevelInstanceID);
-	LevelInstance.LevelStreaming = LevelStreaming;
+	FLevelInstance& LevelInstanceEntry = LevelInstances.Add(LevelInstanceID);
+	LevelInstanceEntry.LevelStreaming = LevelStreaming;
 #if WITH_EDITOR
-	LevelInstanceActor->OnLevelInstanceLoaded();
+	LevelInstance->OnLevelInstanceLoaded();
 #endif
 }
 
@@ -257,8 +257,8 @@ void ULevelInstanceSubsystem::RegisterLoadedLevelStreamingLevelInstanceEditor(UL
 	if (!bIsCreatingLevelInstance)
 	{
 		check(!LevelInstanceEdit.IsValid());
-		ALevelInstance* LevelInstanceActor = LevelStreaming->GetLevelInstanceActor();
-		LevelInstanceEdit = MakeUnique<FLevelInstanceEdit>(LevelStreaming, LevelInstanceActor->GetLevelInstanceID());
+		ILevelInstanceInterface* LevelInstance = LevelStreaming->GetLevelInstance();
+		LevelInstanceEdit = MakeUnique<FLevelInstanceEdit>(LevelStreaming, LevelInstance->GetLevelInstanceID());
 
 		if (ILevelInstanceEditorModule* EditorModule = FModuleManager::GetModulePtr<ILevelInstanceEditorModule>("LevelInstanceEditor"))
 		{
@@ -291,18 +291,19 @@ void ULevelInstanceSubsystem::ResetEdit(TUniquePtr<FLevelInstanceEdit>& InLevelI
 }
 #endif
 
-void ULevelInstanceSubsystem::LoadLevelInstance(ALevelInstance* LevelInstanceActor)
+void ULevelInstanceSubsystem::LoadLevelInstance(ILevelInstanceInterface* LevelInstance)
 {
-	check(LevelInstanceActor);
-	if (IsLoaded(LevelInstanceActor) || !IsValidChecked(LevelInstanceActor) || LevelInstanceActor->IsUnreachable() || !LevelInstanceActor->IsLevelInstancePathValid())
+	check(LevelInstance);
+	AActor* LevelInstanceActor = CastChecked<AActor>(LevelInstance);
+	if (IsLoaded(LevelInstance) || !IsValidChecked(LevelInstanceActor) || LevelInstanceActor->IsUnreachable() || !LevelInstance->IsWorldAssetValid())
 	{
 		return;
 	}
 
-	const FLevelInstanceID& LevelInstanceID = LevelInstanceActor->GetLevelInstanceID();
+	const FLevelInstanceID& LevelInstanceID = LevelInstance->GetLevelInstanceID();
 	check(!LevelInstances.Contains(LevelInstanceID));
 
-	if (ULevelStreamingLevelInstance* LevelStreaming = ULevelStreamingLevelInstance::LoadInstance(LevelInstanceActor))
+	if (ULevelStreamingLevelInstance* LevelStreaming = ULevelStreamingLevelInstance::LoadInstance(LevelInstance))
 	{
 #if WITH_EDITOR
 		check(LevelInstanceActor->GetWorld()->IsGameWorld() || LevelInstances.Contains(LevelInstanceID));
@@ -329,12 +330,12 @@ void ULevelInstanceSubsystem::UnloadLevelInstance(const FLevelInstanceID& LevelI
 		{
 			ForEachActorInLevel(LoadedLevel, [this](AActor* LevelActor)
 			{
-				if (ALevelInstance* LevelInstanceActor = Cast<ALevelInstance>(LevelActor))
+				if (ILevelInstanceInterface* LevelInstance = Cast<ILevelInstanceInterface>(LevelActor))
 				{
 					// Make sure to remove from pending loads if we are unloading child can't be loaded
-					LevelInstancesToLoadOrUpdate.Remove(LevelInstanceActor);
+					LevelInstancesToLoadOrUpdate.Remove(LevelInstance);
 					
-					UnloadLevelInstance(LevelInstanceActor->GetLevelInstanceID());
+					UnloadLevelInstance(LevelInstance->GetLevelInstanceID());
 				}
 				return true;
 			});
@@ -365,11 +366,11 @@ void ULevelInstanceSubsystem::ForEachActorInLevel(ULevel* Level, TFunctionRef<bo
 	}
 }
 
-void ULevelInstanceSubsystem::ForEachLevelInstanceAncestorsAndSelf(AActor* Actor, TFunctionRef<bool(ALevelInstance*)> Operation) const
+void ULevelInstanceSubsystem::ForEachLevelInstanceAncestorsAndSelf(AActor* Actor, TFunctionRef<bool(ILevelInstanceInterface*)> Operation) const
 {
-	if (ALevelInstance* LevelInstanceActor = Cast<ALevelInstance>(Actor))
+	if (ILevelInstanceInterface* LevelInstance = Cast<ILevelInstanceInterface>(Actor))
 	{
-		if (!Operation(LevelInstanceActor))
+		if (!Operation(LevelInstance))
 		{
 			return;
 		}
@@ -378,31 +379,31 @@ void ULevelInstanceSubsystem::ForEachLevelInstanceAncestorsAndSelf(AActor* Actor
 	ForEachLevelInstanceAncestors(Actor, Operation);
 }
 
-void ULevelInstanceSubsystem::ForEachLevelInstanceAncestors(AActor* Actor, TFunctionRef<bool(ALevelInstance*)> Operation) const
+void ULevelInstanceSubsystem::ForEachLevelInstanceAncestors(AActor* Actor, TFunctionRef<bool(ILevelInstanceInterface*)> Operation) const
 {
-	ALevelInstance* ParentLevelInstance = nullptr;
+	ILevelInstanceInterface* ParentLevelInstance = nullptr;
 	do
 	{
 		ParentLevelInstance = GetOwningLevelInstance(Actor->GetLevel());
-		Actor = ParentLevelInstance;
+		Actor = Cast<AActor>(ParentLevelInstance);
 
 	} while (ParentLevelInstance != nullptr && Operation(ParentLevelInstance));
 }
 
-ALevelInstance* ULevelInstanceSubsystem::GetOwningLevelInstance(const ULevel* Level) const
+ILevelInstanceInterface* ULevelInstanceSubsystem::GetOwningLevelInstance(const ULevel* Level) const
 {
 	if (ULevelStreaming* BaseLevelStreaming = FLevelUtils::FindStreamingLevel(Level))
 	{
 #if WITH_EDITOR
 		if (ULevelStreamingLevelInstanceEditor* LevelStreamingEditor = Cast<ULevelStreamingLevelInstanceEditor>(BaseLevelStreaming))
 		{
-			return LevelStreamingEditor->GetLevelInstanceActor();
+			return LevelStreamingEditor->GetLevelInstance();
 		}
 		else 
 #endif
 		if (ULevelStreamingLevelInstance* LevelStreaming = Cast<ULevelStreamingLevelInstance>(BaseLevelStreaming))
 		{
-			return LevelStreaming->GetLevelInstanceActor();
+			return LevelStreaming->GetLevelInstance();
 		}
 		else if (UWorldPartitionLevelStreamingDynamic* WorldPartitionLevelStreaming = Cast<UWorldPartitionLevelStreamingDynamic>(BaseLevelStreaming))
 		{
@@ -448,11 +449,11 @@ bool ULevelInstanceSubsystem::OnExitEditorModeInternal(bool bForceExit)
 	if (LevelInstanceEdit)
 	{
 		TGuardValue<bool> CommitScope(bIsCommittingLevelInstance, true);
-		ALevelInstance* LevelInstance = GetEditingLevelInstance();
+		ILevelInstanceInterface* LevelInstance = GetEditingLevelInstance();
 
 		bool bDiscard = false;
 		bool bIsDirty = IsLevelInstanceEditDirty(LevelInstanceEdit.Get());
-		if (bIsDirty && CanDiscardLevelInstance(LevelInstance))
+		if (bIsDirty && CanCommitLevelInstance(LevelInstance, /*bDiscardEdits=*/true))
 		{
 			FText Title = LOCTEXT("CommitOrDiscardChangesTitle", "Save changes?");
 			// if bForceExit we can't cancel the exiting of the mode so the user needs to decide between saving or discarding
@@ -559,27 +560,27 @@ void ULevelInstanceSubsystem::PackAllLoadedActors()
 
 	for (APackedLevelActor* PackedLevelActor : PackedLevelActorsToUpdate)
 	{
-		PackedLevelActor->UpdateFromLevel();
+		PackedLevelActor->UpdateLevelInstanceFromWorldAsset();
 		UpdateProgress();
 	}
 }
 
-bool ULevelInstanceSubsystem::GetLevelInstanceBounds(const ALevelInstance* LevelInstanceActor, FBox& OutBounds) const
+bool ULevelInstanceSubsystem::GetLevelInstanceBounds(const ILevelInstanceInterface* LevelInstance, FBox& OutBounds) const
 {
-	if (IsLoaded(LevelInstanceActor))
+	if (IsLoaded(LevelInstance))
 	{
-		const FLevelInstance& LevelInstance = LevelInstances.FindChecked(LevelInstanceActor->GetLevelInstanceID());
-		OutBounds = LevelInstance.LevelStreaming->GetBounds();
+		const FLevelInstance& LevelInstanceEntry = LevelInstances.FindChecked(LevelInstance->GetLevelInstanceID());
+		OutBounds = LevelInstanceEntry.LevelStreaming->GetBounds();
 		return true;
 	}
-	else if (const FLevelInstanceEdit* CurrentEdit = GetLevelInstanceEdit(LevelInstanceActor))
+	else if (const FLevelInstanceEdit* CurrentEdit = GetLevelInstanceEdit(LevelInstance))
 	{
 		OutBounds = CurrentEdit->LevelStreaming->GetBounds();
 		return true;
 	}
-	else if(LevelInstanceActor->IsLevelInstancePathValid())
+	else if(LevelInstance->IsWorldAssetValid())
 	{
-		return GetLevelInstanceBoundsFromPackage(LevelInstanceActor->GetActorTransform(), FName(*LevelInstanceActor->GetWorldAssetPackage()), OutBounds);
+		return GetLevelInstanceBoundsFromPackage(CastChecked<AActor>(LevelInstance)->GetActorTransform(), FName(*LevelInstance->GetWorldAssetPackage()), OutBounds);
 	}
 
 	return false;
@@ -605,19 +606,19 @@ bool ULevelInstanceSubsystem::GetLevelInstanceBoundsFromPackage(const FTransform
 	return false;
 }
 
-void ULevelInstanceSubsystem::ForEachActorInLevelInstance(const ALevelInstance* LevelInstanceActor, TFunctionRef<bool(AActor * LevelActor)> Operation) const
+void ULevelInstanceSubsystem::ForEachActorInLevelInstance(const ILevelInstanceInterface* LevelInstance, TFunctionRef<bool(AActor * LevelActor)> Operation) const
 {
-	if (ULevel* LevelInstanceLevel = GetLevelInstanceLevel(LevelInstanceActor))
+	if (ULevel* LevelInstanceLevel = GetLevelInstanceLevel(LevelInstance))
 	{
 		ForEachActorInLevel(LevelInstanceLevel, Operation);
 	}
 }
 
-void ULevelInstanceSubsystem::ForEachLevelInstanceAncestorsAndSelf(const AActor* Actor, TFunctionRef<bool(const ALevelInstance*)> Operation) const
+void ULevelInstanceSubsystem::ForEachLevelInstanceAncestorsAndSelf(const AActor* Actor, TFunctionRef<bool(const ILevelInstanceInterface*)> Operation) const
 {
-	if (const ALevelInstance* LevelInstanceActor = Cast<ALevelInstance>(Actor))
+	if (const ILevelInstanceInterface* LevelInstance = Cast<ILevelInstanceInterface>(Actor))
 	{
-		if (!Operation(LevelInstanceActor))
+		if (!Operation(LevelInstance))
 		{
 			return;
 		}
@@ -626,36 +627,36 @@ void ULevelInstanceSubsystem::ForEachLevelInstanceAncestorsAndSelf(const AActor*
 	ForEachLevelInstanceAncestors(Actor, Operation);
 }
 
-void ULevelInstanceSubsystem::ForEachLevelInstanceAncestors(const AActor* Actor, TFunctionRef<bool(const ALevelInstance*)> Operation) const
+void ULevelInstanceSubsystem::ForEachLevelInstanceAncestors(const AActor* Actor, TFunctionRef<bool(const ILevelInstanceInterface*)> Operation) const
 {
-	const ALevelInstance* ParentLevelInstance = nullptr;
+	const ILevelInstanceInterface* ParentLevelInstance = nullptr;
 	do 
 	{
 		ParentLevelInstance = GetOwningLevelInstance(Actor->GetLevel());
-		Actor = ParentLevelInstance;
+		Actor = Cast<AActor>(ParentLevelInstance);
 	} 
 	while (ParentLevelInstance != nullptr && Operation(ParentLevelInstance));
 }
 
-void ULevelInstanceSubsystem::ForEachLevelInstanceChild(const ALevelInstance* LevelInstanceActor, bool bRecursive, TFunctionRef<bool(const ALevelInstance*)> Operation) const
+void ULevelInstanceSubsystem::ForEachLevelInstanceChild(const ILevelInstanceInterface* LevelInstance, bool bRecursive, TFunctionRef<bool(const ILevelInstanceInterface*)> Operation) const
 {
-	ForEachLevelInstanceChildImpl(LevelInstanceActor, bRecursive, Operation);
+	ForEachLevelInstanceChildImpl(LevelInstance, bRecursive, Operation);
 }
 
-bool ULevelInstanceSubsystem::ForEachLevelInstanceChildImpl(const ALevelInstance* LevelInstanceActor, bool bRecursive, TFunctionRef<bool(const ALevelInstance*)> Operation) const
+bool ULevelInstanceSubsystem::ForEachLevelInstanceChildImpl(const ILevelInstanceInterface* LevelInstance, bool bRecursive, TFunctionRef<bool(const ILevelInstanceInterface*)> Operation) const
 {
 	bool bContinue = true;
-	if (ULevel* LevelInstanceLevel = GetLevelInstanceLevel(LevelInstanceActor))
+	if (ULevel* LevelInstanceLevel = GetLevelInstanceLevel(LevelInstance))
 	{
 		ForEachActorInLevel(LevelInstanceLevel, [&bContinue, this, Operation,bRecursive](AActor* LevelActor)
 		{
-			if (const ALevelInstance* ChildLevelInstanceActor = Cast<ALevelInstance>(LevelActor))
+			if (const ILevelInstanceInterface* ChildLevelInstance = Cast<ILevelInstanceInterface>(LevelActor))
 			{
-				bContinue = Operation(ChildLevelInstanceActor);
+				bContinue = Operation(ChildLevelInstance);
 				
 				if (bContinue && bRecursive)
 				{
-					bContinue = ForEachLevelInstanceChildImpl(ChildLevelInstanceActor, bRecursive, Operation);
+					bContinue = ForEachLevelInstanceChildImpl(ChildLevelInstance, bRecursive, Operation);
 				}
 			}
 			return bContinue;
@@ -665,25 +666,25 @@ bool ULevelInstanceSubsystem::ForEachLevelInstanceChildImpl(const ALevelInstance
 	return bContinue;
 }
 
-void ULevelInstanceSubsystem::ForEachLevelInstanceChild(ALevelInstance* LevelInstanceActor, bool bRecursive, TFunctionRef<bool(ALevelInstance*)> Operation) const
+void ULevelInstanceSubsystem::ForEachLevelInstanceChild(ILevelInstanceInterface* LevelInstance, bool bRecursive, TFunctionRef<bool(ILevelInstanceInterface*)> Operation) const
 {
-	ForEachLevelInstanceChildImpl(LevelInstanceActor, bRecursive, Operation);
+	ForEachLevelInstanceChildImpl(LevelInstance, bRecursive, Operation);
 }
 
-bool ULevelInstanceSubsystem::ForEachLevelInstanceChildImpl(ALevelInstance* LevelInstanceActor, bool bRecursive, TFunctionRef<bool(ALevelInstance*)> Operation) const
+bool ULevelInstanceSubsystem::ForEachLevelInstanceChildImpl(ILevelInstanceInterface* LevelInstance, bool bRecursive, TFunctionRef<bool(ILevelInstanceInterface*)> Operation) const
 {
 	bool bContinue = true;
-	if (ULevel* LevelInstanceLevel = GetLevelInstanceLevel(LevelInstanceActor))
+	if (ULevel* LevelInstanceLevel = GetLevelInstanceLevel(LevelInstance))
 	{
 		ForEachActorInLevel(LevelInstanceLevel, [&bContinue, this, Operation, bRecursive](AActor* LevelActor)
 		{
-			if (ALevelInstance* ChildLevelInstanceActor = Cast<ALevelInstance>(LevelActor))
+			if (ILevelInstanceInterface* ChildLevelInstance = Cast<ILevelInstanceInterface>(LevelActor))
 			{
-				bContinue = Operation(ChildLevelInstanceActor);
+				bContinue = Operation(ChildLevelInstance);
 
 				if (bContinue && bRecursive)
 				{
-					bContinue = ForEachLevelInstanceChildImpl(ChildLevelInstanceActor, bRecursive, Operation);
+					bContinue = ForEachLevelInstanceChildImpl(ChildLevelInstance, bRecursive, Operation);
 				}
 			}
 			return bContinue;
@@ -693,12 +694,12 @@ bool ULevelInstanceSubsystem::ForEachLevelInstanceChildImpl(ALevelInstance* Leve
 	return bContinue;
 }
 
-bool ULevelInstanceSubsystem::HasDirtyChildrenLevelInstances(const ALevelInstance* LevelInstanceActor) const
+bool ULevelInstanceSubsystem::HasDirtyChildrenLevelInstances(const ILevelInstanceInterface* LevelInstance) const
 {
 	bool bDirtyChildren = false;
-	ForEachLevelInstanceChild(LevelInstanceActor, /*bRecursive=*/true, [this, &bDirtyChildren](const ALevelInstance* ChildLevelInstanceActor)
+	ForEachLevelInstanceChild(LevelInstance, /*bRecursive=*/true, [this, &bDirtyChildren](const ILevelInstanceInterface* ChildLevelInstance)
 	{
-		if (IsEditingLevelInstanceDirty(ChildLevelInstanceActor))
+		if (IsEditingLevelInstanceDirty(ChildLevelInstance))
 		{
 			bDirtyChildren = true;
 			return false;
@@ -708,9 +709,9 @@ bool ULevelInstanceSubsystem::HasDirtyChildrenLevelInstances(const ALevelInstanc
 	return bDirtyChildren;
 }
 
-void ULevelInstanceSubsystem::SetIsHiddenEdLayer(ALevelInstance* LevelInstanceActor, bool bIsHiddenEdLayer)
+void ULevelInstanceSubsystem::SetIsHiddenEdLayer(ILevelInstanceInterface* LevelInstance, bool bIsHiddenEdLayer)
 {
-	if (ULevel* LevelInstanceLevel = GetLevelInstanceLevel(LevelInstanceActor))
+	if (ULevel* LevelInstanceLevel = GetLevelInstanceLevel(LevelInstance))
 	{
 		ForEachActorInLevel(LevelInstanceLevel, [bIsHiddenEdLayer](AActor* LevelActor)
 		{
@@ -720,9 +721,9 @@ void ULevelInstanceSubsystem::SetIsHiddenEdLayer(ALevelInstance* LevelInstanceAc
 	}
 }
 
-void ULevelInstanceSubsystem::SetIsTemporarilyHiddenInEditor(ALevelInstance* LevelInstanceActor, bool bIsHidden)
+void ULevelInstanceSubsystem::SetIsTemporarilyHiddenInEditor(ILevelInstanceInterface* LevelInstance, bool bIsHidden)
 {
-	if (ULevel* LevelInstanceLevel = GetLevelInstanceLevel(LevelInstanceActor))
+	if (ULevel* LevelInstanceLevel = GetLevelInstanceLevel(LevelInstance))
 	{
 		ForEachActorInLevel(LevelInstanceLevel, [bIsHidden](AActor* LevelActor)
 		{
@@ -732,21 +733,21 @@ void ULevelInstanceSubsystem::SetIsTemporarilyHiddenInEditor(ALevelInstance* Lev
 	}
 }
 
-bool ULevelInstanceSubsystem::SetCurrent(ALevelInstance* LevelInstanceActor) const
+bool ULevelInstanceSubsystem::SetCurrent(ILevelInstanceInterface* LevelInstance) const
 {
-	if (IsEditingLevelInstance(LevelInstanceActor))
+	if (IsEditingLevelInstance(LevelInstance))
 	{
-		return GetWorld()->SetCurrentLevel(GetLevelInstanceLevel(LevelInstanceActor));
+		return GetWorld()->SetCurrentLevel(GetLevelInstanceLevel(LevelInstance));
 	}
 
 	return false;
 }
 
-bool ULevelInstanceSubsystem::IsCurrent(const ALevelInstance* LevelInstanceActor) const
+bool ULevelInstanceSubsystem::IsCurrent(const ILevelInstanceInterface* LevelInstance) const
 {
-	if (IsEditingLevelInstance(LevelInstanceActor))
+	if (IsEditingLevelInstance(LevelInstance))
 	{
-		return GetLevelInstanceLevel(LevelInstanceActor) == GetWorld()->GetCurrentLevel();
+		return GetLevelInstanceLevel(LevelInstance) == GetWorld()->GetCurrentLevel();
 	}
 
 	return false;
@@ -765,7 +766,7 @@ bool ULevelInstanceSubsystem::MoveActorsToLevel(const TArray<AActor*>& ActorsToR
 		return false;
 	}
 
-	ALevelInstance* OwningInstance = GetOwningLevelInstance(DestinationLevel);
+	ILevelInstanceInterface* OwningInstance = GetOwningLevelInstance(DestinationLevel);
 	if (!OwningInstance || !OwningInstance->IsEditing())
 	{
 		for (const auto& Actor : ActorsToRemove)
@@ -778,16 +779,16 @@ bool ULevelInstanceSubsystem::MoveActorsToLevel(const TArray<AActor*>& ActorsToR
 	return true;
 }
 
-bool ULevelInstanceSubsystem::MoveActorsTo(ALevelInstance* LevelInstanceActor, const TArray<AActor*>& ActorsToMove, TArray<AActor*>* OutActors /*= nullptr*/)
+bool ULevelInstanceSubsystem::MoveActorsTo(ILevelInstanceInterface* LevelInstance, const TArray<AActor*>& ActorsToMove, TArray<AActor*>* OutActors /*= nullptr*/)
 {
-	check(IsEditingLevelInstance(LevelInstanceActor));
-	ULevel* LevelInstanceLevel = GetLevelInstanceLevel(LevelInstanceActor);
+	check(IsEditingLevelInstance(LevelInstance));
+	ULevel* LevelInstanceLevel = GetLevelInstanceLevel(LevelInstance);
 	check(LevelInstanceLevel);
 
 	return MoveActorsToLevel(ActorsToMove, LevelInstanceLevel, OutActors);
 }
 
-ALevelInstance* ULevelInstanceSubsystem::CreateLevelInstanceFrom(const TArray<AActor*>& ActorsToMove, const FNewLevelInstanceParams& CreationParams)
+ILevelInstanceInterface* ULevelInstanceSubsystem::CreateLevelInstanceFrom(const TArray<AActor*>& ActorsToMove, const FNewLevelInstanceParams& CreationParams)
 {
 	check(!bIsCreatingLevelInstance);
 	TGuardValue<bool> CreateLevelInstanceGuard(bIsCreatingLevelInstance, true);
@@ -893,15 +894,23 @@ ALevelInstance* ULevelInstanceSubsystem::CreateLevelInstanceFrom(const TArray<AA
 
 	FActorSpawnParameters SpawnParams;
 	SpawnParams.OverrideLevel = CurrentLevel;
-	ALevelInstance* NewLevelInstanceActor = nullptr;
+	AActor* NewLevelInstanceActor = nullptr;
 	TSoftObjectPtr<UWorld> WorldPtr(LoadedLevel->GetTypedOuter<UWorld>());
-	
+			
 	// Make sure newly created level asset gets scanned
 	ULevel::ScanLevelAssets(LoadedLevel->GetPackage()->GetName());
 
 	if (CreationParams.Type == ELevelInstanceCreationType::LevelInstance)
 	{
-		NewLevelInstanceActor = GetWorld()->SpawnActor<ALevelInstance>(ALevelInstance::StaticClass(), SpawnParams);
+		TSubclassOf<AActor> ActorClass = ALevelInstance::StaticClass();
+		if (CreationParams.LevelInstanceClass)
+		{
+			ActorClass = CreationParams.LevelInstanceClass;
+		}
+
+		check(ActorClass->ImplementsInterface(ULevelInstanceInterface::StaticClass()));
+
+		NewLevelInstanceActor = GetWorld()->SpawnActor<AActor>(ActorClass, SpawnParams);
 	}
 	else
 	{
@@ -934,7 +943,9 @@ ALevelInstance* ULevelInstanceSubsystem::CreateLevelInstanceFrom(const TArray<AA
 	}
 	
 	check(NewLevelInstanceActor);
-	NewLevelInstanceActor->SetWorldAsset(WorldPtr);
+
+	ILevelInstanceInterface* NewLevelInstance = CastChecked<ILevelInstanceInterface>(NewLevelInstanceActor);
+	NewLevelInstance->SetWorldAsset(WorldPtr);
 	NewLevelInstanceActor->SetActorLocation(LevelInstanceLocation);
 	
 	// Actors were moved and kept their World positions so when saving we want their positions to actually be relative to the LevelInstance Actor
@@ -946,11 +957,11 @@ ALevelInstance* ULevelInstanceSubsystem::CreateLevelInstanceFrom(const TArray<AA
 	GEditor->SelectNone(false, true);
 	GEditor->SelectActor(NewLevelInstanceActor, true, true);
 
-	NewLevelInstanceActor->OnEdit();
+	NewLevelInstance->OnEdit();
 
 	// Notify parents of edit
 	TArray<FLevelInstanceID> AncestorIDs;
-	ForEachLevelInstanceAncestors(NewLevelInstanceActor, [&AncestorIDs](ALevelInstance* InAncestor)
+	ForEachLevelInstanceAncestors(NewLevelInstanceActor, [&AncestorIDs](ILevelInstanceInterface* InAncestor)
 	{
 		AncestorIDs.Add(InAncestor->GetLevelInstanceID());
 		return true;
@@ -962,8 +973,8 @@ ALevelInstance* ULevelInstanceSubsystem::CreateLevelInstanceFrom(const TArray<AA
 	}
 	
 	// New level instance
-	const FLevelInstanceID NewLevelInstanceID = NewLevelInstanceActor->GetLevelInstanceID();
-	TUniquePtr<FLevelInstanceEdit> TempLevelInstanceEdit = MakeUnique<FLevelInstanceEdit>(LevelStreaming, NewLevelInstanceActor->GetLevelInstanceID());
+	const FLevelInstanceID NewLevelInstanceID = NewLevelInstance->GetLevelInstanceID();
+	TUniquePtr<FLevelInstanceEdit> TempLevelInstanceEdit = MakeUnique<FLevelInstanceEdit>(LevelStreaming, NewLevelInstance->GetLevelInstanceID());
 	// Force mark it as changed
 	TempLevelInstanceEdit->MarkCommittedChanges();
 
@@ -991,7 +1002,7 @@ ALevelInstance* ULevelInstanceSubsystem::CreateLevelInstanceFrom(const TArray<AA
 	return GetLevelInstance(NewLevelInstanceID);
 }
 
-bool ULevelInstanceSubsystem::BreakLevelInstance(ALevelInstance* LevelInstanceActor, uint32 Levels /* = 1 */, TArray<AActor*>* OutMovedActors /* = nullptr */)
+bool ULevelInstanceSubsystem::BreakLevelInstance(ILevelInstanceInterface* LevelInstance, uint32 Levels /* = 1 */, TArray<AActor*>* OutMovedActors /* = nullptr */)
 {
 	const double StartTime = FPlatformTime::Seconds();
 
@@ -1000,7 +1011,7 @@ bool ULevelInstanceSubsystem::BreakLevelInstance(ALevelInstance* LevelInstanceAc
 	GetMutableDefault<ULevelEditorMiscSettings>()->bAvoidRelabelOnPasteSelected = 1;
 
 	TArray<AActor*> MovedActors;
-	BreakLevelInstance_Impl(LevelInstanceActor, Levels, MovedActors);
+	BreakLevelInstance_Impl(LevelInstance, Levels, MovedActors);
 
 	USelection* ActorSelection = GEditor->GetSelectedActors();
 	ActorSelection->BeginBatchSelectOperation();
@@ -1023,30 +1034,31 @@ bool ULevelInstanceSubsystem::BreakLevelInstance(ALevelInstance* LevelInstanceAc
 	return bStatus;
 }
 
-void ULevelInstanceSubsystem::BreakLevelInstance_Impl(ALevelInstance* LevelInstanceActor, uint32 Levels, TArray<AActor*>& OutMovedActors)
+void ULevelInstanceSubsystem::BreakLevelInstance_Impl(ILevelInstanceInterface* LevelInstance, uint32 Levels, TArray<AActor*>& OutMovedActors)
 {
 	if (Levels > 0)
 	{
+		AActor* LevelInstanceActor = CastChecked<AActor>(LevelInstance);
 		// Can only break the top level LevelInstance
 		check(LevelInstanceActor->GetLevel() == GetWorld()->GetCurrentLevel());
 
 		// Actors in a packed level actor will not be streamed in unless they are editing. Must force this before moving.
 		if (LevelInstanceActor->IsA<APackedLevelActor>())
 		{
-			BlockLoadLevelInstance(LevelInstanceActor);
+			BlockLoadLevelInstance(LevelInstance);
 		}
 
-		// need to ensure that LevelInstanceActor has been streamed in fully
+		// need to ensure that LevelInstance has been streamed in fully
 		GEngine->BlockTillLevelStreamingCompleted(LevelInstanceActor->GetWorld());
 
 		// Cannot break a level instance which has a level script
-		if (LevelInstanceHasLevelScriptBlueprint(LevelInstanceActor))
+		if (LevelInstanceHasLevelScriptBlueprint(LevelInstance))
 		{
 			UE_LOG(LogLevelInstance, Warning, TEXT("Failed to completely break Level Instance because some children have Level Scripts."));
 
 			if (LevelInstanceActor->IsA<APackedLevelActor>())
 			{
-				BlockUnloadLevelInstance(LevelInstanceActor);
+				BlockUnloadLevelInstance(LevelInstance);
 			}
 			return;
 		}
@@ -1094,7 +1106,7 @@ void ULevelInstanceSubsystem::BreakLevelInstance_Impl(ALevelInstance* LevelInsta
 			return false;
 		};
 
-		ForEachActorInLevelInstance(LevelInstanceActor, [this, &ActorsToMove, &AddActorToMove](AActor* Actor)
+		ForEachActorInLevelInstance(LevelInstance, [this, &ActorsToMove, &AddActorToMove](AActor* Actor)
 		{
 			AddActorToMove(Actor);
 			return true;
@@ -1114,14 +1126,14 @@ void ULevelInstanceSubsystem::BreakLevelInstance_Impl(ALevelInstance* LevelInsta
 
 		if (LevelInstanceActor->IsA<APackedLevelActor>())
 		{
-			BlockUnloadLevelInstance(LevelInstanceActor);
+			BlockUnloadLevelInstance(LevelInstance);
 		}
 
 		// Destroy the old LevelInstance instance actor
 		GetWorld()->DestroyActor(LevelInstanceActor);
 	
 		const bool bContinueBreak = Levels > 1;
-		TArray<ALevelInstance*> Children;
+		TArray<ILevelInstanceInterface*> Children;
 
 		for (FSelectionIterator It(GEditor->GetSelectedActorIterator()); It; ++It)
 		{
@@ -1135,9 +1147,9 @@ void ULevelInstanceSubsystem::BreakLevelInstance_Impl(ALevelInstance* LevelInsta
 			// Break up any sub LevelInstances if more levels are requested
 			if (bContinueBreak)
 			{
-				if (ALevelInstance* ChildLevelInstance = Cast<ALevelInstance>(Actor))
+				if (ILevelInstanceInterface* ChildLevelInstance = Cast<ILevelInstanceInterface>(Actor))
 				{
-					OutMovedActors.Remove(ChildLevelInstance);
+					OutMovedActors.Remove(Actor);
 
 					Children.Add(ChildLevelInstance);
 				}
@@ -1153,24 +1165,24 @@ void ULevelInstanceSubsystem::BreakLevelInstance_Impl(ALevelInstance* LevelInsta
 	return;
 }
 
-ULevel* ULevelInstanceSubsystem::GetLevelInstanceLevel(const ALevelInstance* LevelInstanceActor) const
+ULevel* ULevelInstanceSubsystem::GetLevelInstanceLevel(const ILevelInstanceInterface* LevelInstance) const
 {
-	if (LevelInstanceActor->HasValidLevelInstanceID())
+	if (LevelInstance->HasValidLevelInstanceID())
 	{
-		if (const FLevelInstanceEdit* CurrentEdit = GetLevelInstanceEdit(LevelInstanceActor))
+		if (const FLevelInstanceEdit* CurrentEdit = GetLevelInstanceEdit(LevelInstance))
 		{
 			return LevelInstanceEdit->LevelStreaming->GetLoadedLevel();
 		}
-		else if (const FLevelInstance* LevelInstance = LevelInstances.Find(LevelInstanceActor->GetLevelInstanceID()))
+		else if (const FLevelInstance* LevelInstanceEntry = LevelInstances.Find(LevelInstance->GetLevelInstanceID()))
 		{
-			return LevelInstance->LevelStreaming->GetLoadedLevel();
+			return LevelInstanceEntry->LevelStreaming->GetLoadedLevel();
 		}
 	}
 
 	return nullptr;
 }
 
-bool ULevelInstanceSubsystem::LevelInstanceHasLevelScriptBlueprint(const ALevelInstance* LevelInstance) const
+bool ULevelInstanceSubsystem::LevelInstanceHasLevelScriptBlueprint(const ILevelInstanceInterface* LevelInstance) const
 {
 	if (LevelInstance)
 	{
@@ -1243,9 +1255,9 @@ bool ULevelInstanceSubsystem::CanMoveActorToLevel(const AActor* Actor, FText* Ou
 
 	if (Actor->GetWorld() == GetWorld())
 	{
-		if (const ALevelInstance* LevelInstanceActor = Cast<ALevelInstance>(Actor))
+		if (const ILevelInstanceInterface* LevelInstance = Cast<ILevelInstanceInterface>(Actor))
 		{
-			if (IsEditingLevelInstance(LevelInstanceActor))
+			if (IsEditingLevelInstance(LevelInstance))
 			{
 				if (OutReason != nullptr)
 				{
@@ -1255,9 +1267,9 @@ bool ULevelInstanceSubsystem::CanMoveActorToLevel(const AActor* Actor, FText* Ou
 			}
 
 			bool bEditingChildren = false;
-			ForEachLevelInstanceChild(LevelInstanceActor, true, [this, &bEditingChildren](const ALevelInstance* ChildLevelInstanceActor)
+			ForEachLevelInstanceChild(LevelInstance, true, [this, &bEditingChildren](const ILevelInstanceInterface* ChildLevelInstance)
 			{
-				if (IsEditingLevelInstance(ChildLevelInstanceActor))
+				if (IsEditingLevelInstance(ChildLevelInstance))
 				{
 					bEditingChildren = true;
 					return false;
@@ -1281,7 +1293,7 @@ bool ULevelInstanceSubsystem::CanMoveActorToLevel(const AActor* Actor, FText* Ou
 
 void ULevelInstanceSubsystem::OnActorDeleted(AActor* Actor)
 {
-	if (ALevelInstance* LevelInstanceActor = Cast<ALevelInstance>(Actor))
+	if (ILevelInstanceInterface* LevelInstance = Cast<ILevelInstanceInterface>(Actor))
 	{
 		if (Actor->GetClass()->HasAnyClassFlags(CLASS_NewerVersionExists))
 		{
@@ -1292,37 +1304,37 @@ void ULevelInstanceSubsystem::OnActorDeleted(AActor* Actor)
 		}
 
 		// Unregistered Level Instance Actor nothing to do.
-		if (!LevelInstanceActor->HasValidLevelInstanceID())
+		if (!LevelInstance->HasValidLevelInstanceID())
 		{
 			return;
 		}
 
-		const bool bIsEditingLevelInstance = IsEditingLevelInstance(LevelInstanceActor);
-		if (!bIsEditingLevelInstance && LevelInstanceActor->IsA<APackedLevelActor>())
+		const bool bIsEditingLevelInstance = IsEditingLevelInstance(LevelInstance);
+		if (!bIsEditingLevelInstance && Actor->IsA<APackedLevelActor>())
 		{
 			return;
 		}
 
-		const bool bAlreadyRooted = LevelInstanceActor->IsRooted();
+		const bool bAlreadyRooted = Actor->IsRooted();
 		// Unloading LevelInstances leads to GC and Actor can be collected. Add to root temp. It will get collected after the OnActorDeleted callbacks
 		if (!bAlreadyRooted)
 		{
-			LevelInstanceActor->AddToRoot();
+			Actor->AddToRoot();
 		}
 
 		FScopedSlowTask SlowTask(0, LOCTEXT("UnloadingLevelInstances", "Unloading Level Instances..."), !GetWorld()->IsGameWorld());
 		SlowTask.MakeDialog();
-		check(!IsEditingLevelInstanceDirty(LevelInstanceActor) && !HasDirtyChildrenLevelInstances(LevelInstanceActor));
+		check(!IsEditingLevelInstanceDirty(LevelInstance) && !HasDirtyChildrenLevelInstances(LevelInstance));
 		if (bIsEditingLevelInstance)
 		{
-			CommitLevelInstance(LevelInstanceActor);
+			CommitLevelInstance(LevelInstance);
 		}
 		else
 		{
 			// We are ending editing. Discard Non dirty child edits
-			ForEachLevelInstanceChild(LevelInstanceActor, /*bRecursive=*/true, [this](const ALevelInstance* ChildLevelInstanceActor)
+			ForEachLevelInstanceChild(LevelInstance, /*bRecursive=*/true, [this](const ILevelInstanceInterface* ChildLevelInstance)
 			{
-				if (const FLevelInstanceEdit* ChildLevelInstanceEdit = GetLevelInstanceEdit(ChildLevelInstanceActor))
+				if (const FLevelInstanceEdit* ChildLevelInstanceEdit = GetLevelInstanceEdit(ChildLevelInstance))
 				{
 					check(!IsLevelInstanceEditDirty(ChildLevelInstanceEdit));
 					ResetEdit(LevelInstanceEdit);
@@ -1332,14 +1344,14 @@ void ULevelInstanceSubsystem::OnActorDeleted(AActor* Actor)
 			});
 		}
 
-		LevelInstancesToLoadOrUpdate.Remove(LevelInstanceActor);
+		LevelInstancesToLoadOrUpdate.Remove(LevelInstance);
 				
-		UnloadLevelInstance(LevelInstanceActor->GetLevelInstanceID());
+		UnloadLevelInstance(LevelInstance->GetLevelInstanceID());
 		
 		// Remove from root so it gets collected on the next GC if it can be.
 		if (!bAlreadyRooted)
 		{
-			LevelInstanceActor->RemoveFromRoot();
+			Actor->RemoveFromRoot();
 		}
 	}
 }
@@ -1439,9 +1451,9 @@ FLevelInstanceID ULevelInstanceSubsystem::FLevelInstanceEdit::GetLevelInstanceID
 	return LevelStreaming ? LevelStreaming->GetLevelInstanceID() : FLevelInstanceID();
 }
 
-const ULevelInstanceSubsystem::FLevelInstanceEdit* ULevelInstanceSubsystem::GetLevelInstanceEdit(const ALevelInstance* LevelInstanceActor) const
+const ULevelInstanceSubsystem::FLevelInstanceEdit* ULevelInstanceSubsystem::GetLevelInstanceEdit(const ILevelInstanceInterface* LevelInstance) const
 {
-	if (LevelInstanceEdit && LevelInstanceEdit->GetLevelInstanceID() == LevelInstanceActor->GetLevelInstanceID())
+	if (LevelInstanceEdit && LevelInstanceEdit->GetLevelInstanceID() == LevelInstance->GetLevelInstanceID())
 	{
 		return LevelInstanceEdit.Get();
 	}
@@ -1449,9 +1461,9 @@ const ULevelInstanceSubsystem::FLevelInstanceEdit* ULevelInstanceSubsystem::GetL
 	return nullptr;
 }
 
-bool ULevelInstanceSubsystem::IsEditingLevelInstanceDirty(const ALevelInstance* LevelInstanceActor) const
+bool ULevelInstanceSubsystem::IsEditingLevelInstanceDirty(const ILevelInstanceInterface* LevelInstance) const
 {
-	const FLevelInstanceEdit* CurrentEdit = GetLevelInstanceEdit(LevelInstanceActor);
+	const FLevelInstanceEdit* CurrentEdit = GetLevelInstanceEdit(LevelInstance);
 	if (!CurrentEdit)
 	{
 		return false;
@@ -1467,7 +1479,7 @@ bool ULevelInstanceSubsystem::IsLevelInstanceEditDirty(const FLevelInstanceEdit*
 	return OutPackagesToSave.Num() > 0;
 }
 
-ALevelInstance* ULevelInstanceSubsystem::GetEditingLevelInstance() const
+ILevelInstanceInterface* ULevelInstanceSubsystem::GetEditingLevelInstance() const
 {
 	if (LevelInstanceEdit)
 	{
@@ -1477,7 +1489,7 @@ ALevelInstance* ULevelInstanceSubsystem::GetEditingLevelInstance() const
 	return nullptr;
 }
 
-bool ULevelInstanceSubsystem::CanEditLevelInstance(const ALevelInstance* LevelInstanceActor, FText* OutReason) const
+bool ULevelInstanceSubsystem::CanEditLevelInstance(const ILevelInstanceInterface* LevelInstance, FText* OutReason) const
 {
 	// Only allow Editing in Editor World
 	if (GetWorld()->WorldType != EWorldType::Editor)
@@ -1485,13 +1497,13 @@ bool ULevelInstanceSubsystem::CanEditLevelInstance(const ALevelInstance* LevelIn
 		return false;
 	}
 
-	if (ULevel* LevelInstanceLevel = GetLevelInstanceLevel(LevelInstanceActor))
+	if (ULevel* LevelInstanceLevel = GetLevelInstanceLevel(LevelInstance))
 	{
 		if (LevelInstanceLevel->GetWorldPartition())
 		{
 			if (OutReason)
 			{
-				*OutReason = FText::Format(LOCTEXT("CanEditPartitionedLevelInstance", "Can't edit partitioned Level Instance ({0})."), FText::FromString(LevelInstanceActor->GetWorldAssetPackage()));
+				*OutReason = FText::Format(LOCTEXT("CanEditPartitionedLevelInstance", "Can't edit partitioned Level Instance ({0})."), FText::FromString(LevelInstance->GetWorldAssetPackage()));
 			}
 			return false;
 		}
@@ -1499,11 +1511,11 @@ bool ULevelInstanceSubsystem::CanEditLevelInstance(const ALevelInstance* LevelIn
 	
 	if (LevelInstanceEdit)
 	{
-		if (LevelInstanceEdit->GetLevelInstanceID() == LevelInstanceActor->GetLevelInstanceID())
+		if (LevelInstanceEdit->GetLevelInstanceID() == LevelInstance->GetLevelInstanceID())
 		{
 			if (OutReason)
 			{
-				*OutReason = FText::Format(LOCTEXT("CanEditLevelInstanceAlreadyBeingEdited", "Level Instance already being edited ({0})."), FText::FromString(LevelInstanceActor->GetWorldAssetPackage()));
+				*OutReason = FText::Format(LOCTEXT("CanEditLevelInstanceAlreadyBeingEdited", "Level Instance already being edited ({0})."), FText::FromString(LevelInstance->GetWorldAssetPackage()));
 			}
 			return false;
 		}
@@ -1518,7 +1530,7 @@ bool ULevelInstanceSubsystem::CanEditLevelInstance(const ALevelInstance* LevelIn
 		}
 	}
 	
-	if (!LevelInstanceActor->IsLevelInstancePathValid())
+	if (!LevelInstance->IsWorldAssetValid())
 	{
 		if (OutReason)
 		{
@@ -1527,44 +1539,30 @@ bool ULevelInstanceSubsystem::CanEditLevelInstance(const ALevelInstance* LevelIn
 		return false;
 	}
 
-	if (GetWorld()->PersistentLevel->GetPackage()->GetName() == LevelInstanceActor->GetWorldAssetPackage())
+	if (GetWorld()->PersistentLevel->GetPackage()->GetName() == LevelInstance->GetWorldAssetPackage())
 	{
 		if (OutReason)
 		{
-			*OutReason = FText::Format(LOCTEXT("CanEditLevelInstancePersistentLevel", "The Persistent level and the Level Instance are the same ({0})."), FText::FromString(LevelInstanceActor->GetWorldAssetPackage()));
+			*OutReason = FText::Format(LOCTEXT("CanEditLevelInstancePersistentLevel", "The Persistent level and the Level Instance are the same ({0})."), FText::FromString(LevelInstance->GetWorldAssetPackage()));
 		}
 		return false;
 	}
 
-	if (FLevelUtils::FindStreamingLevel(GetWorld(), *LevelInstanceActor->GetWorldAssetPackage()))
+	if (FLevelUtils::FindStreamingLevel(GetWorld(), *LevelInstance->GetWorldAssetPackage()))
 	{
 		if (OutReason)
 		{
-			*OutReason = FText::Format(LOCTEXT("CanEditLevelInstanceAlreadyExists", "The same level was added to world outside of Level Instances ({0})."), FText::FromString(LevelInstanceActor->GetWorldAssetPackage()));
+			*OutReason = FText::Format(LOCTEXT("CanEditLevelInstanceAlreadyExists", "The same level was added to world outside of Level Instances ({0})."), FText::FromString(LevelInstance->GetWorldAssetPackage()));
 		}
 		return false;
 	}
 
 	FPackagePath WorldAssetPath;
-	if (!FPackagePath::TryFromPackageName(LevelInstanceActor->GetWorldAssetPackage(), WorldAssetPath) || !FPackageName::DoesPackageExist(WorldAssetPath))
-	{
-		if (OutReason)
-		{
-			*OutReason = FText::Format(LOCTEXT("CanEditLevelInstanceInvalidAsset", "Level Instance asset is invalid ({0})."), FText::FromString(LevelInstanceActor->GetWorldAssetPackage()));
-		}
-		return false;
-	}
-
-	return true;
-}
-
-bool ULevelInstanceSubsystem::CanCommitLevelInstance(const ALevelInstance* LevelInstanceActor, FText* OutReason) const
-{
-	if (!IsEditingLevelInstance(LevelInstanceActor))
+	if (!FPackagePath::TryFromPackageName(LevelInstance->GetWorldAssetPackage(), WorldAssetPath) || !FPackageName::DoesPackageExist(WorldAssetPath))
 	{
 		if (OutReason)
 		{
-			*OutReason = LOCTEXT("CanCommitLevelInstanceNotEditing", "Level Instance is not currently being edited");
+			*OutReason = FText::Format(LOCTEXT("CanEditLevelInstanceInvalidAsset", "Level Instance asset is invalid ({0})."), FText::FromString(LevelInstance->GetWorldAssetPackage()));
 		}
 		return false;
 	}
@@ -1572,11 +1570,11 @@ bool ULevelInstanceSubsystem::CanCommitLevelInstance(const ALevelInstance* Level
 	return true;
 }
 
-bool ULevelInstanceSubsystem::CanDiscardLevelInstance(const ALevelInstance* LevelInstanceActor, FText* OutReason) const
+bool ULevelInstanceSubsystem::CanCommitLevelInstance(const ILevelInstanceInterface* LevelInstance, bool bDiscardEdits, FText* OutReason) const
 {
-	if (const FLevelInstanceEdit* CurrentEdit = GetLevelInstanceEdit(LevelInstanceActor))
+	if (const FLevelInstanceEdit* CurrentEdit = GetLevelInstanceEdit(LevelInstance))
 	{
-		return LevelInstanceEdit->CanDiscard(OutReason);
+		return !bDiscardEdits || LevelInstanceEdit->CanDiscard(OutReason);
 	}
 
 	if (OutReason)
@@ -1586,18 +1584,18 @@ bool ULevelInstanceSubsystem::CanDiscardLevelInstance(const ALevelInstance* Leve
 	return false;
 }
 
-void ULevelInstanceSubsystem::EditLevelInstance(ALevelInstance* LevelInstanceActor, TWeakObjectPtr<AActor> ContextActorPtr)
+void ULevelInstanceSubsystem::EditLevelInstance(ILevelInstanceInterface* LevelInstance, TWeakObjectPtr<AActor> ContextActorPtr)
 {
-	if (EditLevelInstanceInternal(LevelInstanceActor, ContextActorPtr, FString(), false))
+	if (EditLevelInstanceInternal(LevelInstance, ContextActorPtr, FString(), false))
 	{
 		ILevelInstanceEditorModule& EditorModule = FModuleManager::GetModuleChecked<ILevelInstanceEditorModule>("LevelInstanceEditor");
 		EditorModule.ActivateEditorMode();
 	}
 }
 
-bool ULevelInstanceSubsystem::EditLevelInstanceInternal(ALevelInstance* LevelInstanceActor, TWeakObjectPtr<AActor> ContextActorPtr, const FString& InActorNameToSelect, bool bRecursive)
+bool ULevelInstanceSubsystem::EditLevelInstanceInternal(ILevelInstanceInterface* LevelInstance, TWeakObjectPtr<AActor> ContextActorPtr, const FString& InActorNameToSelect, bool bRecursive)
 {
-	check(CanEditLevelInstance(LevelInstanceActor));
+	check(CanEditLevelInstance(LevelInstance));
 		
 	FScopedSlowTask SlowTask(0, LOCTEXT("BeginEditLevelInstance", "Loading Level Instance for edit..."), !GetWorld()->IsGameWorld());
 	SlowTask.MakeDialog();
@@ -1607,15 +1605,15 @@ bool ULevelInstanceSubsystem::EditLevelInstanceInternal(ALevelInstance* LevelIns
 	if (AActor* ContextActor = ContextActorPtr.Get())
 	{
 		ActorNameToSelect = ContextActor->GetName();
-		ForEachLevelInstanceAncestorsAndSelf(ContextActor, [&ActorNameToSelect,LevelInstanceActor](const ALevelInstance* AncestorLevelInstanceActor)
+		ForEachLevelInstanceAncestorsAndSelf(ContextActor, [&ActorNameToSelect,LevelInstance](const ILevelInstanceInterface* AncestorLevelInstance)
 		{
 			// stop when we hit the LevelInstance we are about to edit
-			if (AncestorLevelInstanceActor == LevelInstanceActor)
+			if (AncestorLevelInstance == LevelInstance)
 			{
 				return false;
 			}
 			
-			ActorNameToSelect = AncestorLevelInstanceActor->GetName();
+			ActorNameToSelect = CastChecked<AActor>(AncestorLevelInstance)->GetName();
 			return true;
 		});
 	}
@@ -1626,7 +1624,7 @@ bool ULevelInstanceSubsystem::EditLevelInstanceInternal(ALevelInstance* LevelIns
 	if (!bRecursive)
 	{
 		TArray<FLevelInstanceID> AncestorIDs;
-		ForEachLevelInstanceAncestors(LevelInstanceActor, [&AncestorIDs](ALevelInstance* InAncestor)
+		ForEachLevelInstanceAncestors(CastChecked<AActor>(LevelInstance), [&AncestorIDs](ILevelInstanceInterface* InAncestor)
 		{
 			AncestorIDs.Add(InAncestor->GetLevelInstanceID());
 			return true;
@@ -1643,33 +1641,33 @@ bool ULevelInstanceSubsystem::EditLevelInstanceInternal(ALevelInstance* LevelIns
 	{	
 		// Only support one level of recursion to commit current edit
 		check(!bRecursive);
-		FLevelInstanceID PendingEditId = LevelInstanceActor->GetLevelInstanceID();
+		FLevelInstanceID PendingEditId = LevelInstance->GetLevelInstanceID();
 		
 		check(!IsLevelInstanceEditDirty(LevelInstanceEdit.Get()));
 		CommitLevelInstanceInternal(LevelInstanceEdit);
 
-		ALevelInstance* LevelInstanceToEdit = GetLevelInstance(PendingEditId);
+		ILevelInstanceInterface* LevelInstanceToEdit = GetLevelInstance(PendingEditId);
 		check(LevelInstanceToEdit);
 
 		return EditLevelInstanceInternal(LevelInstanceToEdit, nullptr, ActorNameToSelect, /*bRecursive=*/true);
 	}
 
 	// Cleanup async requests in case
-	LevelInstancesToUnload.Remove(LevelInstanceActor->GetLevelInstanceID());
-	LevelInstancesToLoadOrUpdate.Remove(LevelInstanceActor);
+	LevelInstancesToUnload.Remove(LevelInstance->GetLevelInstanceID());
+	LevelInstancesToLoadOrUpdate.Remove(LevelInstance);
 	// Unload right away
-	UnloadLevelInstance(LevelInstanceActor->GetLevelInstanceID());
+	UnloadLevelInstance(LevelInstance->GetLevelInstanceID());
 		
 	// Load Edit LevelInstance level
-	ULevelStreamingLevelInstanceEditor* LevelStreaming = ULevelStreamingLevelInstanceEditor::Load(LevelInstanceActor);
+	ULevelStreamingLevelInstanceEditor* LevelStreaming = ULevelStreamingLevelInstanceEditor::Load(LevelInstance);
 	if (!LevelStreaming)
 	{
-		LevelInstanceActor->LoadLevelInstance();
+		LevelInstance->LoadLevelInstance();
 		return false;
 	}
 
 	check(LevelInstanceEdit.IsValid());
-	check(LevelInstanceEdit->GetLevelInstanceID() == LevelInstanceActor->GetLevelInstanceID());
+	check(LevelInstanceEdit->GetLevelInstanceID() == LevelInstance->GetLevelInstanceID());
 	check(LevelInstanceEdit->LevelStreaming == LevelStreaming);
 		
 	// Try and select something meaningful
@@ -1680,6 +1678,7 @@ bool ULevelInstanceSubsystem::EditLevelInstanceInternal(ALevelInstance* LevelIns
 	}
 
 	// default to LevelInstance
+	AActor* LevelInstanceActor = CastChecked<AActor>(LevelInstance);
 	if (!ActorToSelect)
 	{
 		ActorToSelect = LevelInstanceActor;
@@ -1687,7 +1686,7 @@ bool ULevelInstanceSubsystem::EditLevelInstanceInternal(ALevelInstance* LevelIns
 	LevelInstanceActor->SetIsTemporarilyHiddenInEditor(false);
 
 	// Notify
-	LevelInstanceActor->OnEdit();
+	LevelInstance->OnEdit();
 
 	GEditor->SelectActor(ActorToSelect, true, true);
 
@@ -1709,10 +1708,10 @@ bool ULevelInstanceSubsystem::EditLevelInstanceInternal(ALevelInstance* LevelIns
 	return true;
 }
 
-bool ULevelInstanceSubsystem::CommitLevelInstance(ALevelInstance* LevelInstanceActor, bool bDiscardEdits, TSet<FName>* DirtyPackages)
+bool ULevelInstanceSubsystem::CommitLevelInstance(ILevelInstanceInterface* LevelInstance, bool bDiscardEdits, TSet<FName>* DirtyPackages)
 {
-	check(LevelInstanceEdit.Get() == GetLevelInstanceEdit(LevelInstanceActor));
-	check(CanCommitLevelInstance(LevelInstanceActor));
+	check(LevelInstanceEdit.Get() == GetLevelInstanceEdit(LevelInstance));
+	check(CanCommitLevelInstance(LevelInstance));
 	if (CommitLevelInstanceInternal(LevelInstanceEdit, bDiscardEdits, /*bDiscardOnFailure=*/false, DirtyPackages))
 	{
 		ILevelInstanceEditorModule& EditorModule = FModuleManager::GetModuleChecked<ILevelInstanceEditorModule>("LevelInstanceEditor");
@@ -1727,7 +1726,7 @@ bool ULevelInstanceSubsystem::CommitLevelInstance(ALevelInstance* LevelInstanceA
 bool ULevelInstanceSubsystem::CommitLevelInstanceInternal(TUniquePtr<FLevelInstanceEdit>& InLevelInstanceEdit, bool bDiscardEdits, bool bDiscardOnFailure, TSet<FName>* DirtyPackages)
 {
 	TGuardValue<bool> CommitScope(bIsCommittingLevelInstance, true);
-	ALevelInstance* LevelInstanceActor = GetLevelInstance(InLevelInstanceEdit->GetLevelInstanceID());
+	ILevelInstanceInterface* LevelInstance = GetLevelInstance(InLevelInstanceEdit->GetLevelInstanceID());
 	check(InLevelInstanceEdit);
 	UWorld* EditingWorld = InLevelInstanceEdit->GetEditWorld();
 	check(EditingWorld);
@@ -1790,7 +1789,7 @@ bool ULevelInstanceSubsystem::CommitLevelInstanceInternal(TUniquePtr<FLevelInsta
 
 	GEditor->SelectNone(false, true);
 
-	const FString EditPackage = LevelInstanceActor->GetWorldAssetPackage();
+	const FString EditPackage = LevelInstance->GetWorldAssetPackage();
 
 	// Remove from streaming level...
 	ResetEdit(InLevelInstanceEdit);
@@ -1801,47 +1800,50 @@ bool ULevelInstanceSubsystem::CommitLevelInstanceInternal(TUniquePtr<FLevelInsta
 	}
 	
 	// Backup ID on Commit in case Actor gets recreated
-	const FLevelInstanceID LevelInstanceID = LevelInstanceActor->GetLevelInstanceID();
+	const FLevelInstanceID LevelInstanceID = LevelInstance->GetLevelInstanceID();
 
 	// Notify (Actor might get destroyed by this call if its a packed bp)
-	LevelInstanceActor->OnCommit(bChangesCommitted);
+	LevelInstance->OnCommit(bChangesCommitted);
 
-	// Update pointer since BP Compilation might have invalidated LevelInstanceActor
-	LevelInstanceActor = GetLevelInstance(LevelInstanceID);
+	// Update pointer since BP Compilation might have invalidated LevelInstance
+	LevelInstance = GetLevelInstance(LevelInstanceID);
 
-	TArray<FLevelInstanceID> LevelInstancesToUpdate;
+	TArray<TPair<ULevelInstanceSubsystem*, FLevelInstanceID>> LevelInstancesToUpdate;
 	// Gather list to update
 	for (TObjectIterator<UWorld> It(RF_ClassDefaultObject | RF_ArchetypeObject, true); It; ++It)
 	{
 		UWorld* CurrentWorld = *It;
-		if (IsValid(CurrentWorld) && CurrentWorld->GetSubsystem<ULevelInstanceSubsystem>() != nullptr)
+		if (IsValid(CurrentWorld))
 		{
-			for (TActorIterator<ALevelInstance> LevelInstanceIt(CurrentWorld); LevelInstanceIt; ++LevelInstanceIt)
+			if (ULevelInstanceSubsystem* LevelInstanceSubsystem = CurrentWorld->GetSubsystem<ULevelInstanceSubsystem>())
 			{
-				ALevelInstance* CurrentLevelInstanceActor = *LevelInstanceIt;
-				if (CurrentLevelInstanceActor->GetWorldAssetPackage() == EditPackage && (LevelInstanceActor == CurrentLevelInstanceActor || bChangesCommitted))
+				TArray<ILevelInstanceInterface*> WorldLevelInstances = LevelInstanceSubsystem->GetLevelInstances(EditPackage);
+				for (ILevelInstanceInterface* CurrentLevelInstance : WorldLevelInstances)
 				{
-					LevelInstancesToUpdate.Add(CurrentLevelInstanceActor->GetLevelInstanceID());
+					if (CurrentLevelInstance == LevelInstance || bChangesCommitted)
+					{
+						LevelInstancesToUpdate.Add({ LevelInstanceSubsystem, CurrentLevelInstance->GetLevelInstanceID() });
+					}
 				}
 			}
 		}
 	}
 
 	// Do update
-	for (const FLevelInstanceID& LevelInstanceToUpdateID : LevelInstancesToUpdate)
+	for (const auto& KeyValuePair : LevelInstancesToUpdate)
 	{
-		if (ALevelInstance* LevelInstance = GetLevelInstance(LevelInstanceToUpdateID))
+		if (ILevelInstanceInterface* LevelInstanceToUpdate = KeyValuePair.Key->GetLevelInstance(KeyValuePair.Value))
 		{
-			LevelInstance->UpdateFromLevel();
+			LevelInstanceToUpdate->UpdateLevelInstanceFromWorldAsset();
 		}
 	}
 
-	LevelInstanceActor = GetLevelInstance(LevelInstanceID);
+	LevelInstance = GetLevelInstance(LevelInstanceID);
 	
 	// Notify Ancestors
 	FLevelInstanceID LevelInstanceToSelectID = LevelInstanceID;
 	TArray<FLevelInstanceID> AncestorIDs;
-	ForEachLevelInstanceAncestors(LevelInstanceActor, [&LevelInstanceToSelectID, &AncestorIDs](ALevelInstance* AncestorLevelInstance)
+	ForEachLevelInstanceAncestors(CastChecked<AActor>(LevelInstance), [&LevelInstanceToSelectID, &AncestorIDs](ILevelInstanceInterface* AncestorLevelInstance)
 	{
 		LevelInstanceToSelectID = AncestorLevelInstance->GetLevelInstanceID();
 		AncestorIDs.Add(AncestorLevelInstance->GetLevelInstanceID());
@@ -1853,9 +1855,9 @@ bool ULevelInstanceSubsystem::CommitLevelInstanceInternal(TUniquePtr<FLevelInsta
 		OnCommitChild(AncestorID, bChangesCommitted);
 	}
 		
-	if (ALevelInstance* Actor = GetLevelInstance(LevelInstanceToSelectID))
+	if (ILevelInstanceInterface* LevelInstanceToSelect = GetLevelInstance(LevelInstanceToSelectID))
 	{
-		GEditor->SelectActor(Actor, true, true);
+		GEditor->SelectActor(CastChecked<AActor>(LevelInstanceToSelect), true, true);
 	}
 				
 	// Wait for Level Instances to be loaded
@@ -1869,7 +1871,7 @@ bool ULevelInstanceSubsystem::CommitLevelInstanceInternal(TUniquePtr<FLevelInsta
 	return true;
 }
 
-ALevelInstance* ULevelInstanceSubsystem::GetParentLevelInstance(const AActor* Actor) const
+ILevelInstanceInterface* ULevelInstanceSubsystem::GetParentLevelInstance(const AActor* Actor) const
 {
 	check(Actor);
 	const ULevel* OwningLevel = Actor->GetLevel();
@@ -1889,25 +1891,25 @@ void ULevelInstanceSubsystem::BlockOnLoading()
 	}
 }
 
-void ULevelInstanceSubsystem::BlockLoadLevelInstance(ALevelInstance* LevelInstanceActor)
+void ULevelInstanceSubsystem::BlockLoadLevelInstance(ILevelInstanceInterface* LevelInstance)
 {
-	check(!LevelInstanceActor->IsEditing());
-	RequestLoadLevelInstance(LevelInstanceActor, true);
+	check(!LevelInstance->IsEditing());
+	RequestLoadLevelInstance(LevelInstance, true);
 
 	BlockOnLoading();
 }
 
-void ULevelInstanceSubsystem::BlockUnloadLevelInstance(ALevelInstance* LevelInstanceActor)
+void ULevelInstanceSubsystem::BlockUnloadLevelInstance(ILevelInstanceInterface* LevelInstance)
 {
-	check(!LevelInstanceActor->IsEditing());
-	RequestUnloadLevelInstance(LevelInstanceActor);
+	check(!LevelInstance->IsEditing());
+	RequestUnloadLevelInstance(LevelInstance);
 
 	BlockOnLoading();
 }
 
-bool ULevelInstanceSubsystem::HasChildEdit(const ALevelInstance* LevelInstanceActor) const
+bool ULevelInstanceSubsystem::HasChildEdit(const ILevelInstanceInterface* LevelInstance) const
 {
-	const int32* ChildEditCountPtr = ChildEdits.Find(LevelInstanceActor->GetLevelInstanceID());
+	const int32* ChildEditCountPtr = ChildEdits.Find(LevelInstance->GetLevelInstanceID());
 	return ChildEditCountPtr && *ChildEditCountPtr;
 }
 
@@ -1917,7 +1919,7 @@ void ULevelInstanceSubsystem::OnCommitChild(FLevelInstanceID LevelInstanceID, bo
 	check(ChildEditCount > 0);
 	ChildEditCount--;
 
-	if (ALevelInstance* LevelInstance = GetLevelInstance(LevelInstanceID))
+	if (ILevelInstanceInterface* LevelInstance = GetLevelInstance(LevelInstanceID))
 	{
 		LevelInstance->OnCommitChild(bChildChanged);
 	}
@@ -1931,12 +1933,119 @@ void ULevelInstanceSubsystem::OnEditChild(FLevelInstanceID LevelInstanceID)
 	check(ChildEditCount < 2);
 	ChildEditCount++;
 
-	if (ALevelInstance* LevelInstance = GetLevelInstance(LevelInstanceID))
+	if (ILevelInstanceInterface* LevelInstance = GetLevelInstance(LevelInstanceID))
 	{
 		LevelInstance->OnEditChild();
 	}
 }
 
+TArray<ILevelInstanceInterface*> ULevelInstanceSubsystem::GetLevelInstances(const FString& WorldAssetPackage)
+{
+	TArray<ILevelInstanceInterface*> MatchingLevelInstances;
+	for (const auto& KeyValuePair : RegisteredLevelInstances)
+	{
+		if (KeyValuePair.Value->GetWorldAssetPackage() == WorldAssetPackage)
+		{
+			MatchingLevelInstances.Add(KeyValuePair.Value);
+		}
+	}
+
+	return MatchingLevelInstances;
+}
+
+bool ULevelInstanceSubsystem::CheckForLoop(const ILevelInstanceInterface* LevelInstance, TArray<TPair<FText, TSoftObjectPtr<UWorld>>>* LoopInfo, const ILevelInstanceInterface** LoopStart)
+{
+	return CheckForLoop(LevelInstance, LevelInstance->GetWorldAsset(), LoopInfo, LoopStart);
+}
+
+bool ULevelInstanceSubsystem::CheckForLoop(const ILevelInstanceInterface* LevelInstance, TSoftObjectPtr<UWorld> WorldAsset, TArray<TPair<FText, TSoftObjectPtr<UWorld>>>* LoopInfo, const ILevelInstanceInterface** LoopStart)
+{
+	bool bValid = true;
+		
+	if (ULevelInstanceSubsystem* LevelInstanceSubsystem = LevelInstance->GetLevelInstanceSubsystem())
+	{
+		LevelInstanceSubsystem->ForEachLevelInstanceAncestorsAndSelf(CastChecked<AActor>(LevelInstance), [&bValid, &WorldAsset, LevelInstance, LoopInfo, LoopStart](const ILevelInstanceInterface* CurrentLevelInstance)
+		{
+			FName PackageToTest(*WorldAsset.GetLongPackageName());
+			// Check to exclude NAME_None since Preview Levels are in the transient package
+			// Check the level we are spawned in to detect the loop (this will handle loops caused by LevelInstances and by regular level streaming)
+			const AActor* CurrentActor = CastChecked<AActor>(CurrentLevelInstance);
+			if (PackageToTest != NAME_None && CurrentActor->GetLevel()->GetPackage()->GetLoadedPath() == FPackagePath::FromPackageNameChecked(PackageToTest))
+			{
+				bValid = false;
+				if (LoopStart)
+				{
+					*LoopStart = CurrentLevelInstance;
+				}
+			}
+
+			if (LoopInfo)
+			{
+				TSoftObjectPtr<UWorld> CurrentAsset = CurrentLevelInstance == LevelInstance ? WorldAsset : CurrentLevelInstance->GetWorldAsset();
+				FText LevelInstanceName = FText::FromString(CurrentActor->GetPathName());
+				FText Description = FText::Format(LOCTEXT("LevelInstanceLoopLink", "-> Actor: {0} loads"), LevelInstanceName);
+				LoopInfo->Emplace(Description, CurrentAsset);
+			}
+			
+			return bValid;
+		});
+	}
+	
+	return bValid;
+}
+
+bool ULevelInstanceSubsystem::CanUseWorldAsset(const ILevelInstanceInterface* LevelInstance, TSoftObjectPtr<UWorld> WorldAsset, FString* OutReason)
+{
+	if (WorldAsset.IsNull())
+	{
+		return true;
+	}
+
+	FString PackageName;
+	if (!FPackageName::DoesPackageExist(WorldAsset.GetLongPackageName()))
+	{
+		if (OutReason)
+		{
+			*OutReason = FString::Format(TEXT("Attempting to set Level Instance to package {0} which does not exist. Ensure the level was saved before attepting to set the level instance world asset."), { WorldAsset.GetLongPackageName() });
+		}
+		return false;
+	}
+
+	TArray<TPair<FText, TSoftObjectPtr<UWorld>>> LoopInfo;
+	const ILevelInstanceInterface* LoopStart = nullptr;
+
+	if (ULevel::GetIsLevelPartitionedFromPackage(*WorldAsset.GetLongPackageName()))
+	{
+		if (OutReason)
+		{
+			*OutReason = FString::Format(TEXT("LevelInstance doesn't support partitioned world {0}\n"), { WorldAsset.GetLongPackageName() });
+		}
+
+		return false;
+	}
+
+	if (!CheckForLoop(LevelInstance, WorldAsset, OutReason ? &LoopInfo : nullptr, OutReason ? &LoopStart : nullptr))
+	{
+		if (OutReason)
+		{
+			if (ensure(LoopStart))
+			{
+				const AActor* LoopStartActor = CastChecked<AActor>(LoopStart);
+				TSoftObjectPtr<UWorld> LoopStartAsset(LoopStartActor->GetLevel()->GetTypedOuter<UWorld>());
+				*OutReason = FString::Format(TEXT("Setting LevelInstance to {0} would cause loop {1}:{2}\n"), { WorldAsset.GetLongPackageName(), LoopStartActor->GetName(), LoopStartAsset.GetLongPackageName() });
+				for (int32 i = LoopInfo.Num() - 1; i >= 0; --i)
+				{
+					OutReason->Append(FString::Format(TEXT("{0} {1}\n"), { *LoopInfo[i].Key.ToString(), *LoopInfo[i].Value.GetLongPackageName() }));
+				}
+			}
+		}
+
+		return false;
+	}
+
+	return true;
+}
+
 #endif
 
 #undef LOCTEXT_NAMESPACE
\ No newline at end of file
diff --git a/Engine/Source/Runtime/Engine/Private/PackedLevelActor/PackedLevelActor.cpp b/Engine/Source/Runtime/Engine/Private/PackedLevelActor/PackedLevelActor.cpp
--- a/Engine/Source/Runtime/Engine/Private/PackedLevelActor/PackedLevelActor.cpp
+++ b/Engine/Source/Runtime/Engine/Private/PackedLevelActor/PackedLevelActor.cpp
@@ -7,8 +7,8 @@
 #include "LevelInstance/LevelInstancePrivate.h"
 #include "UObject/UE5ReleaseStreamObjectVersion.h"
 
-APackedLevelActor::APackedLevelActor(const FObjectInitializer& ObjectInitializer)
-	: Super(ObjectInitializer)
+APackedLevelActor::APackedLevelActor()
+	: Super()
 #if WITH_EDITORONLY_DATA
 	, bChildChanged(false)
 #endif
@@ -33,7 +33,7 @@ void APackedLevelActor::Serialize(FArchive& Ar)
 #endif
 }
 
-bool APackedLevelActor::SupportsLoading() const
+bool APackedLevelActor::IsLoadingEnabled() const
 {
 #if WITH_EDITOR
 	return HasChildEdit() || IsLoaded();
@@ -104,13 +104,13 @@ FName APackedLevelActor::GetPackedComponentTag()
 	return PackedComponentTag;
 }
 
-void APackedLevelActor::UpdateFromLevel()
+void APackedLevelActor::UpdateLevelInstanceFromWorldAsset()
 {
-	Super::UpdateFromLevel();
+	Super::UpdateLevelInstanceFromWorldAsset();
 
 	if (!Cast<UBlueprint>(GetClass()->ClassGeneratedBy))
 	{
-		if (IsLevelInstancePathValid())
+		if (IsWorldAssetValid())
 		{
 			TSharedPtr<FPackedLevelActorBuilder> Builder = FPackedLevelActorBuilder::CreateDefaultBuilder();
 			Builder->PackActor(this, GetWorldAsset());
diff --git a/Engine/Source/Runtime/Engine/Private/WorldPartition/LevelInstance/LevelInstanceActorDesc.cpp b/Engine/Source/Runtime/Engine/Private/WorldPartition/LevelInstance/LevelInstanceActorDesc.cpp
--- a/Engine/Source/Runtime/Engine/Private/WorldPartition/LevelInstance/LevelInstanceActorDesc.cpp
+++ b/Engine/Source/Runtime/Engine/Private/WorldPartition/LevelInstance/LevelInstanceActorDesc.cpp
@@ -6,7 +6,7 @@
 #if WITH_EDITOR
 #include "Engine/World.h"
 #include "Engine/Level.h"
-#include "LevelInstance/LevelInstanceActor.h"
+#include "LevelInstance/LevelInstanceInterface.h"
 #include "LevelInstance/LevelInstanceSubsystem.h"
 #include "WorldPartition/WorldPartition.h"
 #include "UObject/UE5ReleaseStreamObjectVersion.h"
@@ -45,16 +45,17 @@ void FLevelInstanceActorDesc::Init(const AActor* InActor)
 {
 	FWorldPartitionActorDesc::Init(InActor);
 
-	const ALevelInstance* LevelInstanceActor = CastChecked<ALevelInstance>(InActor);
-	LevelPackage = *LevelInstanceActor->GetWorldAssetPackage();
-	LevelInstanceTransform = LevelInstanceActor->GetActorTransform();
-	DesiredRuntimeBehavior = LevelInstanceActor->GetDesiredRuntimeBehavior();
+	const ILevelInstanceInterface* LevelInstance = CastChecked<ILevelInstanceInterface>(InActor);
+	LevelPackage = *LevelInstance->GetWorldAssetPackage();
+	LevelInstanceTransform = InActor->GetActorTransform();
+	DesiredRuntimeBehavior = LevelInstance->GetDesiredRuntimeBehavior();
 }
 
 void FLevelInstanceActorDesc::Init(const FWorldPartitionActorDescInitData& DescData)
 {
-	ALevelInstance* CDO = DescData.NativeClass->GetDefaultObject<ALevelInstance>();
-	DesiredRuntimeBehavior = CDO->GetDefaultRuntimeBehavior();
+	AActor* CDO = DescData.NativeClass->GetDefaultObject<AActor>();
+	ILevelInstanceInterface* LevelInstanceCDO = CastChecked<ILevelInstanceInterface>(CDO);
+	DesiredRuntimeBehavior = LevelInstanceCDO->GetDefaultRuntimeBehavior();
 
 	FWorldPartitionActorDesc::Init(DescData);
 }
@@ -157,7 +158,9 @@ void FLevelInstanceActorDesc::Serialize(FArchive& Ar)
 		const bool bFixupOldVersion = (Ar.CustomVer(FUE5MainStreamObjectVersion::GUID) < FUE5MainStreamObjectVersion::PackedLevelInstanceBoundsFix) && 
 									  (Ar.CustomVer(FUE5ReleaseStreamObjectVersion::GUID) < FUE5ReleaseStreamObjectVersion::PackedLevelInstanceBoundsFix);
 
-		if (!LevelPackage.IsNone() && (GetActorClass()->GetDefaultObject<ALevelInstance>()->SupportsLoading() || bFixupOldVersion))
+		const AActor* CDO = GetActorClass()->GetDefaultObject<AActor>();
+		const ILevelInstanceInterface* LevelInstanceCDO = CastChecked<ILevelInstanceInterface>(CDO);
+		if (!LevelPackage.IsNone() && (LevelInstanceCDO->IsLoadingEnabled() || bFixupOldVersion))
 		{
 			FBox OutBounds;
 			if (ULevelInstanceSubsystem::GetLevelInstanceBoundsFromPackage(LevelInstanceTransform, LevelPackage, OutBounds))
diff --git a/Engine/Source/Runtime/Engine/Public/LevelInstance/LevelInstanceActor.h b/Engine/Source/Runtime/Engine/Public/LevelInstance/LevelInstanceActor.h
--- a/Engine/Source/Runtime/Engine/Public/LevelInstance/LevelInstanceActor.h
+++ b/Engine/Source/Runtime/Engine/Public/LevelInstance/LevelInstanceActor.h
@@ -7,25 +7,19 @@
 #include "GameFramework/Actor.h"
 #include "Containers/Array.h"
 #include "LevelInstance/LevelInstanceTypes.h"
+#include "LevelInstance/LevelInstanceInterface.h"
+#include "LevelInstance/LevelInstanceActorImpl.h"
+#include "LevelInstance/LevelInstanceActorGuid.h"
 #include "WorldPartition/WorldPartitionActorDesc.h"
 #include "LevelInstanceActor.generated.h"
 
-UENUM()
-enum class ELevelInstanceRuntimeBehavior : uint8
-{
-	None UMETA(Hidden),
-	// Deprecated exists only to avoid breaking Actor Desc serialization
-	Embedded_Deprecated UMETA(Hidden),
-	// Default behavior is to move Level Instance level actors to the main world partition using World Partition clustering rules
-	Partitioned,
-	// Behavior only supported through Conversion Commandlet or on non OFPA Level Instances
-	LevelStreaming UMETA(Hidden)
-};
-
 UCLASS()
-class ENGINE_API ALevelInstance : public AActor
+class ENGINE_API ALevelInstance : public AActor, public ILevelInstanceInterface
 {
-	GENERATED_UCLASS_BODY()
+	GENERATED_BODY()
+
+public:
+	ALevelInstance();
 
 protected:
 #if WITH_EDITORONLY_DATA
@@ -37,6 +31,9 @@ protected:
 	UPROPERTY()
 	TSoftObjectPtr<UWorld> CookedWorldAsset;
 
+private:
+	FLevelInstanceActorGuid LevelInstanceActorGuid;
+	FLevelInstanceActorImpl LevelInstanceActorImpl;
 public:
 
 #if WITH_EDITORONLY_DATA
@@ -44,22 +41,28 @@ public:
 	ELevelInstanceRuntimeBehavior DesiredRuntimeBehavior;
 #endif
 
-	void LoadLevelInstance();
-	void UnloadLevelInstance();
-	virtual bool SupportsLoading() const;
-	bool IsLevelInstancePathValid() const;
-	const TSoftObjectPtr<UWorld>& GetWorldAsset() const;
-	class ULevelInstanceSubsystem* GetLevelInstanceSubsystem() const;
+	// Begin ILevelInstanceInterface
+	const FLevelInstanceID& GetLevelInstanceID() const override;
+	bool HasValidLevelInstanceID() const override;
+	virtual const FGuid& GetLevelInstanceGuid() const override;
+	virtual const TSoftObjectPtr<UWorld>& GetWorldAsset() const override;
 	
-	const FLevelInstanceID& GetLevelInstanceID() const;
-	bool HasValidLevelInstanceID() const;
-
+	virtual void OnLevelInstanceLoaded() override;
+	virtual bool IsLoadingEnabled() const override;
+	// End ILevelInstanceInterface
+			
 	virtual void PostRegisterAllComponents() override;
 	virtual void PostUnregisterAllComponents() override;
 	virtual void Serialize(FArchive& Ar) override;
-
-	const FGuid& GetLevelInstanceActorGuid() const;
+		
 #if WITH_EDITOR
+	// Begin ILevelInstanceInterface
+	virtual ULevelInstanceComponent* GetLevelInstanceComponent() const override;
+	virtual bool SetWorldAsset(TSoftObjectPtr<UWorld> WorldAsset) override;
+	virtual ELevelInstanceRuntimeBehavior GetDesiredRuntimeBehavior() const override { return DesiredRuntimeBehavior; }
+	virtual ELevelInstanceRuntimeBehavior GetDefaultRuntimeBehavior() const override { return ELevelInstanceRuntimeBehavior::Partitioned; }
+	// End ILevelInstanceInterface
+			
 	// UObject overrides
 	virtual void PreEditUndo() override;
 	virtual void PostEditUndo() override;
@@ -82,56 +85,10 @@ public:
 	virtual FBox GetComponentsBoundingBox(bool bNonColliding = false, bool bIncludeFromChildActors = false) const override;
 	virtual FBox GetStreamingBounds() const override;
 	virtual bool IsLockLocation() const override;
-	virtual ELevelInstanceRuntimeBehavior GetDesiredRuntimeBehavior() const { return DesiredRuntimeBehavior; }
-	virtual ELevelInstanceRuntimeBehavior GetDefaultRuntimeBehavior() const { return ELevelInstanceRuntimeBehavior::Partitioned; }
 
-	bool CanEdit(FText* OutReason = nullptr) const;
-	bool CanCommit(FText* OutReason = nullptr) const;
-	bool CanDiscard(FText* OutReason = nullptr) const;
-	bool IsEditing() const;
-	ULevel* GetLoadedLevel() const;
-	bool HasChildEdit() const;
-	void Edit(AActor* ContextActor = nullptr);
-	void Commit();
-	void Discard();
-	bool HasDirtyChildren() const;
-	bool IsDirty() const;
-	bool SetCurrent();
-	bool IsCurrent() const;
-	bool IsLoaded() const;
-	void OnLevelInstanceLoaded();
-	FString GetWorldAssetPackage() const;
-	bool SetWorldAsset(TSoftObjectPtr<UWorld> WorldAsset);
-	bool CheckForLoop(TSoftObjectPtr<UWorld> WorldAsset, TArray<TPair<FText, TSoftObjectPtr<UWorld>>>* LoopInfo = nullptr, const ALevelInstance** LoopStart = nullptr) const;
-	AActor* FindEditorInstanceActor() const;
-
-	virtual void OnWorldAssetChanged() { UpdateFromLevel(); }
-	virtual void OnEdit() {}
-	virtual void OnEditChild() {}
-	virtual void OnCommit(bool bChanged) {}
-	virtual void OnCommitChild(bool bChanged) {}
-	virtual void UpdateFromLevel();
-		
 	DECLARE_MULTICAST_DELEGATE_OneParam(FOnLevelInstanceActorPostLoad, ALevelInstance*);
 	static FOnLevelInstanceActorPostLoad OnLevelInstanceActorPostLoad;
 #endif
-
-private:
-
-#if WITH_EDITOR
-	void PostEditUndoInternal();
-	bool CanSetValue(TSoftObjectPtr<UWorld> LevelInstance, FString* Reason = nullptr) const;
-	
-	TSoftObjectPtr<UWorld> CachedWorldAsset;
-	FLevelInstanceID CachedLevelInstanceID;
-	bool bCachedIsTemporarilyHiddenInEditor;
-	bool bGuardLoadUnload;
-#else
-	// This Guid is used to compute the LevelInstanceID. Because in non-editor build we don't have an ActorGuid, we store it at cook time.
-	FGuid LevelInstanceActorGuid;
-#endif
-
-	FLevelInstanceID LevelInstanceID;
 };
 
 DEFINE_ACTORDESC_TYPE(ALevelInstance, FLevelInstanceActorDesc);
diff --git a/Engine/Source/Runtime/Engine/Public/LevelInstance/LevelInstanceActorGuid.h b/Engine/Source/Runtime/Engine/Public/LevelInstance/LevelInstanceActorGuid.h
new file mode 100644
--- /dev/null
+++ b/Engine/Source/Runtime/Engine/Public/LevelInstance/LevelInstanceActorGuid.h
@@ -0,0 +1,31 @@
+// Copyright Epic Games, Inc. All Rights Reserved.
+
+#pragma once
+
+#include "UObject/ObjectPtr.h"
+#include "Serialization/Archive.h"
+#include "Misc/Guid.h"
+
+class AActor;
+
+/**
+ * Helper struct that allows serializing the ActorGuid for runtime use.
+ */
+struct FLevelInstanceActorGuid
+{
+	// Exists only to support 'FVTableHelper' Actor constructors
+	FLevelInstanceActorGuid() : FLevelInstanceActorGuid(nullptr) {}
+
+	FLevelInstanceActorGuid(AActor* InActor) : Actor(InActor) {}
+
+#if !WITH_EDITOR
+	void AssignIfInvalid();
+#endif
+
+	const FGuid& GetGuid() const;
+
+	TObjectPtr<AActor> Actor = nullptr;
+	FGuid ActorGuid;
+
+	friend FArchive& operator<<(FArchive& Ar, FLevelInstanceActorGuid& LevelInstanceActorGuid);
+};
\ No newline at end of file
diff --git a/Engine/Source/Runtime/Engine/Public/LevelInstance/LevelInstanceActorImpl.h b/Engine/Source/Runtime/Engine/Public/LevelInstance/LevelInstanceActorImpl.h
new file mode 100644
--- /dev/null
+++ b/Engine/Source/Runtime/Engine/Public/LevelInstance/LevelInstanceActorImpl.h
@@ -0,0 +1,93 @@
+// Copyright Epic Games, Inc. All Rights Reserved.
+
+#pragma once
+
+#include "CoreMinimal.h"
+#include "GameFramework/Actor.h"
+#include "LevelInstance/LevelInstanceTypes.h"
+#include "LevelInstance/LevelInstanceInterface.h"
+#include "LevelInstance/LevelInstanceSubsystem.h"
+#include "Misc/ITransaction.h"
+#include "Templates/UnrealTemplate.h"
+
+/**
+ * Provides base implementation for AActor classes implementing the ILevelInstanceInterface
+ *
+ * Should be used as a member of the AActor class implementing ILevelInstanceInterface
+ *
+ * @see ILevelInstanceInterface, ALevelInstance
+ */
+class ENGINE_API FLevelInstanceActorImpl
+{
+private:
+	ILevelInstanceInterface* LevelInstance;
+	FLevelInstanceID LevelInstanceID;
+
+#if WITH_EDITOR
+	FLevelInstanceID CachedLevelInstanceID;
+	bool bCachedIsTemporarilyHiddenInEditor;
+	bool bGuardLoadUnload;
+public:
+	TSoftObjectPtr<UWorld> CachedWorldAsset;
+#endif
+public:
+	// Exists only to support 'FVTableHelper' Actor constructors
+	FLevelInstanceActorImpl()
+		: FLevelInstanceActorImpl(nullptr)
+	{
+	}
+
+	FLevelInstanceActorImpl(ILevelInstanceInterface* InLevelInstance)
+		: LevelInstance(InLevelInstance)
+#if WITH_EDITOR
+		, bCachedIsTemporarilyHiddenInEditor(false)
+		, bGuardLoadUnload(false)
+#endif
+	{
+	}
+
+	virtual ~FLevelInstanceActorImpl()
+	{
+	}
+
+	virtual void RegisterLevelInstance();
+	virtual void UnregisterLevelInstance();
+
+	/**
+	 * Begin ILevelInstanceInterface Implementaion 
+	 * 
+	 */
+	const FLevelInstanceID& GetLevelInstanceID() const;
+	bool HasValidLevelInstanceID() const;
+	virtual bool IsLoadingEnabled() const;
+	virtual void OnLevelInstanceLoaded();
+
+	/**
+	 * Begin AActor Implementation
+	 *
+	 * Actor implementing ILevelInstanceInterface should be overriding most of those methods and forward the call their member FLevelInstanceActorImpl
+	 * 
+	 * @see ALevelInstance as an example
+	 */
+#if WITH_EDITOR
+	virtual void PreEditUndo(TFunctionRef<void()> SuperCall);
+	virtual void PostEditUndo(TFunctionRef<void()> SuperCall);
+	virtual void PostEditUndo(TSharedPtr<ITransactionObjectAnnotation> TransactionAnnotation, TFunctionRef<void(TSharedPtr<ITransactionObjectAnnotation> TransactionAnnotation)> SuperCall);
+	virtual void PostEditImport(TFunctionRef<void()> SuperCall);
+	virtual bool CanEditChange(const FProperty* Property) const;
+	virtual void PreEditChange(FProperty* Property, bool bWorldAssetChange, TFunctionRef<void(FProperty*)> SuperCall);
+	virtual void PostEditChangeProperty(FPropertyChangedEvent& PropertyChangedEvent, bool bWorldAssetChange, TFunction<void(FPropertyChangedEvent&)> SuperCall);
+	virtual bool CanDeleteSelectedActor(FText& OutReason) const;
+	virtual void SetIsTemporarilyHiddenInEditor(bool bIsHidden, TFunctionRef<void(bool)> SuperCall);
+	virtual bool SetIsHiddenEdLayer(bool bIsHiddenEdLayer, TFunctionRef<bool(bool)> SuperCall);
+	virtual void EditorGetUnderlyingActors(TSet<AActor*>& OutUnderlyingActors) const;
+	virtual bool IsLockLocation() const;
+	virtual bool GetBounds(FBox& OutBounds) const;
+	virtual void PushSelectionToProxies();
+	virtual void PushLevelInstanceEditingStateToProxies(bool bInEditingState);
+	virtual void CheckForErrors();
+	// End AActor Implementation
+private:
+	virtual void PostEditUndoInternal();
+#endif
+};
\ No newline at end of file
diff --git a/Engine/Source/Runtime/Engine/Public/LevelInstance/LevelInstanceComponent.h b/Engine/Source/Runtime/Engine/Public/LevelInstance/LevelInstanceComponent.h
new file mode 100644
--- /dev/null
+++ b/Engine/Source/Runtime/Engine/Public/LevelInstance/LevelInstanceComponent.h
@@ -0,0 +1,37 @@
+// Copyright Epic Games, Inc. All Rights Reserved.
+
+#pragma once
+
+#include "CoreMinimal.h"
+#include "Components/SceneComponent.h"
+
+#include "LevelInstanceComponent.generated.h"
+
+/**
+ * ULevelInstanceComponent subclasses USceneComponent for Editing purposes so that we can have a proxy to the LevelInstanceActor's RootComponent transform without attaching to it.
+ *
+ * It is responsible for updating the transform of the ALevelInstanceEditorInstanceActor which is created when loading a LevelInstance Instance Level
+ *
+ * We use this method to avoid attaching the Instance Level Actors to the ILevelInstanceInterface. (Cross level attachment and undo/redo pain)
+ * 
+ * The LevelInstance Level Actors are attached to this ALevelInstanceEditorInstanceActor keeping the attachment local to the Instance Level and shielded from the transaction buffer.
+ *
+ * Avoiding those Level Actors from being part of the Transaction system allows us to unload that level without clearing the transaction buffer. It also allows BP Reinstancing without having to update attachements.
+ */
+UCLASS()
+class ENGINE_API ULevelInstanceComponent : public USceneComponent
+{
+	GENERATED_UCLASS_BODY()
+public:
+#if WITH_EDITOR
+	// Those are the methods that need overriding to be able to properly update the AttachComponent
+	virtual void OnRegister() override;
+	virtual void PostEditUndo() override;
+	virtual void PostEditChangeProperty(FPropertyChangedEvent& PropertyChangedEvent) override;
+	virtual void OnUpdateTransform(EUpdateTransformFlags UpdateTransformFlags, ETeleportType Teleport) override;
+
+	void UpdateEditorInstanceActor();
+private:
+	TWeakObjectPtr<AActor> CachedEditorInstanceActorPtr;
+#endif
+};
\ No newline at end of file
diff --git a/Engine/Source/Runtime/Engine/Public/LevelInstance/LevelInstanceEditorInstanceActor.h b/Engine/Source/Runtime/Engine/Public/LevelInstance/LevelInstanceEditorInstanceActor.h
--- a/Engine/Source/Runtime/Engine/Public/LevelInstance/LevelInstanceEditorInstanceActor.h
+++ b/Engine/Source/Runtime/Engine/Public/LevelInstance/LevelInstanceEditorInstanceActor.h
@@ -8,11 +8,11 @@
 #include "LevelInstance/LevelInstanceTypes.h"
 #include "LevelInstanceEditorInstanceActor.generated.h"
 
-class ALevelInstance;
+class ILevelInstanceInterface;
 class ULevel;
 
 /**
- * Editor Only Actor that is spawned inside every LevelInstance Instance Level so that we can update its Actor Transforms through the ALevelInstance's root component(ULevelInstanceComponent)
+ * Editor Only Actor that is spawned inside every LevelInstance Instance Level so that we can update its Actor Transforms through the ILevelInstanceInterface's (ULevelInstanceComponent)
  * @see ULevelInstanceComponent
  */
 UCLASS(transient, notplaceable)
@@ -22,7 +22,7 @@ class ENGINE_API ALevelInstanceEditorInstanceActor : public AActor
 
 #if WITH_EDITOR
 public:
-	static ALevelInstanceEditorInstanceActor* Create(ALevelInstance* LevelInstanceActor, ULevel* LoadedLevel);
+	static ALevelInstanceEditorInstanceActor* Create(ILevelInstanceInterface* LevelInstance, ULevel* LoadedLevel);
 	
 	void SetLevelInstanceID(const FLevelInstanceID& InLevelInstanceID) { LevelInstanceID = InLevelInstanceID; }
 	const FLevelInstanceID& GetLevelInstanceID() const { return LevelInstanceID; }
diff --git a/Engine/Source/Runtime/Engine/Public/LevelInstance/LevelInstanceEditorLevelStreaming.h b/Engine/Source/Runtime/Engine/Public/LevelInstance/LevelInstanceEditorLevelStreaming.h
--- a/Engine/Source/Runtime/Engine/Public/LevelInstance/LevelInstanceEditorLevelStreaming.h
+++ b/Engine/Source/Runtime/Engine/Public/LevelInstance/LevelInstanceEditorLevelStreaming.h
@@ -8,7 +8,7 @@
 #include "LevelInstance/LevelInstanceTypes.h"
 #include "LevelInstanceEditorLevelStreaming.generated.h"
 
-class ALevelInstance;
+class ILevelInstanceInterface;
 
 UCLASS(Transient, MinimalAPI)
 class ULevelStreamingLevelInstanceEditor : public ULevelStreamingAlwaysLoaded
@@ -18,7 +18,7 @@ class ULevelStreamingLevelInstanceEditor : public ULevelStreamingAlwaysLoaded
 public:
 #if WITH_EDITOR
 	virtual bool ShowInLevelCollection() const override { return false; }
-	ALevelInstance* GetLevelInstanceActor() const;
+	ILevelInstanceInterface* GetLevelInstance() const;
 	const FLevelInstanceID& GetLevelInstanceID() const { return LevelInstanceID; }
 	FBox GetBounds() const;
 
@@ -28,7 +28,7 @@ protected:
 
 	friend class ULevelInstanceSubsystem;
 
-	static ULevelStreamingLevelInstanceEditor* Load(ALevelInstance* LevelInstanceActor);
+	static ULevelStreamingLevelInstanceEditor* Load(ILevelInstanceInterface* LevelInstance);
 	static void Unload(ULevelStreamingLevelInstanceEditor* LevelStreaming);
 
 	virtual void OnLevelLoadedChanged(ULevel* Level) override;
diff --git a/Engine/Source/Runtime/Engine/Public/LevelInstance/LevelInstanceEditorPivotActor.h b/Engine/Source/Runtime/Engine/Public/LevelInstance/LevelInstanceEditorPivotActor.h
--- a/Engine/Source/Runtime/Engine/Public/LevelInstance/LevelInstanceEditorPivotActor.h
+++ b/Engine/Source/Runtime/Engine/Public/LevelInstance/LevelInstanceEditorPivotActor.h
@@ -8,7 +8,7 @@
 #include "LevelInstance/LevelInstanceTypes.h"
 #include "LevelInstanceEditorPivotActor.generated.h"
 
-class ALevelInstance;
+class ILevelInstanceInterface;
 class ULevelStreaming;
 
 /**
@@ -21,7 +21,7 @@ class ENGINE_API ALevelInstancePivot : public AActor
 
 #if WITH_EDITOR
 public:
-	static ALevelInstancePivot* Create(ALevelInstance* LevelInstanceActor, ULevelStreaming* LevelStreaming);
+	static ALevelInstancePivot* Create(ILevelInstanceInterface* LevelInstance, ULevelStreaming* LevelStreaming);
 	
 	virtual bool CanDeleteSelectedActor(FText& OutReason) const override { return false; }
 	virtual void PostEditMove(bool bFinished) override;
diff --git a/Engine/Source/Runtime/Engine/Public/LevelInstance/LevelInstanceInterface.h b/Engine/Source/Runtime/Engine/Public/LevelInstance/LevelInstanceInterface.h
new file mode 100644
--- /dev/null
+++ b/Engine/Source/Runtime/Engine/Public/LevelInstance/LevelInstanceInterface.h
@@ -0,0 +1,99 @@
+// Copyright Epic Games, Inc. All Rights Reserved.
+
+#pragma once
+
+#include "CoreMinimal.h"
+#include "UObject/Interface.h"
+#include "UObject/SoftObjectPtr.h"
+#include "LevelInstance/LevelInstanceTypes.h"
+#include "LevelInstanceInterface.generated.h"
+
+class ULevelInstanceComponent;
+
+UINTERFACE()
+class ENGINE_API ULevelInstanceInterface : public UInterface
+{
+	GENERATED_UINTERFACE_BODY()
+};
+
+/**
+ * Interface to be implemented by Actor classes to implement support for LevelInstances
+ *
+ * Interface exists to allow integration of the LevelInstance workflow on existing Actor class hierarchies.
+ * 
+ * @see FLevelInstanceActorImpl to use as a member or base class member of the implementing class. It provides most of the boiler plate code so that the implementing class behaves properly in editor (Undo/Redo/Delete/Selection/Imporrt/etc)
+ */
+class ENGINE_API ILevelInstanceInterface
+{
+	GENERATED_IINTERFACE_BODY()
+			
+	// Pure Interface Start
+	virtual const FLevelInstanceID& GetLevelInstanceID() const = 0;
+
+	virtual bool HasValidLevelInstanceID() const = 0;
+
+	virtual const FGuid& GetLevelInstanceGuid() const = 0;
+
+	virtual const TSoftObjectPtr<UWorld>& GetWorldAsset() const = 0;
+
+	virtual bool IsLoadingEnabled() const = 0;
+
+#if WITH_EDITOR
+	virtual bool SetWorldAsset(TSoftObjectPtr<UWorld> WorldAsset) = 0;
+
+	virtual ULevelInstanceComponent* GetLevelInstanceComponent() const = 0;
+
+	virtual ELevelInstanceRuntimeBehavior GetDesiredRuntimeBehavior() const = 0;
+
+	virtual ELevelInstanceRuntimeBehavior GetDefaultRuntimeBehavior() const = 0;
+#endif
+	// Pure Interface End
+		
+	virtual void OnLevelInstanceLoaded() {}
+		
+	FString GetWorldAssetPackage() const { return GetWorldAsset().GetUniqueID().GetLongPackageName(); }
+
+	bool IsWorldAssetValid() const { return GetWorldAsset().GetUniqueID().IsValid(); }
+
+	virtual ULevelInstanceSubsystem* GetLevelInstanceSubsystem() const;
+
+	virtual bool IsLoaded() const;
+			
+	virtual void LoadLevelInstance();
+
+	virtual void UnloadLevelInstance();
+		
+#if WITH_EDITOR
+	virtual ULevel* GetLoadedLevel() const;
+
+	virtual void UpdateLevelInstanceFromWorldAsset();
+
+	virtual void OnEdit() {}
+
+	virtual void OnEditChild() {}
+
+	virtual void OnCommit(bool bChanged) {}
+
+	virtual void OnCommitChild(bool bChanged) {}
+	
+	virtual bool IsEditing() const;
+	
+	virtual bool HasChildEdit() const;
+	
+	virtual bool IsDirty() const;
+	
+	virtual bool HasDirtyChildren() const;
+	
+	virtual bool CanEnterEdit(FText* OutReason = nullptr) const;
+	
+	virtual bool EnterEdit(AActor* ContextActor = nullptr);
+	
+	virtual bool CanExitEdit(bool bDiscardEdits = false, FText* OutReason = nullptr) const;
+	
+	virtual bool ExitEdit(bool bDiscardEdits = false);
+	
+	virtual bool SetCurrent();
+	
+	virtual bool MoveActorsTo(const TArray<AActor*>& ActorsToMove);
+#endif
+};
\ No newline at end of file
diff --git a/Engine/Source/Runtime/Engine/Public/LevelInstance/LevelInstanceLevelStreaming.h b/Engine/Source/Runtime/Engine/Public/LevelInstance/LevelInstanceLevelStreaming.h
--- a/Engine/Source/Runtime/Engine/Public/LevelInstance/LevelInstanceLevelStreaming.h
+++ b/Engine/Source/Runtime/Engine/Public/LevelInstance/LevelInstanceLevelStreaming.h
@@ -8,7 +8,7 @@
 #include "LevelInstance/LevelInstanceTypes.h"
 #include "LevelInstanceLevelStreaming.generated.h"
 
-class ALevelInstance;
+class ILevelInstanceInterface;
 
 UCLASS(Transient)
 class ENGINE_API ULevelStreamingLevelInstance : public ULevelStreamingDynamic
@@ -16,7 +16,7 @@ class ENGINE_API ULevelStreamingLevelInstance : public ULevelStreamingDynamic
 	GENERATED_UCLASS_BODY()
 
 public:
-	ALevelInstance* GetLevelInstanceActor() const;
+	ILevelInstanceInterface* GetLevelInstance() const;
 
 #if WITH_EDITOR
 	virtual bool ShowInLevelCollection() const override { return false; }
@@ -26,7 +26,7 @@ public:
 #endif
 	
 protected:
-	static ULevelStreamingLevelInstance* LoadInstance(ALevelInstance* LevelInstanceActor);
+	static ULevelStreamingLevelInstance* LoadInstance(ILevelInstanceInterface* LevelInstanceActor);
 	static void UnloadInstance(ULevelStreamingLevelInstance* LevelStreaming);
 
 	virtual void OnLevelLoadedChanged(ULevel* Level) override;
diff --git a/Engine/Source/Runtime/Engine/Public/LevelInstance/LevelInstanceSubsystem.h b/Engine/Source/Runtime/Engine/Public/LevelInstance/LevelInstanceSubsystem.h
--- a/Engine/Source/Runtime/Engine/Public/LevelInstance/LevelInstanceSubsystem.h
+++ b/Engine/Source/Runtime/Engine/Public/LevelInstance/LevelInstanceSubsystem.h
@@ -11,7 +11,7 @@
 
 #include "LevelInstanceSubsystem.generated.h"
 
-class ALevelInstance;
+class ILevelInstanceInterface;
 class ULevelInstanceEditorObject;
 class ULevelStreamingLevelInstance;
 class ULevelStreamingLevelInstanceEditor;
@@ -39,13 +39,13 @@ public:
 	virtual void UpdateStreamingState() override;
 	//~ End UWorldSubsystem Interface.
 
-	ALevelInstance* GetLevelInstance(FLevelInstanceID LevelInstanceID) const;
-	FLevelInstanceID RegisterLevelInstance(ALevelInstance* LevelInstanceActor);
-	void UnregisterLevelInstance(ALevelInstance* LevelInstanceActor);
-	void RequestLoadLevelInstance(ALevelInstance* LevelInstanceActor, bool bUpdate);
-	void RequestUnloadLevelInstance(ALevelInstance* LevelInstanceActor);
-	bool IsLoaded(const ALevelInstance* LevelInstanceActor) const;
-	void ForEachLevelInstanceAncestorsAndSelf(AActor* Actor, TFunctionRef<bool(ALevelInstance*)> Operation) const;
+	ILevelInstanceInterface* GetLevelInstance(FLevelInstanceID LevelInstanceID) const;
+	FLevelInstanceID RegisterLevelInstance(ILevelInstanceInterface* LevelInstance);
+	void UnregisterLevelInstance(ILevelInstanceInterface* LevelInstance);
+	void RequestLoadLevelInstance(ILevelInstanceInterface* LevelInstance, bool bUpdate);
+	void RequestUnloadLevelInstance(ILevelInstanceInterface* LevelInstance);
+	bool IsLoaded(const ILevelInstanceInterface* LevelInstance) const;
+	void ForEachLevelInstanceAncestorsAndSelf(AActor* Actor, TFunctionRef<bool(ILevelInstanceInterface*)> Operation) const;
 
 #if WITH_EDITOR
 	void Tick();
@@ -55,55 +55,60 @@ public:
 	void PackAllLoadedActors();
 	bool CanPackAllLoadedActors() const;
 
-	ALevelInstance* GetEditingLevelInstance() const;
-	bool CanEditLevelInstance(const ALevelInstance* LevelInstanceActor, FText* OutReason = nullptr) const;
-	bool CanCommitLevelInstance(const ALevelInstance* LevelInstanceActor, FText* OutReason = nullptr) const;
-	bool CanDiscardLevelInstance(const ALevelInstance* LevelInstanceActor, FText* OutReason = nullptr) const;
-	void EditLevelInstance(ALevelInstance* LevelInstanceActor, TWeakObjectPtr<AActor> ContextActorPtr = nullptr);
-	bool CommitLevelInstance(ALevelInstance* LevelInstanceActor, bool bDiscardEdits = false, TSet<FName>* DirtyPackages = nullptr);
-	bool IsEditingLevelInstanceDirty(const ALevelInstance* LevelInstanceActor) const;
-	bool IsEditingLevelInstance(const ALevelInstance* LevelInstanceActor) const { return GetLevelInstanceEdit(LevelInstanceActor) != nullptr; }
+	ILevelInstanceInterface* GetEditingLevelInstance() const;
+	bool CanEditLevelInstance(const ILevelInstanceInterface* LevelInstance, FText* OutReason = nullptr) const;
+	bool CanCommitLevelInstance(const ILevelInstanceInterface* LevelInstance, bool bDiscardEdits = false, FText* OutReason = nullptr) const;
+	void EditLevelInstance(ILevelInstanceInterface* LevelInstance, TWeakObjectPtr<AActor> ContextActorPtr = nullptr);
+	bool CommitLevelInstance(ILevelInstanceInterface* LevelInstance, bool bDiscardEdits = false, TSet<FName>* DirtyPackages = nullptr);
+	bool IsEditingLevelInstanceDirty(const ILevelInstanceInterface* LevelInstance) const;
+	bool IsEditingLevelInstance(const ILevelInstanceInterface* LevelInstance) const { return GetLevelInstanceEdit(LevelInstance) != nullptr; }
 	
-	bool GetLevelInstanceBounds(const ALevelInstance* LevelInstanceActor, FBox& OutBounds) const;
+	bool GetLevelInstanceBounds(const ILevelInstanceInterface* LevelInstance, FBox& OutBounds) const;
 	static bool GetLevelInstanceBoundsFromPackage(const FTransform& InstanceTransform, FName LevelPackage, FBox& OutBounds);
 	
-	void ForEachActorInLevelInstance(const ALevelInstance* LevelInstanceActor, TFunctionRef<bool(AActor * LevelActor)> Operation) const;
-	void ForEachLevelInstanceAncestorsAndSelf(const AActor* Actor, TFunctionRef<bool(const ALevelInstance*)> Operation) const;
-	void ForEachLevelInstanceAncestors(const AActor* Actor, TFunctionRef<bool(const ALevelInstance*)> Operation) const;
-	void ForEachLevelInstanceChild(const ALevelInstance* LevelInstanceActor, bool bRecursive, TFunctionRef<bool(const ALevelInstance*)> Operation) const;
-	void ForEachLevelInstanceChild(ALevelInstance* LevelInstanceActor, bool bRecursive, TFunctionRef<bool(ALevelInstance*)> Operation) const;
-	bool HasDirtyChildrenLevelInstances(const ALevelInstance* LevelInstanceActor) const;
+	void ForEachActorInLevelInstance(const ILevelInstanceInterface* LevelInstance, TFunctionRef<bool(AActor * LevelActor)> Operation) const;
+	void ForEachLevelInstanceAncestorsAndSelf(const AActor* Actor, TFunctionRef<bool(const ILevelInstanceInterface*)> Operation) const;
+	void ForEachLevelInstanceAncestors(const AActor* Actor, TFunctionRef<bool(const ILevelInstanceInterface*)> Operation) const;
+	void ForEachLevelInstanceChild(const ILevelInstanceInterface* LevelInstance, bool bRecursive, TFunctionRef<bool(const ILevelInstanceInterface*)> Operation) const;
+	void ForEachLevelInstanceChild(ILevelInstanceInterface* LevelInstance, bool bRecursive, TFunctionRef<bool(ILevelInstanceInterface*)> Operation) const;
+	bool HasDirtyChildrenLevelInstances(const ILevelInstanceInterface* LevelInstance) const;
 	
-	void SetIsHiddenEdLayer(ALevelInstance* LevelInstanceActor, bool bIsHiddenEdLayer);
-	void SetIsTemporarilyHiddenInEditor(ALevelInstance* LevelInstanceActor, bool bIsHidden);
+	void SetIsHiddenEdLayer(ILevelInstanceInterface* LevelInstance, bool bIsHiddenEdLayer);
+	void SetIsTemporarilyHiddenInEditor(ILevelInstanceInterface* LevelInstance, bool bIsHidden);
 
-	bool SetCurrent(ALevelInstance* LevelInstanceActor) const;
-	bool IsCurrent(const ALevelInstance* LevelInstanceActor) const;
-	ALevelInstance* CreateLevelInstanceFrom(const TArray<AActor*>& ActorsToMove, const FNewLevelInstanceParams& CreationParams);
+	bool SetCurrent(ILevelInstanceInterface* LevelInstance) const;
+	bool IsCurrent(const ILevelInstanceInterface* LevelInstance) const;
+	ILevelInstanceInterface* CreateLevelInstanceFrom(const TArray<AActor*>& ActorsToMove, const FNewLevelInstanceParams& CreationParams);
 	bool MoveActorsToLevel(const TArray<AActor*>& ActorsToRemove, ULevel* DestinationLevel, TArray<AActor*>* OutActors = nullptr) const;
-	bool MoveActorsTo(ALevelInstance* LevelInstanceActor, const TArray<AActor*>& ActorsToMove, TArray<AActor*>* OutActors = nullptr);
-	bool BreakLevelInstance(ALevelInstance* LevelInstanceActor, uint32 Levels = 1, TArray<AActor*>* OutMovedActors = nullptr);
+	bool MoveActorsTo(ILevelInstanceInterface* LevelInstance, const TArray<AActor*>& ActorsToMove, TArray<AActor*>* OutActors = nullptr);
+	bool BreakLevelInstance(ILevelInstanceInterface* LevelInstance, uint32 Levels = 1, TArray<AActor*>* OutMovedActors = nullptr);
 
 	bool CanMoveActorToLevel(const AActor* Actor, FText* OutReason = nullptr) const;
 	void OnActorDeleted(AActor* Actor);
-	ULevel* GetLevelInstanceLevel(const ALevelInstance* LevelInstanceActor) const;
+	ULevel* GetLevelInstanceLevel(const ILevelInstanceInterface* LevelInstance) const;
 
-	bool LevelInstanceHasLevelScriptBlueprint(const ALevelInstance* LevelInstance) const;
+	bool LevelInstanceHasLevelScriptBlueprint(const ILevelInstanceInterface* LevelInstance) const;
 
-	ALevelInstance* GetParentLevelInstance(const AActor* Actor) const;
-	void BlockLoadLevelInstance(ALevelInstance* LevelInstanceActor);
-	void BlockUnloadLevelInstance(ALevelInstance* LevelInstanceActor);
+	ILevelInstanceInterface* GetParentLevelInstance(const AActor* Actor) const;
+	void BlockLoadLevelInstance(ILevelInstanceInterface* LevelInstance);
+	void BlockUnloadLevelInstance(ILevelInstanceInterface* LevelInstance);
 		
-	bool HasChildEdit(const ALevelInstance* LevelInstanceActor) const;
+	bool HasChildEdit(const ILevelInstanceInterface* LevelInstance) const;
+
+	TArray<ILevelInstanceInterface*> GetLevelInstances(const FString& WorldAssetPackage);
+	
+	static bool CheckForLoop(const ILevelInstanceInterface* LevelInstance, TArray<TPair<FText, TSoftObjectPtr<UWorld>>>* LoopInfo = nullptr, const ILevelInstanceInterface** LoopStart = nullptr);
+	static bool CheckForLoop(const ILevelInstanceInterface* LevelInstance, TSoftObjectPtr<UWorld> WorldAsset, TArray<TPair<FText, TSoftObjectPtr<UWorld>>>* LoopInfo = nullptr, const ILevelInstanceInterface** LoopStart = nullptr);
+	static bool CanUseWorldAsset(const ILevelInstanceInterface* LevelInstance, TSoftObjectPtr<UWorld> WorldAsset, FString* OutReason);
 #endif
 
 private:
 	void BlockOnLoading();
-	void LoadLevelInstance(ALevelInstance* LevelInstanceActor);
+	void LoadLevelInstance(ILevelInstanceInterface* LevelInstance);
 	void UnloadLevelInstance(const FLevelInstanceID& LevelInstanceID);
 	void ForEachActorInLevel(ULevel* Level, TFunctionRef<bool(AActor * LevelActor)> Operation) const;
-	void ForEachLevelInstanceAncestors(AActor* Actor, TFunctionRef<bool(ALevelInstance*)> Operation) const;
-	ALevelInstance* GetOwningLevelInstance(const ULevel* Level) const;
+	void ForEachLevelInstanceAncestors(AActor* Actor, TFunctionRef<bool(ILevelInstanceInterface*)> Operation) const;
+	ILevelInstanceInterface* GetOwningLevelInstance(const ULevel* Level) const;
 	
 	void RegisterLoadedLevelStreamingLevelInstance(ULevelStreamingLevelInstance* LevelStreaming);
 
@@ -113,10 +118,10 @@ private:
 	void OnEditChild(FLevelInstanceID LevelInstanceID);
 	void OnCommitChild(FLevelInstanceID LevelInstanceID, bool bChildChanged);
 
-	bool ForEachLevelInstanceChildImpl(const ALevelInstance* LevelInstanceActor, bool bRecursive, TFunctionRef<bool(const ALevelInstance*)> Operation) const;
-	bool ForEachLevelInstanceChildImpl(ALevelInstance* LevelInstanceActor, bool bRecursive, TFunctionRef<bool(ALevelInstance*)> Operation) const;
+	bool ForEachLevelInstanceChildImpl(const ILevelInstanceInterface* LevelInstance, bool bRecursive, TFunctionRef<bool(const ILevelInstanceInterface*)> Operation) const;
+	bool ForEachLevelInstanceChildImpl(ILevelInstanceInterface* LevelInstance, bool bRecursive, TFunctionRef<bool(ILevelInstanceInterface*)> Operation) const;
 
-	void BreakLevelInstance_Impl(ALevelInstance* LevelInstanceActor, uint32 Levels, TArray<AActor*>& OutMovedActors);
+	void BreakLevelInstance_Impl(ILevelInstanceInterface* LevelInstance, uint32 Levels, TArray<AActor*>& OutMovedActors);
 
 	static bool ShouldIgnoreDirtyPackage(UPackage* DirtyPackage, const UWorld* EditingWorld);
 
@@ -142,10 +147,10 @@ private:
 	};
 
 	void ResetEdit(TUniquePtr<FLevelInstanceEdit>& InLevelInstanceEdit);
-	bool EditLevelInstanceInternal(ALevelInstance* LevelInstanceActor, TWeakObjectPtr<AActor> ContextActorPtr, const FString& InActorNameToSelect, bool bRecursive);
+	bool EditLevelInstanceInternal(ILevelInstanceInterface* LevelInstance, TWeakObjectPtr<AActor> ContextActorPtr, const FString& InActorNameToSelect, bool bRecursive);
 	bool CommitLevelInstanceInternal(TUniquePtr<FLevelInstanceEdit>& InLevelInstanceEdit, bool bDiscardEdits = false, bool bDiscardOnFailure = false, TSet<FName>* DirtyPackages = nullptr);
 	
-	const FLevelInstanceEdit* GetLevelInstanceEdit(const ALevelInstance* LevelInstanceActor) const;
+	const FLevelInstanceEdit* GetLevelInstanceEdit(const ILevelInstanceInterface* LevelInstance) const;
 	bool IsLevelInstanceEditDirty(const FLevelInstanceEdit* LevelInstanceEdit) const;
 	
 	struct FLevelsToRemoveScope
@@ -175,10 +180,10 @@ private:
 		ULevelStreamingLevelInstance* LevelStreaming = nullptr;
 	};
 
-	TMap<ALevelInstance*, bool> LevelInstancesToLoadOrUpdate;
+	TMap<ILevelInstanceInterface*, bool> LevelInstancesToLoadOrUpdate;
 	TSet<FLevelInstanceID> LevelInstancesToUnload;
 	TMap<FLevelInstanceID, FLevelInstance> LevelInstances;
-	TMap<FLevelInstanceID, ALevelInstance*> RegisteredLevelInstances;
+	TMap<FLevelInstanceID, ILevelInstanceInterface*> RegisteredLevelInstances;
 
 #if WITH_EDITORONLY_DATA
 	// Optional scope to accelerate level unload by batching them
diff --git a/Engine/Source/Runtime/Engine/Public/LevelInstance/LevelInstanceTypes.h b/Engine/Source/Runtime/Engine/Public/LevelInstance/LevelInstanceTypes.h
--- a/Engine/Source/Runtime/Engine/Public/LevelInstance/LevelInstanceTypes.h
+++ b/Engine/Source/Runtime/Engine/Public/LevelInstance/LevelInstanceTypes.h
@@ -4,6 +4,7 @@
 
 #include "CoreMinimal.h"
 #include "Misc/Guid.h"
+#include "Templates/SubclassOf.h"
 #include "LevelInstanceTypes.generated.h"
 
 // FLevelInstanceID is a runtime unique id that is computed from the Hash of LevelInstance Actor Guid and all its ancestor LevelInstance Actor Guids.
@@ -11,7 +12,7 @@
 struct FLevelInstanceID
 {
 	FLevelInstanceID() {}
-	FLevelInstanceID(class ULevelInstanceSubsystem* LevelInstanceSubsystem, class ALevelInstance* Actor);
+	FLevelInstanceID(class ULevelInstanceSubsystem* LevelInstanceSubsystem, class ILevelInstanceInterface* LevelInstance);
 
 	inline friend uint32 GetTypeHash(const FLevelInstanceID& Key)
 	{
@@ -37,6 +38,18 @@ private:
 	TArray<FGuid> Guids;
 };
 
+UENUM()
+enum class ELevelInstanceRuntimeBehavior : uint8
+{
+	None UMETA(Hidden),
+	// Deprecated exists only to avoid breaking Actor Desc serialization
+	Embedded_Deprecated UMETA(Hidden),
+	// Default behavior is to move Level Instance level actors to the main world partition using World Partition clustering rules
+	Partitioned,
+	// Behavior only supported through Conversion Commandlet or on non OFPA Level Instances
+	LevelStreaming UMETA(Hidden)
+};
+
 UENUM()
 enum class ELevelInstanceCreationType : uint8
 {
@@ -76,6 +89,9 @@ struct FNewLevelInstanceParams
 	UPROPERTY()
 	bool bPromptForSave = false;
 
+	UPROPERTY()
+	TSubclassOf<AActor> LevelInstanceClass;
+
 private:
 	UPROPERTY(EditAnywhere, Category = Default, meta = (EditCondition = "!bForceExternalActors", EditConditionHides, HideEditConditionToggle))
 	bool bExternalActors = true;
diff --git a/Engine/Source/Runtime/Engine/Public/PackedLevelActor/PackedLevelActor.h b/Engine/Source/Runtime/Engine/Public/PackedLevelActor/PackedLevelActor.h
--- a/Engine/Source/Runtime/Engine/Public/PackedLevelActor/PackedLevelActor.h
+++ b/Engine/Source/Runtime/Engine/Public/PackedLevelActor/PackedLevelActor.h
@@ -24,10 +24,12 @@ class UBlueprint;
 UCLASS()
 class ENGINE_API APackedLevelActor : public ALevelInstance
 {
-	GENERATED_UCLASS_BODY()
+	GENERATED_BODY()
 
 public:
-	virtual bool SupportsLoading() const override;
+	APackedLevelActor();
+
+	virtual bool IsLoadingEnabled() const override;
 
 	virtual void Serialize(FArchive& Ar) override;
 #if WITH_EDITOR
@@ -39,7 +41,7 @@ public:
 	static FName GetPackedComponentTag();
 
 	virtual void PostLoad() override;
-	virtual void UpdateFromLevel() override;
+	virtual void UpdateLevelInstanceFromWorldAsset() override;
 	virtual void OnCommit(bool bChanged) override;
 	virtual void OnCommitChild(bool bChanged) override;
 	virtual void OnEdit() override;
