World Partition
- Resaved all templates skybox actors to be non-spatially loaded.
- Only include spatially loaded actors when checking for enabling streaming logic.

#rnx

diff --git a/Engine/Source/Runtime/Engine/Private/WorldPartition/WorldPartition.cpp b/Engine/Source/Runtime/Engine/Private/WorldPartition/WorldPartition.cpp
--- a/Engine/Source/Runtime/Engine/Private/WorldPartition/WorldPartition.cpp
+++ b/Engine/Source/Runtime/Engine/Private/WorldPartition/WorldPartition.cpp
@@ -1293,21 +1293,24 @@ void UWorldPartition::Tick(float DeltaSeconds)
 			FBox AllActorsBounds(ForceInit);
 			for (UActorDescContainer::TConstIterator<> ActorDescIterator(this); ActorDescIterator; ++ActorDescIterator)
 			{
-				AllActorsBounds += ActorDescIterator->GetBounds();
-
-				// Warn the user if the world becomes larger that 4km in any axis
-				if (AllActorsBounds.GetSize().GetMax() >= 400000.0f)
+				if (ActorDescIterator->GetIsSpatiallyLoadedRaw())
 				{
-					bShouldEnableStreamingWarned = true;
+					AllActorsBounds += ActorDescIterator->GetBounds();
 
-					if (FMessageDialog::Open(EAppMsgType::YesNo, LOCTEXT("WorldPartitionShouldEnableStreaming", "The size of your world has grown enough to justify enabling streaming. Enable streaming now?")) == EAppReturnType::Yes)
+					// Warn the user if the world becomes larger that 4km in any axis
+					if (AllActorsBounds.GetSize().GetMax() >= 400000.0f)
 					{
-						bEnableStreaming = true;
-						bStreamingWasEnabled = true;
-						FMessageDialog::Open(EAppMsgType::Ok, LOCTEXT("WorldPartitionEnableStreamingDialolg", "Please refer to https://docs.unrealengine.com/5.0/en-US/building-virtual-worlds/world-partition for how to set up streaming."));			
-					}
+						bShouldEnableStreamingWarned = true;
+
+						if (FMessageDialog::Open(EAppMsgType::YesNo, LOCTEXT("WorldPartitionShouldEnableStreaming", "The size of your world has grown enough to justify enabling streaming. Enable streaming now?")) == EAppReturnType::Yes)
+						{
+							SetEnableStreaming(true);
+							bStreamingWasEnabled = true;
+							FMessageDialog::Open(EAppMsgType::Ok, LOCTEXT("WorldPartitionEnableStreamingDialolg", "Please refer to https://docs.unrealengine.com/5.0/en-US/building-virtual-worlds/world-partition for how to set up streaming."));			
+						}
 
-					break;
+						break;
+					}
 				}
 			}
 		}
@@ -1569,4 +1572,4 @@ FBox UWorldPartition::GetEditorWorldBounds() const
 }
 #endif
 
-#undef LOCTEXT_NAMESPACE
+#undef LOCTEXT_NAMESPACE
\ No newline at end of file
diff --git a/Engine/Source/Runtime/Engine/Public/WorldPartition/WorldPartitionActorDesc.h b/Engine/Source/Runtime/Engine/Public/WorldPartition/WorldPartitionActorDesc.h
--- a/Engine/Source/Runtime/Engine/Public/WorldPartition/WorldPartitionActorDesc.h
+++ b/Engine/Source/Runtime/Engine/Public/WorldPartition/WorldPartitionActorDesc.h
@@ -50,6 +50,7 @@ public:
 	inline FVector GetOrigin() const { return GetBounds().GetCenter(); }
 	inline FName GetRuntimeGrid() const { return RuntimeGrid; }
 	inline bool GetIsSpatiallyLoaded() const { return bForceAlwaysLoaded ? false : bIsSpatiallyLoaded; }
+	inline bool GetIsSpatiallyLoadedRaw() const { return bIsSpatiallyLoaded; }
 	inline bool GetActorIsEditorOnly() const { return bActorIsEditorOnly; }
 	inline bool GetLevelBoundsRelevant() const { return bLevelBoundsRelevant; }
 	inline bool GetActorIsHLODRelevant() const { return bActorIsHLODRelevant; }
