Fixed refactored enum name conflict.

diff --git a/Engine/Source/Runtime/Engine/Private/Net/RPCDoSDetection.cpp b/Engine/Source/Runtime/Engine/Private/Net/RPCDoSDetection.cpp
--- a/Engine/Source/Runtime/Engine/Private/Net/RPCDoSDetection.cpp
+++ b/Engine/Source/Runtime/Engine/Private/Net/RPCDoSDetection.cpp
@@ -72,29 +72,29 @@ TAutoConsoleVariable<FString> CVarRPCDoSForcedRPCTracking(
 
 
 /**
- * EEscalateReason
+ * ERPCDoSEscalateReason
  */
 
 /**
- * Convert EEscalateReason enum values, to a string.
+ * Convert ERPCDoSEscalateReason enum values, to a string.
  *
  * @param Reason	The enum value to convert.
  * @return			The string name for the enum value.
  */
-const TCHAR* LexToString(EEscalateReason Reason)
+const TCHAR* LexToString(ERPCDoSEscalateReason Reason)
 {
 	switch (Reason)
 	{
-	case EEscalateReason::CountLimit:
+	case ERPCDoSEscalateReason::CountLimit:
 		return TEXT("CountLimit");
 
-	case EEscalateReason::TimeLimit:
+	case ERPCDoSEscalateReason::TimeLimit:
 		return TEXT("TimeLimit");
 
-	case EEscalateReason::AutoEscalate:
+	case ERPCDoSEscalateReason::AutoEscalate:
 		return TEXT("AutoEscalate");
 
-	case EEscalateReason::Deescalate:
+	case ERPCDoSEscalateReason::Deescalate:
 		return TEXT("Deescalate");
 
 	default:
@@ -488,9 +488,9 @@ void FRPCDoSDetection::InitConfig(FName NetDriverName)
 	}
 }
 
-void FRPCDoSDetection::UpdateSeverity(ESeverityUpdate Update, EEscalateReason Reason)
+void FRPCDoSDetection::UpdateSeverity(ERPCDoSSeverityUpdate Update, ERPCDoSEscalateReason Reason)
 {
-	bool bEscalate = Update == ESeverityUpdate::Escalate || Update == ESeverityUpdate::AutoEscalate;
+	bool bEscalate = Update == ERPCDoSSeverityUpdate::Escalate || Update == ERPCDoSSeverityUpdate::AutoEscalate;
 	int32 NewStateIdx = FMath::Clamp(ActiveState + (bEscalate ? 1 : -1), 0, DetectionSeverity.Num()-1);
 
 	// If kicking is disabled, and the new state is the kick state, exclude that state (otherwise RPC DoS Detection will become stuck)
@@ -766,11 +766,11 @@ void FRPCDoSDetection::PreTickDispatch(double TimeSeconds)
 
 			if (CooloffTime > 0 && ActiveStateTime > CooloffTime)
 			{
-				UpdateSeverity(ESeverityUpdate::Deescalate, EEscalateReason::Deescalate);
+				UpdateSeverity(ERPCDoSSeverityUpdate::Deescalate, ERPCDoSEscalateReason::Deescalate);
 			}
 			else if (AutoEscalateTime > 0 && ActiveStateTime > AutoEscalateTime)
 			{
-				UpdateSeverity(ESeverityUpdate::AutoEscalate, EEscalateReason::AutoEscalate);
+				UpdateSeverity(ERPCDoSSeverityUpdate::AutoEscalate, ERPCDoSEscalateReason::AutoEscalate);
 			}
 		}
 
diff --git a/Engine/Source/Runtime/Engine/Public/Net/RPCDoSDetection.h b/Engine/Source/Runtime/Engine/Public/Net/RPCDoSDetection.h
--- a/Engine/Source/Runtime/Engine/Public/Net/RPCDoSDetection.h
+++ b/Engine/Source/Runtime/Engine/Public/Net/RPCDoSDetection.h
@@ -363,7 +363,7 @@ enum class ERPCNotifyResult : uint8
 /**
  * RPC DoS escalation severity update types
  */
-enum class ESeverityUpdate : uint8
+enum class ERPCDoSSeverityUpdate : uint8
 {
 	Escalate,		// Escalating to a higher RPC DoS escalation state
 	AutoEscalate,	// Automatic/non-RPC escalation to a higher RPC DoS escalation state, e.g. when too much time was spent in a lower state
@@ -373,7 +373,7 @@ enum class ESeverityUpdate : uint8
 /**
  * The reason for an RPC DoS severity update
  */
-enum class EEscalateReason : uint8
+enum class ERPCDoSEscalateReason : uint8
 {
 	CountLimit,		// RPC Count quota's were hit
 	TimeLimit,		// RPC Time quota's were hit
@@ -457,7 +457,7 @@ public:
 	 * @param Update	Whether or not we are escalating or de-escalating the severity state
 	 * @param Reason	The reason for the escalation change
 	 */
-	void UpdateSeverity(ESeverityUpdate Update, EEscalateReason Reason);
+	void UpdateSeverity(ERPCDoSSeverityUpdate Update, ERPCDoSEscalateReason Reason);
 
 
 	/**
@@ -679,7 +679,7 @@ private:
 		// 64 RPC interval
 		if ((RPCIntervalCounter & 0x3F) == 0 && HasHitQuota_Count(CounterPerPeriodHistory, FrameCounter))
 		{
-			UpdateSeverity(ESeverityUpdate::Escalate, EEscalateReason::CountLimit);
+			UpdateSeverity(ERPCDoSSeverityUpdate::Escalate, ERPCDoSEscalateReason::CountLimit);
 		}
 	}
 
@@ -696,7 +696,7 @@ private:
 
 			if (HasHitQuota_Time(CounterPerPeriodHistory, FrameCounter))
 			{
-				UpdateSeverity(ESeverityUpdate::Escalate, EEscalateReason::TimeLimit);
+				UpdateSeverity(ERPCDoSSeverityUpdate::Escalate, ERPCDoSEscalateReason::TimeLimit);
 			}
 		}
 	}
