Converting uses of RHICreateTargetableShaderResource2D to use RHICreateTargetableShaderResource.

diff --git a/Engine/Plugins/Media/AndroidCamera/Source/AndroidCamera/Private/Player/AndroidCameraTextureSample.h b/Engine/Plugins/Media/AndroidCamera/Source/AndroidCamera/Private/Player/AndroidCameraTextureSample.h
--- a/Engine/Plugins/Media/AndroidCamera/Source/AndroidCamera/Private/Player/AndroidCameraTextureSample.h
+++ b/Engine/Plugins/Media/AndroidCamera/Source/AndroidCamera/Private/Player/AndroidCameraTextureSample.h
@@ -142,23 +142,13 @@ public:
 			return Texture;
 		}
 
-		const ETextureCreateFlags CreateFlags = TexCreate_Dynamic | TexCreate_SRGB;
-
-		TRefCountPtr<FRHITexture2D> DummyTexture2DRHI;
-		FRHIResourceCreateInfo CreateInfo(TEXT("DummyTexture2D"));
-
-		RHICreateTargetableShaderResource2D(
-			Dim.X,
-			Dim.Y,
-			PF_B8G8R8A8,
-			1,
-			CreateFlags,
-			TexCreate_RenderTargetable,
-			false,
-			CreateInfo,
-			Texture,
-			DummyTexture2DRHI
-		);
+		const FRHITextureCreateDesc Desc =
+			FRHITextureCreateDesc::Create2D(TEXT("DummyTexture2D"))
+			.SetExtent(Dim)
+			.SetFormat(PF_B8G8R8A8)
+			.SetFlags(ETextureCreateFlags::Dynamic | ETextureCreateFlags::SRGB);
+
+		RHICreateTargetableShaderResource(Desc, ETextureCreateFlags::RenderTargetable, Texture);
 
 		return Texture;
 	}
diff --git a/Engine/Plugins/Media/AndroidMedia/Source/AndroidMedia/Private/Player/AndroidMediaTextureSample.h b/Engine/Plugins/Media/AndroidMedia/Source/AndroidMedia/Private/Player/AndroidMediaTextureSample.h
--- a/Engine/Plugins/Media/AndroidMedia/Source/AndroidMedia/Private/Player/AndroidMediaTextureSample.h
+++ b/Engine/Plugins/Media/AndroidMedia/Source/AndroidMedia/Private/Player/AndroidMediaTextureSample.h
@@ -143,23 +143,13 @@ public:
 			return Texture;
 		}
 
-		const ETextureCreateFlags CreateFlags = TexCreate_Dynamic | TexCreate_SRGB;
-
-		TRefCountPtr<FRHITexture2D> DummyTexture2DRHI;
-		FRHIResourceCreateInfo CreateInfo(TEXT("DummyTexture2D"));
-
-		RHICreateTargetableShaderResource2D(
-			Dim.X,
-			Dim.Y,
-			PF_B8G8R8A8,
-			1,
-			CreateFlags,
-			TexCreate_RenderTargetable,
-			false,
-			CreateInfo,
-			Texture,
-			DummyTexture2DRHI
-		);
+		const FRHITextureCreateDesc Desc =
+			FRHITextureCreateDesc::Create2D(TEXT("DummyTexture2D"))
+			.SetExtent(Dim)
+			.SetFormat(PF_B8G8R8A8)
+			.SetFlags(ETextureCreateFlags::Dynamic | ETextureCreateFlags::SRGB);
+
+		RHICreateTargetableShaderResource(Desc, ETextureCreateFlags::RenderTargetable, Texture);
 
 		return Texture;
 	}
diff --git a/Engine/Plugins/Media/BinkMedia/Source/BinkMediaPlayer/Private/Assets/BinkMediaTextureResource.cpp b/Engine/Plugins/Media/BinkMedia/Source/BinkMediaPlayer/Private/Assets/BinkMediaTextureResource.cpp
--- a/Engine/Plugins/Media/BinkMedia/Source/BinkMediaPlayer/Private/Assets/BinkMediaTextureResource.cpp
+++ b/Engine/Plugins/Media/BinkMedia/Source/BinkMediaPlayer/Private/Assets/BinkMediaTextureResource.cpp
@@ -19,7 +19,6 @@ void FBinkMediaTextureResource::InitDynamicRHI()
 
 	// Create the RHI texture. Only one mip is used and the texture is targetable or resolve.
 	ETextureCreateFlags TexCreateFlags = Owner->SRGB ? TexCreate_SRGB : TexCreate_None;
-	FRHIResourceCreateInfo CreateInfo(TEXT("Bink"));
 	
 	if (bink_force_pixel_format != PF_Unknown) 
 	{
@@ -32,27 +31,22 @@ void FBinkMediaTextureResource::InitDynamicRHI()
 		TexCreateFlags = TexCreate_None;
 	}
 
+	const TCHAR* DebugName = TEXT("Bink");
+
 #if !(UE_BUILD_TEST || UE_BUILD_SHIPPING)
 	FString DebugNameString = TEXT("Bink:");
 	DebugNameString += Owner->GetName();
-	CreateInfo.DebugName = *DebugNameString;
+	DebugName = *DebugNameString;
 #endif // ARK_EXTRA_RESOURCE_NAMES
 
-	TRefCountPtr<FRHITexture2D> Texture2DRHI;
-	RHICreateTargetableShaderResource2D(
-		w, h,
-		PixelFormat,
-		1,
-		TexCreateFlags,
-		TexCreate_RenderTargetable,
-		false,
-		CreateInfo,
-		RenderTargetTextureRHI,
-		Texture2DRHI
-	);
+	const FRHITextureCreateDesc Desc =
+		FRHITextureCreateDesc::Create2D(DebugName)
+		.SetExtent(w, h)
+		.SetFormat(PixelFormat)
+		.SetFlags(TexCreateFlags);
+
+	RHICreateTargetableShaderResource(Desc, ETextureCreateFlags::RenderTargetable, TextureRHI);
 
-	TextureRHI = (FTextureRHIRef&)Texture2DRHI;
-	
 	// Don't bother updating if its not a valid video
 	if (Owner->GetSurfaceWidth() && Owner->GetSurfaceHeight()) 
 	{
diff --git a/Engine/Plugins/Media/ElectraUtil/Source/ElectraSamples/Private/ElectraTextureSampleAndroid.cpp b/Engine/Plugins/Media/ElectraUtil/Source/ElectraSamples/Private/ElectraTextureSampleAndroid.cpp
--- a/Engine/Plugins/Media/ElectraUtil/Source/ElectraSamples/Private/ElectraTextureSampleAndroid.cpp
+++ b/Engine/Plugins/Media/ElectraUtil/Source/ElectraSamples/Private/ElectraTextureSampleAndroid.cpp
@@ -422,21 +422,13 @@ void FElectraTextureSample::InitializeTexture()
 		return;
 	}
 
-	TRefCountPtr<FRHITexture2D> DummyTexture2DRHI;
-	FRHIResourceCreateInfo CreateInfo(TEXT("FElectraTextureSample"));
-
-	RHICreateTargetableShaderResource2D(
-		Dim.X,
-		Dim.Y,
-		PF_B8G8R8A8,
-		1,
-		TexCreate_Dynamic | TexCreate_SRGB,
-		TexCreate_RenderTargetable,
-		false,
-		CreateInfo,
-		Texture,
-		DummyTexture2DRHI
-	);
+	const FRHITextureCreateDesc Desc =
+		FRHITextureCreateDesc::Create2D(TEXT("FElectraTextureSample"))
+		.SetExtent(Dim)
+		.SetFormat(PF_B8G8R8A8)
+		.SetFlags(ETextureCreateFlags::Dynamic | ETextureCreateFlags::SRGB);
+
+	RHICreateTargetableShaderResource(Desc, ETextureCreateFlags::RenderTargetable, Texture);
 
 	return;
 }
diff --git a/Engine/Plugins/Media/ElectraUtil/Source/ElectraSamples/Public/GStreamer/ElectraTextureSample.h b/Engine/Plugins/Media/ElectraUtil/Source/ElectraSamples/Public/GStreamer/ElectraTextureSample.h
--- a/Engine/Plugins/Media/ElectraUtil/Source/ElectraSamples/Public/GStreamer/ElectraTextureSample.h
+++ b/Engine/Plugins/Media/ElectraUtil/Source/ElectraSamples/Public/GStreamer/ElectraTextureSample.h
@@ -34,23 +34,13 @@ public:
 	{
 		check(IsInRenderingThread());
 
-		const ETextureCreateFlags CreateFlags = TexCreate_Dynamic | TexCreate_SRGB;
-
-		TRefCountPtr<FRHITexture2D> DummyTexture2DRHI;
-		FRHIResourceCreateInfo CreateInfo(TEXT("DummyTexture2D"));
-
-		RHICreateTargetableShaderResource2D(
-			TotalSize.X,
-			TotalSize.Y,
-			PF_B8G8R8A8,
-			1,
-			CreateFlags,
-			TexCreate_RenderTargetable,
-			false,
-			CreateInfo,
-			Texture,
-			DummyTexture2DRHI
-		);
+		const FRHITextureCreateDesc Desc =
+			FRHITextureCreateDesc::Create2D(TEXT("DummyTexture2D"))
+			.SetExtent(TotalSize)
+			.SetFormat(PF_B8G8R8A8)
+			.SetFlags(ETextureCreateFlags::Dynamic | ETextureCreateFlags::SRGB);
+
+		RHICreateTargetableShaderResource(Desc, ETextureCreateFlags::RenderTargetable, Texture);
 	}	
 
 public:
diff --git a/Engine/Plugins/Media/WebMMedia/Source/WebMMedia/Private/Player/WebMMediaTextureSample.cpp b/Engine/Plugins/Media/WebMMedia/Source/WebMMedia/Private/Player/WebMMediaTextureSample.cpp
--- a/Engine/Plugins/Media/WebMMedia/Source/WebMMedia/Private/Player/WebMMediaTextureSample.cpp
+++ b/Engine/Plugins/Media/WebMMedia/Source/WebMMedia/Private/Player/WebMMediaTextureSample.cpp
@@ -20,23 +20,13 @@ void FWebMMediaTextureSample::CreateTexture()
 {
 	check(IsInRenderingThread());
 
-	const ETextureCreateFlags CreateFlags = TexCreate_Dynamic | TexCreate_SRGB;
-
-	TRefCountPtr<FRHITexture2D> DummyTexture2DRHI;
-	FRHIResourceCreateInfo CreateInfo(TEXT("FWebMMediaTextureSample_DummyTexture2D"));
-
-	RHICreateTargetableShaderResource2D(
-		TotalSize.X,
-		TotalSize.Y,
-		PF_B8G8R8A8,
-		1,
-		CreateFlags,
-		TexCreate_RenderTargetable,
-		false,
-		CreateInfo,
-		Texture,
-		DummyTexture2DRHI
-	);
+	const FRHITextureCreateDesc Desc =
+		FRHITextureCreateDesc::Create2D(TEXT("FWebMMediaTextureSample_DummyTexture2D"))
+		.SetExtent(TotalSize)
+		.SetFormat(PF_B8G8R8A8)
+		.SetFlags(ETextureCreateFlags::Dynamic | ETextureCreateFlags::SRGB);
+
+	RHICreateTargetableShaderResource(Desc, ETextureCreateFlags::RenderTargetable, Texture);
 }
 
 const void* FWebMMediaTextureSample::GetBuffer()
diff --git a/Engine/Plugins/Runtime/Steam/SteamVR/Source/SteamVR/Private/SteamVRHMD.cpp b/Engine/Plugins/Runtime/Steam/SteamVR/Source/SteamVR/Private/SteamVRHMD.cpp
--- a/Engine/Plugins/Runtime/Steam/SteamVR/Source/SteamVR/Private/SteamVRHMD.cpp
+++ b/Engine/Plugins/Runtime/Steam/SteamVR/Source/SteamVR/Private/SteamVRHMD.cpp
@@ -1649,21 +1649,25 @@ bool FSteamVRHMD::AllocateRenderTargetTexture(uint32 Index, uint32 SizeX, uint32
 		return true;
 	}
 
+	FRHITextureCreateDesc Desc =
+		FRHITextureCreateDesc::Create2D(TEXT("FSteamVRHMD"))
+		.SetExtent(SizeX, SizeY)
+		.SetFormat(PF_B8G8R8A8)
+		.SetNumSamples(NumSamples);
+
 	for (uint32 SwapChainIter = 0; SwapChainIter < SteamVRSwapChainLength; ++SwapChainIter)
 	{
-		FRHIResourceCreateInfo CreateInfo(TEXT("FSteamVRHMD"));
 #if PLATFORM_MAC
 		IOSurfaceRef Surface = MetalBridgePtr->GetSurface(SizeX, SizeY);
 		check(Surface != nil);
 
-		CreateInfo.BulkData = new FIOSurfaceResourceWrapper(Surface);
+		Desc.BulkData = new FIOSurfaceResourceWrapper(Surface);
 		CFRelease(Surface);
-		CreateInfo.ResourceArray = nullptr;
 #endif
 
-		FTexture2DRHIRef TargetableTexture, ShaderResourceTexture;
-		RHICreateTargetableShaderResource2D(SizeX, SizeY, PF_B8G8R8A8, 1, TexCreate_None, TexCreate_RenderTargetable, false, CreateInfo, TargetableTexture, ShaderResourceTexture, NumSamples);
-		check(TargetableTexture == ShaderResourceTexture);
+		FTexture2DRHIRef TargetableTexture;
+
+		RHICreateTargetableShaderResource(Desc, ETextureCreateFlags::RenderTargetable, TargetableTexture);
 
 		SwapChainTextures.Add((FTextureRHIRef&)TargetableTexture);
 
@@ -1711,28 +1715,21 @@ bool FSteamVRHMD::AllocateDepthTexture(uint32 Index, uint32 SizeX, uint32 SizeY,
 
 	FClearValueBinding ClearValue(0.0f, 0);
 	ClearValue.ColorBinding = EClearBinding::EDepthStencilBound;
-	FRHIResourceCreateInfo CreateInfo(TEXT("SteamVRDepthStencil"), ClearValue);
+
+	const FRHITextureCreateDesc Desc =
+		FRHITextureCreateDesc::Create2D(TEXT("SteamVRDepthStencil"))
+		.SetExtent(SizeX, SizeY)
+		.SetFormat(PF_DepthStencil)
+		.SetFlags(Flags)
+		.SetClearValue(ClearValue);
 
 	for (uint32 SwapChainIter = 0; SwapChainIter < SteamVRSwapChainLength; ++SwapChainIter)
 	{
-		FTexture2DRHIRef TargetableTexture, ShaderResourceTexture;
+		FTextureRHIRef TargetableTexture;
 
-		RHICreateTargetableShaderResource2D(
-			SizeX,
-			SizeY,
-			PF_DepthStencil,
-			1,			// 1 mip for depth (0 is usually passed in above, which is invalid).
-			Flags,
-			TargetableTextureFlags,
-			false,
-			CreateInfo,
-			TargetableTexture,
-			ShaderResourceTexture,
-			1);
+		RHICreateTargetableShaderResource(Desc, TargetableTextureFlags, TargetableTexture);
 
-		check(TargetableTexture == ShaderResourceTexture);
-
-		SwapChainTextures.Add((FTextureRHIRef&)TargetableTexture);
+		SwapChainTextures.Add(TargetableTexture);
 
 		if (BindingTexture == nullptr)
 		{
diff --git a/Engine/Plugins/Runtime/WarpUtils/Source/PFMExporter/Private/PFMExporterMesh.cpp b/Engine/Plugins/Runtime/WarpUtils/Source/PFMExporter/Private/PFMExporterMesh.cpp
--- a/Engine/Plugins/Runtime/WarpUtils/Source/PFMExporter/Private/PFMExporterMesh.cpp
+++ b/Engine/Plugins/Runtime/WarpUtils/Source/PFMExporter/Private/PFMExporterMesh.cpp
@@ -49,8 +49,12 @@ bool FPFMExporterMesh::SaveToFile(const FString& FileName)
 bool FPFMExporterMesh::BeginExport_RenderThread(FRHICommandListImmediate& RHICmdList)
 {
 	// Allocate RT+SR RHI resources
-	FRHIResourceCreateInfo CreateInfo(TEXT("FPFMExporterMesh"));
-	RHICreateTargetableShaderResource2D(DimWidth, DimHeight, PFMTextureFormat, 1, TexCreate_None, TexCreate_RenderTargetable, true, CreateInfo, RenderTargetTexture, ShaderResourceTexture);
+	const FRHITextureCreateDesc Desc =
+		FRHITextureCreateDesc::Create2D(TEXT("FPFMExporterMesh"))
+		.SetExtent(DimWidth, DimHeight)
+		.SetFormat(PFMTextureFormat);
+
+	RHICreateTargetableShaderResource(Desc, ETextureCreateFlags::RenderTargetable, true, RenderTargetTexture, ShaderResourceTexture);
 
 	return true;
 }
diff --git a/Engine/Plugins/Runtime/WindowsMixedReality/Source/HoloLensAR/Private/HoloLensCameraImageTexture.cpp b/Engine/Plugins/Runtime/WindowsMixedReality/Source/HoloLensAR/Private/HoloLensCameraImageTexture.cpp
--- a/Engine/Plugins/Runtime/WindowsMixedReality/Source/HoloLensAR/Private/HoloLensCameraImageTexture.cpp
+++ b/Engine/Plugins/Runtime/WindowsMixedReality/Source/HoloLensAR/Private/HoloLensCameraImageTexture.cpp
@@ -217,15 +217,25 @@ public:
 
 			// Create the copy target
 			{
-				FRHIResourceCreateInfo CreateInfo(TEXT("FHoloLensCameraImageResource_CopyTextureRef"));
-				CopyTextureRef = RHICreateTexture2D(Size.X, Size.Y, PF_NV12, 1, 1, TexCreate_Dynamic | TexCreate_ShaderResource, CreateInfo);
+				const FRHITextureCreateDesc CreateDesc =
+					FRHITextureCreateDesc::Create2D(TEXT("FHoloLensCameraImageResource_CopyTextureRef"))
+					.SetExtent(Size)
+					.SetFormat(PF_NV12)
+					.SetFlags(ETextureCreateFlags::Dynamic | ETextureCreateFlags::ShaderResource);
+
+				CopyTextureRef = RHICreateTexture(CreateDesc);
 			}
 			// Create the render target
 			{
-				FRHIResourceCreateInfo CreateInfo(TEXT("FHoloLensCameraImageResource_DummyTexture2D"));
-				TRefCountPtr<FRHITexture2D> DummyTexture2DRHI;
+
+				const FRHITextureCreateDesc CreateDesc =
+					FRHITextureCreateDesc::Create2D(TEXT("FHoloLensCameraImageResource_DummyTexture2D"))
+					.SetExtent(Size)
+					.SetFormat(PF_B8G8R8A8)
+					.SetFlags(ETextureCreateFlags::Dynamic);
+
 				// Create our render target that we'll convert to
-				RHICreateTargetableShaderResource2D(Size.X, Size.Y, PF_B8G8R8A8, 1, TexCreate_Dynamic, TexCreate_RenderTargetable, false, CreateInfo, DecodedTextureRef, DummyTexture2DRHI);
+				RHICreateTargetableShaderResource(CreateDesc, ETextureCreateFlags::RenderTargetable, DecodedTextureRef);
 			}
 
 			if (PerformCopy(cameraImageTexture, D3D11DeviceContext))
diff --git a/Engine/Plugins/Runtime/WindowsMoviePlayer/Source/WindowsMoviePlayer/Private/WindowsMovieStreamer.cpp b/Engine/Plugins/Runtime/WindowsMoviePlayer/Source/WindowsMoviePlayer/Private/WindowsMovieStreamer.cpp
--- a/Engine/Plugins/Runtime/WindowsMoviePlayer/Source/WindowsMoviePlayer/Private/WindowsMovieStreamer.cpp
+++ b/Engine/Plugins/Runtime/WindowsMoviePlayer/Source/WindowsMoviePlayer/Private/WindowsMovieStreamer.cpp
@@ -88,21 +88,13 @@ void FMediaFoundationMovieStreamer::ConvertSample()
 		// create a new input render target if necessary
 		if (!InputTarget.IsValid() || (InputTarget->GetSizeXY() != SourceFormat.BufferDim) || (InputTarget->GetFormat() != InputPixelFormat) || ((InputTarget->GetFlags() & InputCreateFlags) != InputCreateFlags))
 		{
-			TRefCountPtr<FRHITexture2D> DummyTexture2DRHI;
-			FRHIResourceCreateInfo CreateInfo(TEXT("FMediaFoundationMovieStreamer"));
-
-			RHICreateTargetableShaderResource2D(
-				SourceFormat.BufferDim.X,
-				SourceFormat.BufferDim.Y,
-				InputPixelFormat,
-				1,
-				InputCreateFlags,
-				TexCreate_RenderTargetable,
-				false,
-				CreateInfo,
-				InputTarget,
-				DummyTexture2DRHI
-			);
+			const FRHITextureCreateDesc Desc =
+				FRHITextureCreateDesc::Create2D(TEXT("FMediaFoundationMovieStreamer"))
+				.SetExtent(SourceFormat.BufferDim)
+				.SetFormat(InputPixelFormat)
+				.SetFlags(InputCreateFlags);
+
+			RHICreateTargetableShaderResource(Desc, ETextureCreateFlags::RenderTargetable, InputTarget);
 		}
 
 		// copy sample data to input render target
diff --git a/Engine/Plugins/Runtime/nDisplay/Source/DisplayCluster/Private/Render/Viewport/RenderTarget/DisplayClusterRenderTargetResource.cpp b/Engine/Plugins/Runtime/nDisplay/Source/DisplayCluster/Private/Render/Viewport/RenderTarget/DisplayClusterRenderTargetResource.cpp
--- a/Engine/Plugins/Runtime/nDisplay/Source/DisplayCluster/Private/Render/Viewport/RenderTarget/DisplayClusterRenderTargetResource.cpp
+++ b/Engine/Plugins/Runtime/nDisplay/Source/DisplayCluster/Private/Render/Viewport/RenderTarget/DisplayClusterRenderTargetResource.cpp
@@ -38,20 +38,15 @@ void FDisplayClusterViewportResource::ImplInitDynamicRHI_RenderTargetResource2D(
 		NumMips = ViewportResourceSettings.NumMips;
 	}
 
-	FRHIResourceCreateInfo CreateInfo(TEXT("DisplayClusterViewportRenderTargetResource"), FClearValueBinding::Black);
-
-	RHICreateTargetableShaderResource2D(
-		GetSizeX(),
-		GetSizeY(),
-		ViewportResourceSettings.Format,
-		NumMips,
-		CreateFlags,
-		TexCreate_RenderTargetable,
-		false,
-		CreateInfo,
-		OutRenderTargetTextureRHI,
-		OutTextureRHI
-	);
+	const FRHITextureCreateDesc Desc =
+		FRHITextureCreateDesc::Create2D(TEXT("DisplayClusterViewportRenderTargetResource"))
+		.SetExtent(GetSizeX(), GetSizeY())
+		.SetFormat(ViewportResourceSettings.Format)
+		.SetNumMips(NumMips)
+		.SetFlags(CreateFlags)
+		.SetClearValue(FClearValueBinding::Black);
+
+	RHICreateTargetableShaderResource(Desc, ETextureCreateFlags::RenderTargetable, OutRenderTargetTextureRHI, OutTextureRHI);
 }
 
 void FDisplayClusterViewportResource::ImplInitDynamicRHI_TextureResource2D(FTexture2DRHIRef& OutTextureRHI)
diff --git a/Engine/Source/Runtime/Engine/Private/Particles/ParticleCurveTexture.cpp b/Engine/Source/Runtime/Engine/Private/Particles/ParticleCurveTexture.cpp
--- a/Engine/Source/Runtime/Engine/Private/Particles/ParticleCurveTexture.cpp
+++ b/Engine/Source/Runtime/Engine/Private/Particles/ParticleCurveTexture.cpp
@@ -515,20 +515,13 @@ FParticleCurveTexture::FParticleCurveTexture()
 void FParticleCurveTexture::InitRHI()
 {
 	// 8-bit per channel RGBA texture for curves.
-	FRHIResourceCreateInfo CreateInfo(TEXT("ParticleCurveTexture"), FClearValueBinding(FLinearColor::Blue));
-
-	RHICreateTargetableShaderResource2D(
-		GParticleCurveTextureSizeX,
-		GParticleCurveTextureSizeY,
-		PF_B8G8R8A8,
-		/*NumMips=*/ 1,
-		TexCreate_None,
-		TexCreate_RenderTargetable | TexCreate_NoFastClear,
-		/*bForceSeparateTargetAndShaderResource=*/ false,
-		CreateInfo,
-		CurveTextureTargetRHI,
-		CurveTextureRHI
-	);
+	const FRHITextureCreateDesc Desc =
+		FRHITextureCreateDesc::Create2D(TEXT("ParticleCurveTexture"))
+		.SetExtent(GParticleCurveTextureSizeX, GParticleCurveTextureSizeY)
+		.SetFormat(PF_B8G8R8A8)
+		.SetClearValue(FClearValueBinding(FLinearColor::Blue));
+
+	RHICreateTargetableShaderResource(Desc, ETextureCreateFlags::RenderTargetable | ETextureCreateFlags::NoFastClear, CurveTextureTargetRHI, CurveTextureRHI);
 }
 
 /**
diff --git a/Engine/Source/Runtime/Engine/Private/Particles/ParticleGpuSimulation.cpp b/Engine/Source/Runtime/Engine/Private/Particles/ParticleGpuSimulation.cpp
--- a/Engine/Source/Runtime/Engine/Private/Particles/ParticleGpuSimulation.cpp
+++ b/Engine/Source/Runtime/Engine/Private/Particles/ParticleGpuSimulation.cpp
@@ -251,37 +251,29 @@ public:
 #define PARTICLE_STATE_POSITION_TEXTURE_NAME	TEXT("ParticleStatePosition")
 #define PARTICLE_STATE_VELOCITY_TEXTURE_NAME	TEXT("ParticleStateVelocity")
 
-		FRHIResourceCreateInfo CreateInfo(PARTICLE_STATE_POSITION_TEXTURE_NAME, FClearValueBinding::Transparent);
-		RHICreateTargetableShaderResource2D(
-			SizeX,
-			SizeY,
-			PF_A32B32G32R32F,
-			/*NumMips=*/ 1,
-			TexCreate_None,
-			TexCreate_RenderTargetable,
-			/*bForceSeparateTargetAndShaderResource=*/ false,
-			CreateInfo,
-			PositionTextureTargetRHI,
-			PositionTextureRHI
-			);
+		{
+			const FRHITextureCreateDesc Desc =
+				FRHITextureCreateDesc::Create2D(PARTICLE_STATE_POSITION_TEXTURE_NAME)
+				.SetExtent(SizeX, SizeY)
+				.SetFormat(PF_A32B32G32R32F)
+				.SetClearValue(FClearValueBinding::Transparent);
+
+			RHICreateTargetableShaderResource(Desc, ETextureCreateFlags::RenderTargetable, PositionTextureTargetRHI, PositionTextureRHI);
+		}
 
 		// 16-bit per channel RGBA texture for velocity.
 		check( !IsValidRef( VelocityTextureTargetRHI ) );
 		check( !IsValidRef( VelocityTextureRHI ) );
 
-		CreateInfo.DebugName = PARTICLE_STATE_VELOCITY_TEXTURE_NAME;
-		RHICreateTargetableShaderResource2D(
-			SizeX,
-			SizeY,
-			PF_FloatRGBA,
-			/*NumMips=*/ 1,
-			TexCreate_None,
-			TexCreate_RenderTargetable,
-			/*bForceSeparateTargetAndShaderResource=*/ false,
-			CreateInfo,
-			VelocityTextureTargetRHI,
-			VelocityTextureRHI
-			);
+		{
+			const FRHITextureCreateDesc Desc =
+				FRHITextureCreateDesc::Create2D(PARTICLE_STATE_VELOCITY_TEXTURE_NAME)
+				.SetExtent(SizeX, SizeY)
+				.SetFormat(PF_FloatRGBA)
+				.SetClearValue(FClearValueBinding::Transparent);
+
+			RHICreateTargetableShaderResource(Desc, ETextureCreateFlags::RenderTargetable, VelocityTextureTargetRHI, VelocityTextureRHI);
+		}
 
 		// using FName's ability to append a number to a string (..._0) without an extra string allocation, except suffixing is done when number > 0 hence the +1 here : 
 		FName PositionTextureName(PARTICLE_STATE_POSITION_TEXTURE_NAME, ParticleStateIndex + 1); 
@@ -331,19 +323,14 @@ public:
 		const ETextureCreateFlags ExtraFlags = CVarGPUParticleAFRReinject.GetValueOnRenderThread() == 1 ? TexCreate_AFRManual : TexCreate_None;
 
 #define ATTRIBUTES_TEXTURE_NAME	TEXT("ParticleAttributes")
-		FRHIResourceCreateInfo CreateInfo(ATTRIBUTES_TEXTURE_NAME, FClearValueBinding::Transparent);
-		RHICreateTargetableShaderResource2D(
-			SizeX,
-			SizeY,
-			PF_B8G8R8A8,
-			/*NumMips=*/ 1,
-			TexCreate_None,
-			TexCreate_RenderTargetable | TexCreate_NoFastClear | ExtraFlags,
-			/*bForceSeparateTargetAndShaderResource=*/ false,
-			CreateInfo,
-			TextureTargetRHI,
-			TextureRHI
-			);
+
+		const FRHITextureCreateDesc Desc =
+			FRHITextureCreateDesc::Create2D(ATTRIBUTES_TEXTURE_NAME)
+			.SetExtent(SizeX, SizeY)
+			.SetFormat(PF_B8G8R8A8)
+			.SetClearValue(FClearValueBinding::Transparent);
+
+		RHICreateTargetableShaderResource(Desc, ETextureCreateFlags::RenderTargetable | ETextureCreateFlags::NoFastClear | ExtraFlags, TextureTargetRHI, TextureRHI);
 
 		static FName AttributesTextureName(ATTRIBUTES_TEXTURE_NAME);
 		TextureTargetRHI->SetName(AttributesTextureName);
diff --git a/Engine/Source/Runtime/Engine/Private/Slate/SceneViewport.cpp b/Engine/Source/Runtime/Engine/Private/Slate/SceneViewport.cpp
--- a/Engine/Source/Runtime/Engine/Private/Slate/SceneViewport.cpp
+++ b/Engine/Source/Runtime/Engine/Private/Slate/SceneViewport.cpp
@@ -2008,16 +2008,20 @@ void FSceneViewport::InitDynamicRHI()
 		}
 #endif
 
-		FRHIResourceCreateInfo CreateInfo(TEXT("BufferedRT"));
-		FTexture2DRHIRef BufferedRTRHI;
-		FTexture2DRHIRef BufferedSRVRHI;
+		FTextureRHIRef BufferedRTRHI;
+		FTextureRHIRef BufferedSRVRHI;
 
 		for (int32 i = 0; i < NumBufferedFrames; ++i)
 		{
 			// try to allocate texture via StereoRenderingDevice; if not successful, use the default way
 			if (StereoRenderTargetManager == nullptr || !StereoRenderTargetManager->AllocateRenderTargetTexture(i, TexSizeX, TexSizeY, SceneTargetFormat, 1, TexCreate_None, TexCreate_RenderTargetable, BufferedRTRHI, BufferedSRVRHI))
 			{
-				RHICreateTargetableShaderResource2D(TexSizeX, TexSizeY, SceneTargetFormat, 1, TexCreate_None, TexCreate_RenderTargetable, false, CreateInfo, BufferedRTRHI, BufferedSRVRHI);
+				const FRHITextureCreateDesc Desc =
+					FRHITextureCreateDesc::Create2D(TEXT("BufferedRT"))
+					.SetExtent(TexSizeX, TexSizeY)
+					.SetFormat(SceneTargetFormat);
+
+				RHICreateTargetableShaderResource(Desc, ETextureCreateFlags::RenderTargetable, BufferedRTRHI, BufferedSRVRHI);
 			}
 			BufferedRenderTargetsRHI[i] = BufferedRTRHI;
 			BufferedShaderResourceTexturesRHI[i] = BufferedSRVRHI;
diff --git a/Engine/Source/Runtime/Engine/Private/Slate/SlateTextures.cpp b/Engine/Source/Runtime/Engine/Private/Slate/SlateTextures.cpp
--- a/Engine/Source/Runtime/Engine/Private/Slate/SlateTextures.cpp
+++ b/Engine/Source/Runtime/Engine/Private/Slate/SlateTextures.cpp
@@ -315,20 +315,13 @@ void FSlateTextureRenderTarget2DResource::InitDynamicRHI()
 	if( TargetSizeX > 0 && TargetSizeY > 0 )
 	{
 		// Create the RHI texture. Only one mip is used and the texture is targetable for resolve.
-		FRHIResourceCreateInfo CreateInfo(TEXT("FSlateTextureRenderTarget2DResource"), FClearValueBinding(ClearColor));
-		RHICreateTargetableShaderResource2D(
-			TargetSizeX, 
-			TargetSizeY, 
-			Format, 
-			1,
-			/*TexCreateFlags=*/TexCreate_None,
-			TexCreate_RenderTargetable,
-			/*bNeedsTwoCopies=*/false,
-			CreateInfo,
-			RenderTargetTextureRHI,
-			Texture2DRHI
-			);
-		TextureRHI = (FTextureRHIRef&)Texture2DRHI;
+		const FRHITextureCreateDesc Desc =
+			FRHITextureCreateDesc::Create2D(TEXT("FSlateTextureRenderTarget2DResource"))
+			.SetExtent(TargetSizeX, TargetSizeY)
+			.SetFormat((EPixelFormat)Format)
+			.SetClearValue(FClearValueBinding(ClearColor));
+
+		RHICreateTargetableShaderResource(Desc, ETextureCreateFlags::RenderTargetable, RenderTargetTextureRHI, Texture2DRHI);
 	}
 
 	// Create the sampler state RHI resource.
diff --git a/Engine/Source/Runtime/Engine/Private/TextureRenderTarget2D.cpp b/Engine/Source/Runtime/Engine/Private/TextureRenderTarget2D.cpp
--- a/Engine/Source/Runtime/Engine/Private/TextureRenderTarget2D.cpp
+++ b/Engine/Source/Runtime/Engine/Private/TextureRenderTarget2D.cpp
@@ -557,44 +557,40 @@ void FTextureRenderTarget2DResource::InitDynamicRHI()
 	if( TargetSizeX > 0 && TargetSizeY > 0 )
 	{
 		// Create the RHI texture. Only one mip is used and the texture is targetable for resolve.
-		ETextureCreateFlags TexCreateFlags = Owner->IsSRGB() ? TexCreate_SRGB : TexCreate_None;
-		TexCreateFlags |= Owner->bGPUSharedFlag ? TexCreate_Shared : TexCreate_None;
+		ETextureCreateFlags TexCreateFlags = Owner->IsSRGB() ? ETextureCreateFlags::SRGB : ETextureCreateFlags::None;
+		TexCreateFlags |= Owner->bGPUSharedFlag ? ETextureCreateFlags::Shared : ETextureCreateFlags::None;
 		FString ResourceName = Owner->GetName();
-		FRHIResourceCreateInfo CreateInfo = FRHIResourceCreateInfo(*ResourceName, FClearValueBinding(ClearColor));
 
 		if (Owner->bAutoGenerateMips)
 		{
-			TexCreateFlags |= TexCreate_GenerateMipCapable;
+			TexCreateFlags |= ETextureCreateFlags::GenerateMipCapable;
 			if (FGenerateMips::WillFormatSupportCompute(Format))
 			{
-				TexCreateFlags |= TexCreate_UAV;
+				TexCreateFlags |= ETextureCreateFlags::UAV;
 			}
 		}
 
 		if (Owner->bCanCreateUAV)
 		{
-			TexCreateFlags |= TexCreate_UAV;
+			TexCreateFlags |= ETextureCreateFlags::UAV;
 		}
 
-		RHICreateTargetableShaderResource2D(
-			Owner->SizeX, 
-			Owner->SizeY, 
-			Format,
-			Owner->GetNumMips(),
-			TexCreateFlags,
-			TexCreate_RenderTargetable,
-			Owner->bNeedsTwoCopies,
-			CreateInfo,
-			RenderTargetTextureRHI,
-			Texture2DRHI
-			);
+		const FRHITextureCreateDesc Desc =
+			FRHITextureCreateDesc::Create2D(*ResourceName)
+			.SetExtent(Owner->SizeX, Owner->SizeY)
+			.SetFormat(Format)
+			.SetNumMips(Owner->GetNumMips())
+			.SetFlags(TexCreateFlags)
+			.SetClearValue(FClearValueBinding(ClearColor));
+
+		RHICreateTargetableShaderResource(Desc, ETextureCreateFlags::RenderTargetable, Owner->bNeedsTwoCopies, RenderTargetTextureRHI, Texture2DRHI);
 
-		if (EnumHasAnyFlags(TexCreateFlags, TexCreate_UAV))
+		if (EnumHasAnyFlags(TexCreateFlags, ETextureCreateFlags::UAV))
 		{
 			UnorderedAccessViewRHI = RHICreateUnorderedAccessView(RenderTargetTextureRHI);
 		}
 
-		SetGPUMask(CreateInfo.GPUMask);
+		SetGPUMask(Desc.GPUMask);
 		TextureRHI = (FTextureRHIRef&)Texture2DRHI;
 		RHIUpdateTextureReference(Owner->TextureReference.TextureReferenceRHI,TextureRHI);
 
diff --git a/Engine/Source/Runtime/MediaAssets/Private/Misc/MediaTextureResource.cpp b/Engine/Source/Runtime/MediaAssets/Private/Misc/MediaTextureResource.cpp
--- a/Engine/Source/Runtime/MediaAssets/Private/Misc/MediaTextureResource.cpp
+++ b/Engine/Source/Runtime/MediaAssets/Private/Misc/MediaTextureResource.cpp
@@ -1199,24 +1199,17 @@ void FMediaTextureResource::CreateOutputRenderTarget(const FIntPoint & InDim, EP
 
 	if ((InClearColor != CurrentClearColor) || !OutputTarget.IsValid() || (OutputTarget->GetSizeXY() != InDim) || (OutputTarget->GetFormat() != InPixelFormat) || ((OutputTarget->GetFlags() & OutputCreateFlags) != OutputCreateFlags) || CurrentNumMips != InNumMips)
 	{
-		TRefCountPtr<FRHITexture2D> DummyTexture2DRHI;
-
 		MipGenerationCache.SafeRelease();
 
-		FRHIResourceCreateInfo CreateInfo(TEXT("MediaTextureResourceOutput"), FClearValueBinding(InClearColor));
-
-		RHICreateTargetableShaderResource2D(
-			InDim.X,
-			InDim.Y,
-			InPixelFormat,
-			InNumMips,
-			OutputCreateFlags,
-			TexCreate_RenderTargetable,
-			false,
-			CreateInfo,
-			OutputTarget,
-			DummyTexture2DRHI
-		);
+		const FRHITextureCreateDesc Desc =
+			FRHITextureCreateDesc::Create2D(TEXT("MediaTextureResourceOutput"))
+			.SetExtent(InDim)
+			.SetFormat(InPixelFormat)
+			.SetNumMips(InNumMips)
+			.SetFlags(OutputCreateFlags)
+			.SetClearValue(FClearValueBinding(InClearColor));
+
+		RHICreateTargetableShaderResource(Desc, ETextureCreateFlags::RenderTargetable, OutputTarget);
 
 		OutputTarget->SetName(TEXT("MediaTextureResourceOutput"));
 
diff --git a/Engine/Source/Runtime/RHI/Public/RHIUtilities.h b/Engine/Source/Runtime/RHI/Public/RHIUtilities.h
--- a/Engine/Source/Runtime/RHI/Public/RHIUtilities.h
+++ b/Engine/Source/Runtime/RHI/Public/RHIUtilities.h
@@ -722,7 +722,7 @@ inline void RHICreateTargetableShaderResource(
  * Two texture references are always returned, but they may reference the same texture.
  * If two different textures are returned, the render-target texture must be manually copied to the shader-resource texture.
  */
-//UE_DEPRECATED(5.1, "RHICreateTargetableShaderResource2D is deprecated, RHICreateTargetableShaderResource with a FRHITextureCreateDesc should be used instead")
+UE_DEPRECATED(5.1, "RHICreateTargetableShaderResource2D is deprecated, RHICreateTargetableShaderResource with a FRHITextureCreateDesc should be used instead")
 RHI_API void RHICreateTargetableShaderResource2D(
 	uint32 SizeX,
 	uint32 SizeY,
@@ -737,7 +737,7 @@ RHI_API void RHICreateTargetableShaderResource2D(
 	FTextureRHIRef& OutShaderResourceTexture,
 	uint32 NumSamples = 1);
 
-//UE_DEPRECATED(5.1, "RHICreateTargetableShaderResource2D is deprecated, RHICreateTargetableShaderResource with a FRHITextureCreateDesc should be used instead")
+UE_DEPRECATED(5.1, "RHICreateTargetableShaderResource2D is deprecated, RHICreateTargetableShaderResource with a FRHITextureCreateDesc should be used instead")
 RHI_API void RHICreateTargetableShaderResource2D(
 	uint32 SizeX,
 	uint32 SizeY,
diff --git a/Engine/Source/Runtime/SlateRHIRenderer/Private/Slate3DRenderer.cpp b/Engine/Source/Runtime/SlateRHIRenderer/Private/Slate3DRenderer.cpp
--- a/Engine/Source/Runtime/SlateRHIRenderer/Private/Slate3DRenderer.cpp
+++ b/Engine/Source/Runtime/SlateRHIRenderer/Private/Slate3DRenderer.cpp
@@ -195,9 +195,14 @@ void FSlate3DRenderer::DrawWindowToTarget_RenderThread(FRHICommandListImmediate&
 				{
 					DepthStencil.SafeRelease();
 
-					FTexture2DRHIRef ShaderResourceUnused;
-					FRHIResourceCreateInfo CreateInfo(TEXT("SlateWindowDepthStencil"), FClearValueBinding::DepthZero);
-					RHICreateTargetableShaderResource2D(ColorTarget->GetSizeX(), ColorTarget->GetSizeY(), PF_DepthStencil, 1, TexCreate_None, TexCreate_DepthStencilTargetable, false, CreateInfo, DepthStencil, ShaderResourceUnused);
+					const FRHITextureCreateDesc Desc =
+						FRHITextureCreateDesc::Create2D(TEXT("SlateWindowDepthStencil"))
+						.SetExtent(ColorTarget->GetSizeXY())
+						.SetFormat(PF_DepthStencil)
+						.SetClearValue(FClearValueBinding::DepthZero);
+
+					RHICreateTargetableShaderResource(Desc, ETextureCreateFlags::DepthStencilTargetable, DepthStencil);
+
 					check(IsValidRef(DepthStencil));
 				}
 			}
diff --git a/Engine/Source/Runtime/SlateRHIRenderer/Private/SlatePostProcessResource.cpp b/Engine/Source/Runtime/SlateRHIRenderer/Private/SlatePostProcessResource.cpp
--- a/Engine/Source/Runtime/SlateRHIRenderer/Private/SlatePostProcessResource.cpp
+++ b/Engine/Source/Runtime/SlateRHIRenderer/Private/SlatePostProcessResource.cpp
@@ -42,21 +42,13 @@ void FSlatePostProcessResource::ResizeTargets(const FIntPoint& NewSize)
 	{
 		for (int32 TexIndex = 0; TexIndex < RenderTargetCount; ++TexIndex)
 		{
+			const FRHITextureCreateDesc Desc =
+				FRHITextureCreateDesc::Create2D(TEXT("FSlatePostProcessResource"))
+				.SetExtent(RenderTargetSize)
+				.SetFormat(PixelFormat);
+
 			FTexture2DRHIRef RenderTargetTextureRHI;
-			FTexture2DRHIRef ShaderResourceUnused;
-			FRHIResourceCreateInfo CreateInfo(TEXT("FSlatePostProcessResource"));
-			RHICreateTargetableShaderResource2D(
-				RenderTargetSize.X,
-				RenderTargetSize.Y,
-				PixelFormat,
-				1,
-				TexCreate_None,
-				TexCreate_RenderTargetable,
-				/*bNeedsTwoCopies=*/false,
-				CreateInfo,
-				RenderTargetTextureRHI,
-				ShaderResourceUnused
-			);
+			RHICreateTargetableShaderResource(Desc, ETextureCreateFlags::RenderTargetable, RenderTargetTextureRHI);
 
 			RenderTargets.Add(RenderTargetTextureRHI);
 		}
diff --git a/Engine/Source/Runtime/WebBrowserTexture/Private/WebBrowserTextureResource.cpp b/Engine/Source/Runtime/WebBrowserTexture/Private/WebBrowserTextureResource.cpp
--- a/Engine/Source/Runtime/WebBrowserTexture/Private/WebBrowserTextureResource.cpp
+++ b/Engine/Source/Runtime/WebBrowserTexture/Private/WebBrowserTextureResource.cpp
@@ -148,24 +148,16 @@ void FWebBrowserTextureResource::ClearTexture(const FLinearColor& ClearColor)
 	// create output render target if we don't have one yet
 	const ETextureCreateFlags OutputCreateFlags = TexCreate_Dynamic | TexCreate_SRGB;
 
-	if ((ClearColor != CurrentClearColor) || !OutputTarget.IsValid() || ((OutputTarget->GetFlags() & OutputCreateFlags) != OutputCreateFlags))
+	if ((ClearColor != CurrentClearColor) || !OutputTarget.IsValid() || !EnumHasAllFlags(OutputTarget->GetFlags(), OutputCreateFlags))
 	{
-		FRHIResourceCreateInfo CreateInfo(TEXT("FWebBrowserTextureResource"), FClearValueBinding(ClearColor));
-
-		TRefCountPtr<FRHITexture2D> DummyTexture2DRHI;
-
-		RHICreateTargetableShaderResource2D(
-			2,
-			2,
-			PF_B8G8R8A8,
-			1,
-			OutputCreateFlags,
-			TexCreate_RenderTargetable,
-			false,
-			CreateInfo,
-			OutputTarget,
-			DummyTexture2DRHI
-		);
+		const FRHITextureCreateDesc Desc =
+			FRHITextureCreateDesc::Create2D(TEXT("FWebBrowserTextureResource"))
+			.SetExtent(2, 2)
+			.SetFormat(PF_B8G8R8A8)
+			.SetFlags(OutputCreateFlags)
+			.SetClearValue(FClearValueBinding(ClearColor));
+
+		RHICreateTargetableShaderResource(Desc, ETextureCreateFlags::RenderTargetable, OutputTarget);
 
 		CurrentClearColor = ClearColor;
 		UpdateResourceSize();
@@ -219,22 +211,15 @@ void FWebBrowserTextureResource::CopySample(const TSharedPtr<FWebBrowserTextureS
 		if ((ClearColor != CurrentClearColor) || !OutputTarget.IsValid() || (OutputTarget->GetSizeXY() != SampleDim) || ((OutputTarget->GetFlags() & OutputCreateFlags) != OutputCreateFlags))
 		{
 			FPlatformMisc::LowLevelOutputDebugStringf(TEXT("FWebBrowserTextureResource:CopySample 1"));
-			TRefCountPtr<FRHITexture2D> DummyTexture2DRHI;
-
-			FRHIResourceCreateInfo CreateInfo(TEXT("FWebBrowserTextureResource"), FClearValueBinding(ClearColor));
-
-			RHICreateTargetableShaderResource2D(
-				SampleDim.X,
-				SampleDim.Y,
-				PF_B8G8R8A8,
-				1,
-				OutputCreateFlags,
-				TexCreate_RenderTargetable,
-				false,
-				CreateInfo,
-				OutputTarget,
-				DummyTexture2DRHI
-			);
+
+			const FRHITextureCreateDesc Desc =
+				FRHITextureCreateDesc::Create2D(TEXT("FWebBrowserTextureResource"))
+				.SetExtent(SampleDim)
+				.SetFormat(PF_B8G8R8A8)
+				.SetFlags(OutputCreateFlags)
+				.SetClearValue(FClearValueBinding(ClearColor));
+
+			RHICreateTargetableShaderResource(Desc, ETextureCreateFlags::RenderTargetable, OutputTarget);
 
 			CurrentClearColor = ClearColor;
 			UpdateResourceSize();
diff --git a/Engine/Source/Runtime/WebBrowserTexture/Public/WebBrowserTextureSample.h b/Engine/Source/Runtime/WebBrowserTexture/Public/WebBrowserTextureSample.h
--- a/Engine/Source/Runtime/WebBrowserTexture/Public/WebBrowserTextureSample.h
+++ b/Engine/Source/Runtime/WebBrowserTexture/Public/WebBrowserTextureSample.h
@@ -127,23 +127,13 @@ public:
 			return Texture;
 		}
 
-		const ETextureCreateFlags CreateFlags = TexCreate_Dynamic | TexCreate_SRGB;
-
-		TRefCountPtr<FRHITexture2D> DummyTexture2DRHI;
-		FRHIResourceCreateInfo CreateInfo(TEXT("FWebBrowserTextureSample"));
-
-		RHICreateTargetableShaderResource2D(
-			Dim.X,
-			Dim.Y,
-			PF_B8G8R8A8,
-			1,
-			CreateFlags,
-			TexCreate_RenderTargetable,
-			false,
-			CreateInfo,
-			Texture,
-			DummyTexture2DRHI
-		);
+		const FRHITextureCreateDesc Desc =
+			FRHITextureCreateDesc::Create2D(TEXT("FWebBrowserTextureSample"))
+			.SetExtent(Dim)
+			.SetFormat(PF_B8G8R8A8)
+			.SetFlags(ETextureCreateFlags::Dynamic | ETextureCreateFlags::SRGB);
+
+		RHICreateTargetableShaderResource(Desc, ETextureCreateFlags::RenderTargetable, Texture);
 
 		return Texture;
 	}
