fix Texture mip count was wrong
make it so 0-dim textures get 0 mips, not 1
nonpow2 textures use floor log2 not ceil
also validate passed in mip counts
also fix inconsistency in null return from lockmip

diff --git a/Engine/Source/Runtime/Engine/Private/Tests/Texture2DTests.cpp b/Engine/Source/Runtime/Engine/Private/Tests/Texture2DTests.cpp
--- a/Engine/Source/Runtime/Engine/Private/Tests/Texture2DTests.cpp
+++ b/Engine/Source/Runtime/Engine/Private/Tests/Texture2DTests.cpp
@@ -48,18 +48,18 @@ bool FTexture2DTestLockingWhenEmpty::RunTest(const FString& Parameters)
 		// Test that we can lock and unlock repeatedly
 		uint8* FirstLockPtr = EmptyTexture->Source.LockMip(0);
 		TestNull(TEXT("Locking an empty texture"), FirstLockPtr);
-		EmptyTexture->Source.UnlockMip(0);
 
-		// If a lock on the primary mip (0) fails because the bulkdata is empty we still
-		// get a successful lock and so need to unlock it
+		// If LockMip returns null you do not get the lock, you do not need to unlock it
 		uint8* SecondLockPtr = EmptyTexture->Source.LockMip(0);
 		TestNull(TEXT("Locking an empty texture a second time"), SecondLockPtr);
-		EmptyTexture->Source.UnlockMip(0);
 
 		// If we fail to get the lock because the mip does not exist then we do not get
 		// a lock and don't need to unlock it.
 		uint8* InvalidLockPtr = EmptyTexture->Source.LockMip(1);
 		TestNull(TEXT("Locking a submip of an empty texture"), InvalidLockPtr);
+
+		FTextureSource::FMipLock MipLock(FTextureSource::ELockState::ReadOnly,&EmptyTexture->Source,0);
+		TestFalse( TEXT("MipLock on empty texture should not be valid"), MipLock.IsValid() );
 	}
 
 	// Create a texture with  valid dimensions and default data
@@ -76,6 +76,11 @@ bool FTexture2DTestLockingWhenEmpty::RunTest(const FString& Parameters)
 		TestNotNull(TEXT("Locking a valid a second time"), SecondLockPtr);
 		Texture->Source.UnlockMip(0);
 
+		{
+			FTextureSource::FMipLock MipLock(FTextureSource::ELockState::ReadOnly,&Texture->Source,0);
+			TestTrue( TEXT("MipLock on valid texture should be valid"), MipLock.IsValid() );
+		}
+
 		// Test that we can lock each mip before unlocking them all
 		for (int32 MipIndex = 0; MipIndex < Texture->Source.GetNumMips(); ++MipIndex)
 		{
@@ -87,6 +92,16 @@ bool FTexture2DTestLockingWhenEmpty::RunTest(const FString& Parameters)
 		{
 			Texture->Source.UnlockMip(MipIndex);
 		}
+		
+		for (int32 MipIndex = 0; MipIndex < Texture->Source.GetNumMips(); ++MipIndex)
+		{
+			FTextureSource::FMipLock MipLock(FTextureSource::ELockState::ReadOnly,&Texture->Source,MipIndex);
+			TestTrue( TEXT("MipLock on valid texture should be valid"), MipLock.IsValid() );
+			
+			// does fail, but fails with a check, so we can't test it :
+			//uint8* MipPtr = Texture->Source.LockMip(MipIndex);
+			//TestNull(TEXT("Locking for write after readonly should fail"), MipPtr);
+		}
 	}
 
 	return true;
diff --git a/Engine/Source/Runtime/Engine/Private/Texture.cpp b/Engine/Source/Runtime/Engine/Private/Texture.cpp
--- a/Engine/Source/Runtime/Engine/Private/Texture.cpp
+++ b/Engine/Source/Runtime/Engine/Private/Texture.cpp
@@ -1393,6 +1393,20 @@ int FTextureSource::GetVolumeSizeZ() const
 	}
 }
 
+static int GetFullMipCount(int32 SizeX,int32 SizeY,int32 SizeZ = 1)
+{
+	if ( SizeX == 0 || SizeY == 0 || SizeZ == 0 )
+	{
+		return 0;
+	}
+	
+	int MaxDim = FMath::Max3(SizeX,SizeY,SizeZ);
+
+	int32 MipCount = FMath::FloorLog2(MaxDim) + 1;
+
+	return MipCount;
+}
+
 void FTextureSource::InitBlocked(const ETextureSourceFormat* InLayerFormats,
 	const FTextureSourceBlock* InBlocks,
 	int32 InNumLayers,
@@ -1538,7 +1552,7 @@ void FTextureSource::Init2DWithMipChain(
 	ETextureSourceFormat NewFormat
 	)
 {
-	int32 NewMipCount = FMath::Max(FMath::CeilLogTwo(NewSizeX),FMath::CeilLogTwo(NewSizeY)) + 1;
+	int32 NewMipCount = GetFullMipCount(NewSizeX,NewSizeY);
 	Init(NewSizeX, NewSizeY, 1, NewMipCount, NewFormat);
 }
 
@@ -1548,7 +1562,7 @@ void FTextureSource::InitLayered2DWithMipChain(
 	int32 NewNumLayers,
 	const ETextureSourceFormat* NewFormat)
 {
-	int32 NewMipCount = FMath::Max(FMath::CeilLogTwo(NewSizeX), FMath::CeilLogTwo(NewSizeY)) + 1;
+	int32 NewMipCount = GetFullMipCount(NewSizeX,NewSizeY);
 	InitLayered(NewSizeX, NewSizeY, 1, NewNumLayers, NewMipCount, NewFormat);
 }
 
@@ -1558,7 +1572,7 @@ void FTextureSource::InitCubeWithMipChain(
 	ETextureSourceFormat NewFormat
 	)
 {
-	int32 NewMipCount = FMath::Max(FMath::CeilLogTwo(NewSizeX),FMath::CeilLogTwo(NewSizeY)) + 1;
+	int32 NewMipCount = GetFullMipCount(NewSizeX,NewSizeY);
 	Init(NewSizeX, NewSizeY, 6, NewMipCount, NewFormat);
 }
 
@@ -1580,6 +1594,8 @@ void FTextureSource::InitWithCompressedSourceData(
 	NumSlices = 1;
 	NumMips = NewNumMips;
 
+	check( NewNumMips <= GetFullMipCount(SizeX,SizeY) );
+
 	Format = NewFormat;
 	LayerFormat.SetNum(1, true);
 	LayerFormat[0] = NewFormat;
@@ -1785,6 +1801,12 @@ uint8* FTextureSource::LockMipInternal(int32 BlockIndex, int32 LayerIndex, int32
 		{
 			MipData = LockedMipData.GetDataReadWrite();
 		}
+
+		if ( MipData == nullptr )
+		{
+			// no data, you did not get the lock, do not call Unlock
+			return nullptr;
+		}
 		
 		MipData += CalcMipOffset(BlockIndex, LayerIndex, MipIndex);
 
@@ -3017,6 +3039,11 @@ void FTextureSource::InitLayeredImpl(
 	NumLayers = NewNumLayers;
 	NumSlices = NewNumSlices;
 	NumMips = NewNumMips;
+	
+	//check( NewNumMips <= MAX_TEXTURE_MIP_COUNT );
+	// NumSlices could be volume size Z or not
+	check( NewNumMips <= GetFullMipCount(SizeX,SizeY,GetVolumeSizeZ()) );
+
 	Format = NewLayerFormat[0];
 	LayerFormat.SetNum(NewNumLayers, true);
 	for (int i = 0; i < NewNumLayers; ++i)
@@ -3045,6 +3072,8 @@ void FTextureSource::InitBlockedImpl(const ETextureSourceFormat* InLayerFormats,
 	SizeY = InBlocks[0].SizeY;
 	NumSlices = InBlocks[0].NumSlices;
 	NumMips = InBlocks[0].NumMips;
+	
+	check( NumMips <= GetFullMipCount(SizeX,SizeY) );
 
 	NumLayers = InNumLayers;
 	Format = InLayerFormats[0];
