change the Texture DDC Key for the case of VT+Pow22 gamma

diff --git a/Engine/Source/Runtime/Engine/Private/TextureDerivedData.cpp b/Engine/Source/Runtime/Engine/Private/TextureDerivedData.cpp
--- a/Engine/Source/Runtime/Engine/Private/TextureDerivedData.cpp
+++ b/Engine/Source/Runtime/Engine/Private/TextureDerivedData.cpp
@@ -225,6 +225,12 @@ static void SerializeForKey(FArchive& Ar, const FTextureBuildSettings& Settings)
 		Ar << TempGuid;
 	}
 
+	if ( Settings.bVirtualStreamable && Settings.bSRGB && Settings.bUseLegacyGamma )
+	{
+		TempGuid = FGuid(0xCAEDDFB6,0xEDC2455D,0x8D45B90C,0x3A1B7783);
+		Ar << TempGuid;
+	}
+
 	// Note - compression quality is added to the DDC by the formats (based on whether they
 	// use them or not).
 	// This is true for:
diff --git a/Engine/Source/Runtime/Engine/Private/TextureDerivedDataBuildUtils.cpp b/Engine/Source/Runtime/Engine/Private/TextureDerivedDataBuildUtils.cpp
--- a/Engine/Source/Runtime/Engine/Private/TextureDerivedDataBuildUtils.cpp
+++ b/Engine/Source/Runtime/Engine/Private/TextureDerivedDataBuildUtils.cpp
@@ -130,10 +130,6 @@ static void WriteBuildSettings(FCbWriter& Writer, const FTextureBuildSettings& B
 	WriteCbFieldWithDefault<bool>(Writer, "bCubemap", BuildSettings.bCubemap, DefaultSettings.bCubemap);
 	WriteCbFieldWithDefault<bool>(Writer, "bTextureArray", BuildSettings.bTextureArray, DefaultSettings.bTextureArray);
 	WriteCbFieldWithDefault<bool>(Writer, "bVolume", BuildSettings.bVolume, DefaultSettings.bVolume);
-	if ( BuildSettings.bVolume )
-	{
-		WriteCbField<bool>(Writer, "bVolume_ForceNewDDcKey", true); 
-	}
 	WriteCbFieldWithDefault<bool>(Writer, "bLongLatSource", BuildSettings.bLongLatSource, DefaultSettings.bLongLatSource);
 	WriteCbFieldWithDefault<bool>(Writer, "bSRGB", BuildSettings.bSRGB, DefaultSettings.bSRGB);
 	WriteCbFieldWithDefault(Writer, "SourceEncodingOverride", BuildSettings.SourceEncodingOverride, DefaultSettings.SourceEncodingOverride);
@@ -187,6 +183,15 @@ static void WriteBuildSettings(FCbWriter& Writer, const FTextureBuildSettings& B
 	
 	WriteCbFieldWithDefault(Writer, "OodleTextureSdkVersion", BuildSettings.OodleTextureSdkVersion, DefaultSettings.OodleTextureSdkVersion);
 	
+	if ( BuildSettings.bVolume )
+	{
+		WriteCbField<bool>(Writer, "bVolume_ForceNewDDcKey", true); 
+	}
+	if ( BuildSettings.bVirtualStreamable && BuildSettings.bSRGB && BuildSettings.bUseLegacyGamma )
+	{
+		WriteCbField<bool>(Writer, "VTPow22_ForceNewDDcKey", true); 
+	}
+
 	Writer.EndObject();
 }
 
