Suppress spurious warning from CleanupWorld when a streaming sublevel is deleted from the Level List in the editor.
#rnx

diff --git a/Engine/Source/Runtime/Engine/Private/World.cpp b/Engine/Source/Runtime/Engine/Private/World.cpp
--- a/Engine/Source/Runtime/Engine/Private/World.cpp
+++ b/Engine/Source/Runtime/Engine/Private/World.cpp
@@ -4932,7 +4932,11 @@ void UWorld::CleanupWorld(bool bSessionEnded, bool bCleanupResources, UWorld* Ne
     CleanupWorldGlobalTag++;
 	if (!bInitializedAndNeedsCleanup)
 	{
-		UE_LOG(LogWorld, Warning, TEXT("UWorld::CleanupWorld called twice or called without InitWorld called first."));
+		// Only issue the warning for the TopLevelWorld. It is currently valid to call CleanupWorld on the UWorld of
+		// streaming sublevels, and they never call InitWorld. (this is done by PrivateDestroyLevel when removing a
+		// streaming level from the Level List in the editor.)
+		bool bIsStreamingSubWorld = PersistentLevel && PersistentLevel->OwningWorld != this;
+		UE_CLOG(!bIsStreamingSubWorld, LogWorld, Warning, TEXT("UWorld::CleanupWorld called twice or called without InitWorld called first."));
 	}
 	CleanupWorldInternal(bSessionEnded, bCleanupResources, NewWorld);
 	bInitializedAndNeedsCleanup = false;
