World Partition
- Refactored how the code handles disabling streaming so we can support multiple actor descriptor containers.
- Store if the partitioned level has streaming disabled in the asset registry tags so it can be queried without having the map loaded.
- Adapted changelist validation to query the disable streaming state for unloaded levels.

#rnx

diff --git a/Engine/Plugins/Editor/DataValidation/Source/DataValidation/Private/WorldPartitionChangelistValidator.cpp b/Engine/Plugins/Editor/DataValidation/Source/DataValidation/Private/WorldPartitionChangelistValidator.cpp
--- a/Engine/Plugins/Editor/DataValidation/Source/DataValidation/Private/WorldPartitionChangelistValidator.cpp
+++ b/Engine/Plugins/Editor/DataValidation/Source/DataValidation/Private/WorldPartitionChangelistValidator.cpp
@@ -200,7 +200,9 @@ EDataValidationResult UWorldPartitionChangelistValidator::ValidateActorsAndDataL
 		}
 
 		// Invoke static WorldPartition Validation from the ActorDescContainer
-		UWorldPartition::CheckForErrors(this, ActorDescContainer);		
+		const FSoftObjectPath MapPath = FSoftObjectPath(MapName);
+		const bool bIsStreamingDisabled = ULevel::GetIsStreamingDisabledFromPackage(MapPath.GetLongPackageFName());
+		UWorldPartition::CheckForErrors(this, ActorDescContainer, !bIsStreamingDisabled);
 	}
 
 	if (Errors->Num())
diff --git a/Engine/Source/Runtime/Engine/Classes/Engine/Level.h b/Engine/Source/Runtime/Engine/Classes/Engine/Level.h
--- a/Engine/Source/Runtime/Engine/Classes/Engine/Level.h
+++ b/Engine/Source/Runtime/Engine/Classes/Engine/Level.h
@@ -770,12 +770,13 @@ public:
 	ENGINE_API static bool GetIsLevelPartitionedFromAsset(const FAssetData& Asset);
 	ENGINE_API static bool GetIsLevelUsingExternalActorsFromAsset(const FAssetData& Asset);
 	ENGINE_API static bool GetIsUsingActorFoldersFromAsset(const FAssetData& Asset);
-
+	ENGINE_API static bool GetIsStreamingDisabledFromAsset(const FAssetData& Asset);
 
 	ENGINE_API static bool GetLevelBoundsFromPackage(FName LevelPackage, FBox& OutLevelBounds);
 	ENGINE_API static bool GetIsLevelPartitionedFromPackage(FName LevelPackage);
 	ENGINE_API static bool GetIsLevelUsingExternalActorsFromPackage(FName LevelPackage);
 	ENGINE_API static bool GetIsUsingActorFoldersFromPackage(FName LevelPackage);
+	ENGINE_API static bool GetIsStreamingDisabledFromPackage(FName LevelPackage);
 
 	ENGINE_API bool GetPromptWhenAddingToLevelOutsideBounds() const;
 	ENGINE_API bool GetPromptWhenAddingToLevelBeforeCheckout() const;
diff --git a/Engine/Source/Runtime/Engine/Private/Level.cpp b/Engine/Source/Runtime/Engine/Private/Level.cpp
--- a/Engine/Source/Runtime/Engine/Private/Level.cpp
+++ b/Engine/Source/Runtime/Engine/Private/Level.cpp
@@ -2188,6 +2188,26 @@ bool ULevel::GetIsUsingActorFoldersFromPackage(FName LevelPackage)
 	});
 }
 
+bool ULevel::GetIsStreamingDisabledFromAsset(const FAssetData& Asset)
+{
+	FString LevelHasStreamingDisabledStr;
+	static const FName NAME_LevelHasStreamingDisabled(TEXT("LevelHasStreamingDisabled"));
+	if (Asset.GetTagValue(NAME_LevelHasStreamingDisabled, LevelHasStreamingDisabledStr))
+	{
+		check(LevelHasStreamingDisabledStr == TEXT("1"));
+		return true;
+	}
+	return false;
+}
+
+bool ULevel::GetIsStreamingDisabledFromPackage(FName LevelPackage)
+{
+	return LevelAssetRegistryHelper::GetLevelInfoFromAssetRegistry(LevelPackage, [](const FAssetData& Asset)
+	{
+		return GetIsStreamingDisabledFromAsset(Asset);
+	});
+}
+
 bool ULevel::GetLevelBoundsFromAsset(const FAssetData& Asset, FBox& OutLevelBounds)
 {
 	FString LevelBoundsLocationStr;
diff --git a/Engine/Source/Runtime/Engine/Private/World.cpp b/Engine/Source/Runtime/Engine/Private/World.cpp
--- a/Engine/Source/Runtime/Engine/Private/World.cpp
+++ b/Engine/Source/Runtime/Engine/Private/World.cpp
@@ -8414,10 +8414,9 @@ void UWorld::GetAssetRegistryTags(TArray<FAssetRegistryTag>& OutTags) const
 		}
 		// If there are no blueprints FiBData will be empty, the search manager will treat this as indexed
 								
-		if(IsPartitionedWorld())
+		if(UWorldPartition* WorldPartition = GetWorldPartition())
 		{
-			static const FName NAME_LevelIsPartitioned(TEXT("LevelIsPartitioned"));
-			OutTags.Add(FAssetRegistryTag(NAME_LevelIsPartitioned, TEXT("1"), FAssetRegistryTag::TT_Hidden));
+			WorldPartition->AppendAssetRegistryTags(OutTags);
 		}
 		else
 		{
diff --git a/Engine/Source/Runtime/Engine/Private/WorldPartition/WorldPartition.cpp b/Engine/Source/Runtime/Engine/Private/WorldPartition/WorldPartition.cpp
--- a/Engine/Source/Runtime/Engine/Private/WorldPartition/WorldPartition.cpp
+++ b/Engine/Source/Runtime/Engine/Private/WorldPartition/WorldPartition.cpp
@@ -404,8 +404,6 @@ void UWorldPartition::Initialize(UWorld* InWorld, const FTransform& InTransform)
 
 	UE_LOG(LogWorldPartition, Log, TEXT("UWorldPartition::Initialize(IsEditor=%d, bPIEWorldTravel=%d IsGame=%d, IsCooking=%d)"), bIsEditor ? 1 : 0, bPIEWorldTravel ? 1 : 0, bIsGame ? 1 : 0, bIsCooking ? 1 : 0);
 
-	FWorldPartitionActorDesc::bForceAlwaysLoaded = !IsStreamingEnabled();
-
 	if (bEnableStreaming)
 	{
 		bStreamingWasEnabled = true;
@@ -468,6 +466,8 @@ void UWorldPartition::Initialize(UWorld* InWorld, const FTransform& InTransform)
 					ActorDescIterator->TransformInstance(SourceWorldPath, RemappedWorldPath, InstanceTransform);
 				}
 
+				ActorDescIterator->bIsForcedNonSpatiallyLoaded = !IsStreamingEnabled();
+
 				if (bIsEditor && !bIsCooking)
 				{
 					HashActorDesc(*ActorDescIterator);
@@ -1075,6 +1075,8 @@ void UWorldPartition::OnActorDescAdded(FWorldPartitionActorDesc* NewActorDesc)
 {
 	Super::OnActorDescAdded(NewActorDesc);
 
+	NewActorDesc->bIsForcedNonSpatiallyLoaded = !IsStreamingEnabled();
+
 	HashActorDesc(NewActorDesc);
 
 	if (WorldPartitionEditor)
@@ -1184,10 +1186,9 @@ void UWorldPartition::OnEnableStreamingChanged()
 		UnhashActorDesc(*ActorDescIterator);
 	}
 
-	FWorldPartitionActorDesc::bForceAlwaysLoaded = !IsStreamingEnabled();
-
 	for (UActorDescContainer::TIterator<> ActorDescIterator(this); ActorDescIterator; ++ActorDescIterator)
 	{
+		ActorDescIterator->bIsForcedNonSpatiallyLoaded = !IsStreamingEnabled();
 		HashActorDesc(*ActorDescIterator);
 	}
 
@@ -1469,6 +1470,18 @@ void UWorldPartition::DumpActorDescs(const FString& Path)
 	}
 }
 
+void UWorldPartition::AppendAssetRegistryTags(TArray<FAssetRegistryTag>& OutTags) const
+{
+	static const FName NAME_LevelIsPartitioned(TEXT("LevelIsPartitioned"));
+	OutTags.Add(FAssetRegistryTag(NAME_LevelIsPartitioned, TEXT("1"), FAssetRegistryTag::TT_Hidden));
+
+	if (!IsStreamingEnabled())
+	{
+		static const FName NAME_LevelHasStreamingDisabled(TEXT("LevelHasStreamingDisabled"));
+		OutTags.Add(FAssetRegistryTag(NAME_LevelHasStreamingDisabled, TEXT("1"), FAssetRegistryTag::TT_Hidden));
+	}
+}
+
 uint32 UWorldPartition::GetWantedEditorCellSize() const
 {
 	return EditorHash->GetWantedEditorCellSize();
diff --git a/Engine/Source/Runtime/Engine/Private/WorldPartition/WorldPartitionActorDesc.cpp b/Engine/Source/Runtime/Engine/Private/WorldPartition/WorldPartitionActorDesc.cpp
--- a/Engine/Source/Runtime/Engine/Private/WorldPartition/WorldPartitionActorDesc.cpp
+++ b/Engine/Source/Runtime/Engine/Private/WorldPartition/WorldPartitionActorDesc.cpp
@@ -24,13 +24,13 @@
 
 #if WITH_EDITOR
 uint32 FWorldPartitionActorDesc::GlobalTag = 0;
-bool FWorldPartitionActorDesc::bForceAlwaysLoaded = false;
 
 FWorldPartitionActorDesc::FWorldPartitionActorDesc()
 	: bIsUsingDataLayerAsset(false)
 	, SoftRefCount(0)
 	, HardRefCount(0)
 	, Container(nullptr)
+	, bIsForcedNonSpatiallyLoaded(false)
 	, Tag(0)
 {}
 
diff --git a/Engine/Source/Runtime/Engine/Private/WorldPartition/WorldPartitionStreamingGeneration.cpp b/Engine/Source/Runtime/Engine/Private/WorldPartition/WorldPartitionStreamingGeneration.cpp
--- a/Engine/Source/Runtime/Engine/Private/WorldPartition/WorldPartitionStreamingGeneration.cpp
+++ b/Engine/Source/Runtime/Engine/Private/WorldPartition/WorldPartitionStreamingGeneration.cpp
@@ -38,21 +38,32 @@ class FWorldPartitionStreamingGenerator
 {
 	void CreateActorDescViewMap(const UActorDescContainer* InContainer, TMap<FGuid, FWorldPartitionActorDescView>& OutActorDescViewMap, const FActorContainerID& InContainerID)
 	{
+		// should we handle unsaved or newly created actors?
 		const bool bHandleUnsavedActors = ModifiedActorsDescList && InContainerID.IsMainContainer();
 
 		// Consider all actors of a /Temp/ container package as Unsaved because loading them from disk will fail (Outer world name mismatch)
 		const bool bIsTempContainerPackage = FPackageName::IsTempPackage(InContainer->GetPackage()->GetName());
 		
-		// Test whether an actor is editor only. Will fallback to the actor descriptor only if the actor is not loaded.
+		// Test whether an actor is editor only. Will fallback to the actor descriptor only if the actor is not loaded
 		auto IsActorEditorOnly = [](const FWorldPartitionActorDesc* ActorDesc) { return ActorDesc->IsLoaded() ? ActorDesc->GetActor()->IsEditorOnly() : ActorDesc->GetActorIsEditorOnly(); };
 
-		auto GetModifiedActorDesc = [this, InContainer](AActor* InActor)
+		// Create an actor descriptor view for the specified actor (modified or unsaved actors)
+		auto GetModifiedActorDesc = [this](AActor* InActor) { return ModifiedActorsDescList->AddActor(InActor); };
+
+		// Adjust and register the actor descriptor view
+		auto RegisterActorDescView = [this, InContainer, &OutActorDescViewMap](const FGuid& ActorGuid, FWorldPartitionActorDescView& ActorDescView, bool bResolveRuntimeDataLayers)
 		{
-			FWorldPartitionActorDesc* ModifiedActorDesc = ModifiedActorsDescList->AddActor(InActor);
-			FWorldPartitionActorDescView ModifiedActorDescView(ModifiedActorDesc);
-			// ModifiedActorDescView's ActorDesc doesn't have a container, resolve view RuntimeDataLayers with provided ActorDescContainer
-			ModifiedActorDescView.ResolveRuntimeDataLayers(InContainer);
-			return ModifiedActorDescView;
+			if (!bEnableStreaming)
+			{
+				ActorDescView.SetForcedNonSpatiallyLoaded();
+			}
+
+			if (bResolveRuntimeDataLayers)
+			{
+				ActorDescView.ResolveRuntimeDataLayers(InContainer);
+			}
+
+			OutActorDescViewMap.Emplace(ActorGuid, ActorDescView);
 		};
 		
 		TMap<FGuid, FGuid> ContainerGuidsRemap;
@@ -74,13 +85,14 @@ class FWorldPartitionStreamingGenerator
 					{
 						// Dirty, unsaved actor for PIE
 						FWorldPartitionActorDescView ModifiedActorDescView = GetModifiedActorDesc(Actor);
-						OutActorDescViewMap.Emplace(ActorDescIt->GetGuid(), ModifiedActorDescView);
+						RegisterActorDescView(ActorDescIt->GetGuid(), ModifiedActorDescView, true);
 						continue;
 					}
 				}
 
 				// Non-dirty actor
-				OutActorDescViewMap.Emplace(ActorDescIt->GetGuid(), *ActorDescIt);
+				FWorldPartitionActorDescView ActorDescView(*ActorDescIt);
+				RegisterActorDescView(ActorDescIt->GetGuid(), ActorDescView, false);
 			}
 		}
 
@@ -92,7 +104,7 @@ class FWorldPartitionStreamingGenerator
 				if (IsValid(Actor) && Actor->IsPackageExternal() && Actor->IsMainPackageActor() && !Actor->IsEditorOnly() && !InContainer->GetActorDesc(Actor->GetActorGuid()))
 				{
 					FWorldPartitionActorDescView ModifiedActorDescView = GetModifiedActorDesc(Actor);
-					OutActorDescViewMap.Emplace(Actor->GetActorGuid(), ModifiedActorDescView);
+					RegisterActorDescView(Actor->GetActorGuid(), ModifiedActorDescView, true);
 				}
 			}
 		}
@@ -434,9 +446,10 @@ class FWorldPartitionStreamingGenerator
 	}
 
 public:
-	FWorldPartitionStreamingGenerator(FActorDescList* InModifiedActorsDescList, IStreamingGenerationErrorHandler* InErrorHandler)
-	: ModifiedActorsDescList(InModifiedActorsDescList)
-	, ErrorHandler(InErrorHandler ? InErrorHandler : &NullErrorHandler)
+	FWorldPartitionStreamingGenerator(FActorDescList* InModifiedActorsDescList, IStreamingGenerationErrorHandler* InErrorHandler, bool bInEnableStreaming)
+	: bEnableStreaming(bInEnableStreaming)
+	, ModifiedActorsDescList(InModifiedActorsDescList)
+	, ErrorHandler(InErrorHandler ? InErrorHandler : &NullErrorHandler)	
 	{}
 
 	void PreparationPhase(const UActorDescContainer* Container)
@@ -549,6 +562,7 @@ public:
 	}
 
 private:
+	bool bEnableStreaming;
 	FActorDescList* ModifiedActorsDescList;
 	IStreamingGenerationErrorHandler* ErrorHandler;
 	FStreamingGenerationNullErrorHandler NullErrorHandler;
@@ -592,12 +606,12 @@ bool UWorldPartition::GenerateStreaming(TArray<FString>* OutPackagesToGenerate)
 		ModifiedActorsDescList = &RuntimeHash->ModifiedActorDescListForPIE;
 		
 		// In PIE, we always want to populate the map check dialog
-		ErrorHandler = &MapCheckErrorHandler;		
+		ErrorHandler = &MapCheckErrorHandler;
 	}
 
 	FActorClusterContext ActorClusterContext;
 	{
-		FWorldPartitionStreamingGenerator StreamingGenerator(ModifiedActorsDescList, ErrorHandler);
+		FWorldPartitionStreamingGenerator StreamingGenerator(ModifiedActorsDescList, ErrorHandler, IsStreamingEnabled());
 
 		// Preparation Phase
 		StreamingGenerator.PreparationPhase(this);
@@ -627,7 +641,7 @@ bool UWorldPartition::GenerateStreaming(TArray<FString>* OutPackagesToGenerate)
 void UWorldPartition::GenerateHLOD(ISourceControlHelper* SourceControlHelper, bool bCreateActorsOnly)
 {
 	FStreamingGenerationLogErrorHandler LogErrorHandler;
-	FWorldPartitionStreamingGenerator StreamingGenerator(nullptr, &LogErrorHandler);
+	FWorldPartitionStreamingGenerator StreamingGenerator(nullptr, &LogErrorHandler, IsStreamingEnabled());
 	StreamingGenerator.PreparationPhase(this);
 	StreamingGenerator.DumpStateLog(TEXT("HLOD"));
 
@@ -642,15 +656,15 @@ void UWorldPartition::GenerateHLOD(ISourceControlHelper* SourceControlHelper, bo
 
 void UWorldPartition::CheckForErrors(IStreamingGenerationErrorHandler* ErrorHandler) const
 {
-	CheckForErrors(ErrorHandler, this);
+	CheckForErrors(ErrorHandler, this, IsStreamingEnabled());
 }
 
-void UWorldPartition::CheckForErrors(IStreamingGenerationErrorHandler* ErrorHandler, const UActorDescContainer* ActorDescContainer)
+void UWorldPartition::CheckForErrors(IStreamingGenerationErrorHandler* ErrorHandler, const UActorDescContainer* ActorDescContainer, bool bEnableStreaming)
 {
 	FActorClusterContext ActorClusterContext;
 	{
 		FActorDescList ModifiedActorDescList;
-		FWorldPartitionStreamingGenerator StreamingGenerator(ActorDescContainer->GetWorld() ? &ModifiedActorDescList : nullptr, ErrorHandler);
+		FWorldPartitionStreamingGenerator StreamingGenerator(ActorDescContainer->GetWorld() ? &ModifiedActorDescList : nullptr, ErrorHandler, bEnableStreaming);
 		StreamingGenerator.PreparationPhase(ActorDescContainer);
 	}
 }
diff --git a/Engine/Source/Runtime/Engine/Public/WorldPartition/WorldPartition.h b/Engine/Source/Runtime/Engine/Public/WorldPartition/WorldPartition.h
--- a/Engine/Source/Runtime/Engine/Public/WorldPartition/WorldPartition.h
+++ b/Engine/Source/Runtime/Engine/Public/WorldPartition/WorldPartition.h
@@ -163,7 +163,9 @@ public:
 	void DumpActorDescs(const FString& Path);
 
 	void CheckForErrors(IStreamingGenerationErrorHandler* ErrorHandler) const;
-	static void CheckForErrors(IStreamingGenerationErrorHandler* ErrorHandler, const UActorDescContainer* ActorDescContainer);
+	static void CheckForErrors(IStreamingGenerationErrorHandler* ErrorHandler, const UActorDescContainer* ActorDescContainer, bool bEnableStreaming);
+
+	void AppendAssetRegistryTags(TArray<FAssetRegistryTag>& OutTags) const;
 
 	uint32 GetWantedEditorCellSize() const;
 	void SetEditorWantedCellSize(uint32 InCellSize);
diff --git a/Engine/Source/Runtime/Engine/Public/WorldPartition/WorldPartitionActorDesc.h b/Engine/Source/Runtime/Engine/Public/WorldPartition/WorldPartitionActorDesc.h
--- a/Engine/Source/Runtime/Engine/Public/WorldPartition/WorldPartitionActorDesc.h
+++ b/Engine/Source/Runtime/Engine/Public/WorldPartition/WorldPartitionActorDesc.h
@@ -35,6 +35,7 @@ class ENGINE_API FWorldPartitionActorDesc
 #if WITH_EDITOR
 	friend class AActor;
 	friend class UActorDescContainer;
+	friend class UWorldPartition;
 	friend struct FWorldPartitionHandleImpl;
 	friend struct FWorldPartitionReferenceImpl;
 	friend class UDataLayerEditorSubsystem;
@@ -49,7 +50,7 @@ public:
 	inline UClass* GetActorClass() const { return ActorClass; }
 	inline FVector GetOrigin() const { return GetBounds().GetCenter(); }
 	inline FName GetRuntimeGrid() const { return RuntimeGrid; }
-	inline bool GetIsSpatiallyLoaded() const { return bForceAlwaysLoaded ? false : bIsSpatiallyLoaded; }
+	inline bool GetIsSpatiallyLoaded() const { return bIsForcedNonSpatiallyLoaded ? false : bIsSpatiallyLoaded; }
 	inline bool GetIsSpatiallyLoadedRaw() const { return bIsSpatiallyLoaded; }
 	inline bool GetActorIsEditorOnly() const { return bActorIsEditorOnly; }
 	inline bool GetLevelBoundsRelevant() const { return bLevelBoundsRelevant; }
@@ -160,6 +161,7 @@ protected:
 		Container = From->Container;
 		SoftRefCount = From->SoftRefCount;
 		HardRefCount = From->HardRefCount;
+		bIsForcedNonSpatiallyLoaded = From->bIsForcedNonSpatiallyLoaded;
 	}
 
 	virtual void TransferWorldData(const FWorldPartitionActorDesc* From)
@@ -199,13 +201,11 @@ protected:
 	mutable TWeakObjectPtr<AActor>	ActorPtr;
 	UActorDescContainer*			Container;
 	TArray<FName>					DataLayerInstanceNames;
+	bool							bIsForcedNonSpatiallyLoaded;
 
 public:
 	// Tagging
 	mutable uint32					Tag;
 	static uint32					GlobalTag;
-
-	// Temp
-	static bool						bForceAlwaysLoaded;
 #endif
 };
