Cleaning up some global Engine header dependencies.
- Focus around removing includes of "Engine/NetSerialization.h" from global headers. This file and its dependencies can take anywhere between 15 and 20 seconds to parse in MSVC.
- Everything in EngineTypes.h that depends on NetSerialization.h now lives inside their own headers.
- Moved FOverlapInfo and FScopedMovementUpdate out of SceneComponent.h/cpp and into their own files since they depend on the types moved out of EngineTypes.h

diff --git a/Engine/Source/Editor/UnrealEd/Private/Kismet2/KismetReinstanceUtilities.cpp b/Engine/Source/Editor/UnrealEd/Private/Kismet2/KismetReinstanceUtilities.cpp
--- a/Engine/Source/Editor/UnrealEd/Private/Kismet2/KismetReinstanceUtilities.cpp
+++ b/Engine/Source/Editor/UnrealEd/Private/Kismet2/KismetReinstanceUtilities.cpp
@@ -36,6 +36,7 @@
 #include "Engine/NetDriver.h"
 #include "Engine/ActorChannel.h"
 #include "Subsystems/AssetEditorSubsystem.h"
+#include "Engine/ScopedMovementUpdate.h"
 
 // Enabling this will validate cached dependent Blueprints against the full set of loaded Blueprints when updating bytecode references.
 // Note: Enabling this may potentially increase editor/Blueprint load time and/or decrease performance related to Blueprint compilation.
diff --git a/Engine/Source/Runtime/Engine/Classes/Components/PrimitiveComponent.h b/Engine/Source/Runtime/Engine/Classes/Components/PrimitiveComponent.h
--- a/Engine/Source/Runtime/Engine/Classes/Components/PrimitiveComponent.h
+++ b/Engine/Source/Runtime/Engine/Classes/Components/PrimitiveComponent.h
@@ -23,6 +23,9 @@
 #include "HitProxies.h"
 #include "Interfaces/Interface_AsyncCompilation.h"
 #include "HLOD/HLODBatchingPolicy.h"
+#if UE_ENABLE_INCLUDE_ORDER_DEPRECATED_IN_5_1
+#include "Engine/OverlapInfo.h"
+#endif
 #include "PrimitiveComponent.generated.h"
 
 class AController;
@@ -163,35 +166,6 @@ struct FRendererStencilMaskEvaluation
 	}
 };
 
-/*
- * Predicate for comparing FOverlapInfos when exact weak object pointer index/serial numbers should match, assuming one is not null and not invalid.
- * Compare to operator== for WeakObjectPtr which does both HasSameIndexAndSerialNumber *and* IsValid() checks on both pointers.
- */
-struct FFastOverlapInfoCompare
-{
-	FFastOverlapInfoCompare(const FOverlapInfo& BaseInfo)
-		: MyBaseInfo(BaseInfo)
-	{
-	}
-
-	bool operator() (const FOverlapInfo& Info)
-	{
-		return MyBaseInfo.OverlapInfo.Component.HasSameIndexAndSerialNumber(Info.OverlapInfo.Component)
-			&& MyBaseInfo.GetBodyIndex() == Info.GetBodyIndex();
-	}
-
-	bool operator() (const FOverlapInfo* Info)
-	{
-		return MyBaseInfo.OverlapInfo.Component.HasSameIndexAndSerialNumber(Info->OverlapInfo.Component)
-			&& MyBaseInfo.GetBodyIndex() == Info->GetBodyIndex();
-	}
-
-private:
-	const FOverlapInfo& MyBaseInfo;
-
-};
-
-
 /**
  * Delegate for notification of blocking collision against a specific component.  
  * NormalImpulse will be filled in for physics-simulating bodies, but will be zero for swept-component blocking collisions. 
@@ -232,6 +206,8 @@ public:
 	 * Default UObject constructor.
 	 */
 	UPrimitiveComponent(const FObjectInitializer& ObjectInitializer = FObjectInitializer::Get());
+	UPrimitiveComponent(FVTableHelper& Helper);
+	~UPrimitiveComponent();
 
 	// Rendering
 	
diff --git a/Engine/Source/Runtime/Engine/Classes/Components/SceneComponent.h b/Engine/Source/Runtime/Engine/Classes/Components/SceneComponent.h
--- a/Engine/Source/Runtime/Engine/Classes/Components/SceneComponent.h
+++ b/Engine/Source/Runtime/Engine/Classes/Components/SceneComponent.h
@@ -6,7 +6,11 @@
 #include "UObject/ObjectMacros.h"
 #include "UObject/UObjectGlobals.h"
 #include "UObject/CoreNet.h"
+#if UE_ENABLE_INCLUDE_ORDER_DEPRECATED_IN_5_1
 #include "Engine/EngineTypes.h"
+#include "Engine/OverlapInfo.h"
+#include "Engine/ScopedMovementUpdate.h"
+#endif
 #include "ComponentInstanceDataCache.h"
 #include "Components/ActorComponent.h"
 #include "RHIDefinitions.h"
@@ -15,37 +19,10 @@
 class AActor;
 class APhysicsVolume;
 class USceneComponent;
+class FScopedMovementUpdate;
 struct FLevelCollection;
 
-/** Overlap info consisting of the primitive and the body that is overlapping */
-struct ENGINE_API FOverlapInfo
-{
-	FOverlapInfo()
-	{}
-
-	explicit FOverlapInfo(const FHitResult& InSweepResult)
-		: bFromSweep(true), OverlapInfo(InSweepResult)
-	{
-	}
-
-	explicit FOverlapInfo(class UPrimitiveComponent* InComponent, int32 InBodyIndex = INDEX_NONE);
-	
-	int32 GetBodyIndex() const { return OverlapInfo.Item;  }
-
-	//This function completely ignores SweepResult information. It seems that places that use this function do not care, but it still seems risky
-	friend bool operator == (const FOverlapInfo& LHS, const FOverlapInfo& RHS) { return LHS.OverlapInfo.Component == RHS.OverlapInfo.Component && LHS.OverlapInfo.Item == RHS.OverlapInfo.Item; }
-	bool bFromSweep;
-
-	/** Information for both sweep and overlap queries. Different parts are valid depending on bFromSweep.
-	  * If bFromSweep is true then FHitResult is completely valid just like a regular sweep result.
-	  * If bFromSweep is false only FHitResult::Component, FHitResult::Actor, FHitResult::Item are valid as this is really just an FOverlapResult*/
-	FHitResult OverlapInfo;
-};
-
-// All added members of FOverlapInfo are PODs.
-template<> struct TIsPODType<FOverlapInfo> { enum { Value = TIsPODType<FHitResult>::Value }; };
-
-typedef TArray<FOverlapInfo, TInlineAllocator<3>> TInlineOverlapInfoArray;
+struct FOverlapInfo;
 typedef TArrayView<const FOverlapInfo> TOverlapArrayView;
 
 /** Detail mode for scene component rendering, corresponds with the integer value of UWorld::GetDetailMode() */
@@ -320,7 +297,7 @@ public:
 	FTransformUpdated TransformUpdated;
 
 	/** Returns the current scoped movement update, or NULL if there is none. @see FScopedMovementUpdate */
-	class FScopedMovementUpdate* GetCurrentScopedMovement() const;
+	FScopedMovementUpdate* GetCurrentScopedMovement() const;
 
 #if WITH_EDITORONLY_DATA
 	/**
@@ -341,10 +318,10 @@ public:
 
 private:
 	/** Stack of current movement scopes. */
-	TArray<class FScopedMovementUpdate*> ScopedMovementStack;
+	TArray<FScopedMovementUpdate*> ScopedMovementStack;
 
-	void BeginScopedMovementUpdate(class FScopedMovementUpdate& ScopedUpdate);
-	void EndScopedMovementUpdate(class FScopedMovementUpdate& ScopedUpdate);
+	void BeginScopedMovementUpdate(FScopedMovementUpdate& ScopedUpdate);
+	void EndScopedMovementUpdate(FScopedMovementUpdate& ScopedUpdate);
 
 	/** Cache that avoids Quat<->Rotator conversions if possible. Only to be used with GetComponentTransform().GetRotation(). */
 	FRotationConversionCache WorldRotationCache;
@@ -1009,6 +986,8 @@ protected:
 	/** Override this method for custom behavior for MoveComponent */
 	virtual bool MoveComponentImpl(const FVector& Delta, const FQuat& NewRotation, bool bSweep, FHitResult* Hit = NULL, EMoveComponentFlags MoveFlags = MOVECOMP_NoFlags, ETeleportType Teleport = ETeleportType::None);
 
+	bool IsDeferringMovementUpdates(const FScopedMovementUpdate& ScopedUpdate) const;
+
 public:
 	/** Call UpdateComponentToWorld if bComponentToWorldUpdated is false. */
 	void ConditionalUpdateComponentToWorld();
@@ -1739,216 +1718,7 @@ private:
 	void	operator delete[]	(void*);
 };
 
-//////////////////////////////////////////////////////////////////////////
-
-/**
- * Enum that controls the scoping behavior of FScopedMovementUpdate.
- * Note that EScopedUpdate::ImmediateUpdates is not allowed within outer scopes that defer updates,
- * and any attempt to do so will change the new inner scope to use deferred updates instead.
- */
-namespace EScopedUpdate
-{
-	enum Type
-	{
-		ImmediateUpdates,
-		DeferredUpdates	
-	};
-}
-
-
-/**
- * FScopedMovementUpdate creates a new movement scope, within which propagation of moves may be deferred until the end of the outermost scope that does not defer updates.
- * Moves within this scope will avoid updates such as UpdateBounds(), OnUpdateTransform(), UpdatePhysicsVolume(), UpdateChildTransforms() etc
- * until the move is committed (which happens when the last deferred scope goes out of context).
- *
- * Note that non-deferred scopes are not allowed within outer scopes that defer updates, and any attempt to use one will change the inner scope to use deferred updates.
- */
-class ENGINE_API FScopedMovementUpdate : private FNoncopyable
-{
-public:
-	
-	typedef TArray<struct FHitResult, TInlineAllocator<2>> TScopedBlockingHitArray;
-	typedef TArray<struct FOverlapInfo, TInlineAllocator<3>> TScopedOverlapInfoArray;
-
-	FScopedMovementUpdate( USceneComponent* Component, EScopedUpdate::Type ScopeBehavior = EScopedUpdate::DeferredUpdates, bool bRequireOverlapsEventFlagToQueueOverlaps = true );
-	~FScopedMovementUpdate();
-
-	enum class EHasMovedTransformOption
-	{
-		eTestTransform,
-		eIgnoreTransform
-	};
-
-	enum class EOverlapState
-	{
-		eUseParent,
-		eUnknown,
-		eIncludesOverlaps,
-		eForceUpdate,
-	};
-
-	/** Get the scope containing this scope. A scope only has an outer scope if they both defer updates. */
-	const FScopedMovementUpdate* GetOuterDeferredScope() const;
-
-	/** Return true if deferring updates, false if updates are applied immediately. */
-	bool IsDeferringUpdates() const;
-	
-	/** Revert movement to the initial location of the Component at the start of the scoped update. Also clears pending overlaps and sets bHasMoved to false. */
-	void RevertMove();
-
-	/** Returns whether movement has occurred at all during this scope, optionally checking if the transform is different (since changing scale does not go through a move). RevertMove() sets this back to false. */
-	bool HasMoved(EHasMovedTransformOption CheckTransform) const;
-
-	/** Returns true if the Component's transform differs from that at the start of the scoped update. */
-	bool IsTransformDirty() const;
-
-	/** Returns true if there are pending overlaps queued in this scope. */
-	bool HasPendingOverlaps() const;
-
-	/**
-	* Returns true if we require GetGenerateOverlapEvents() on both the moving object and the overlapped object to add them to the pending overlaps list.
-	* These flags will still be required when dispatching calls to UpdateOverlaps(), but this allows some custom processing of queued overlaps that would be otherwise missed along the way.
-	*/
-	bool RequiresOverlapsEventFlag() const;
-
-	/** Returns the pending overlaps within this scope. */
-	const TScopedOverlapInfoArray& GetPendingOverlaps() const;
-
-	/** Returns the list of pending blocking hits, which will be used for notifications once the move is committed. */
-	const TScopedBlockingHitArray& GetPendingBlockingHits() const;
-
-	//--------------------------------------------------------------------------------------------------------//
-	// These methods are intended only to be used by SceneComponent and derived classes.
-
-	/** Add overlaps to the queued overlaps array. This is intended for use only by SceneComponent and its derived classes whenever movement is performed. */
-	void AppendOverlapsAfterMove(const TOverlapArrayView& NewPendingOverlaps, bool bSweep, bool bIncludesOverlapsAtEnd);
-
-	/** Keep current pending overlaps after a move but make note that there was movement (just a symmetric rotation). */
-	void KeepCurrentOverlapsAfterRotation(bool bSweep);
-
-	/** Add blocking hit that will get processed once the move is committed. This is intended for use only by SceneComponent and its derived classes. */
-	void AppendBlockingHitAfterMove(const FHitResult& Hit);
-
-	/** Clear overlap state at current location, we don't know what it is. */
-	void InvalidateCurrentOverlaps();
-
-	/** Force full overlap update once this scope finishes. */
-	void ForceOverlapUpdate();
-
-	/** Registers that this move is a teleport */
-	void SetHasTeleported(ETeleportType InTeleportType);
-
-protected:
-	/** Fills in the list of overlaps at the end location (in EndOverlaps). Returns pointer to the list, or null if it can't be computed. */
-	template<typename AllocatorType>
-	TOptional<TOverlapArrayView> GetOverlapsAtEnd(class UPrimitiveComponent& PrimComponent, TArray<FOverlapInfo, AllocatorType>& OutEndOverlaps, bool bTransformChanged) const;
-
-	bool SetWorldLocationAndRotation(FVector NewLocation, const FQuat& NewQuat, bool bNoPhysics = false, ETeleportType Teleport = ETeleportType::None);
-
-private:
-	/** Notify this scope that the given inner scope completed its update (ie is going out of scope). Only occurs for deferred updates. */
-	void OnInnerScopeComplete(const FScopedMovementUpdate& InnerScope);
-	
-	// This class can only be created on the stack, otherwise the ordering constraints
-	// of the constructor and destructor between encapsulated scopes could be violated.
-	void*	operator new		(size_t);
-	void*	operator new[]		(size_t);
-	void	operator delete		(void *);
-	void	operator delete[]	(void*);
-
-protected:
-	USceneComponent* Owner;
-	FScopedMovementUpdate* OuterDeferredScope;
-
-	EOverlapState CurrentOverlapState;
-	ETeleportType TeleportType;
-
-	FTransform InitialTransform;
-	FVector InitialRelativeLocation;
-	FRotator InitialRelativeRotation;
-	FVector InitialRelativeScale;
-
-	int32 FinalOverlapCandidatesIndex;		// If not INDEX_NONE, overlaps at this index and beyond in PendingOverlaps are at the final destination
-	TScopedOverlapInfoArray PendingOverlaps;	// All overlaps encountered during the scope of moves.
-	TScopedBlockingHitArray BlockingHits;		// All blocking hits encountered during the scope of moves.
-
-	uint8 bDeferUpdates:1;
-	uint8 bHasMoved:1;
-	uint8 bRequireOverlapsEventFlag:1;
-
-	friend class USceneComponent;
-};
-
-//////////////////////////////////////////////////////////////////////////
-// FScopedMovementUpdate inlines
-
-FORCEINLINE const FScopedMovementUpdate* FScopedMovementUpdate::GetOuterDeferredScope() const
-{
-	return OuterDeferredScope;
-}
-
-FORCEINLINE bool FScopedMovementUpdate::IsDeferringUpdates() const
-{
-	return bDeferUpdates;
-}
-
-FORCEINLINE bool FScopedMovementUpdate::HasMoved(EHasMovedTransformOption CheckTransform) const
-{
-	return bHasMoved || (CheckTransform == EHasMovedTransformOption::eTestTransform && IsTransformDirty());
-}
-
-FORCEINLINE bool FScopedMovementUpdate::HasPendingOverlaps() const
-{
-	return PendingOverlaps.Num() > 0;
-}
-
-FORCEINLINE bool FScopedMovementUpdate::RequiresOverlapsEventFlag() const
-{
-	return bRequireOverlapsEventFlag;
-}
-
-FORCEINLINE const FScopedMovementUpdate::TScopedOverlapInfoArray& FScopedMovementUpdate::GetPendingOverlaps() const
-{
-	return PendingOverlaps;
-}
-
-FORCEINLINE const FScopedMovementUpdate::TScopedBlockingHitArray& FScopedMovementUpdate::GetPendingBlockingHits() const
-{
-	return BlockingHits;
-}
-
-FORCEINLINE_DEBUGGABLE void FScopedMovementUpdate::AppendBlockingHitAfterMove(const FHitResult& Hit)
-{
-	BlockingHits.Add(Hit);
-}
-
-FORCEINLINE_DEBUGGABLE void FScopedMovementUpdate::KeepCurrentOverlapsAfterRotation(bool bSweep)
-{
-	bHasMoved = true;
-	// CurrentOverlapState is unchanged
-}
-
-FORCEINLINE_DEBUGGABLE void FScopedMovementUpdate::InvalidateCurrentOverlaps()
-{
-	bHasMoved = true;
-	CurrentOverlapState = EOverlapState::eUnknown;
-	FinalOverlapCandidatesIndex = INDEX_NONE;
-}
-
-FORCEINLINE_DEBUGGABLE void FScopedMovementUpdate::ForceOverlapUpdate()
-{
-	bHasMoved = true;
-	CurrentOverlapState = EOverlapState::eForceUpdate;
-	FinalOverlapCandidatesIndex = INDEX_NONE;
-}
-
-FORCEINLINE_DEBUGGABLE void FScopedMovementUpdate::SetHasTeleported(ETeleportType InTeleportType)
-{
-	// Request an initialization. Teleport type can only go higher - i.e. if we have requested a reset, then a teleport will still reset fully
-	TeleportType = ((InTeleportType > TeleportType) ? InTeleportType : TeleportType); 
-}
-
-FORCEINLINE_DEBUGGABLE class FScopedMovementUpdate* USceneComponent::GetCurrentScopedMovement() const
+FORCEINLINE_DEBUGGABLE FScopedMovementUpdate* USceneComponent::GetCurrentScopedMovement() const
 {
 	if (ScopedMovementStack.Num() > 0)
 	{
@@ -1957,40 +1727,23 @@ FORCEINLINE_DEBUGGABLE class FScopedMovementUpdate* USceneComponent::GetCurrentS
 	return nullptr;
 }
 
-//////////////////////////////////////////////////////////////////////////
-// SceneComponent inlines that depend on FScopedMovementUpdate
-
 FORCEINLINE_DEBUGGABLE bool USceneComponent::IsDeferringMovementUpdates() const
 {
 	if (ScopedMovementStack.Num() > 0)
 	{
-		checkSlow(ScopedMovementStack.Last()->IsDeferringUpdates());
+		checkSlow(IsDeferringMovementUpdates(*ScopedMovementStack.Last()));
 		return true;
 	}
 	return false;
 }
 
-FORCEINLINE_DEBUGGABLE void USceneComponent::BeginScopedMovementUpdate(class FScopedMovementUpdate& ScopedUpdate)
+FORCEINLINE_DEBUGGABLE void USceneComponent::BeginScopedMovementUpdate(FScopedMovementUpdate& ScopedUpdate)
 {
 	checkSlow(IsInGameThread());
-	checkSlow(ScopedUpdate.IsDeferringUpdates());
+	checkSlow(IsDeferringMovementUpdates(ScopedUpdate));
 	ScopedMovementStack.Push(&ScopedUpdate);
 }
 
-FORCEINLINE_DEBUGGABLE bool USceneComponent::UpdateOverlaps(const TOverlapArrayView* PendingOverlaps /* = nullptr */, bool bDoNotifies /* = true */, const TOverlapArrayView* OverlapsAtEndLocation /* = nullptr */)
-{
-	if (IsDeferringMovementUpdates())
-	{
-		GetCurrentScopedMovement()->ForceOverlapUpdate();
-	}
-	else if (!ShouldSkipUpdateOverlaps())
-	{
-		bSkipUpdateOverlaps = UpdateOverlapsImpl(PendingOverlaps, bDoNotifies, OverlapsAtEndLocation);
-	}
-
-	return bSkipUpdateOverlaps;
-}
-
 FORCEINLINE_DEBUGGABLE bool USceneComponent::GetShouldUpdatePhysicsVolume() const
 {
 	return bShouldUpdatePhysicsVolume;
diff --git a/Engine/Source/Runtime/Engine/Classes/Engine/ActorInstanceHandle.h b/Engine/Source/Runtime/Engine/Classes/Engine/ActorInstanceHandle.h
new file mode 100644
--- /dev/null
+++ b/Engine/Source/Runtime/Engine/Classes/Engine/ActorInstanceHandle.h
@@ -0,0 +1,100 @@
+// Copyright Epic Games, Inc. All Rights Reserved.
+
+#pragma once
+
+#include "CoreMinimal.h"
+#include "ActorInstanceHandle.generated.h"
+
+// Handle to a unique object. This may specify a full weigh actor or it may only specify the light weight instance that represents the same object.
+USTRUCT(BlueprintType)
+struct ENGINE_API FActorInstanceHandle
+{
+	GENERATED_BODY()
+
+	friend struct FLightWeightInstanceSubsystem;
+	friend class ALightWeightInstanceManager;
+	friend class UActorInstanceHandleInterface;
+
+	FActorInstanceHandle();
+
+	explicit FActorInstanceHandle(AActor* InActor);
+	explicit FActorInstanceHandle(class ALightWeightInstanceManager* Manager, int32 InInstanceIndex);
+
+	FActorInstanceHandle(const FActorInstanceHandle& Other);
+
+	bool IsValid() const;
+
+	bool DoesRepresentClass(const UClass* OtherClass) const;
+
+	UClass* GetRepresentedClass() const;
+
+	FVector GetLocation() const;
+	FRotator GetRotation() const;
+	FTransform GetTransform() const;
+
+	FName GetFName() const;
+	FString GetName() const;
+
+	/** If this handle has a valid actor, return it; otherwise return the actor responsible for managing the instances. */
+	AActor* GetManagingActor() const;
+
+	/** Returns either the actor's root component or the root component for the manager associated with the handle */
+	USceneComponent* GetRootComponent() const;
+
+	/** Returns the actor specified by this handle. This may require loading and creating the actor object. */
+	AActor* FetchActor() const;
+	template <typename T>
+	T* FetchActor() const;
+
+	/* Returns the index used internally by the manager */
+	FORCEINLINE int32 GetInstanceIndex() const { return InstanceIndex; }
+
+	/* Returns the index used by rendering and collision */
+	int32 GetRenderingInstanceIndex() const;
+
+	FActorInstanceHandle& operator=(const FActorInstanceHandle& Other) = default;
+	FActorInstanceHandle& operator=(FActorInstanceHandle&& Other) = default;
+	FActorInstanceHandle& operator=(AActor* OtherActor);
+
+	bool operator==(const FActorInstanceHandle& Other) const;
+	bool operator!=(const FActorInstanceHandle& Other) const;
+
+	bool operator==(const AActor* OtherActor) const;
+	bool operator!=(const AActor* OtherActor) const;
+
+	friend ENGINE_API uint32 GetTypeHash(const FActorInstanceHandle& Handle);
+
+	friend ENGINE_API FArchive& operator<<(FArchive& Ar, FActorInstanceHandle& Handle);
+
+	uint32 GetInstanceUID() const { return InstanceUID; }
+
+private:
+	/**
+	 * helper functions that let us treat the actor pointer as a UObject in templated functions
+	 * these do NOT fetch the actor so they will return nullptr if we don't have a full actor representation
+	 */
+	UObject* GetActorAsUObject();
+	const UObject* GetActorAsUObject() const;
+
+	/** Returns true if Actor is not null and not pending kill */
+	bool IsActorValid() const;
+
+	/** this is cached here for convenience */
+	UPROPERTY()
+		mutable TWeakObjectPtr<AActor> Actor;
+
+	/** Identifies the light weight instance manager to use */
+	TWeakObjectPtr<ALightWeightInstanceManager> Manager;
+
+	/** Identifies the instance within the manager */
+	int32 InstanceIndex;
+
+	/** Unique identifier for instances represented by the handle */
+	uint32 InstanceUID;
+};
+
+template <typename T>
+T* FActorInstanceHandle::FetchActor() const
+{
+	return Cast<T>(FetchActor());
+}
diff --git a/Engine/Source/Runtime/Engine/Classes/Engine/DamageEvents.h b/Engine/Source/Runtime/Engine/Classes/Engine/DamageEvents.h
new file mode 100644
--- /dev/null
+++ b/Engine/Source/Runtime/Engine/Classes/Engine/DamageEvents.h
@@ -0,0 +1,158 @@
+// Copyright Epic Games, Inc. All Rights Reserved.
+
+#pragma once
+
+#include "CoreMinimal.h"
+#include "Templates/SubclassOf.h"
+#include "Engine/NetSerialization.h"
+#include "Engine/HitResult.h"
+#include "DamageEvents.generated.h"
+
+class UDamageType;
+
+/** Event used by AActor::TakeDamage and related functions */
+USTRUCT(BlueprintType)
+struct ENGINE_API FDamageEvent
+{
+	GENERATED_BODY()
+
+public:
+	/** Default constructor (no initialization). */
+	FDamageEvent() { }
+
+	FDamageEvent(FDamageEvent const& InDamageEvent)
+		: DamageTypeClass(InDamageEvent.DamageTypeClass)
+	{ }
+	
+	virtual ~FDamageEvent() { }
+
+	explicit FDamageEvent(TSubclassOf<UDamageType> InDamageTypeClass)
+		: DamageTypeClass(InDamageTypeClass)
+	{ }
+
+	/** Optional DamageType for this event.  If nullptr, UDamageType will be assumed. */
+	UPROPERTY()
+	TSubclassOf<UDamageType> DamageTypeClass;
+
+	/** ID for this class. NOTE this must be unique for all damage events. */
+	static const int32 ClassID = 0;
+
+	virtual int32 GetTypeID() const { return FDamageEvent::ClassID; }
+	virtual bool IsOfType(int32 InID) const { return FDamageEvent::ClassID == InID; };
+
+	/** This is for compatibility with old-style functions which want a unified set of hit data regardless of type of hit.  Ideally this will go away over time. */
+	virtual void GetBestHitInfo(AActor const* HitActor, AActor const* HitInstigator, FHitResult& OutHitInfo, FVector& OutImpulseDir) const;
+};
+
+/** Damage subclass that handles damage with a single impact location and source direction */
+USTRUCT()
+struct ENGINE_API FPointDamageEvent : public FDamageEvent
+{
+	GENERATED_BODY()
+
+	/** Actual damage done */
+	UPROPERTY()
+	float Damage;
+	
+	/** Direction the shot came from. Should be normalized. */
+	UPROPERTY()
+	FVector_NetQuantizeNormal ShotDirection;
+	
+	/** Describes the trace/location that caused this damage */
+	UPROPERTY()
+	FHitResult HitInfo;
+
+	FPointDamageEvent() : Damage(0.0f), ShotDirection(ForceInitToZero), HitInfo() {}
+	FPointDamageEvent(float InDamage, const FHitResult& InHitInfo, FVector const& InShotDirection, TSubclassOf<UDamageType> InDamageTypeClass)
+		: FDamageEvent(InDamageTypeClass), Damage(InDamage), ShotDirection(InShotDirection), HitInfo(InHitInfo)
+	{}
+	
+	/** ID for this class. NOTE this must be unique for all damage events. */
+	static const int32 ClassID = 1;
+	
+	virtual int32 GetTypeID() const override { return FPointDamageEvent::ClassID; };
+	virtual bool IsOfType(int32 InID) const override { return (FPointDamageEvent::ClassID == InID) || FDamageEvent::IsOfType(InID); };
+
+	/** Simple API for common cases where we are happy to assume a single hit is expected, even though damage event may have multiple hits. */
+	virtual void GetBestHitInfo(AActor const* HitActor, AActor const* HitInstigator, FHitResult& OutHitInfo, FVector& OutImpulseDir) const override;
+};
+
+/** Parameters used to compute radial damage */
+USTRUCT(BlueprintType)
+struct ENGINE_API FRadialDamageParams
+{
+	GENERATED_BODY()
+
+	/** Max damage done */
+	UPROPERTY(EditAnywhere, BlueprintReadWrite, Category=RadialDamageParams)
+	float BaseDamage;
+
+	/** Damage will not fall below this if within range */
+	UPROPERTY(EditAnywhere, BlueprintReadWrite, Category=RadialDamageParams)
+	float MinimumDamage;
+	
+	/** Within InnerRadius, do max damage */
+	UPROPERTY(EditAnywhere, BlueprintReadWrite, Category=RadialDamageParams)
+	float InnerRadius;
+		
+	/** Outside OuterRadius, do no damage */
+	UPROPERTY(EditAnywhere, BlueprintReadWrite, Category=RadialDamageParams)
+	float OuterRadius;
+		
+	/** Describes amount of exponential damage falloff */
+	UPROPERTY(EditAnywhere, BlueprintReadWrite, Category=RadialDamageParams)
+	float DamageFalloff;
+
+	FRadialDamageParams()
+		: BaseDamage(0.f), MinimumDamage(0.f), InnerRadius(0.f), OuterRadius(0.f), DamageFalloff(1.f)
+	{}
+	FRadialDamageParams(float InBaseDamage, float InInnerRadius, float InOuterRadius, float InDamageFalloff)
+		: BaseDamage(InBaseDamage), MinimumDamage(0.f), InnerRadius(InInnerRadius), OuterRadius(InOuterRadius), DamageFalloff(InDamageFalloff)
+	{}
+	FRadialDamageParams(float InBaseDamage, float InMinimumDamage, float InInnerRadius, float InOuterRadius, float InDamageFalloff)
+		: BaseDamage(InBaseDamage), MinimumDamage(InMinimumDamage), InnerRadius(InInnerRadius), OuterRadius(InOuterRadius), DamageFalloff(InDamageFalloff)
+	{}
+	FRadialDamageParams(float InBaseDamage, float InRadius)
+		: BaseDamage(InBaseDamage), MinimumDamage(0.f), InnerRadius(0.f), OuterRadius(InRadius), DamageFalloff(1.f)
+	{}
+
+	/** Returns damage done at a certain distance */
+	float GetDamageScale(float DistanceFromEpicenter) const;
+
+	/** Return outermost radius of the damage area. Protects against malformed data. */
+	float GetMaxRadius() const { return FMath::Max( FMath::Max(InnerRadius, OuterRadius), 0.f ); }
+};
+
+/** Damage subclass that handles damage with a source location and falloff radius */
+USTRUCT()
+struct ENGINE_API FRadialDamageEvent : public FDamageEvent
+{
+	GENERATED_BODY()
+
+	/** Static parameters describing damage falloff math */
+	UPROPERTY()
+	FRadialDamageParams Params;
+	
+	/** Location of origin point */
+	UPROPERTY()
+	FVector Origin;
+
+	/** Hit reslts of specific impacts */
+	UPROPERTY()
+	TArray<FHitResult> ComponentHits;
+
+	/** ID for this class. NOTE this must be unique for all damage events. */
+	static const int32 ClassID = 2;
+
+	virtual int32 GetTypeID() const override { return FRadialDamageEvent::ClassID; };
+	virtual bool IsOfType(int32 InID) const override { return (FRadialDamageEvent::ClassID == InID) || FDamageEvent::IsOfType(InID); };
+
+	/** Simple API for common cases where we are happy to assume a single hit is expected, even though damage event may have multiple hits. */
+	virtual void GetBestHitInfo(AActor const* HitActor, AActor const* HitInstigator, FHitResult& OutHitInfo, FVector& OutImpulseDir) const override;
+
+	FRadialDamageEvent()
+		: Origin(ForceInitToZero)
+	{}
+};
+
+
diff --git a/Engine/Source/Runtime/Engine/Classes/Engine/EngineTypes.h b/Engine/Source/Runtime/Engine/Classes/Engine/EngineTypes.h
--- a/Engine/Source/Runtime/Engine/Classes/Engine/EngineTypes.h
+++ b/Engine/Source/Runtime/Engine/Classes/Engine/EngineTypes.h
@@ -12,7 +12,14 @@
 #include "UObject/Object.h"
 #include "UObject/Class.h"
 #include "Templates/SubclassOf.h"
+#if UE_ENABLE_INCLUDE_ORDER_DEPRECATED_IN_5_1
 #include "Engine/NetSerialization.h"
+#include "Engine/ActorInstanceHandle.h"
+#include "Engine/HitResult.h"
+#include "Engine/OverlapResult.h"
+#include "Engine/DamageEvents.h"
+#include "Engine/ReplicatedState.h"
+#endif
 #include "EngineTypes.generated.h"
 
 class AActor;
@@ -1423,47 +1430,6 @@ FORCEINLINE ECollisionEnabled::Type CollisionEnabledIntersection(ECollisionEnabl
 	return ECollisionEnabled::NoCollision;
 }
 
-/** Describes the physical state of a rigid body. */
-USTRUCT()
-struct FRigidBodyState
-{
-	GENERATED_BODY()
-
-	UPROPERTY()
-	FVector_NetQuantize100 Position;
-
-	UPROPERTY()
-	FQuat Quaternion;
-
-	UPROPERTY()
-	FVector_NetQuantize100 LinVel;
-
-	UPROPERTY()
-	FVector_NetQuantize100 AngVel;
-
-	UPROPERTY()
-	uint8 Flags;
-
-	FRigidBodyState()
-		: Position(ForceInit)
-		, Quaternion(ForceInit)
-		, LinVel(ForceInit)
-		, AngVel(ForceInit)
-		, Flags(0)
-	{ }
-};
-
-/** Describes extra state about a specific rigid body */
-namespace ERigidBodyFlags
-{
-	enum Type
-	{
-		None				= 0x00,
-		Sleeping			= 0x01,
-		NeedsUpdate			= 0x02,
-	};
-}
-
 /** Describes type of wake/sleep event sent to the physics system */
 enum class ESleepEvent : uint8
 {
@@ -2131,388 +2097,6 @@ struct FPrimitiveMaterialRef
 	{ 	}
 };
 
-// Handle to a unique object. This may specify a full weigh actor or it may only specify the light weight instance that represents the same object.
-USTRUCT(BlueprintType)
-struct ENGINE_API FActorInstanceHandle
-{
-	GENERATED_BODY()
-
-	friend struct FLightWeightInstanceSubsystem;
-	friend class ALightWeightInstanceManager;
-	friend class UActorInstanceHandleInterface;
-
-	FActorInstanceHandle();
-
-	explicit FActorInstanceHandle(AActor* InActor);
-	explicit FActorInstanceHandle(class ALightWeightInstanceManager* Manager, int32 InInstanceIndex);
-
-	FActorInstanceHandle(const FActorInstanceHandle& Other);
-
-	bool IsValid() const;
-
-	bool DoesRepresentClass(const UClass* OtherClass) const;
-
-	UClass* GetRepresentedClass() const;
-
-	FVector GetLocation() const;
-	FRotator GetRotation() const;
-	FTransform GetTransform() const;
-
-	FName GetFName() const;
-	FString GetName() const;
-
-	/** If this handle has a valid actor, return it; otherwise return the actor responsible for managing the instances. */
-	AActor* GetManagingActor() const;
-
-	/** Returns either the actor's root component or the root component for the manager associated with the handle */
-	USceneComponent* GetRootComponent() const;
-
-	/** Returns the actor specified by this handle. This may require loading and creating the actor object. */
-	AActor* FetchActor() const;
-	template <typename T>
-	T* FetchActor() const;
-
-	/* Returns the index used internally by the manager */
-	FORCEINLINE int32 GetInstanceIndex() const { return InstanceIndex; }
-
-	/* Returns the index used by rendering and collision */
-	int32 GetRenderingInstanceIndex() const;
-
-	FActorInstanceHandle& operator=(const FActorInstanceHandle& Other) = default;
-	FActorInstanceHandle& operator=(FActorInstanceHandle&& Other) = default;
-	FActorInstanceHandle& operator=(AActor* OtherActor);
-
-	bool operator==(const FActorInstanceHandle& Other) const;
-	bool operator!=(const FActorInstanceHandle& Other) const;
-
-	bool operator==(const AActor* OtherActor) const;
-	bool operator!=(const AActor* OtherActor) const;
-
-	friend ENGINE_API uint32 GetTypeHash(const FActorInstanceHandle& Handle);
-
-	friend ENGINE_API FArchive& operator<<(FArchive& Ar, FActorInstanceHandle& Handle);
-
-	uint32 GetInstanceUID() const { return InstanceUID; }
-
-private:
-	/**
-	 * helper functions that let us treat the actor pointer as a UObject in templated functions
-	 * these do NOT fetch the actor so they will return nullptr if we don't have a full actor representation
-	 */
-	UObject* GetActorAsUObject();
-	const UObject* GetActorAsUObject() const;
-
-	/** Returns true if Actor is not null and not pending kill */
-	bool IsActorValid() const;
-
-	/** this is cached here for convenience */
-	UPROPERTY()
-	mutable TWeakObjectPtr<AActor> Actor;
-
-	/** Identifies the light weight instance manager to use */
-	TWeakObjectPtr<ALightWeightInstanceManager> Manager;
-
-	/** Identifies the instance within the manager */
-	int32 InstanceIndex;
-
-	/** Unique identifier for instances represented by the handle */
-	uint32 InstanceUID;
-};
-
-template <typename T>
-T* FActorInstanceHandle::FetchActor() const
-{
-	return Cast<T>(FetchActor());
-}
-
-/**
- * Structure containing information about one hit of a trace, such as point of impact and surface normal at that point.
- */
-USTRUCT(BlueprintType, meta = (HasNativeBreak = "Engine.GameplayStatics.BreakHitResult", HasNativeMake = "Engine.GameplayStatics.MakeHitResult"))
-struct ENGINE_API FHitResult
-{
-	GENERATED_BODY()
-
-	/** Face index we hit (for complex hits with triangle meshes). */
-	UPROPERTY()
-	int32 FaceIndex;
-
-	/**
-	 * 'Time' of impact along trace direction (ranging from 0.0 to 1.0) if there is a hit, indicating time between TraceStart and TraceEnd.
-	 * For swept movement (but not queries) this may be pulled back slightly from the actual time of impact, to prevent precision problems with adjacent geometry.
-	 */
-	UPROPERTY()
-	float Time;
-	 
-	/** The distance from the TraceStart to the Location in world space. This value is 0 if there was an initial overlap (trace started inside another colliding object). */
-	UPROPERTY()
-	float Distance; 
-	
-	/**
-	 * The location in world space where the moving shape would end up against the impacted object, if there is a hit. Equal to the point of impact for line tests.
-	 * Example: for a sphere trace test, this is the point where the center of the sphere would be located when it touched the other object.
-	 * For swept movement (but not queries) this may not equal the final location of the shape since hits are pulled back slightly to prevent precision issues from overlapping another surface.
-	 */
-	UPROPERTY()
-	FVector_NetQuantize Location;
-
-	/**
-	 * Location in world space of the actual contact of the trace shape (box, sphere, ray, etc) with the impacted object.
-	 * Example: for a sphere trace test, this is the point where the surface of the sphere touches the other object.
-	 * @note: In the case of initial overlap (bStartPenetrating=true), ImpactPoint will be the same as Location because there is no meaningful single impact point to report.
-	 */
-	UPROPERTY()
-	FVector_NetQuantize ImpactPoint;
-
-	/**
-	 * Normal of the hit in world space, for the object that was swept. Equal to ImpactNormal for line tests.
-	 * This is computed for capsules and spheres, otherwise it will be the same as ImpactNormal.
-	 * Example: for a sphere trace test, this is a normalized vector pointing in towards the center of the sphere at the point of impact.
-	 */
-	UPROPERTY()
-	FVector_NetQuantizeNormal Normal;
-
-	/**
-	 * Normal of the hit in world space, for the object that was hit by the sweep, if any.
-	 * For example if a sphere hits a flat plane, this is a normalized vector pointing out from the plane.
-	 * In the case of impact with a corner or edge of a surface, usually the "most opposing" normal (opposed to the query direction) is chosen.
-	 */
-	UPROPERTY()
-	FVector_NetQuantizeNormal ImpactNormal;
-
-	/**
-	 * Start location of the trace.
-	 * For example if a sphere is swept against the world, this is the starting location of the center of the sphere.
-	 */
-	UPROPERTY()
-	FVector_NetQuantize TraceStart;
-
-	/**
-	 * End location of the trace; this is NOT where the impact occurred (if any), but the furthest point in the attempted sweep.
-	 * For example if a sphere is swept against the world, this would be the center of the sphere if there was no blocking hit.
-	 */
-	UPROPERTY()
-	FVector_NetQuantize TraceEnd;
-
-	/**
-	  * If this test started in penetration (bStartPenetrating is true) and a depenetration vector can be computed,
-	  * this value is the distance along Normal that will result in moving out of penetration.
-	  * If the distance cannot be computed, this distance will be zero.
-	  */
-	UPROPERTY()
-	float PenetrationDepth;
-
-	/** If the hit result is from a collision this will have extra info about the item that hit the second item. */
-	UPROPERTY()
-	int32 MyItem;
-
-	/** Extra data about item that was hit (hit primitive specific). */
-	UPROPERTY()
-	int32 Item;
-
-	/** Index to item that was hit, also hit primitive specific. */
-	UPROPERTY()
-	uint8 ElementIndex;
-
-	/** Indicates if this hit was a result of blocking collision. If false, there was no hit or it was an overlap/touch instead. */
-	UPROPERTY()
-	uint8 bBlockingHit : 1;
-
-	/**
-	 * Whether the trace started in penetration, i.e. with an initial blocking overlap.
-	 * In the case of penetration, if PenetrationDepth > 0.f, then it will represent the distance along the Normal vector that will result in
-	 * minimal contact between the swept shape and the object that was hit. In this case, ImpactNormal will be the normal opposed to movement at that location
-	 * (ie, Normal may not equal ImpactNormal). ImpactPoint will be the same as Location, since there is no single impact point to report.
-	 */
-	UPROPERTY()
-	uint8 bStartPenetrating : 1;
-
-	/**
-	 * Physical material that was hit.
-	 * @note Must set bReturnPhysicalMaterial on the swept PrimitiveComponent or in the query params for this to be returned.
-	 */
-	UPROPERTY()
-	TWeakObjectPtr<class UPhysicalMaterial> PhysMaterial;
-
-	/** Handle to the object hit by the trace. */
-	UPROPERTY()
-	FActorInstanceHandle HitObjectHandle;
-
-	/** PrimitiveComponent hit by the trace. */
-	UPROPERTY()
-	TWeakObjectPtr<class UPrimitiveComponent> Component;
-
-	/** Name of bone we hit (for skeletal meshes). */
-	UPROPERTY()
-	FName BoneName;
-
-	/** Name of the _my_ bone which took part in hit event (in case of two skeletal meshes colliding). */
-	UPROPERTY()
-	FName MyBoneName;
-
-
-	FHitResult()
-	{
-		Init();
-	}
-	
-	explicit FHitResult(float InTime)
-	{
-		Init();
-		Time = InTime;
-	}
-
-	explicit FHitResult(EForceInit InInit)
-	{
-		Init();
-	}
-
-	explicit FHitResult(ENoInit NoInit)
-	{
-	}
-
-	explicit FHitResult(FVector Start, FVector End)
-	{
-		Init(Start, End);
-	}
-
-	/** Initialize empty hit result with given time. */
-	FORCEINLINE void Init()
-	{
-		FMemory::Memzero(this, sizeof(FHitResult));
-		HitObjectHandle = FActorInstanceHandle();
-		Time = 1.f;
-		MyItem = INDEX_NONE;
-	}
-
-	/** Initialize empty hit result with given time, TraceStart, and TraceEnd */
-	FORCEINLINE void Init(FVector Start, FVector End)
-	{
-		FMemory::Memzero(this, sizeof(FHitResult));
-		HitObjectHandle = FActorInstanceHandle();
-		Time = 1.f;
-		TraceStart = Start;
-		TraceEnd = End;
-		MyItem = INDEX_NONE;
-	}
-
-	/** Ctor for easily creating "fake" hits from limited data. */
-	FHitResult(class AActor* InActor, class UPrimitiveComponent* InComponent, FVector const& HitLoc, FVector const& HitNorm);
- 
-	/** Reset hit result while optionally saving TraceStart and TraceEnd. */
-	FORCEINLINE void Reset(float InTime = 1.f, bool bPreserveTraceData = true)
-	{
-		const FVector SavedTraceStart = TraceStart;
-		const FVector SavedTraceEnd = TraceEnd;
-		Init();
-		Time = InTime;
-		if (bPreserveTraceData)
-		{
-			TraceStart = SavedTraceStart;
-			TraceEnd = SavedTraceEnd;
-		}
-	}
-
-	/** Utility to return the Actor that owns the Component that was hit. */
-	FORCEINLINE AActor* GetActor() const
-	{
-		return HitObjectHandle.FetchActor();
-	}
-
-	FORCEINLINE FActorInstanceHandle GetHitObjectHandle() const
-	{
-		return HitObjectHandle;
-	}
-
-	FORCEINLINE bool HasValidHitObjectHandle() const
-	{
-		return HitObjectHandle.IsValid();
-	}
-
-	/** Utility to return the Component that was hit. */
-	FORCEINLINE UPrimitiveComponent* GetComponent() const
-	{
-		return Component.Get();
-	}
-
-	/** Optimized serialize function */
-	bool NetSerialize(FArchive& Ar, class UPackageMap* Map, bool& bOutSuccess);
-
-	/** Return true if there was a blocking hit that was not caused by starting in penetration. */
-	FORCEINLINE bool IsValidBlockingHit() const
-	{
-		return bBlockingHit && !bStartPenetrating;
-	}
-
-	/** Static utility function that returns the first 'blocking' hit in an array of results. */
-	static FHitResult* GetFirstBlockingHit(TArray<FHitResult>& InHits)
-	{
-		for(int32 HitIdx=0; HitIdx<InHits.Num(); HitIdx++)
-		{
-			if(InHits[HitIdx].bBlockingHit)
-			{
-				return &InHits[HitIdx];
-			}
-		}
-		return nullptr;
-	}
-
-	/** Static utility function that returns the number of blocking hits in array. */
-	static int32 GetNumBlockingHits(const TArray<FHitResult>& InHits)
-	{
-		int32 NumBlocks = 0;
-		for(int32 HitIdx=0; HitIdx<InHits.Num(); HitIdx++)
-		{
-			if(InHits[HitIdx].bBlockingHit)
-			{
-				NumBlocks++;
-			}
-		}
-		return NumBlocks;
-	}
-
-	/** Static utility function that returns the number of overlapping hits in array. */
-	static int32 GetNumOverlapHits(const TArray<FHitResult>& InHits)
-	{
-		return (InHits.Num() - GetNumBlockingHits(InHits));
-	}
-
-	/**
-	 * Get a copy of the HitResult with relevant information reversed.
-	 * For example when receiving a hit from another object, we reverse the normals.
-	 */
-	static FHitResult GetReversedHit(const FHitResult& Hit)
-	{
-		FHitResult Result(Hit);
-		Result.Normal = -Result.Normal;
-		Result.ImpactNormal = -Result.ImpactNormal;
-
-		int32 TempItem = Result.Item;
-		Result.Item = Result.MyItem;
-		Result.MyItem = TempItem;
-
-		FName TempBoneName = Result.BoneName;
-		Result.BoneName = Result.MyBoneName;
-		Result.MyBoneName = TempBoneName;
-		return Result;
-	}
-
-	FString ToString() const;
-};
-
-// All members of FHitResult are PODs.
-template<> struct TIsPODType<FHitResult> { enum { Value = true }; };
-
-template<>
-struct TStructOpsTypeTraits<FHitResult> : public TStructOpsTypeTraitsBase2<FHitResult>
-{
-	enum
-	{
-		WithNetSerializer = true,
-	};
-};
-
-
 /** Whether to teleport physics body or not */
 UENUM()
 enum class ETeleportType : uint8
@@ -2530,44 +2114,6 @@ enum class ETeleportType : uint8
 FORCEINLINE ETeleportType TeleportFlagToEnum(bool bTeleport) { return bTeleport ? ETeleportType::TeleportPhysics : ETeleportType::None; }
 FORCEINLINE bool TeleportEnumToFlag(ETeleportType Teleport) { return ETeleportType::TeleportPhysics == Teleport; }
 
-
-/** Structure containing information about one hit of an overlap test */
-USTRUCT()
-struct ENGINE_API FOverlapResult
-{
-	GENERATED_BODY()
-
-	UPROPERTY()
-	FActorInstanceHandle OverlapObjectHandle;
-
-	/** PrimitiveComponent that the check hit. */
-	UPROPERTY()
-	TWeakObjectPtr<class UPrimitiveComponent> Component;
-
-	/** This is the index of the overlapping item. 
-		For DestructibleComponents, this is the ChunkInfo index. 
-		For SkeletalMeshComponents this is the Body index or INDEX_NONE for single body */
-	int32 ItemIndex;
-
-	/** Utility to return the Actor that owns the Component that was hit */
-	AActor* GetActor() const;
-
-	/** Utility to return the Component that was hit */
-	UPrimitiveComponent* GetComponent() const;
-
-	/** Indicates if this hit was requesting a block - if false, was requesting a touch instead */
-	UPROPERTY()
-	uint32 bBlockingHit:1;
-
-	FOverlapResult()
-	{
-		FMemory::Memzero(this, sizeof(FOverlapResult));
-	}
-};
-
-// All members of FOverlapResult are PODs.
-template<> struct TIsPODType<FOverlapResult> { enum { Value = true }; };
-
 /** Structure containing information about minimum translation direction (MTD) */
 USTRUCT()
 struct ENGINE_API FMTDResult
@@ -3199,151 +2745,6 @@ struct FMeshNaniteSettings
 	}
 };
 
-/** Event used by AActor::TakeDamage and related functions */
-USTRUCT(BlueprintType)
-struct ENGINE_API FDamageEvent
-{
-	GENERATED_BODY()
-
-public:
-	/** Default constructor (no initialization). */
-	FDamageEvent() { }
-
-	FDamageEvent(FDamageEvent const& InDamageEvent)
-		: DamageTypeClass(InDamageEvent.DamageTypeClass)
-	{ }
-	
-	virtual ~FDamageEvent() { }
-
-	explicit FDamageEvent(TSubclassOf<class UDamageType> InDamageTypeClass)
-		: DamageTypeClass(InDamageTypeClass)
-	{ }
-
-	/** Optional DamageType for this event.  If nullptr, UDamageType will be assumed. */
-	UPROPERTY()
-	TSubclassOf<class UDamageType> DamageTypeClass;
-
-	/** ID for this class. NOTE this must be unique for all damage events. */
-	static const int32 ClassID = 0;
-
-	virtual int32 GetTypeID() const { return FDamageEvent::ClassID; }
-	virtual bool IsOfType(int32 InID) const { return FDamageEvent::ClassID == InID; };
-
-	/** This is for compatibility with old-style functions which want a unified set of hit data regardless of type of hit.  Ideally this will go away over time. */
-	virtual void GetBestHitInfo(AActor const* HitActor, AActor const* HitInstigator, struct FHitResult& OutHitInfo, FVector& OutImpulseDir) const;
-};
-
-/** Damage subclass that handles damage with a single impact location and source direction */
-USTRUCT()
-struct ENGINE_API FPointDamageEvent : public FDamageEvent
-{
-	GENERATED_BODY()
-
-	/** Actual damage done */
-	UPROPERTY()
-	float Damage;
-	
-	/** Direction the shot came from. Should be normalized. */
-	UPROPERTY()
-	FVector_NetQuantizeNormal ShotDirection;
-	
-	/** Describes the trace/location that caused this damage */
-	UPROPERTY()
-	struct FHitResult HitInfo;
-
-	FPointDamageEvent() : Damage(0.0f), ShotDirection(ForceInitToZero), HitInfo() {}
-	FPointDamageEvent(float InDamage, struct FHitResult const& InHitInfo, FVector const& InShotDirection, TSubclassOf<class UDamageType> InDamageTypeClass)
-		: FDamageEvent(InDamageTypeClass), Damage(InDamage), ShotDirection(InShotDirection), HitInfo(InHitInfo)
-	{}
-	
-	/** ID for this class. NOTE this must be unique for all damage events. */
-	static const int32 ClassID = 1;
-	
-	virtual int32 GetTypeID() const override { return FPointDamageEvent::ClassID; };
-	virtual bool IsOfType(int32 InID) const override { return (FPointDamageEvent::ClassID == InID) || FDamageEvent::IsOfType(InID); };
-
-	/** Simple API for common cases where we are happy to assume a single hit is expected, even though damage event may have multiple hits. */
-	virtual void GetBestHitInfo(AActor const* HitActor, AActor const* HitInstigator, struct FHitResult& OutHitInfo, FVector& OutImpulseDir) const override;
-};
-
-/** Parameters used to compute radial damage */
-USTRUCT(BlueprintType)
-struct ENGINE_API FRadialDamageParams
-{
-	GENERATED_BODY()
-
-	/** Max damage done */
-	UPROPERTY(EditAnywhere, BlueprintReadWrite, Category=RadialDamageParams)
-	float BaseDamage;
-
-	/** Damage will not fall below this if within range */
-	UPROPERTY(EditAnywhere, BlueprintReadWrite, Category=RadialDamageParams)
-	float MinimumDamage;
-	
-	/** Within InnerRadius, do max damage */
-	UPROPERTY(EditAnywhere, BlueprintReadWrite, Category=RadialDamageParams)
-	float InnerRadius;
-		
-	/** Outside OuterRadius, do no damage */
-	UPROPERTY(EditAnywhere, BlueprintReadWrite, Category=RadialDamageParams)
-	float OuterRadius;
-		
-	/** Describes amount of exponential damage falloff */
-	UPROPERTY(EditAnywhere, BlueprintReadWrite, Category=RadialDamageParams)
-	float DamageFalloff;
-
-	FRadialDamageParams()
-		: BaseDamage(0.f), MinimumDamage(0.f), InnerRadius(0.f), OuterRadius(0.f), DamageFalloff(1.f)
-	{}
-	FRadialDamageParams(float InBaseDamage, float InInnerRadius, float InOuterRadius, float InDamageFalloff)
-		: BaseDamage(InBaseDamage), MinimumDamage(0.f), InnerRadius(InInnerRadius), OuterRadius(InOuterRadius), DamageFalloff(InDamageFalloff)
-	{}
-	FRadialDamageParams(float InBaseDamage, float InMinimumDamage, float InInnerRadius, float InOuterRadius, float InDamageFalloff)
-		: BaseDamage(InBaseDamage), MinimumDamage(InMinimumDamage), InnerRadius(InInnerRadius), OuterRadius(InOuterRadius), DamageFalloff(InDamageFalloff)
-	{}
-	FRadialDamageParams(float InBaseDamage, float InRadius)
-		: BaseDamage(InBaseDamage), MinimumDamage(0.f), InnerRadius(0.f), OuterRadius(InRadius), DamageFalloff(1.f)
-	{}
-
-	/** Returns damage done at a certain distance */
-	float GetDamageScale(float DistanceFromEpicenter) const;
-
-	/** Return outermost radius of the damage area. Protects against malformed data. */
-	float GetMaxRadius() const { return FMath::Max( FMath::Max(InnerRadius, OuterRadius), 0.f ); }
-};
-
-/** Damage subclass that handles damage with a source location and falloff radius */
-USTRUCT()
-struct ENGINE_API FRadialDamageEvent : public FDamageEvent
-{
-	GENERATED_BODY()
-
-	/** Static parameters describing damage falloff math */
-	UPROPERTY()
-	FRadialDamageParams Params;
-	
-	/** Location of origin point */
-	UPROPERTY()
-	FVector Origin;
-
-	/** Hit reslts of specific impacts */
-	UPROPERTY()
-	TArray<struct FHitResult> ComponentHits;
-
-	/** ID for this class. NOTE this must be unique for all damage events. */
-	static const int32 ClassID = 2;
-
-	virtual int32 GetTypeID() const override { return FRadialDamageEvent::ClassID; };
-	virtual bool IsOfType(int32 InID) const override { return (FRadialDamageEvent::ClassID == InID) || FDamageEvent::IsOfType(InID); };
-
-	/** Simple API for common cases where we are happy to assume a single hit is expected, even though damage event may have multiple hits. */
-	virtual void GetBestHitInfo(AActor const* HitActor, AActor const* HitInstigator, struct FHitResult& OutHitInfo, FVector& OutImpulseDir) const override;
-
-	FRadialDamageEvent()
-		: Origin(ForceInitToZero)
-	{}
-};
-
 /** The network role of an actor on a local/remote network context */
 UENUM()
 enum ENetRole
@@ -3506,235 +2907,6 @@ private:
 	}
 };
 
-/** Describes rules for network replicating a vector efficiently */
-UENUM()
-enum class EVectorQuantization : uint8
-{
-	/** Each vector component will be rounded to the nearest whole number. */
-	RoundWholeNumber,
-	/** Each vector component will be rounded, preserving one decimal place. */
-	RoundOneDecimal,
-	/** Each vector component will be rounded, preserving two decimal places. */
-	RoundTwoDecimals
-};
-
-/** Describes rules for network replicating a vector efficiently */
-UENUM()
-enum class ERotatorQuantization : uint8
-{
-	/** The rotator will be compressed to 8 bits per component. */
-	ByteComponents,
-	/** The rotator will be compressed to 16 bits per component. */
-	ShortComponents
-};
-
-/** Replicated movement data of our RootComponent.
-  * Struct used for efficient replication as velocity and location are generally replicated together (this saves a repindex) 
-  * and velocity.Z is commonly zero (most position replications are for walking pawns). 
-  */
-USTRUCT()
-struct ENGINE_API FRepMovement
-{
-	GENERATED_BODY()
-
-	/** Velocity of component in world space */
-	UPROPERTY(Transient)
-	FVector LinearVelocity;
-
-	/** Velocity of rotation for component */
-	UPROPERTY(Transient)
-	FVector AngularVelocity;
-	
-	/** Location in world space */
-	UPROPERTY(Transient)
-	FVector Location;
-
-	/** Current rotation */
-	UPROPERTY(Transient)
-	FRotator Rotation;
-
-	/** If set, RootComponent should be sleeping. */
-	UPROPERTY(Transient)
-	uint8 bSimulatedPhysicSleep : 1;
-
-	/** If set, additional physic data (angular velocity) will be replicated. */
-	UPROPERTY(Transient)
-	uint8 bRepPhysics : 1;
-
-	/** Server physics step */
-	UPROPERTY(Transient)
-	int32 ServerFrame;
-
-	/** ID assigned by server used to ensure determinism by physics. */
-	UPROPERTY(Transient)
-	int32 ServerPhysicsHandle = INDEX_NONE;
-
-	/** Allows tuning the compression level for the replicated location vector. You should only need to change this from the default if you see visual artifacts. */
-	UPROPERTY(EditDefaultsOnly, Category=Replication, AdvancedDisplay)
-	EVectorQuantization LocationQuantizationLevel;
-
-	/** Allows tuning the compression level for the replicated velocity vectors. You should only need to change this from the default if you see visual artifacts. */
-	UPROPERTY(EditDefaultsOnly, Category=Replication, AdvancedDisplay)
-	EVectorQuantization VelocityQuantizationLevel;
-
-	/** Allows tuning the compression level for replicated rotation. You should only need to change this from the default if you see visual artifacts. */
-	UPROPERTY(EditDefaultsOnly, Category=Replication, AdvancedDisplay)
-	ERotatorQuantization RotationQuantizationLevel;
-
-	FRepMovement();
-
-	bool SerializeQuantizedVector(FArchive& Ar, FVector& Vector, EVectorQuantization QuantizationLevel)
-	{
-		// Since FRepMovement used to use FVector_NetQuantize100, we're allowing enough bits per component
-		// regardless of the quantization level so that we can still support at least the same maximum magnitude
-		// (2^30 / 100, or ~10 million).
-		// This uses no inherent extra bandwidth since we're still using the same number of bits to store the
-		// bits-per-component value. Of course, larger magnitudes will still use more bandwidth,
-		// as has always been the case.
-		switch(QuantizationLevel)
-		{
-			case EVectorQuantization::RoundTwoDecimals:
-			{
-				return SerializePackedVector<100, 30>(Vector, Ar);
-			}
-
-			case EVectorQuantization::RoundOneDecimal:
-			{
-				return SerializePackedVector<10, 27>(Vector, Ar);
-			}
-
-			default:
-			{
-				return SerializePackedVector<1, 24>(Vector, Ar);
-			}
-		}
-	}
-
-	bool NetSerialize(FArchive& Ar, class UPackageMap* Map, bool& bOutSuccess);
-
-	void FillFrom(const struct FRigidBodyState& RBState, const AActor* const Actor = nullptr, int32 InServerFrame = 0)
-	{
-		Location = RebaseOntoZeroOrigin(RBState.Position, Actor);
-		Rotation = RBState.Quaternion.Rotator();
-		LinearVelocity = RBState.LinVel;
-		AngularVelocity = RBState.AngVel;
-		bSimulatedPhysicSleep = (RBState.Flags & ERigidBodyFlags::Sleeping) != 0;
-		bRepPhysics = true;
-		ServerFrame = InServerFrame;
-	}
-
-	void CopyTo(struct FRigidBodyState& RBState, const AActor* const Actor = nullptr) const
-	{
-		RBState.Position = RebaseOntoLocalOrigin(Location, Actor);
-		RBState.Quaternion = Rotation.Quaternion();
-		RBState.LinVel = LinearVelocity;
-		RBState.AngVel = AngularVelocity;
-		RBState.Flags = (decltype(FRigidBodyState::Flags))(bSimulatedPhysicSleep ? ERigidBodyFlags::Sleeping : ERigidBodyFlags::None) | ERigidBodyFlags::NeedsUpdate;
-	}
-
-	bool operator==(const FRepMovement& Other) const
-	{
-		if ( LinearVelocity != Other.LinearVelocity )
-		{
-			return false;
-		}
-
-		if ( AngularVelocity != Other.AngularVelocity )
-		{
-			return false;
-		}
-
-		if ( Location != Other.Location )
-		{
-			return false;
-		}
-
-		if ( Rotation != Other.Rotation )
-		{
-			return false;
-		}
-
-		if ( bSimulatedPhysicSleep != Other.bSimulatedPhysicSleep )
-		{
-			return false;
-		}
-
-		if ( bRepPhysics != Other.bRepPhysics )
-		{
-			return false;
-		}
-
-		return true;
-	}
-
-	bool operator!=(const FRepMovement& Other) const
-	{
-		return !(*this == Other);
-	}
-
-	/** True if multiplayer rebasing is enabled, corresponds to p.EnableMultiplayerWorldOriginRebasing console variable */
-	static int32 EnableMultiplayerWorldOriginRebasing;
-
-	/** Rebase zero-origin position onto local world origin value. */
-	static FVector RebaseOntoLocalOrigin(const FVector& Location, const FIntVector& LocalOrigin);
-
-	/** Rebase local-origin position onto zero world origin value. */
-	static FVector RebaseOntoZeroOrigin(const FVector& Location, const FIntVector& LocalOrigin);
-
-	/** Rebase zero-origin position onto an Actor's local world origin. */
-	static FVector RebaseOntoLocalOrigin(const FVector& Location, const AActor* const WorldContextActor);
-
-	/** Rebase an Actor's local-origin position onto zero world origin value. */
-	static FVector RebaseOntoZeroOrigin(const FVector& Location, const AActor* const WorldContextActor);
-
-	/** Rebase zero-origin position onto local world origin value based on an actor component's world. */
-	static FVector RebaseOntoLocalOrigin(const FVector& Location, const class UActorComponent* const WorldContextActorComponent);
-
-	/** Rebase local-origin position onto zero world origin value based on an actor component's world.*/
-	static FVector RebaseOntoZeroOrigin(const FVector& Location, const class UActorComponent* const WorldContextActorComponent);
-};
-
-
-/** Handles attachment replication to clients.  */
-USTRUCT()
-struct FRepAttachment
-{
-	GENERATED_BODY()
-
-	/** Actor we are attached to, movement replication will not happen while AttachParent is non-nullptr */
-	UPROPERTY()
-	TObjectPtr<class AActor> AttachParent;
-
-	/** Location offset from attach parent */
-	UPROPERTY()
-	FVector_NetQuantize100 LocationOffset;
-
-	/** Scale relative to attach parent */
-	UPROPERTY()
-	FVector_NetQuantize100 RelativeScale3D;
-
-	/** Rotation offset from attach parent */
-	UPROPERTY()
-	FRotator RotationOffset;
-
-	/** Specific socket we are attached to */
-	UPROPERTY()
-	FName AttachSocket;
-
-	/** Specific component we are attached to */
-	UPROPERTY()
-	TObjectPtr<class USceneComponent> AttachComponent;
-
-	FRepAttachment()
-		: AttachParent(nullptr)
-		, LocationOffset(ForceInit)
-		, RelativeScale3D(ForceInit)
-		, RotationOffset(ForceInit)
-		, AttachSocket(NAME_None)
-		, AttachComponent(nullptr)
-	{ }
-};
-
 /**
  * Controls behavior of WalkableSlopeOverride, determining how to affect walkability of surfaces for Characters.
  * @see FWalkableSlopeOverride
@@ -3887,16 +3059,6 @@ private:
 
 template<> struct TIsPODType<FWalkableSlopeOverride> { enum { Value = true }; };
 
-template<>
-struct TStructOpsTypeTraits<FRepMovement> : public TStructOpsTypeTraitsBase2<FRepMovement>
-{
-	enum 
-	{
-		WithNetSerializer = true,
-		WithNetSharedSerialization = true,
-	};
-};
-
 /** Structure to hold and pass around transient flags used during replication. */
 struct FReplicationFlags
 {
diff --git a/Engine/Source/Runtime/Engine/Classes/Engine/HitResult.h b/Engine/Source/Runtime/Engine/Classes/Engine/HitResult.h
new file mode 100644
--- /dev/null
+++ b/Engine/Source/Runtime/Engine/Classes/Engine/HitResult.h
@@ -0,0 +1,300 @@
+// Copyright Epic Games, Inc. All Rights Reserved.
+
+#pragma once
+
+#include "CoreMinimal.h"
+#include "Engine/NetSerialization.h"
+#include "Engine/ActorInstanceHandle.h"
+#include "HitResult.generated.h"
+
+class AActor;
+class UPackageMap;
+class UPhysicalMaterial;
+class UPrimitiveComponent;
+
+/**
+ * Structure containing information about one hit of a trace, such as point of impact and surface normal at that point.
+ */
+USTRUCT(BlueprintType, meta = (HasNativeBreak = "Engine.GameplayStatics.BreakHitResult", HasNativeMake = "Engine.GameplayStatics.MakeHitResult"))
+struct ENGINE_API FHitResult
+{
+	GENERATED_BODY()
+
+	/** Face index we hit (for complex hits with triangle meshes). */
+	UPROPERTY()
+	int32 FaceIndex;
+
+	/**
+	 * 'Time' of impact along trace direction (ranging from 0.0 to 1.0) if there is a hit, indicating time between TraceStart and TraceEnd.
+	 * For swept movement (but not queries) this may be pulled back slightly from the actual time of impact, to prevent precision problems with adjacent geometry.
+	 */
+	UPROPERTY()
+	float Time;
+	 
+	/** The distance from the TraceStart to the Location in world space. This value is 0 if there was an initial overlap (trace started inside another colliding object). */
+	UPROPERTY()
+	float Distance;
+	
+	/**
+	 * The location in world space where the moving shape would end up against the impacted object, if there is a hit. Equal to the point of impact for line tests.
+	 * Example: for a sphere trace test, this is the point where the center of the sphere would be located when it touched the other object.
+	 * For swept movement (but not queries) this may not equal the final location of the shape since hits are pulled back slightly to prevent precision issues from overlapping another surface.
+	 */
+	UPROPERTY()
+	FVector_NetQuantize Location;
+
+	/**
+	 * Location in world space of the actual contact of the trace shape (box, sphere, ray, etc) with the impacted object.
+	 * Example: for a sphere trace test, this is the point where the surface of the sphere touches the other object.
+	 * @note: In the case of initial overlap (bStartPenetrating=true), ImpactPoint will be the same as Location because there is no meaningful single impact point to report.
+	 */
+	UPROPERTY()
+	FVector_NetQuantize ImpactPoint;
+
+	/**
+	 * Normal of the hit in world space, for the object that was swept. Equal to ImpactNormal for line tests.
+	 * This is computed for capsules and spheres, otherwise it will be the same as ImpactNormal.
+	 * Example: for a sphere trace test, this is a normalized vector pointing in towards the center of the sphere at the point of impact.
+	 */
+	UPROPERTY()
+	FVector_NetQuantizeNormal Normal;
+
+	/**
+	 * Normal of the hit in world space, for the object that was hit by the sweep, if any.
+	 * For example if a sphere hits a flat plane, this is a normalized vector pointing out from the plane.
+	 * In the case of impact with a corner or edge of a surface, usually the "most opposing" normal (opposed to the query direction) is chosen.
+	 */
+	UPROPERTY()
+	FVector_NetQuantizeNormal ImpactNormal;
+
+	/**
+	 * Start location of the trace.
+	 * For example if a sphere is swept against the world, this is the starting location of the center of the sphere.
+	 */
+	UPROPERTY()
+	FVector_NetQuantize TraceStart;
+
+	/**
+	 * End location of the trace; this is NOT where the impact occurred (if any), but the furthest point in the attempted sweep.
+	 * For example if a sphere is swept against the world, this would be the center of the sphere if there was no blocking hit.
+	 */
+	UPROPERTY()
+	FVector_NetQuantize TraceEnd;
+
+	/**
+	  * If this test started in penetration (bStartPenetrating is true) and a depenetration vector can be computed,
+	  * this value is the distance along Normal that will result in moving out of penetration.
+	  * If the distance cannot be computed, this distance will be zero.
+	  */
+	UPROPERTY()
+	float PenetrationDepth;
+
+	/** If the hit result is from a collision this will have extra info about the item that hit the second item. */
+	UPROPERTY()
+	int32 MyItem;
+
+	/** Extra data about item that was hit (hit primitive specific). */
+	UPROPERTY()
+	int32 Item;
+
+	/** Index to item that was hit, also hit primitive specific. */
+	UPROPERTY()
+	uint8 ElementIndex;
+
+	/** Indicates if this hit was a result of blocking collision. If false, there was no hit or it was an overlap/touch instead. */
+	UPROPERTY()
+	uint8 bBlockingHit : 1;
+
+	/**
+	 * Whether the trace started in penetration, i.e. with an initial blocking overlap.
+	 * In the case of penetration, if PenetrationDepth > 0.f, then it will represent the distance along the Normal vector that will result in
+	 * minimal contact between the swept shape and the object that was hit. In this case, ImpactNormal will be the normal opposed to movement at that location
+	 * (ie, Normal may not equal ImpactNormal). ImpactPoint will be the same as Location, since there is no single impact point to report.
+	 */
+	UPROPERTY()
+	uint8 bStartPenetrating : 1;
+
+	/**
+	 * Physical material that was hit.
+	 * @note Must set bReturnPhysicalMaterial on the swept PrimitiveComponent or in the query params for this to be returned.
+	 */
+	UPROPERTY()
+	TWeakObjectPtr<UPhysicalMaterial> PhysMaterial;
+
+	/** Handle to the object hit by the trace. */
+	UPROPERTY()
+	FActorInstanceHandle HitObjectHandle;
+
+	/** PrimitiveComponent hit by the trace. */
+	UPROPERTY()
+	TWeakObjectPtr<UPrimitiveComponent> Component;
+
+	/** Name of bone we hit (for skeletal meshes). */
+	UPROPERTY()
+	FName BoneName;
+
+	/** Name of the _my_ bone which took part in hit event (in case of two skeletal meshes colliding). */
+	UPROPERTY()
+	FName MyBoneName;
+
+
+	FHitResult()
+	{
+		Init();
+	}
+	
+	explicit FHitResult(float InTime)
+	{
+		Init();
+		Time = InTime;
+	}
+
+	explicit FHitResult(EForceInit InInit)
+	{
+		Init();
+	}
+
+	explicit FHitResult(ENoInit NoInit)
+	{
+	}
+
+	explicit FHitResult(FVector Start, FVector End)
+	{
+		Init(Start, End);
+	}
+
+	/** Initialize empty hit result with given time. */
+	FORCEINLINE void Init()
+	{
+		FMemory::Memzero(this, sizeof(FHitResult));
+		HitObjectHandle = FActorInstanceHandle();
+		Time = 1.f;
+		MyItem = INDEX_NONE;
+	}
+
+	/** Initialize empty hit result with given time, TraceStart, and TraceEnd */
+	FORCEINLINE void Init(FVector Start, FVector End)
+	{
+		FMemory::Memzero(this, sizeof(FHitResult));
+		HitObjectHandle = FActorInstanceHandle();
+		Time = 1.f;
+		TraceStart = Start;
+		TraceEnd = End;
+		MyItem = INDEX_NONE;
+	}
+
+	/** Ctor for easily creating "fake" hits from limited data. */
+	FHitResult(AActor* InActor, UPrimitiveComponent* InComponent, FVector const& HitLoc, FVector const& HitNorm);
+ 
+	/** Reset hit result while optionally saving TraceStart and TraceEnd. */
+	FORCEINLINE void Reset(float InTime = 1.f, bool bPreserveTraceData = true)
+	{
+		const FVector SavedTraceStart = TraceStart;
+		const FVector SavedTraceEnd = TraceEnd;
+		Init();
+		Time = InTime;
+		if (bPreserveTraceData)
+		{
+			TraceStart = SavedTraceStart;
+			TraceEnd = SavedTraceEnd;
+		}
+	}
+
+	/** Utility to return the Actor that owns the Component that was hit. */
+	FORCEINLINE AActor* GetActor() const
+	{
+		return HitObjectHandle.FetchActor();
+	}
+
+	FORCEINLINE FActorInstanceHandle GetHitObjectHandle() const
+	{
+		return HitObjectHandle;
+	}
+
+	FORCEINLINE bool HasValidHitObjectHandle() const
+	{
+		return HitObjectHandle.IsValid();
+	}
+
+	/** Utility to return the Component that was hit. */
+	FORCEINLINE UPrimitiveComponent* GetComponent() const
+	{
+		return Component.Get();
+	}
+
+	/** Optimized serialize function */
+	bool NetSerialize(FArchive& Ar, UPackageMap* Map, bool& bOutSuccess);
+
+	/** Return true if there was a blocking hit that was not caused by starting in penetration. */
+	FORCEINLINE bool IsValidBlockingHit() const
+	{
+		return bBlockingHit && !bStartPenetrating;
+	}
+
+	/** Static utility function that returns the first 'blocking' hit in an array of results. */
+	static FHitResult* GetFirstBlockingHit(TArray<FHitResult>& InHits)
+	{
+		for(int32 HitIdx=0; HitIdx<InHits.Num(); HitIdx++)
+		{
+			if(InHits[HitIdx].bBlockingHit)
+			{
+				return &InHits[HitIdx];
+			}
+		}
+		return nullptr;
+	}
+
+	/** Static utility function that returns the number of blocking hits in array. */
+	static int32 GetNumBlockingHits(const TArray<FHitResult>& InHits)
+	{
+		int32 NumBlocks = 0;
+		for(int32 HitIdx=0; HitIdx<InHits.Num(); HitIdx++)
+		{
+			if(InHits[HitIdx].bBlockingHit)
+			{
+				NumBlocks++;
+			}
+		}
+		return NumBlocks;
+	}
+
+	/** Static utility function that returns the number of overlapping hits in array. */
+	static int32 GetNumOverlapHits(const TArray<FHitResult>& InHits)
+	{
+		return (InHits.Num() - GetNumBlockingHits(InHits));
+	}
+
+	/**
+	 * Get a copy of the HitResult with relevant information reversed.
+	 * For example when receiving a hit from another object, we reverse the normals.
+	 */
+	static FHitResult GetReversedHit(const FHitResult& Hit)
+	{
+		FHitResult Result(Hit);
+		Result.Normal = -Result.Normal;
+		Result.ImpactNormal = -Result.ImpactNormal;
+
+		int32 TempItem = Result.Item;
+		Result.Item = Result.MyItem;
+		Result.MyItem = TempItem;
+
+		FName TempBoneName = Result.BoneName;
+		Result.BoneName = Result.MyBoneName;
+		Result.MyBoneName = TempBoneName;
+		return Result;
+	}
+
+	FString ToString() const;
+};
+
+// All members of FHitResult are PODs.
+template<> struct TIsPODType<FHitResult> { enum { Value = true }; };
+
+template<>
+struct TStructOpsTypeTraits<FHitResult> : public TStructOpsTypeTraitsBase2<FHitResult>
+{
+	enum
+	{
+		WithNetSerializer = true,
+	};
+};
diff --git a/Engine/Source/Runtime/Engine/Classes/Engine/OverlapResult.h b/Engine/Source/Runtime/Engine/Classes/Engine/OverlapResult.h
new file mode 100644
--- /dev/null
+++ b/Engine/Source/Runtime/Engine/Classes/Engine/OverlapResult.h
@@ -0,0 +1,53 @@
+
+#pragma once
+
+#include "CoreMinimal.h"
+#include "Engine/ActorInstanceHandle.h"
+#include "OverlapResult.generated.h"
+
+/** Structure containing information about one hit of an overlap test */
+USTRUCT()
+struct ENGINE_API FOverlapResult
+{
+	GENERATED_BODY()
+
+	UPROPERTY()
+	FActorInstanceHandle OverlapObjectHandle;
+
+	/** PrimitiveComponent that the check hit. */
+	UPROPERTY()
+	TWeakObjectPtr<class UPrimitiveComponent> Component;
+
+	/** This is the index of the overlapping item.
+		For DestructibleComponents, this is the ChunkInfo index.
+		For SkeletalMeshComponents this is the Body index or INDEX_NONE for single body */
+	int32 ItemIndex;
+
+	/** Utility to return the Actor that owns the Component that was hit */
+	AActor* GetActor() const;
+
+	/** Utility to return the Component that was hit */
+	UPrimitiveComponent* GetComponent() const;
+
+	/** Indicates if this hit was requesting a block - if false, was requesting a touch instead */
+	UPROPERTY()
+	uint32 bBlockingHit : 1;
+
+	FOverlapResult()
+	{
+		FMemory::Memzero(this, sizeof(FOverlapResult));
+	}
+};
+
+// All members of FOverlapResult are PODs.
+template<> struct TIsPODType<FOverlapResult> { enum { Value = true }; };
+
+inline AActor* FOverlapResult::GetActor() const
+{
+	return OverlapObjectHandle.FetchActor();
+}
+
+inline UPrimitiveComponent* FOverlapResult::GetComponent() const
+{
+	return Component.Get();
+}
diff --git a/Engine/Source/Runtime/Engine/Classes/Engine/ReplicatedState.h b/Engine/Source/Runtime/Engine/Classes/Engine/ReplicatedState.h
new file mode 100644
--- /dev/null
+++ b/Engine/Source/Runtime/Engine/Classes/Engine/ReplicatedState.h
@@ -0,0 +1,262 @@
+// Copyright Epic Games, Inc. All Rights Reserved.
+
+#pragma once
+
+#include "CoreMinimal.h"
+#include "UObject/ObjectMacros.h"
+#include "Engine/NetSerialization.h"
+#include "ReplicatedState.generated.h"
+
+/** Describes rules for network replicating a vector efficiently */
+UENUM()
+enum class EVectorQuantization : uint8
+{
+	/** Each vector component will be rounded to the nearest whole number. */
+	RoundWholeNumber,
+	/** Each vector component will be rounded, preserving one decimal place. */
+	RoundOneDecimal,
+	/** Each vector component will be rounded, preserving two decimal places. */
+	RoundTwoDecimals
+};
+
+/** Describes rules for network replicating a vector efficiently */
+UENUM()
+enum class ERotatorQuantization : uint8
+{
+	/** The rotator will be compressed to 8 bits per component. */
+	ByteComponents,
+	/** The rotator will be compressed to 16 bits per component. */
+	ShortComponents
+};
+
+/** Handles attachment replication to clients.  */
+USTRUCT()
+struct FRepAttachment
+{
+	GENERATED_BODY()
+
+	/** Actor we are attached to, movement replication will not happen while AttachParent is non-nullptr */
+	UPROPERTY()
+	TObjectPtr<class AActor> AttachParent;
+
+	/** Location offset from attach parent */
+	UPROPERTY()
+	FVector_NetQuantize100 LocationOffset;
+
+	/** Scale relative to attach parent */
+	UPROPERTY()
+	FVector_NetQuantize100 RelativeScale3D;
+
+	/** Rotation offset from attach parent */
+	UPROPERTY()
+	FRotator RotationOffset;
+
+	/** Specific socket we are attached to */
+	UPROPERTY()
+	FName AttachSocket;
+
+	/** Specific component we are attached to */
+	UPROPERTY()
+	TObjectPtr<class USceneComponent> AttachComponent;
+
+	FRepAttachment()
+		: AttachParent(nullptr)
+		, LocationOffset(ForceInit)
+		, RelativeScale3D(ForceInit)
+		, RotationOffset(ForceInit)
+		, AttachSocket(NAME_None)
+		, AttachComponent(nullptr)
+	{ }
+};
+
+/** Describes extra state about a specific rigid body */
+namespace ERigidBodyFlags
+{
+	enum Type
+	{
+		None = 0x00,
+		Sleeping = 0x01,
+		NeedsUpdate = 0x02,
+	};
+}
+
+/** Describes the physical state of a rigid body. */
+USTRUCT()
+struct FRigidBodyState
+{
+	GENERATED_BODY()
+
+	UPROPERTY()
+	FVector_NetQuantize100 Position;
+
+	UPROPERTY()
+	FQuat Quaternion;
+
+	UPROPERTY()
+	FVector_NetQuantize100 LinVel;
+
+	UPROPERTY()
+	FVector_NetQuantize100 AngVel;
+
+	UPROPERTY()
+	uint8 Flags;
+
+	FRigidBodyState()
+		: Position(ForceInit)
+		, Quaternion(ForceInit)
+		, LinVel(ForceInit)
+		, AngVel(ForceInit)
+		, Flags(0)
+	{ }
+};
+
+/** Replicated movement data of our RootComponent.
+  * Struct used for efficient replication as velocity and location are generally replicated together (this saves a repindex) 
+  * and velocity.Z is commonly zero (most position replications are for walking pawns). 
+  */
+USTRUCT()
+struct ENGINE_API FRepMovement
+{
+	GENERATED_BODY()
+
+	/** Velocity of component in world space */
+	UPROPERTY(Transient)
+	FVector LinearVelocity;
+
+	/** Velocity of rotation for component */
+	UPROPERTY(Transient)
+	FVector AngularVelocity;
+	
+	/** Location in world space */
+	UPROPERTY(Transient)
+	FVector Location;
+
+	/** Current rotation */
+	UPROPERTY(Transient)
+	FRotator Rotation;
+
+	/** If set, RootComponent should be sleeping. */
+	UPROPERTY(Transient)
+	uint8 bSimulatedPhysicSleep : 1;
+
+	/** If set, additional physic data (angular velocity) will be replicated. */
+	UPROPERTY(Transient)
+	uint8 bRepPhysics : 1;
+
+	/** Server physics step */
+	UPROPERTY(Transient)
+	int32 ServerFrame;
+
+	/** ID assigned by server used to ensure determinism by physics. */
+	UPROPERTY(Transient)
+	int32 ServerPhysicsHandle = INDEX_NONE;
+
+	/** Allows tuning the compression level for the replicated location vector. You should only need to change this from the default if you see visual artifacts. */
+	UPROPERTY(EditDefaultsOnly, Category=Replication, AdvancedDisplay)
+	EVectorQuantization LocationQuantizationLevel;
+
+	/** Allows tuning the compression level for the replicated velocity vectors. You should only need to change this from the default if you see visual artifacts. */
+	UPROPERTY(EditDefaultsOnly, Category=Replication, AdvancedDisplay)
+	EVectorQuantization VelocityQuantizationLevel;
+
+	/** Allows tuning the compression level for replicated rotation. You should only need to change this from the default if you see visual artifacts. */
+	UPROPERTY(EditDefaultsOnly, Category=Replication, AdvancedDisplay)
+	ERotatorQuantization RotationQuantizationLevel;
+
+	FRepMovement();
+
+	bool SerializeQuantizedVector(FArchive& Ar, FVector& Vector, EVectorQuantization QuantizationLevel);
+
+	bool NetSerialize(FArchive& Ar, class UPackageMap* Map, bool& bOutSuccess);
+
+	void FillFrom(const FRigidBodyState& RBState, const AActor* const Actor = nullptr, int32 InServerFrame = 0)
+	{
+		Location = RebaseOntoZeroOrigin(RBState.Position, Actor);
+		Rotation = RBState.Quaternion.Rotator();
+		LinearVelocity = RBState.LinVel;
+		AngularVelocity = RBState.AngVel;
+		bSimulatedPhysicSleep = (RBState.Flags & ERigidBodyFlags::Sleeping) != 0;
+		bRepPhysics = true;
+		ServerFrame = InServerFrame;
+	}
+
+	void CopyTo(FRigidBodyState& RBState, const AActor* const Actor = nullptr) const
+	{
+		RBState.Position = RebaseOntoLocalOrigin(Location, Actor);
+		RBState.Quaternion = Rotation.Quaternion();
+		RBState.LinVel = LinearVelocity;
+		RBState.AngVel = AngularVelocity;
+		RBState.Flags = (decltype(FRigidBodyState::Flags))(bSimulatedPhysicSleep ? ERigidBodyFlags::Sleeping : ERigidBodyFlags::None) | ERigidBodyFlags::NeedsUpdate;
+	}
+
+	bool operator==(const FRepMovement& Other) const
+	{
+		if ( LinearVelocity != Other.LinearVelocity )
+		{
+			return false;
+		}
+
+		if ( AngularVelocity != Other.AngularVelocity )
+		{
+			return false;
+		}
+
+		if ( Location != Other.Location )
+		{
+			return false;
+		}
+
+		if ( Rotation != Other.Rotation )
+		{
+			return false;
+		}
+
+		if ( bSimulatedPhysicSleep != Other.bSimulatedPhysicSleep )
+		{
+			return false;
+		}
+
+		if ( bRepPhysics != Other.bRepPhysics )
+		{
+			return false;
+		}
+
+		return true;
+	}
+
+	bool operator!=(const FRepMovement& Other) const
+	{
+		return !(*this == Other);
+	}
+
+	/** True if multiplayer rebasing is enabled, corresponds to p.EnableMultiplayerWorldOriginRebasing console variable */
+	static int32 EnableMultiplayerWorldOriginRebasing;
+
+	/** Rebase zero-origin position onto local world origin value. */
+	static FVector RebaseOntoLocalOrigin(const FVector& Location, const FIntVector& LocalOrigin);
+
+	/** Rebase local-origin position onto zero world origin value. */
+	static FVector RebaseOntoZeroOrigin(const FVector& Location, const FIntVector& LocalOrigin);
+
+	/** Rebase zero-origin position onto an Actor's local world origin. */
+	static FVector RebaseOntoLocalOrigin(const FVector& Location, const AActor* const WorldContextActor);
+
+	/** Rebase an Actor's local-origin position onto zero world origin value. */
+	static FVector RebaseOntoZeroOrigin(const FVector& Location, const AActor* const WorldContextActor);
+
+	/** Rebase zero-origin position onto local world origin value based on an actor component's world. */
+	static FVector RebaseOntoLocalOrigin(const FVector& Location, const class UActorComponent* const WorldContextActorComponent);
+
+	/** Rebase local-origin position onto zero world origin value based on an actor component's world.*/
+	static FVector RebaseOntoZeroOrigin(const FVector& Location, const class UActorComponent* const WorldContextActorComponent);
+};
+
+template<>
+struct TStructOpsTypeTraits<FRepMovement> : public TStructOpsTypeTraitsBase2<FRepMovement>
+{
+	enum
+	{
+		WithNetSerializer = true,
+		WithNetSharedSerialization = true,
+	};
+};
diff --git a/Engine/Source/Runtime/Engine/Classes/GameFramework/Actor.h b/Engine/Source/Runtime/Engine/Classes/GameFramework/Actor.h
--- a/Engine/Source/Runtime/Engine/Classes/GameFramework/Actor.h
+++ b/Engine/Source/Runtime/Engine/Classes/GameFramework/Actor.h
@@ -18,6 +18,8 @@
 #include "Misc/ITransaction.h"
 #include "Engine/Level.h"
 #include "Net/Core/Misc/NetSubObjectRegistry.h"
+#include "Engine/HitResult.h"
+#include "Engine/ReplicatedState.h"
 
 #if WITH_EDITOR
 #include "WorldPartition/DataLayer/ActorDataLayer.h"
diff --git a/Engine/Source/Runtime/Engine/Private/Actor.cpp b/Engine/Source/Runtime/Engine/Private/Actor.cpp
--- a/Engine/Source/Runtime/Engine/Private/Actor.cpp
+++ b/Engine/Source/Runtime/Engine/Private/Actor.cpp
@@ -51,6 +51,8 @@
 #include "GameFramework/InputSettings.h"
 #include "Algo/AnyOf.h"
 #include "PrimitiveSceneProxy.h"
+#include "Engine/DamageEvents.h"
+#include "Engine/OverlapInfo.h"
 
 #if WITH_EDITOR
 #include "FoliageHelper.h"
diff --git a/Engine/Source/Runtime/Engine/Private/Character.cpp b/Engine/Source/Runtime/Engine/Private/Character.cpp
--- a/Engine/Source/Runtime/Engine/Private/Character.cpp
+++ b/Engine/Source/Runtime/Engine/Private/Character.cpp
@@ -17,6 +17,7 @@
 #include "DisplayDebugHelpers.h"
 #include "Engine/Canvas.h"
 #include "Animation/AnimInstance.h"
+#include "Engine/DamageEvents.h"
 
 DEFINE_LOG_CATEGORY_STATIC(LogCharacter, Log, All);
 
diff --git a/Engine/Source/Runtime/Engine/Private/CharacterMovementComponentAsync.cpp b/Engine/Source/Runtime/Engine/Private/CharacterMovementComponentAsync.cpp
--- a/Engine/Source/Runtime/Engine/Private/CharacterMovementComponentAsync.cpp
+++ b/Engine/Source/Runtime/Engine/Private/CharacterMovementComponentAsync.cpp
@@ -5,6 +5,7 @@
 #include "PhysicsProxy/SingleParticlePhysicsProxy.h"
 #include "Chaos/SpatialAccelerationCollection.h"
 #include "PBDRigidsSolver.h"
+#include "Engine/ScopedMovementUpdate.h"
 
 void FCharacterMovementComponentAsyncInput::Simulate(const float DeltaSeconds, FCharacterMovementComponentAsyncOutput& Output) const
 {
diff --git a/Engine/Source/Runtime/Engine/Private/CheatManager.cpp b/Engine/Source/Runtime/Engine/Private/CheatManager.cpp
--- a/Engine/Source/Runtime/Engine/Private/CheatManager.cpp
+++ b/Engine/Source/Runtime/Engine/Private/CheatManager.cpp
@@ -39,6 +39,7 @@
 #include "Misc/CoreDelegates.h"
 #include "Engine/NetConnection.h"
 #include "HAL/PlatformApplicationMisc.h"
+#include "Engine/DamageEvents.h"
 
 DEFINE_LOG_CATEGORY_STATIC(LogCheatManager, Log, All);
 
diff --git a/Engine/Source/Runtime/Engine/Private/Collision/Collision.cpp b/Engine/Source/Runtime/Engine/Private/Collision/Collision.cpp
--- a/Engine/Source/Runtime/Engine/Private/Collision/Collision.cpp
+++ b/Engine/Source/Runtime/Engine/Private/Collision/Collision.cpp
@@ -26,158 +26,6 @@
 // TEMP until crash is fixed in IsCollisionEnabled().
 #include "LandscapeHeightfieldCollisionComponent.h"
 
-//////////////////////////////////////////////////////////////////////////
-// FHitResult
-
-FHitResult::FHitResult(class AActor* InActor, class UPrimitiveComponent* InComponent, FVector const& HitLoc, FVector const& HitNorm)
-{
-	FMemory::Memzero(this, sizeof(FHitResult));
-	Location = HitLoc;
-	ImpactPoint = HitLoc;
-	Normal = HitNorm;
-	ImpactNormal = HitNorm;
-	HitObjectHandle = FActorInstanceHandle(InActor);
-	Component = InComponent;
-}
-
-bool FHitResult::NetSerialize(FArchive& Ar, class UPackageMap* Map, bool& bOutSuccess)
-{
-	// Most of the time the vectors are the same values, use that as an optimization
-	bool bImpactPointEqualsLocation = 0, bImpactNormalEqualsNormal = 0;
-
-	// Often times the indexes are invalid, use that as an optimization
-	bool bInvalidItem = 0, bInvalidFaceIndex = 0, bNoPenetrationDepth = 0, bInvalidElementIndex = 0;
-
-	if (Ar.IsSaving())
-	{
-		bImpactPointEqualsLocation = (ImpactPoint == Location);
-		bImpactNormalEqualsNormal = (ImpactNormal == Normal);
-		bInvalidItem = (Item == INDEX_NONE);
-		bInvalidFaceIndex = (FaceIndex == INDEX_NONE);
-		bNoPenetrationDepth = (PenetrationDepth == 0.0f);
-		bInvalidElementIndex = (ElementIndex == INDEX_NONE);
-	}
-
-	// pack bitfield with flags
-	uint8 Flags = (bBlockingHit << 0) | (bStartPenetrating << 1) | (bImpactPointEqualsLocation << 2) | (bImpactNormalEqualsNormal << 3) | (bInvalidItem << 4) | (bInvalidFaceIndex << 5) | (bNoPenetrationDepth << 6) | (bInvalidElementIndex << 7);
-	Ar.SerializeBits(&Flags, 8); 
-	bBlockingHit = (Flags & (1 << 0)) ? 1 : 0;
-	bStartPenetrating = (Flags & (1 << 1)) ? 1 : 0;
-	bImpactPointEqualsLocation = (Flags & (1 << 2)) ? 1 : 0;
-	bImpactNormalEqualsNormal = (Flags & (1 << 3)) ? 1 : 0;
-	bInvalidItem = (Flags & (1 << 4)) ? 1 : 0;
-	bInvalidFaceIndex = (Flags & (1 << 5)) ? 1 : 0;
-	bNoPenetrationDepth = (Flags & (1 << 6)) ? 1 : 0;
-	bInvalidElementIndex = (Flags & (1 << 7)) ? 1 : 0;
-	// NOTE: Every bit in Flags is being used. If any more bits are added,
-	// Flags will need to be upgraded to a uint16.
-
-	Ar << Time;
-
-	bOutSuccess = true;
-
-	bool bOutSuccessLocal = true;
-
-	Location.NetSerialize(Ar, Map, bOutSuccessLocal);
-	bOutSuccess &= bOutSuccessLocal;
-	Normal.NetSerialize(Ar, Map, bOutSuccessLocal);
-	bOutSuccess &= bOutSuccessLocal;
-
-	if (!bImpactPointEqualsLocation)
-	{
-		ImpactPoint.NetSerialize(Ar, Map, bOutSuccessLocal);
-		bOutSuccess &= bOutSuccessLocal;
-	}
-	else if (Ar.IsLoading())
-	{
-		ImpactPoint = Location;
-	}
-	
-	if (!bImpactNormalEqualsNormal)
-	{
-		ImpactNormal.NetSerialize(Ar, Map, bOutSuccessLocal);
-		bOutSuccess &= bOutSuccessLocal;
-	}
-	else if (Ar.IsLoading())
-	{
-		ImpactNormal = Normal;
-	}
-	TraceStart.NetSerialize(Ar, Map, bOutSuccessLocal);
-	bOutSuccess &= bOutSuccessLocal;
-	TraceEnd.NetSerialize(Ar, Map, bOutSuccessLocal);
-	bOutSuccess &= bOutSuccessLocal;
-
-	if (!bNoPenetrationDepth)
-	{
-		Ar << PenetrationDepth;
-	}
-	else if(Ar.IsLoading())
-	{
-		PenetrationDepth = 0.0f;
-	}
-
-	if (Ar.IsLoading() && bOutSuccess)
-	{
-		Distance = (ImpactPoint - TraceStart).Size();
-	}
-	
-	if (!bInvalidItem)
-	{
-		Ar << Item;
-	}
-	else if (Ar.IsLoading())
-	{
-		Item = INDEX_NONE;
-	}
-
-	Ar << PhysMaterial;
-
-	if (Ar.IsLoading() && Ar.EngineNetVer() < HISTORY_HITRESULT_INSTANCEHANDLE)
-	{
-		AActor* HitActor = nullptr;
-		Ar << HitActor;
-		HitObjectHandle = HitActor;
-	}
-	else
-	{
-		Ar << HitObjectHandle;
-	}
-	Ar << Component;
-	Ar << BoneName;
-	if (!bInvalidFaceIndex)
-	{
-		Ar << FaceIndex;
-	}
-	else if (Ar.IsLoading())
-	{
-		FaceIndex = INDEX_NONE;
-	}
-	
-	if (!bInvalidElementIndex)
-	{
-		Ar << ElementIndex;
-	}
-	else if (Ar.IsLoading())
-	{
-		ElementIndex = INDEX_NONE;
-	}
-
-	return true;
-}
-
-//////////////////////////////////////////////////////////////////////////
-// FOverlapResult
-
-AActor* FOverlapResult::GetActor() const
-{
-	return OverlapObjectHandle.FetchActor();
-}
-
-UPrimitiveComponent* FOverlapResult::GetComponent() const
-{
-	return Component.Get();
-}
-
 //////////////////////////////////////////////////////////////////////////
 // FCollisionQueryFlag
 
diff --git a/Engine/Source/Runtime/Engine/Private/Components/CharacterMovementComponent.cpp b/Engine/Source/Runtime/Engine/Private/Components/CharacterMovementComponent.cpp
--- a/Engine/Source/Runtime/Engine/Private/Components/CharacterMovementComponent.cpp
+++ b/Engine/Source/Runtime/Engine/Private/Components/CharacterMovementComponent.cpp
@@ -33,6 +33,7 @@
 
 #include "Net/PerfCountersHelpers.h"
 #include "ProfilingDebugging/CsvProfiler.h"
+#include "Engine/ScopedMovementUpdate.h"
 
 CSV_DEFINE_CATEGORY(CharacterMovement, true);
 
diff --git a/Engine/Source/Runtime/Engine/Private/Components/PrimitiveComponent.cpp b/Engine/Source/Runtime/Engine/Private/Components/PrimitiveComponent.cpp
--- a/Engine/Source/Runtime/Engine/Private/Components/PrimitiveComponent.cpp
+++ b/Engine/Source/Runtime/Engine/Private/Components/PrimitiveComponent.cpp
@@ -36,6 +36,9 @@
 #include "UObject/UE5PrivateFrostyStreamObjectVersion.h"
 #include "EngineModule.h"
 #include "UObject/ObjectSaveContext.h"
+#include "Engine/OverlapInfo.h"
+#include "Engine/DamageEvents.h"
+#include "Engine/ScopedMovementUpdate.h"
 
 #if WITH_EDITOR
 #include "Engine/LODActor.h"
@@ -114,6 +117,13 @@ DECLARE_CYCLE_STAT(TEXT("BeginComponentOverlap"), STAT_BeginComponentOverlap, ST
 DECLARE_CYCLE_STAT(TEXT("EndComponentOverlap"), STAT_EndComponentOverlap, STATGROUP_Game);
 DECLARE_CYCLE_STAT(TEXT("PrimComp DispatchBlockingHit"), STAT_DispatchBlockingHit, STATGROUP_Game);
 
+FOverlapInfo::FOverlapInfo(UPrimitiveComponent* InComponent, int32 InBodyIndex)
+	: bFromSweep(false)
+{
+	OverlapInfo.HitObjectHandle = FActorInstanceHandle(InComponent ? InComponent->GetOwner() : nullptr);
+	OverlapInfo.Component = InComponent;
+	OverlapInfo.Item = InBodyIndex;
+}
 
 // Predicate to determine if an overlap is with a certain AActor.
 struct FPredicateOverlapHasSameActor
@@ -290,6 +300,9 @@ FThreadSafeCounter UPrimitiveComponent::NextRegistrationSerialNumber;
 // 0 is reserved to mean invalid
 FThreadSafeCounter UPrimitiveComponent::NextComponentId;
 
+UPrimitiveComponent::UPrimitiveComponent(FVTableHelper& Helper) : Super(Helper) { }
+UPrimitiveComponent::~UPrimitiveComponent() = default;
+
 UPrimitiveComponent::UPrimitiveComponent(const FObjectInitializer& ObjectInitializer /*= FObjectInitializer::Get()*/)
 	: Super(ObjectInitializer)
 {
diff --git a/Engine/Source/Runtime/Engine/Private/Components/SceneComponent.cpp b/Engine/Source/Runtime/Engine/Private/Components/SceneComponent.cpp
--- a/Engine/Source/Runtime/Engine/Private/Components/SceneComponent.cpp
+++ b/Engine/Source/Runtime/Engine/Private/Components/SceneComponent.cpp
@@ -35,6 +35,7 @@
 #include "Interfaces/ITargetPlatform.h"
 #include "DeviceProfiles/DeviceProfile.h"
 #include "Net/Core/PushModel/PushModel.h"
+#include "Engine/ScopedMovementUpdate.h"
 
 #define LOCTEXT_NAMESPACE "SceneComponent"
 
@@ -55,13 +56,6 @@ DECLARE_CYCLE_STAT(TEXT("Component UpdateNavData"), STAT_ComponentUpdateNavData,
 DECLARE_CYCLE_STAT(TEXT("Component PostUpdateNavData"), STAT_ComponentPostUpdateNavData, STATGROUP_Component);
 
 
-FOverlapInfo::FOverlapInfo(UPrimitiveComponent* InComponent, int32 InBodyIndex)
-	: bFromSweep(false)
-{
-	OverlapInfo.HitObjectHandle = FActorInstanceHandle(InComponent ? InComponent->GetOwner() : nullptr);
-	OverlapInfo.Component = InComponent;
-	OverlapInfo.Item = InBodyIndex;
-}
 
 FName USceneComponent::GetDefaultSceneRootVariableName()
 {
@@ -841,6 +835,24 @@ void USceneComponent::PropagateTransformUpdate(bool bTransformChanged, EUpdateTr
 	}
 }
 
+bool USceneComponent::IsDeferringMovementUpdates(const FScopedMovementUpdate& ScopedUpdate) const
+{
+	return ScopedUpdate.IsDeferringUpdates();
+}
+
+bool USceneComponent::UpdateOverlaps(const TOverlapArrayView* PendingOverlaps /* = nullptr */, bool bDoNotifies /* = true */, const TOverlapArrayView* OverlapsAtEndLocation /* = nullptr */)
+{
+	if (IsDeferringMovementUpdates())
+	{
+		GetCurrentScopedMovement()->ForceOverlapUpdate();
+	}
+	else if (!ShouldSkipUpdateOverlaps())
+	{
+		bSkipUpdateOverlaps = UpdateOverlapsImpl(PendingOverlaps, bDoNotifies, OverlapsAtEndLocation);
+	}
+
+	return bSkipUpdateOverlaps;
+}
 
 void USceneComponent::EndScopedMovementUpdate(class FScopedMovementUpdate& CompletedScope)
 {
@@ -3527,264 +3539,6 @@ FScopedPreventAttachedComponentMove::~FScopedPreventAttachedComponentMove()
 	}
 }
 
-
-/**
- * FScopedMovementUpdate implementation
- */
-
-static uint32 s_ScopedWarningCount = 0;
-
-FScopedMovementUpdate::FScopedMovementUpdate( class USceneComponent* Component, EScopedUpdate::Type ScopeBehavior, bool bRequireOverlapsEventFlagToQueueOverlaps )
-: Owner(Component)
-, OuterDeferredScope(nullptr)
-, CurrentOverlapState(EOverlapState::eUseParent)
-, TeleportType(ETeleportType::None)
-, FinalOverlapCandidatesIndex(INDEX_NONE)
-, bDeferUpdates(ScopeBehavior == EScopedUpdate::DeferredUpdates)
-, bHasMoved(false)
-, bRequireOverlapsEventFlag(bRequireOverlapsEventFlagToQueueOverlaps)
-{
-	if (IsValid(Component))
-	{
-		OuterDeferredScope = Component->GetCurrentScopedMovement();
-		InitialTransform = Component->GetComponentToWorld();
-		InitialRelativeLocation = Component->GetRelativeLocation();
-		InitialRelativeRotation = Component->GetRelativeRotation();
-		InitialRelativeScale = Component->GetRelativeScale3D();
-
-		if (ScopeBehavior == EScopedUpdate::ImmediateUpdates)
-		{
-			// We only allow ScopeUpdateImmediately if there is no outer scope, or if the outer scope is also ScopeUpdateImmediately.
-			if (OuterDeferredScope && OuterDeferredScope->bDeferUpdates)
-			{
-				if (s_ScopedWarningCount < 100 || (GFrameCounter & 31) == 0)
-				{
-					s_ScopedWarningCount++;
-					UE_LOG(LogSceneComponent, Error, TEXT("FScopedMovementUpdate attempting to use immediate updates within deferred scope, will use deferred updates instead."));
-				}
-
-				bDeferUpdates = true;
-			}
-		}			
-
-		if (bDeferUpdates)
-		{
-			Component->BeginScopedMovementUpdate(*this);
-		}
-	}
-	else
-	{
-		Owner = nullptr;
-	}
-}
-
-
-FScopedMovementUpdate::~FScopedMovementUpdate()
-{
-	if (bDeferUpdates && IsValid(Owner))
-	{
-		Owner->EndScopedMovementUpdate(*this);
-	}
-	Owner = nullptr;
-}
-
-
-bool FScopedMovementUpdate::IsTransformDirty() const
-{
-	if (Owner)
-	{
-		return !InitialTransform.Equals(Owner->GetComponentToWorld(), SMALL_NUMBER);
-	}
-
-	return false;
-}
-
-
-void FScopedMovementUpdate::RevertMove()
-{
-	USceneComponent* Component = Owner;
-	if (IsValid(Component))
-	{
-		FinalOverlapCandidatesIndex = INDEX_NONE;
-		PendingOverlaps.Reset();
-		BlockingHits.Reset();
-		
-		if (IsTransformDirty())
-		{
-			// Teleport to start
-			Component->ComponentToWorld = InitialTransform;
-			Component->RelativeLocation = InitialRelativeLocation;
-			Component->RelativeRotation = InitialRelativeRotation;
-			Component->RelativeScale3D = InitialRelativeScale;
-
-			if (!IsDeferringUpdates())
-			{
-				Component->PropagateTransformUpdate(true);
-				Component->UpdateOverlaps();
-			}
-		}
-	}
-	bHasMoved = false;
-	CurrentOverlapState = EOverlapState::eUseParent;
-	TeleportType = ETeleportType::None;
-}
-
-void FScopedMovementUpdate::AppendOverlapsAfterMove(const TOverlapArrayView& NewPendingOverlaps, bool bSweep, bool bIncludesOverlapsAtEnd)
-{
-	bHasMoved = true;
-	const bool bWasForcing = (CurrentOverlapState == EOverlapState::eForceUpdate);
-
-	if (bIncludesOverlapsAtEnd)
-	{
-		CurrentOverlapState = EOverlapState::eIncludesOverlaps;
-		if (NewPendingOverlaps.Num())
-		{
-			FinalOverlapCandidatesIndex = PendingOverlaps.Num();
-			PendingOverlaps.Append(NewPendingOverlaps.GetData(), NewPendingOverlaps.Num());
-		}
-		else
-		{
-			// No new pending overlaps means we're not overlapping anything at the end location.
-			FinalOverlapCandidatesIndex = INDEX_NONE;
-		}
-	}
-	else
-	{
-		// We don't know about the final overlaps in the case of a teleport.
-		CurrentOverlapState = EOverlapState::eUnknown;
-		FinalOverlapCandidatesIndex = INDEX_NONE;
-		PendingOverlaps.Append(NewPendingOverlaps.GetData(), NewPendingOverlaps.Num());
-	}
-
-	if (bWasForcing)
-	{
-		CurrentOverlapState = EOverlapState::eForceUpdate;
-	}
-}
-
-
-void FScopedMovementUpdate::OnInnerScopeComplete(const FScopedMovementUpdate& InnerScope)
-{
-	if (IsValid(Owner))
-	{
-		checkSlow(IsDeferringUpdates());
-		checkSlow(InnerScope.IsDeferringUpdates());
-		checkSlow(InnerScope.OuterDeferredScope == this);
-
-		// Combine with the next item on the stack.
-		if (InnerScope.HasMoved(EHasMovedTransformOption::eTestTransform))
-		{
-			bHasMoved = true;
-			
-			if (InnerScope.CurrentOverlapState == EOverlapState::eUseParent)
-			{
-				// Unchanged, use our own
-			}
-			else
-			{
-				// Bubble up from inner scope.
-				CurrentOverlapState = InnerScope.CurrentOverlapState;
-				if (InnerScope.FinalOverlapCandidatesIndex == INDEX_NONE)
-				{
-					FinalOverlapCandidatesIndex = INDEX_NONE;
-				}
-				else
-				{
-					checkSlow(InnerScope.GetPendingOverlaps().Num() > 0);
-					FinalOverlapCandidatesIndex = PendingOverlaps.Num() + InnerScope.FinalOverlapCandidatesIndex;
-				}
-				PendingOverlaps.Append(InnerScope.GetPendingOverlaps());
-				checkSlow(FinalOverlapCandidatesIndex < PendingOverlaps.Num());
-			}
-
-			if (InnerScope.TeleportType > TeleportType)
-			{
-				SetHasTeleported(InnerScope.TeleportType);
-			}
-		}
-		else
-		{
-			// Don't want to invalidate a parent scope when nothing changed in the child.
-			checkSlow(InnerScope.CurrentOverlapState == EOverlapState::eUseParent);
-		}
-
-		BlockingHits.Append(InnerScope.GetPendingBlockingHits());
-	}	
-}
-
-template<typename AllocatorType>
-TOptional<TOverlapArrayView> FScopedMovementUpdate::GetOverlapsAtEnd(class UPrimitiveComponent& PrimComponent, TArray<FOverlapInfo, AllocatorType>& OutEndOverlaps, bool bTransformChanged) const
-{
-	TOptional<TOverlapArrayView> Result;
-	switch (CurrentOverlapState)
-	{
-		case FScopedMovementUpdate::EOverlapState::eUseParent:
-		{
-			// Only rotation could have possibly changed
-			if (bTransformChanged && PrimComponent.AreSymmetricRotations(InitialTransform.GetRotation(), PrimComponent.GetComponentQuat(), PrimComponent.GetComponentScale()))
-			{
-				if (PrimComponent.ConvertRotationOverlapsToCurrentOverlaps(OutEndOverlaps, PrimComponent.GetOverlapInfos()))
-				{
-					Result = OutEndOverlaps;
-				}
-			}
-			else
-			{
-				// Use current overlaps (unchanged)
-				Result = PrimComponent.GetOverlapInfos();
-			}
-			break;
-		}
-		case FScopedMovementUpdate::EOverlapState::eUnknown:
-		case FScopedMovementUpdate::EOverlapState::eForceUpdate:
-		{
-			break;
-		}
-		case FScopedMovementUpdate::EOverlapState::eIncludesOverlaps:
-		{
-			if (FinalOverlapCandidatesIndex == INDEX_NONE)
-			{
-				// Overlapping nothing
-				Result = OutEndOverlaps;
-			}
-			else
-			{
-				// Fill in EndOverlaps with overlaps valid at the end location.
-				const bool bMatchingScale = FTransform::AreScale3DsEqual(InitialTransform, PrimComponent.GetComponentTransform());
-				if (bMatchingScale)
-				{
-					const bool bHasEndOverlaps = PrimComponent.ConvertSweptOverlapsToCurrentOverlaps(
-						OutEndOverlaps, PendingOverlaps, FinalOverlapCandidatesIndex,
-						PrimComponent.GetComponentLocation(), PrimComponent.GetComponentQuat());
-					
-					if (bHasEndOverlaps)
-					{
-						Result = OutEndOverlaps;
-					}
-				}
-			}
-			break;
-		}
-		default:
-		{
-			checkf(false, TEXT("Unknown FScopedMovementUpdate::EOverlapState value"));
-			break;
-		}
-	}
-
-	return Result;
-}
-
-
-bool FScopedMovementUpdate::SetWorldLocationAndRotation(FVector NewLocation, const FQuat& NewQuat, bool bNoPhysics /*= false*/, ETeleportType Teleport /*= ETeleportType::None*/)
-{
-	if (Owner)
-	{
-		return Owner->InternalSetWorldLocationAndRotation(NewLocation, NewQuat, bNoPhysics, Teleport);
-	}
-	return false;
-}
-
 FBoxSphereBounds USceneComponent::GetLocalBounds() const
 {
 	if (bComputeFastLocalBounds)
diff --git a/Engine/Source/Runtime/Engine/Private/Engine/ActorInstanceHandle.cpp b/Engine/Source/Runtime/Engine/Private/Engine/ActorInstanceHandle.cpp
new file mode 100644
--- /dev/null
+++ b/Engine/Source/Runtime/Engine/Private/Engine/ActorInstanceHandle.cpp
@@ -0,0 +1,335 @@
+// Copyright Epic Games, Inc. All Rights Reserved.
+
+#include "Engine/ActorInstanceHandle.h"
+#include "WorldPartition/DataLayer/DataLayerInstance.h"
+#include "GameFramework/Actor.h"
+#include "GameFramework/LightWeightInstanceSubsystem.h"
+#include "Engine/World.h"
+
+FActorInstanceHandle::FActorInstanceHandle()
+	: Actor(nullptr)
+	, InstanceIndex(INDEX_NONE)
+	, InstanceUID(0)
+{
+	// do nothing
+}
+
+FActorInstanceHandle::FActorInstanceHandle(AActor* InActor)
+	: Actor(InActor)
+	, InstanceIndex(INDEX_NONE)
+	, InstanceUID(0)
+{
+	if (InActor)
+	{
+#if WITH_EDITOR
+		// use the first data layer the actor is in if it's in multiple layers
+		TArray<const UDataLayerInstance*> DataLayerInstances = InActor->GetDataLayerInstances();
+		const UDataLayerInstance* DataLayerInstance = DataLayerInstances.Num() > 0 ? DataLayerInstances[0] : nullptr;
+#else
+		const UDataLayerInstance* DataLayerInstance = nullptr;
+#endif // WITH_EDITOR
+		Manager = FLightWeightInstanceSubsystem::Get().FindLightWeightInstanceManager(InActor->StaticClass(), DataLayerInstance);
+	}
+}
+
+FActorInstanceHandle::FActorInstanceHandle(ALightWeightInstanceManager* InManager, int32 InInstanceIndex)
+	: Actor(nullptr)
+	, InstanceIndex(InInstanceIndex)
+	, InstanceUID(0)
+{
+	Manager = InManager;
+	if (Manager.IsValid())
+	{
+		InstanceIndex = Manager->ConvertCollisionIndexToLightWeightIndex(InInstanceIndex);
+		if (AActor* const* FoundActor = Manager->Actors.Find(InstanceIndex))
+		{
+			Actor = *FoundActor;
+		}
+
+		UWorld* World = Manager->GetWorld();
+		if (ensure(World))
+		{
+			InstanceUID = World->LWILastAssignedUID++;
+		}
+	}
+}
+
+FActorInstanceHandle::FActorInstanceHandle(const FActorInstanceHandle& Other)
+{
+	Actor = Other.Actor;
+
+	Manager = Other.Manager;
+	InstanceIndex = Other.InstanceIndex;
+	InstanceUID = Other.InstanceUID;
+}
+
+bool FActorInstanceHandle::IsValid() const
+{
+	return (Manager.IsValid() && InstanceIndex != INDEX_NONE) || IsActorValid();
+}
+
+bool FActorInstanceHandle::DoesRepresentClass(const UClass* OtherClass) const
+{
+	if (OtherClass == nullptr)
+	{
+		return false;
+	}
+
+	if (IsActorValid())
+	{
+		return Actor->IsA(OtherClass);
+	}
+
+	if (Manager.IsValid())
+	{
+		return Manager->DoesRepresentClass(OtherClass);
+	}
+
+	return false;
+}
+
+UClass* FActorInstanceHandle::GetRepresentedClass() const
+{
+	if (!IsValid())
+	{
+		return nullptr;
+	}
+
+	if (IsActorValid())
+	{
+		return Actor->GetClass();
+	}
+
+	if (Manager.IsValid())
+	{
+		return Manager->GetRepresentedClass();
+	}
+
+	return nullptr;
+}
+
+FVector FActorInstanceHandle::GetLocation() const
+{
+	if (IsActorValid())
+	{
+		return Actor->GetActorLocation();
+	}
+
+	if (Manager.IsValid())
+	{
+		Manager->GetLocation(*this);
+	}
+
+	return FVector();
+}
+
+FRotator FActorInstanceHandle::GetRotation() const
+{
+	if (IsActorValid())
+	{
+		return Actor->GetActorRotation();
+	}
+
+	if (Manager.IsValid())
+	{
+		Manager->GetRotation(*this);
+	}
+
+	return FRotator();
+}
+
+FTransform FActorInstanceHandle::GetTransform() const
+{
+	if (IsActorValid())
+	{
+		return Actor->GetActorTransform();
+	}
+
+	if (Manager.IsValid())
+	{
+		Manager->GetTransform(*this);
+	}
+
+	return FTransform();
+}
+
+FName FActorInstanceHandle::GetFName() const
+{
+	if (IsActorValid())
+	{
+		return Actor->GetFName();
+	}
+
+	return NAME_None;
+}
+
+FString FActorInstanceHandle::GetName() const
+{
+	if (IsActorValid())
+	{
+		return Actor->GetName();
+	}
+
+	if (Manager.IsValid())
+	{
+		Manager->GetName(*this);
+	}
+
+	return FString();
+}
+
+AActor* FActorInstanceHandle::GetManagingActor() const
+{
+	if (IsActorValid())
+	{
+		return Actor.Get();
+	}
+
+	return Manager.Get();
+}
+
+USceneComponent* FActorInstanceHandle::GetRootComponent() const
+{
+	if (IsActorValid())
+	{
+		return Actor->GetRootComponent();
+	}
+
+	return Manager.IsValid() ? Manager->GetRootComponent() : nullptr;
+}
+
+AActor* FActorInstanceHandle::FetchActor() const
+{
+	if (IsActorValid())
+	{
+		return Actor.Get();
+	}
+
+	return FLightWeightInstanceSubsystem::Get().FetchActor(*this);
+}
+
+int32 FActorInstanceHandle::GetRenderingInstanceIndex() const
+{
+	return Manager.IsValid() ? Manager->ConvertLightWeightIndexToCollisionIndex(InstanceIndex) : INDEX_NONE;
+}
+
+UObject* FActorInstanceHandle::GetActorAsUObject()
+{
+	if (IsActorValid())
+	{
+		return Cast<UObject>(Actor.Get());
+	}
+
+	return nullptr;
+}
+
+const UObject* FActorInstanceHandle::GetActorAsUObject() const
+{
+	if (IsActorValid())
+	{
+		return Cast<UObject>(Actor.Get());
+	}
+
+	return nullptr;
+}
+
+bool FActorInstanceHandle::IsActorValid() const
+{
+	return Actor.IsValid();
+}
+
+FActorInstanceHandle& FActorInstanceHandle::operator=(AActor* OtherActor)
+{
+	Actor = OtherActor;
+	Manager.Reset();
+	InstanceIndex = INDEX_NONE;
+
+	return *this;
+}
+
+bool FActorInstanceHandle::operator==(const FActorInstanceHandle& Other) const
+{
+	// try to compare managers and indices first if we have them
+	if (Manager.IsValid() && Other.Manager.IsValid() && InstanceIndex != INDEX_NONE && Other.InstanceIndex != INDEX_NONE)
+	{
+		return Manager == Other.Manager && InstanceIndex == Other.InstanceIndex;
+	}
+
+	// try to compare the actors
+	const AActor* MyActor = FetchActor();
+	const AActor* OtherActor = Other.FetchActor();
+
+	return MyActor == OtherActor;
+}
+
+bool FActorInstanceHandle::operator!=(const FActorInstanceHandle& Other) const
+{
+	return !(*this == Other);
+}
+
+bool FActorInstanceHandle::operator==(const AActor* OtherActor) const
+{
+	// if we have an actor, compare the two actors
+	if (Actor.IsValid())
+	{
+		return Actor.Get() == OtherActor;
+	}
+
+	// if OtherActor is null then we're only equal if this doesn't refer to a valid instance
+	if (OtherActor == nullptr)
+	{
+		return !Manager.IsValid() && InstanceIndex == INDEX_NONE;
+	}
+
+	// we don't have an actor so see if we can look up an instance associated with OtherActor and see if we refer to the same instance
+
+#if WITH_EDITOR
+	// use the first layer the actor is in if it's in multiple layers
+	TArray<const UDataLayerInstance*> DataLayerInstances = OtherActor->GetDataLayerInstances();
+	const UDataLayerInstance* DataLayerInstance = DataLayerInstances.Num() > 0 ? DataLayerInstances[0] : nullptr;
+#else
+	const UDataLayerInstance* DataLayerInstance = nullptr;
+#endif // WITH_EDITOR
+
+	if (ALightWeightInstanceManager* LWIManager = FLightWeightInstanceSubsystem::Get().FindLightWeightInstanceManager(OtherActor->StaticClass(), DataLayerInstance))
+	{
+		if (Manager.Get() != LWIManager)
+		{
+			return false;
+		}
+
+		return Manager->FindIndexForActor(OtherActor) == InstanceIndex;
+	}
+
+	return false;
+}
+
+bool FActorInstanceHandle::operator!=(const AActor* OtherActor) const
+{
+	return !(*this == OtherActor);
+}
+
+uint32 GetTypeHash(const FActorInstanceHandle& Handle)
+{
+	uint32 Hash = 0;
+	if (Handle.Actor.IsValid())
+	{
+		FCrc::StrCrc32(*(Handle.Actor->GetPathName()), Hash);
+	}
+	if (Handle.Manager.IsValid())
+	{
+		Hash = HashCombine(Hash, GetTypeHash(Handle.Manager.Get()));
+	}
+	Hash = HashCombine(Hash, Handle.InstanceIndex);
+
+	return Hash;
+}
+
+FArchive& operator<<(FArchive& Ar, FActorInstanceHandle& Handle)
+{
+	Ar << Handle.Actor;
+	Ar << Handle.Manager;
+	Ar << Handle.InstanceIndex;
+
+	return Ar;
+}
diff --git a/Engine/Source/Runtime/Engine/Private/Engine/DamageEvents.cpp b/Engine/Source/Runtime/Engine/Private/Engine/DamageEvents.cpp
new file mode 100644
--- /dev/null
+++ b/Engine/Source/Runtime/Engine/Private/Engine/DamageEvents.cpp
@@ -0,0 +1,71 @@
+// Copyright Epic Games, Inc. All Rights Reserved.
+
+#include "Engine/DamageEvents.h"
+#include "GameFramework/Actor.h"
+#include "Components/PrimitiveComponent.h"
+
+void FDamageEvent::GetBestHitInfo(AActor const* HitActor, AActor const* HitInstigator, FHitResult& OutHitInfo, FVector& OutImpulseDir) const
+{
+	if (ensure(HitActor))
+	{
+		// fill out the hitinfo as best we can
+		OutHitInfo.HitObjectHandle = FActorInstanceHandle(const_cast<AActor*>(HitActor));
+		OutHitInfo.bBlockingHit = true;
+		OutHitInfo.BoneName = NAME_None;
+		OutHitInfo.Component = Cast<UPrimitiveComponent>(HitActor->GetRootComponent());
+
+		// assume the actor got hit in the center of his root component
+		OutHitInfo.ImpactPoint = HitActor->GetActorLocation();
+		OutHitInfo.Location = OutHitInfo.ImpactPoint;
+
+		// assume hit came from instigator's location
+		OutImpulseDir = HitInstigator ?
+			(OutHitInfo.ImpactPoint - HitInstigator->GetActorLocation()).GetSafeNormal()
+			: FVector::ZeroVector;
+
+		// assume normal points back toward instigator
+		OutHitInfo.ImpactNormal = -OutImpulseDir;
+		OutHitInfo.Normal = OutHitInfo.ImpactNormal;
+	}
+}
+
+void FPointDamageEvent::GetBestHitInfo(AActor const* HitActor, AActor const* HitInstigator, FHitResult& OutHitInfo, FVector& OutImpulseDir) const
+{
+	// assume the actor got hit in the center of his root component
+	OutHitInfo = HitInfo;
+	OutImpulseDir = ShotDirection;
+}
+
+void FRadialDamageEvent::GetBestHitInfo(AActor const* HitActor, AActor const* HitInstigator, FHitResult& OutHitInfo, FVector& OutImpulseDir) const
+{
+	ensure(ComponentHits.Num() > 0);
+
+	// for now, just return the first one
+	OutHitInfo = ComponentHits[0];
+	OutImpulseDir = (OutHitInfo.ImpactPoint - Origin).GetSafeNormal();
+}
+
+float FRadialDamageParams::GetDamageScale(float DistanceFromEpicenter) const
+{
+	float const ValidatedInnerRadius = FMath::Max(0.f, InnerRadius);
+	float const ValidatedOuterRadius = FMath::Max(OuterRadius, ValidatedInnerRadius);
+	float const ValidatedDist = FMath::Max(0.f, DistanceFromEpicenter);
+
+	if (ValidatedDist >= ValidatedOuterRadius)
+	{
+		// outside the radius, no effect
+		return 0.f;
+	}
+
+	if ((DamageFalloff == 0.f) || (ValidatedDist <= ValidatedInnerRadius))
+	{
+		// no falloff or inside inner radius means full effect
+		return 1.f;
+	}
+
+	// calculate the interpolated scale
+	float DamageScale = 1.f - ((ValidatedDist - ValidatedInnerRadius) / (ValidatedOuterRadius - ValidatedInnerRadius));
+	DamageScale = FMath::Pow(DamageScale, DamageFalloff);
+
+	return DamageScale;
+}
diff --git a/Engine/Source/Runtime/Engine/Private/Engine/HitResult.cpp b/Engine/Source/Runtime/Engine/Private/Engine/HitResult.cpp
new file mode 100644
--- /dev/null
+++ b/Engine/Source/Runtime/Engine/Private/Engine/HitResult.cpp
@@ -0,0 +1,163 @@
+// Copyright Epic Games, Inc. All Rights Reserved.
+
+#include "Engine/HitResult.h"
+#include "PhysicalMaterials/PhysicalMaterial.h"
+#include "GameFramework/LightWeightInstanceSubsystem.h"
+#include "Components/PrimitiveComponent.h"
+
+FHitResult::FHitResult(class AActor* InActor, class UPrimitiveComponent* InComponent, FVector const& HitLoc, FVector const& HitNorm)
+{
+	FMemory::Memzero(this, sizeof(FHitResult));
+	Location = HitLoc;
+	ImpactPoint = HitLoc;
+	Normal = HitNorm;
+	ImpactNormal = HitNorm;
+	HitObjectHandle = FActorInstanceHandle(InActor);
+	Component = InComponent;
+}
+
+bool FHitResult::NetSerialize(FArchive& Ar, class UPackageMap* Map, bool& bOutSuccess)
+{
+	// Most of the time the vectors are the same values, use that as an optimization
+	bool bImpactPointEqualsLocation = 0, bImpactNormalEqualsNormal = 0;
+
+	// Often times the indexes are invalid, use that as an optimization
+	bool bInvalidItem = 0, bInvalidFaceIndex = 0, bNoPenetrationDepth = 0, bInvalidElementIndex = 0;
+
+	if (Ar.IsSaving())
+	{
+		bImpactPointEqualsLocation = (ImpactPoint == Location);
+		bImpactNormalEqualsNormal = (ImpactNormal == Normal);
+		bInvalidItem = (Item == INDEX_NONE);
+		bInvalidFaceIndex = (FaceIndex == INDEX_NONE);
+		bNoPenetrationDepth = (PenetrationDepth == 0.0f);
+		bInvalidElementIndex = (ElementIndex == INDEX_NONE);
+	}
+
+	// pack bitfield with flags
+	uint8 Flags = (bBlockingHit << 0) | (bStartPenetrating << 1) | (bImpactPointEqualsLocation << 2) | (bImpactNormalEqualsNormal << 3) | (bInvalidItem << 4) | (bInvalidFaceIndex << 5) | (bNoPenetrationDepth << 6) | (bInvalidElementIndex << 7);
+	Ar.SerializeBits(&Flags, 8); 
+	bBlockingHit = (Flags & (1 << 0)) ? 1 : 0;
+	bStartPenetrating = (Flags & (1 << 1)) ? 1 : 0;
+	bImpactPointEqualsLocation = (Flags & (1 << 2)) ? 1 : 0;
+	bImpactNormalEqualsNormal = (Flags & (1 << 3)) ? 1 : 0;
+	bInvalidItem = (Flags & (1 << 4)) ? 1 : 0;
+	bInvalidFaceIndex = (Flags & (1 << 5)) ? 1 : 0;
+	bNoPenetrationDepth = (Flags & (1 << 6)) ? 1 : 0;
+	bInvalidElementIndex = (Flags & (1 << 7)) ? 1 : 0;
+	// NOTE: Every bit in Flags is being used. If any more bits are added,
+	// Flags will need to be upgraded to a uint16.
+
+	Ar << Time;
+
+	bOutSuccess = true;
+
+	bool bOutSuccessLocal = true;
+
+	Location.NetSerialize(Ar, Map, bOutSuccessLocal);
+	bOutSuccess &= bOutSuccessLocal;
+	Normal.NetSerialize(Ar, Map, bOutSuccessLocal);
+	bOutSuccess &= bOutSuccessLocal;
+
+	if (!bImpactPointEqualsLocation)
+	{
+		ImpactPoint.NetSerialize(Ar, Map, bOutSuccessLocal);
+		bOutSuccess &= bOutSuccessLocal;
+	}
+	else if (Ar.IsLoading())
+	{
+		ImpactPoint = Location;
+	}
+	
+	if (!bImpactNormalEqualsNormal)
+	{
+		ImpactNormal.NetSerialize(Ar, Map, bOutSuccessLocal);
+		bOutSuccess &= bOutSuccessLocal;
+	}
+	else if (Ar.IsLoading())
+	{
+		ImpactNormal = Normal;
+	}
+	TraceStart.NetSerialize(Ar, Map, bOutSuccessLocal);
+	bOutSuccess &= bOutSuccessLocal;
+	TraceEnd.NetSerialize(Ar, Map, bOutSuccessLocal);
+	bOutSuccess &= bOutSuccessLocal;
+
+	if (!bNoPenetrationDepth)
+	{
+		Ar << PenetrationDepth;
+	}
+	else if(Ar.IsLoading())
+	{
+		PenetrationDepth = 0.0f;
+	}
+
+	if (Ar.IsLoading() && bOutSuccess)
+	{
+		Distance = (ImpactPoint - TraceStart).Size();
+	}
+	
+	if (!bInvalidItem)
+	{
+		Ar << Item;
+	}
+	else if (Ar.IsLoading())
+	{
+		Item = INDEX_NONE;
+	}
+
+	Ar << PhysMaterial;
+
+	if (Ar.IsLoading() && Ar.EngineNetVer() < HISTORY_HITRESULT_INSTANCEHANDLE)
+	{
+		AActor* HitActor = nullptr;
+		Ar << HitActor;
+		HitObjectHandle = HitActor;
+	}
+	else
+	{
+		Ar << HitObjectHandle;
+	}
+	Ar << Component;
+	Ar << BoneName;
+	if (!bInvalidFaceIndex)
+	{
+		Ar << FaceIndex;
+	}
+	else if (Ar.IsLoading())
+	{
+		FaceIndex = INDEX_NONE;
+	}
+	
+	if (!bInvalidElementIndex)
+	{
+		Ar << ElementIndex;
+	}
+	else if (Ar.IsLoading())
+	{
+		ElementIndex = INDEX_NONE;
+	}
+
+	return true;
+}
+
+FString FHitResult::ToString() const
+{
+	return FString::Printf(TEXT("bBlockingHit:%s bStartPenetrating:%s Time:%f Location:%s ImpactPoint:%s Normal:%s ImpactNormal:%s TraceStart:%s TraceEnd:%s PenetrationDepth:%f Item:%d PhysMaterial:%s Actor:%s Component:%s BoneName:%s FaceIndex:%d"),
+		bBlockingHit == true ? TEXT("True") : TEXT("False"),
+		bStartPenetrating == true ? TEXT("True") : TEXT("False"),
+		Time,
+		*Location.ToString(),
+		*ImpactPoint.ToString(),
+		*Normal.ToString(),
+		*ImpactNormal.ToString(),
+		*TraceStart.ToString(),
+		*TraceEnd.ToString(),
+		PenetrationDepth,
+		Item,
+		PhysMaterial.IsValid() ? *PhysMaterial->GetName() : TEXT("None"),
+		*FLightWeightInstanceSubsystem::Get().GetName(HitObjectHandle),
+		Component.IsValid() ? *Component->GetName() : TEXT("None"),
+		BoneName.IsValid() ? *BoneName.ToString() : TEXT("None"),
+		FaceIndex);
+}
\ No newline at end of file
diff --git a/Engine/Source/Runtime/Engine/Private/Engine/ReplicatedState.cpp b/Engine/Source/Runtime/Engine/Private/Engine/ReplicatedState.cpp
new file mode 100644
--- /dev/null
+++ b/Engine/Source/Runtime/Engine/Private/Engine/ReplicatedState.cpp
@@ -0,0 +1,196 @@
+// Copyright Epic Games, Inc. All Rights Reserved.
+
+#include "Engine/ReplicatedState.h"
+#include "GameFramework/Actor.h"
+#include "Engine/World.h"
+
+// compile switch for disabling quantization of replicated movement, meant for testing.
+#ifndef REP_MOVEMENT_DISABLE_QUANTIZATION
+#define REP_MOVEMENT_DISABLE_QUANTIZATION 0
+#endif
+
+/** If true, origin rebasing is enabled in multiplayer games, meaning that servers and clients can have different local world origins. */
+int32 FRepMovement::EnableMultiplayerWorldOriginRebasing = 0;
+
+/** Console variable ref to enable multiplayer world origin rebasing. */
+FAutoConsoleVariableRef CVarEnableMultiplayerWorldOriginRebasing(
+	TEXT("p.EnableMultiplayerWorldOriginRebasing"),
+	FRepMovement::EnableMultiplayerWorldOriginRebasing,
+	TEXT("Enable world origin rebasing for multiplayer, meaning that servers and clients can have different world origin locations."),
+	ECVF_ReadOnly);
+
+FRepMovement::FRepMovement()
+	: LinearVelocity(ForceInit)
+	, AngularVelocity(ForceInit)
+	, Location(ForceInit)
+	, Rotation(ForceInit)
+	, bSimulatedPhysicSleep(false)
+	, bRepPhysics(false)
+	, ServerFrame(0)
+	, LocationQuantizationLevel(EVectorQuantization::RoundWholeNumber)
+	, VelocityQuantizationLevel(EVectorQuantization::RoundWholeNumber)
+	, RotationQuantizationLevel(ERotatorQuantization::ByteComponents)
+{
+}
+
+bool FRepMovement::SerializeQuantizedVector(FArchive& Ar, FVector& Vector, EVectorQuantization QuantizationLevel)
+{
+	// Since FRepMovement used to use FVector_NetQuantize100, we're allowing enough bits per component
+	// regardless of the quantization level so that we can still support at least the same maximum magnitude
+	// (2^30 / 100, or ~10 million).
+	// This uses no inherent extra bandwidth since we're still using the same number of bits to store the
+	// bits-per-component value. Of course, larger magnitudes will still use more bandwidth,
+	// as has always been the case.
+	switch (QuantizationLevel)
+	{
+	case EVectorQuantization::RoundTwoDecimals:
+	{
+		return SerializePackedVector<100, 30>(Vector, Ar);
+	}
+
+	case EVectorQuantization::RoundOneDecimal:
+	{
+		return SerializePackedVector<10, 27>(Vector, Ar);
+	}
+
+	default:
+	{
+		return SerializePackedVector<1, 24>(Vector, Ar);
+	}
+	}
+}
+
+bool FRepMovement::NetSerialize(FArchive& Ar, class UPackageMap* Map, bool& bOutSuccess)
+{
+	// pack bitfield with flags
+	const bool bServerFrameAndHandleSupported = Ar.EngineNetVer() >= HISTORY_REPMOVE_SERVERFRAME_AND_HANDLE;
+	uint8 Flags = (bSimulatedPhysicSleep << 0) | (bRepPhysics << 1) | ((ServerFrame > 0) << 2) | ((ServerPhysicsHandle != INDEX_NONE) << 3);
+	Ar.SerializeBits(&Flags, bServerFrameAndHandleSupported ? 4 : 2);
+	bSimulatedPhysicSleep = (Flags & (1 << 0)) ? 1 : 0;
+	bRepPhysics = (Flags & (1 << 1)) ? 1 : 0;
+	const bool bRepServerFrame = (Flags & (1 << 2) && bServerFrameAndHandleSupported) ? 1 : 0;
+	const bool bRepServerHandle = (Flags & (1 << 3) && bServerFrameAndHandleSupported) ? 1 : 0;
+
+	bOutSuccess = true;
+
+#if REP_MOVEMENT_DISABLE_QUANTIZATION
+	Ar << Location;
+	Ar << LinearVelocity;
+	Ar << Rotation;
+	if (bRepPhysics)
+	{
+		Ar << AngularVelocity;
+	}
+#else
+
+	// update location, rotation, linear velocity
+	bOutSuccess &= SerializeQuantizedVector(Ar, Location, LocationQuantizationLevel);
+
+	switch (RotationQuantizationLevel)
+	{
+	case ERotatorQuantization::ByteComponents:
+	{
+		Rotation.SerializeCompressed(Ar);
+		break;
+	}
+
+	case ERotatorQuantization::ShortComponents:
+	{
+		Rotation.SerializeCompressedShort(Ar);
+		break;
+	}
+	}
+
+	bOutSuccess &= SerializeQuantizedVector(Ar, LinearVelocity, VelocityQuantizationLevel);
+
+	// update angular velocity if required
+	if (bRepPhysics)
+	{
+		bOutSuccess &= SerializeQuantizedVector(Ar, AngularVelocity, VelocityQuantizationLevel);
+	}
+#endif
+
+	if (bRepServerFrame)
+	{
+		uint32 uServerFrame = (uint32)ServerFrame;
+		Ar.SerializeIntPacked(uServerFrame);
+		ServerFrame = (int32)uServerFrame;
+	}
+
+	if (bRepServerHandle)
+	{
+		uint32 uServerPhysicsHandle = (uint32)ServerPhysicsHandle;
+		Ar.SerializeIntPacked(uServerPhysicsHandle);
+		ServerPhysicsHandle = (int32)uServerPhysicsHandle;
+	}
+
+	return true;
+}
+
+/** Rebase zero-origin position onto local world origin value. */
+FVector FRepMovement::RebaseOntoLocalOrigin(const FVector& Location, const FIntVector& LocalOrigin)
+{
+	if (EnableMultiplayerWorldOriginRebasing <= 0 || LocalOrigin == FIntVector::ZeroValue)
+	{
+		return Location;
+	}
+
+	return FVector(Location.X - LocalOrigin.X, Location.Y - LocalOrigin.Y, Location.Z - LocalOrigin.Z);
+}
+
+/** Rebase local-origin position onto zero world origin value. */
+FVector FRepMovement::RebaseOntoZeroOrigin(const FVector& Location, const FIntVector& LocalOrigin)
+{
+	if (EnableMultiplayerWorldOriginRebasing <= 0 || LocalOrigin == FIntVector::ZeroValue)
+	{
+		return Location;
+	}
+
+	return FVector(Location.X + LocalOrigin.X, Location.Y + LocalOrigin.Y, Location.Z + LocalOrigin.Z);
+}
+
+/** Rebase zero-origin position onto local world origin value based on an actor's world. */
+FVector FRepMovement::RebaseOntoLocalOrigin(const FVector& Location, const AActor* const WorldContextActor)
+{
+	if (WorldContextActor == nullptr || EnableMultiplayerWorldOriginRebasing <= 0)
+	{
+		return Location;
+	}
+
+	return RebaseOntoLocalOrigin(Location, WorldContextActor->GetWorld()->OriginLocation);
+}
+
+/** Rebase local-origin position onto zero world origin value based on an actor's world.*/
+FVector FRepMovement::RebaseOntoZeroOrigin(const FVector& Location, const AActor* const WorldContextActor)
+{
+	if (WorldContextActor == nullptr || EnableMultiplayerWorldOriginRebasing <= 0)
+	{
+		return Location;
+	}
+
+	return RebaseOntoZeroOrigin(Location, WorldContextActor->GetWorld()->OriginLocation);
+}
+
+/// @cond DOXYGEN_WARNINGS
+
+/** Rebase zero-origin position onto local world origin value based on an actor component's world. */
+FVector FRepMovement::RebaseOntoLocalOrigin(const FVector& Location, const UActorComponent* const WorldContextActorComponent)
+{
+	if (WorldContextActorComponent == nullptr || EnableMultiplayerWorldOriginRebasing <= 0)
+	{
+		return Location;
+	}
+
+	return RebaseOntoLocalOrigin(Location, WorldContextActorComponent->GetWorld()->OriginLocation);
+}
+
+/** Rebase local-origin position onto zero world origin value based on an actor component's world.*/
+FVector FRepMovement::RebaseOntoZeroOrigin(const FVector& Location, const UActorComponent* const WorldContextActorComponent)
+{
+	if (WorldContextActorComponent == nullptr || EnableMultiplayerWorldOriginRebasing <= 0)
+	{
+		return Location;
+	}
+
+	return RebaseOntoZeroOrigin(Location, WorldContextActorComponent->GetWorld()->OriginLocation);
+}
\ No newline at end of file
diff --git a/Engine/Source/Runtime/Engine/Private/Engine/ScopedMovementUpdate.cpp b/Engine/Source/Runtime/Engine/Private/Engine/ScopedMovementUpdate.cpp
new file mode 100644
--- /dev/null
+++ b/Engine/Source/Runtime/Engine/Private/Engine/ScopedMovementUpdate.cpp
@@ -0,0 +1,254 @@
+// Copyright Epic Games, Inc. All Rights Reserved.
+
+#include "Engine/ScopedMovementUpdate.h"
+#include "Components/SceneComponent.h"
+#include "Components/PrimitiveComponent.h"
+
+DEFINE_LOG_CATEGORY_STATIC(LogScopedMovementUpdate, Log, All);
+
+static uint32 s_ScopedWarningCount = 0;
+
+FScopedMovementUpdate::FScopedMovementUpdate( class USceneComponent* Component, EScopedUpdate::Type ScopeBehavior, bool bRequireOverlapsEventFlagToQueueOverlaps )
+: Owner(Component)
+, OuterDeferredScope(nullptr)
+, CurrentOverlapState(EOverlapState::eUseParent)
+, TeleportType(ETeleportType::None)
+, FinalOverlapCandidatesIndex(INDEX_NONE)
+, bDeferUpdates(ScopeBehavior == EScopedUpdate::DeferredUpdates)
+, bHasMoved(false)
+, bRequireOverlapsEventFlag(bRequireOverlapsEventFlagToQueueOverlaps)
+{
+	if (IsValid(Component))
+	{
+		OuterDeferredScope = Component->GetCurrentScopedMovement();
+		InitialTransform = Component->GetComponentToWorld();
+		InitialRelativeLocation = Component->GetRelativeLocation();
+		InitialRelativeRotation = Component->GetRelativeRotation();
+		InitialRelativeScale = Component->GetRelativeScale3D();
+
+		if (ScopeBehavior == EScopedUpdate::ImmediateUpdates)
+		{
+			// We only allow ScopeUpdateImmediately if there is no outer scope, or if the outer scope is also ScopeUpdateImmediately.
+			if (OuterDeferredScope && OuterDeferredScope->bDeferUpdates)
+			{
+				if (s_ScopedWarningCount < 100 || (GFrameCounter & 31) == 0)
+				{
+					s_ScopedWarningCount++;
+					UE_LOG(LogScopedMovementUpdate, Error, TEXT("FScopedMovementUpdate attempting to use immediate updates within deferred scope, will use deferred updates instead."));
+				}
+
+				bDeferUpdates = true;
+			}
+		}			
+
+		if (bDeferUpdates)
+		{
+			Component->BeginScopedMovementUpdate(*this);
+		}
+	}
+	else
+	{
+		Owner = nullptr;
+	}
+}
+
+FScopedMovementUpdate::~FScopedMovementUpdate()
+{
+	if (bDeferUpdates && IsValid(Owner))
+	{
+		Owner->EndScopedMovementUpdate(*this);
+	}
+	Owner = nullptr;
+}
+
+bool FScopedMovementUpdate::IsTransformDirty() const
+{
+	if (Owner)
+	{
+		return !InitialTransform.Equals(Owner->GetComponentToWorld(), SMALL_NUMBER);
+	}
+
+	return false;
+}
+
+void FScopedMovementUpdate::RevertMove()
+{
+	USceneComponent* Component = Owner;
+	if (IsValid(Component))
+	{
+		FinalOverlapCandidatesIndex = INDEX_NONE;
+		PendingOverlaps.Reset();
+		BlockingHits.Reset();
+		
+		if (IsTransformDirty())
+		{
+			// Teleport to start
+			Component->ComponentToWorld = InitialTransform;
+			Component->RelativeLocation = InitialRelativeLocation;
+			Component->RelativeRotation = InitialRelativeRotation;
+			Component->RelativeScale3D = InitialRelativeScale;
+
+			if (!IsDeferringUpdates())
+			{
+				Component->PropagateTransformUpdate(true);
+				Component->UpdateOverlaps();
+			}
+		}
+	}
+	bHasMoved = false;
+	CurrentOverlapState = EOverlapState::eUseParent;
+	TeleportType = ETeleportType::None;
+}
+
+void FScopedMovementUpdate::AppendOverlapsAfterMove(const TOverlapArrayView& NewPendingOverlaps, bool bSweep, bool bIncludesOverlapsAtEnd)
+{
+	bHasMoved = true;
+	const bool bWasForcing = (CurrentOverlapState == EOverlapState::eForceUpdate);
+
+	if (bIncludesOverlapsAtEnd)
+	{
+		CurrentOverlapState = EOverlapState::eIncludesOverlaps;
+		if (NewPendingOverlaps.Num())
+		{
+			FinalOverlapCandidatesIndex = PendingOverlaps.Num();
+			PendingOverlaps.Append(NewPendingOverlaps.GetData(), NewPendingOverlaps.Num());
+		}
+		else
+		{
+			// No new pending overlaps means we're not overlapping anything at the end location.
+			FinalOverlapCandidatesIndex = INDEX_NONE;
+		}
+	}
+	else
+	{
+		// We don't know about the final overlaps in the case of a teleport.
+		CurrentOverlapState = EOverlapState::eUnknown;
+		FinalOverlapCandidatesIndex = INDEX_NONE;
+		PendingOverlaps.Append(NewPendingOverlaps.GetData(), NewPendingOverlaps.Num());
+	}
+
+	if (bWasForcing)
+	{
+		CurrentOverlapState = EOverlapState::eForceUpdate;
+	}
+}
+
+void FScopedMovementUpdate::OnInnerScopeComplete(const FScopedMovementUpdate& InnerScope)
+{
+	if (IsValid(Owner))
+	{
+		checkSlow(IsDeferringUpdates());
+		checkSlow(InnerScope.IsDeferringUpdates());
+		checkSlow(InnerScope.OuterDeferredScope == this);
+
+		// Combine with the next item on the stack.
+		if (InnerScope.HasMoved(EHasMovedTransformOption::eTestTransform))
+		{
+			bHasMoved = true;
+			
+			if (InnerScope.CurrentOverlapState == EOverlapState::eUseParent)
+			{
+				// Unchanged, use our own
+			}
+			else
+			{
+				// Bubble up from inner scope.
+				CurrentOverlapState = InnerScope.CurrentOverlapState;
+				if (InnerScope.FinalOverlapCandidatesIndex == INDEX_NONE)
+				{
+					FinalOverlapCandidatesIndex = INDEX_NONE;
+				}
+				else
+				{
+					checkSlow(InnerScope.GetPendingOverlaps().Num() > 0);
+					FinalOverlapCandidatesIndex = PendingOverlaps.Num() + InnerScope.FinalOverlapCandidatesIndex;
+				}
+				PendingOverlaps.Append(InnerScope.GetPendingOverlaps());
+				checkSlow(FinalOverlapCandidatesIndex < PendingOverlaps.Num());
+			}
+
+			if (InnerScope.TeleportType > TeleportType)
+			{
+				SetHasTeleported(InnerScope.TeleportType);
+			}
+		}
+		else
+		{
+			// Don't want to invalidate a parent scope when nothing changed in the child.
+			checkSlow(InnerScope.CurrentOverlapState == EOverlapState::eUseParent);
+		}
+
+		BlockingHits.Append(InnerScope.GetPendingBlockingHits());
+	}	
+}
+
+bool FScopedMovementUpdate::SetWorldLocationAndRotation(FVector NewLocation, const FQuat& NewQuat, bool bNoPhysics /*= false*/, ETeleportType Teleport /*= ETeleportType::None*/)
+{
+	if (Owner)
+	{
+		return Owner->InternalSetWorldLocationAndRotation(NewLocation, NewQuat, bNoPhysics, Teleport);
+	}
+	return false;
+}
+
+TOptional<TOverlapArrayView> FScopedMovementUpdate::GetOverlapsAtEnd(class UPrimitiveComponent& PrimComponent, TInlineOverlapInfoArray& OutEndOverlaps, bool bTransformChanged) const
+{
+	TOptional<TOverlapArrayView> Result;
+	switch (CurrentOverlapState)
+	{
+	case FScopedMovementUpdate::EOverlapState::eUseParent:
+	{
+		// Only rotation could have possibly changed
+		if (bTransformChanged && PrimComponent.AreSymmetricRotations(InitialTransform.GetRotation(), PrimComponent.GetComponentQuat(), PrimComponent.GetComponentScale()))
+		{
+			if (PrimComponent.ConvertRotationOverlapsToCurrentOverlaps(OutEndOverlaps, PrimComponent.GetOverlapInfos()))
+			{
+				Result = OutEndOverlaps;
+			}
+		}
+		else
+		{
+			// Use current overlaps (unchanged)
+			Result = PrimComponent.GetOverlapInfos();
+		}
+		break;
+	}
+	case FScopedMovementUpdate::EOverlapState::eUnknown:
+	case FScopedMovementUpdate::EOverlapState::eForceUpdate:
+	{
+		break;
+	}
+	case FScopedMovementUpdate::EOverlapState::eIncludesOverlaps:
+	{
+		if (FinalOverlapCandidatesIndex == INDEX_NONE)
+		{
+			// Overlapping nothing
+			Result = OutEndOverlaps;
+		}
+		else
+		{
+			// Fill in EndOverlaps with overlaps valid at the end location.
+			const bool bMatchingScale = FTransform::AreScale3DsEqual(InitialTransform, PrimComponent.GetComponentTransform());
+			if (bMatchingScale)
+			{
+				const bool bHasEndOverlaps = PrimComponent.ConvertSweptOverlapsToCurrentOverlaps(
+					OutEndOverlaps, PendingOverlaps, FinalOverlapCandidatesIndex,
+					PrimComponent.GetComponentLocation(), PrimComponent.GetComponentQuat());
+
+				if (bHasEndOverlaps)
+				{
+					Result = OutEndOverlaps;
+				}
+			}
+		}
+		break;
+	}
+	default:
+	{
+		checkf(false, TEXT("Unknown FScopedMovementUpdate::EOverlapState value"));
+		break;
+	}
+	}
+
+	return Result;
+}
diff --git a/Engine/Source/Runtime/Engine/Private/EngineTypes.cpp b/Engine/Source/Runtime/Engine/Private/EngineTypes.cpp
--- a/Engine/Source/Runtime/Engine/Private/EngineTypes.cpp
+++ b/Engine/Source/Runtime/Engine/Private/EngineTypes.cpp
@@ -22,16 +22,6 @@ FAttachmentTransformRules FAttachmentTransformRules::SnapToTargetIncludingScale(
 FDetachmentTransformRules FDetachmentTransformRules::KeepRelativeTransform(EDetachmentRule::KeepRelative, true);
 FDetachmentTransformRules FDetachmentTransformRules::KeepWorldTransform(EDetachmentRule::KeepWorld, true);
 
-/** If true, origin rebasing is enabled in multiplayer games, meaning that servers and clients can have different local world origins. */
-int32 FRepMovement::EnableMultiplayerWorldOriginRebasing = 0;
-
-/** Console variable ref to enable multiplayer world origin rebasing. */
-FAutoConsoleVariableRef CVarEnableMultiplayerWorldOriginRebasing(
-	TEXT("p.EnableMultiplayerWorldOriginRebasing"),
-	FRepMovement::EnableMultiplayerWorldOriginRebasing,
-	TEXT("Enable world origin rebasing for multiplayer, meaning that servers and clients can have different world origin locations."),
-	ECVF_ReadOnly);
-
 #if WITH_EDITORONLY_DATA
 void FMeshProxySettings::PostLoadDeprecated()
 {
@@ -111,75 +101,6 @@ ETraceTypeQuery UEngineTypes::ConvertToTraceType(ECollisionChannel CollisionChan
 	return UCollisionProfile::Get()->ConvertToTraceType(CollisionChannel);
 }
 
-void FDamageEvent::GetBestHitInfo(AActor const* HitActor, AActor const* HitInstigator, FHitResult& OutHitInfo, FVector& OutImpulseDir) const
-{
-	ensure(HitActor);
-	if (HitActor)
-	{
-		// fill out the hitinfo as best we can
-		OutHitInfo.HitObjectHandle = FActorInstanceHandle(const_cast<AActor*>(HitActor));
-		OutHitInfo.bBlockingHit = true;
-		OutHitInfo.BoneName = NAME_None;
-		OutHitInfo.Component = Cast<UPrimitiveComponent>(HitActor->GetRootComponent());
-		
-		// assume the actor got hit in the center of his root component
-		OutHitInfo.ImpactPoint = HitActor->GetActorLocation();
-		OutHitInfo.Location = OutHitInfo.ImpactPoint;
-		
-		// assume hit came from instigator's location
-		OutImpulseDir = HitInstigator ? 
-			( OutHitInfo.ImpactPoint - HitInstigator->GetActorLocation() ).GetSafeNormal()
-			: FVector::ZeroVector;
-
-		// assume normal points back toward instigator
-		OutHitInfo.ImpactNormal = -OutImpulseDir;
-		OutHitInfo.Normal = OutHitInfo.ImpactNormal;
-	}
-}
-
-void FPointDamageEvent::GetBestHitInfo(AActor const* HitActor, AActor const* HitInstigator, FHitResult& OutHitInfo, FVector& OutImpulseDir) const
-{
-	// assume the actor got hit in the center of his root component
-	OutHitInfo = HitInfo;
-	OutImpulseDir = ShotDirection;
-}
-
-
-void FRadialDamageEvent::GetBestHitInfo(AActor const* HitActor, AActor const* HitInstigator, FHitResult& OutHitInfo, FVector& OutImpulseDir) const
-{
-	ensure(ComponentHits.Num() > 0);
-
-	// for now, just return the first one
-	OutHitInfo = ComponentHits[0];
-	OutImpulseDir = (OutHitInfo.ImpactPoint - Origin).GetSafeNormal();
-}
-
-
-float FRadialDamageParams::GetDamageScale(float DistanceFromEpicenter) const
-{
-	float const ValidatedInnerRadius = FMath::Max(0.f, InnerRadius);
-	float const ValidatedOuterRadius = FMath::Max(OuterRadius, ValidatedInnerRadius);
-	float const ValidatedDist = FMath::Max(0.f, DistanceFromEpicenter);
-
-	if (ValidatedDist >= ValidatedOuterRadius)
-	{
-		// outside the radius, no effect
-		return 0.f;
-	}
-
-	if ( (DamageFalloff == 0.f)	|| (ValidatedDist <= ValidatedInnerRadius) )
-	{
-		// no falloff or inside inner radius means full effect
-		return 1.f;
-	}
-
-	// calculate the interpolated scale
-	float DamageScale = 1.f - ( (ValidatedDist - ValidatedInnerRadius) / (ValidatedOuterRadius - ValidatedInnerRadius) );
-	DamageScale = FMath::Pow(DamageScale, DamageFalloff);
-
-	return DamageScale;
-}
-
 FLightmassDebugOptions::FLightmassDebugOptions()
 	: bDebugMode(false)
 	, bStatsEnabled(false)
@@ -237,509 +158,6 @@ UActorComponent* FComponentReference::GetComponent(AActor* OwningActor) const
 	return Result;
 }
 
-FActorInstanceHandle::FActorInstanceHandle()
-	: Actor(nullptr)
-	, InstanceIndex(INDEX_NONE)
-	, InstanceUID(0)
-{
-	// do nothing
-}
-
-FActorInstanceHandle::FActorInstanceHandle(AActor* InActor)
-	: Actor(InActor)
-	, InstanceIndex(INDEX_NONE)
-	, InstanceUID(0)
-{
-	if (InActor)
-	{
-#if WITH_EDITOR
-		// use the first data layer the actor is in if it's in multiple layers
-		TArray<const UDataLayerInstance*> DataLayerInstances = InActor->GetDataLayerInstances();
-		const UDataLayerInstance* DataLayerInstance = DataLayerInstances.Num() > 0 ? DataLayerInstances[0] : nullptr;
-#else
-		const UDataLayerInstance* DataLayerInstance = nullptr;
-#endif // WITH_EDITOR
-		Manager = FLightWeightInstanceSubsystem::Get().FindLightWeightInstanceManager(InActor->StaticClass(), DataLayerInstance);
-	}
-}
-
-FActorInstanceHandle::FActorInstanceHandle(ALightWeightInstanceManager* InManager, int32 InInstanceIndex)
-	: Actor(nullptr)
-	, InstanceIndex(InInstanceIndex)
-	, InstanceUID(0)
-{
-	Manager = InManager;
-	if (Manager.IsValid())
-	{
-		InstanceIndex = Manager->ConvertCollisionIndexToLightWeightIndex(InInstanceIndex);
-		if (AActor* const* FoundActor = Manager->Actors.Find(InstanceIndex))
-		{
-			Actor = *FoundActor;
-		}
-
-		UWorld* World = Manager->GetWorld();
-		if (ensure(World))
-		{
-			InstanceUID = World->LWILastAssignedUID++;
-		}
-	}
-}
-
-FActorInstanceHandle::FActorInstanceHandle(const FActorInstanceHandle& Other)
-{
-	Actor = Other.Actor;
-
-	Manager = Other.Manager;
-	InstanceIndex = Other.InstanceIndex;
-	InstanceUID = Other.InstanceUID;
-}
-
-bool FActorInstanceHandle::IsValid() const
-{
-	return (Manager.IsValid() && InstanceIndex != INDEX_NONE) || IsActorValid();
-}
-
-bool FActorInstanceHandle::DoesRepresentClass(const UClass* OtherClass) const
-{
-	if (OtherClass == nullptr)
-	{
-		return false;
-	}
-
-	if (IsActorValid())
-	{
-		return Actor->IsA(OtherClass);
-	}
-
-	if (Manager.IsValid())
-	{
-		return Manager->DoesRepresentClass(OtherClass);
-	}
-
-	return false;
-}
-
-UClass* FActorInstanceHandle::GetRepresentedClass() const
-{
-	if (!IsValid())
-	{
-		return nullptr;
-	}
-
-	if (IsActorValid())
-	{
-		return Actor->GetClass();
-	}
-
-	if (Manager.IsValid())
-	{
-		return Manager->GetRepresentedClass();
-	}
-
-	return nullptr;
-}
-
-FVector FActorInstanceHandle::GetLocation() const
-{
-	if (IsActorValid())
-	{
-		return Actor->GetActorLocation();
-	}
-
-	if (Manager.IsValid())
-	{
-		Manager->GetLocation(*this);
-	}
-
-	return FVector();
-}
-
-FRotator FActorInstanceHandle::GetRotation() const
-{
-	if (IsActorValid())
-	{
-		return Actor->GetActorRotation();
-	}
-
-	if (Manager.IsValid())
-	{
-		Manager->GetRotation(*this);
-	}
-
-	return FRotator();
-}
-
-FTransform FActorInstanceHandle::GetTransform() const
-{
-	if (IsActorValid())
-	{
-		return Actor->GetActorTransform();
-	}
-
-	if (Manager.IsValid())
-	{
-		Manager->GetTransform(*this);
-	}
-
-	return FTransform();
-}
-
-FName FActorInstanceHandle::GetFName() const
-{
-	if (IsActorValid())
-	{
-		return Actor->GetFName();
-	}
-
-	return NAME_None;
-}
-
-FString FActorInstanceHandle::GetName() const
-{
-	if (IsActorValid())
-	{
-		return Actor->GetName();
-	}
-
-	if (Manager.IsValid())
-	{
-		Manager->GetName(*this);
-	}
-
-	return FString();
-}
-
-AActor* FActorInstanceHandle::GetManagingActor() const
-{
-	if (IsActorValid())
-	{
-		return Actor.Get();
-	}
-
-	return Manager.Get();
-}
-
-USceneComponent* FActorInstanceHandle::GetRootComponent() const
-{
-	if (IsActorValid())
-	{
-		return Actor->GetRootComponent();
-	}
-
-	return Manager.IsValid() ? Manager->GetRootComponent() : nullptr;
-}
-
-AActor* FActorInstanceHandle::FetchActor() const
-{
-	if (IsActorValid())
-	{
-		return Actor.Get();
-	}
-
-	return FLightWeightInstanceSubsystem::Get().FetchActor(*this);
-}
-
-int32 FActorInstanceHandle::GetRenderingInstanceIndex() const
-{
-	return Manager.IsValid() ? Manager->ConvertLightWeightIndexToCollisionIndex(InstanceIndex) : INDEX_NONE;
-}
-
-UObject* FActorInstanceHandle::GetActorAsUObject()
-{
-	if (IsActorValid())
-	{
-		return Cast<UObject>(Actor.Get());
-	}
-
-	return nullptr;
-}
-
-const UObject* FActorInstanceHandle::GetActorAsUObject() const
-{
-	if (IsActorValid())
-	{
-		return Cast<UObject>(Actor.Get());
-	}
-
-	return nullptr;
-}
-
-bool FActorInstanceHandle::IsActorValid() const
-{
-	return Actor.IsValid();
-}
-
-FActorInstanceHandle& FActorInstanceHandle::operator=(AActor* OtherActor)
-{
-	Actor = OtherActor;
-	Manager.Reset();
-	InstanceIndex = INDEX_NONE;
-
-	return *this;
-}
-
-bool FActorInstanceHandle::operator==(const FActorInstanceHandle& Other) const
-{
-	// try to compare managers and indices first if we have them
-	if (Manager.IsValid() && Other.Manager.IsValid() && InstanceIndex != INDEX_NONE && Other.InstanceIndex != INDEX_NONE)
-	{
-		return Manager == Other.Manager && InstanceIndex == Other.InstanceIndex;
-	}
-
-	// try to compare the actors
-	const AActor* MyActor = FetchActor();
-	const AActor* OtherActor = Other.FetchActor();
-
-	return MyActor == OtherActor;
-}
-
-bool FActorInstanceHandle::operator!=(const FActorInstanceHandle& Other) const
-{
-	return !(*this == Other);
-}
-
-bool FActorInstanceHandle::operator==(const AActor* OtherActor) const
-{
-	// if we have an actor, compare the two actors
-	if (Actor.IsValid())
-	{
-		return Actor.Get() == OtherActor;
-	}
-
-	// if OtherActor is null then we're only equal if this doesn't refer to a valid instance
-	if (OtherActor == nullptr)
-	{
-		return !Manager.IsValid() && InstanceIndex == INDEX_NONE;
-	}
-
-	// we don't have an actor so see if we can look up an instance associated with OtherActor and see if we refer to the same instance
-
-#if WITH_EDITOR
-	// use the first layer the actor is in if it's in multiple layers
-	TArray<const UDataLayerInstance*> DataLayerInstances = OtherActor->GetDataLayerInstances();
-	const UDataLayerInstance* DataLayerInstance = DataLayerInstances.Num() > 0 ? DataLayerInstances[0] : nullptr;
-#else
-	const UDataLayerInstance* DataLayerInstance = nullptr;
-#endif // WITH_EDITOR
-
-	if (ALightWeightInstanceManager* LWIManager = FLightWeightInstanceSubsystem::Get().FindLightWeightInstanceManager(OtherActor->StaticClass(), DataLayerInstance))
-	{
-		if (Manager.Get() != LWIManager)
-		{
-			return false;
-		}
-
-		return Manager->FindIndexForActor(OtherActor) == InstanceIndex;
-	}
-
-	return false;
-}
-
-bool FActorInstanceHandle::operator!=(const AActor* OtherActor) const
-{
-	return !(*this == OtherActor);
-}
-
-uint32 GetTypeHash(const FActorInstanceHandle& Handle)
-{
-	uint32 Hash = 0;
-	if (Handle.Actor.IsValid())
-	{
-		FCrc::StrCrc32(*(Handle.Actor->GetPathName()), Hash);
-	}
-	if (Handle.Manager.IsValid())
-	{
-		Hash = HashCombine(Hash, GetTypeHash(Handle.Manager.Get()));
-	}
-	Hash = HashCombine(Hash, Handle.InstanceIndex);
-
-	return Hash;
-}
-
-FArchive& operator<<(FArchive& Ar, FActorInstanceHandle& Handle)
-{
-	Ar << Handle.Actor;
-	Ar << Handle.Manager;
-	Ar << Handle.InstanceIndex;
-
-	return Ar;
-}
-
-FString FHitResult::ToString() const
-{
-	return FString::Printf(TEXT("bBlockingHit:%s bStartPenetrating:%s Time:%f Location:%s ImpactPoint:%s Normal:%s ImpactNormal:%s TraceStart:%s TraceEnd:%s PenetrationDepth:%f Item:%d PhysMaterial:%s Actor:%s Component:%s BoneName:%s FaceIndex:%d"),
-		bBlockingHit == true ? TEXT("True") : TEXT("False"),
-		bStartPenetrating == true ? TEXT("True") : TEXT("False"),
-		Time,
-		*Location.ToString(),
-		*ImpactPoint.ToString(),
-		*Normal.ToString(),
-		*ImpactNormal.ToString(),
-		*TraceStart.ToString(),
-		*TraceEnd.ToString(),
-		PenetrationDepth,
-		Item,
-		PhysMaterial.IsValid() ? *PhysMaterial->GetName() : TEXT("None"),
-		*FLightWeightInstanceSubsystem::Get().GetName(HitObjectHandle),
-		Component.IsValid() ? *Component->GetName() : TEXT("None"),
-		BoneName.IsValid() ? *BoneName.ToString() : TEXT("None"),
-		FaceIndex);
-}
-
-FRepMovement::FRepMovement()
-	: LinearVelocity(ForceInit)
-	, AngularVelocity(ForceInit)
-	, Location(ForceInit)
-	, Rotation(ForceInit)
-	, bSimulatedPhysicSleep(false)
-	, bRepPhysics(false)
-	, ServerFrame(0)
-	, LocationQuantizationLevel(EVectorQuantization::RoundWholeNumber)
-	, VelocityQuantizationLevel(EVectorQuantization::RoundWholeNumber)
-	, RotationQuantizationLevel(ERotatorQuantization::ByteComponents)
-{
-}
-
-// compile switch for disabling quantization of replicated movement, meant for testing.
-#ifndef REP_MOVEMENT_DISABLE_QUANTIZATION
-#define REP_MOVEMENT_DISABLE_QUANTIZATION 0
-#endif
-
-bool FRepMovement::NetSerialize(FArchive& Ar, class UPackageMap* Map, bool& bOutSuccess)
-{
-	// pack bitfield with flags
-	const bool bServerFrameAndHandleSupported = Ar.EngineNetVer() >= HISTORY_REPMOVE_SERVERFRAME_AND_HANDLE;
-	uint8 Flags = (bSimulatedPhysicSleep << 0) | (bRepPhysics << 1) | ((ServerFrame > 0) << 2) | ((ServerPhysicsHandle != INDEX_NONE) << 3);
-	Ar.SerializeBits(&Flags, bServerFrameAndHandleSupported ? 4 : 2);
-	bSimulatedPhysicSleep = ( Flags & ( 1 << 0 ) ) ? 1 : 0;
-	bRepPhysics = ( Flags & ( 1 << 1 ) ) ? 1 : 0;
-	const bool bRepServerFrame = (Flags & (1 << 2) && bServerFrameAndHandleSupported) ? 1 : 0;
-	const bool bRepServerHandle = (Flags & (1 << 3) && bServerFrameAndHandleSupported) ? 1 : 0;
-
-	bOutSuccess = true;
-
-#if REP_MOVEMENT_DISABLE_QUANTIZATION
-	Ar << Location;
-	Ar << LinearVelocity;
-	Ar << Rotation;
-	if (bRepPhysics)
-	{
-		Ar << AngularVelocity;
-	}
-#else
-
-	// update location, rotation, linear velocity
-	bOutSuccess &= SerializeQuantizedVector( Ar, Location, LocationQuantizationLevel );
-
-	switch(RotationQuantizationLevel)
-	{
-	case ERotatorQuantization::ByteComponents:
-	{
-		Rotation.SerializeCompressed( Ar );
-		break;
-	}
-
-	case ERotatorQuantization::ShortComponents:
-	{
-		Rotation.SerializeCompressedShort( Ar );
-		break;
-	}
-	}
-
-	bOutSuccess &= SerializeQuantizedVector( Ar, LinearVelocity, VelocityQuantizationLevel );
-
-	// update angular velocity if required
-	if ( bRepPhysics )
-	{
-		bOutSuccess &= SerializeQuantizedVector( Ar, AngularVelocity, VelocityQuantizationLevel );
-	}
-#endif
-
-	if (bRepServerFrame)
-	{
-		uint32 uServerFrame = (uint32)ServerFrame;
-		Ar.SerializeIntPacked(uServerFrame);
-		ServerFrame = (int32)uServerFrame;
-	}
-
-	if (bRepServerHandle)
-	{
-		uint32 uServerPhysicsHandle = (uint32)ServerPhysicsHandle;
-		Ar.SerializeIntPacked(uServerPhysicsHandle);
-		ServerPhysicsHandle = (int32)uServerPhysicsHandle;
-	}
-
-	return true;
-}
-
-/** Rebase zero-origin position onto local world origin value. */
-FVector FRepMovement::RebaseOntoLocalOrigin(const FVector& Location, const FIntVector& LocalOrigin)
-{
-	if (EnableMultiplayerWorldOriginRebasing <= 0 || LocalOrigin == FIntVector::ZeroValue)
-	{
-		return Location;
-	}
-
-	return FVector(Location.X - LocalOrigin.X, Location.Y - LocalOrigin.Y, Location.Z - LocalOrigin.Z);
-}
-
-/** Rebase local-origin position onto zero world origin value. */
-FVector FRepMovement::RebaseOntoZeroOrigin(const FVector& Location, const FIntVector& LocalOrigin)
-{
-	if (EnableMultiplayerWorldOriginRebasing <= 0 || LocalOrigin == FIntVector::ZeroValue)
-	{
-		return Location;
-	}
-
-	return FVector(Location.X + LocalOrigin.X, Location.Y + LocalOrigin.Y, Location.Z + LocalOrigin.Z);
-}
-
-/** Rebase zero-origin position onto local world origin value based on an actor's world. */
-FVector FRepMovement::RebaseOntoLocalOrigin(const FVector& Location, const AActor* const WorldContextActor)
-{
-	if (WorldContextActor == nullptr || EnableMultiplayerWorldOriginRebasing <= 0)
-	{
-		return Location;
-	}
-
-	return RebaseOntoLocalOrigin(Location, WorldContextActor->GetWorld()->OriginLocation);
-}
-
-/** Rebase local-origin position onto zero world origin value based on an actor's world.*/
-FVector FRepMovement::RebaseOntoZeroOrigin(const FVector& Location, const AActor* const WorldContextActor)
-{
-	if (WorldContextActor == nullptr || EnableMultiplayerWorldOriginRebasing <= 0)
-	{
-		return Location;
-	}
-
-	return RebaseOntoZeroOrigin(Location, WorldContextActor->GetWorld()->OriginLocation);
-}
-
-/// @cond DOXYGEN_WARNINGS
-
-/** Rebase zero-origin position onto local world origin value based on an actor component's world. */
-FVector FRepMovement::RebaseOntoLocalOrigin(const FVector& Location, const UActorComponent* const WorldContextActorComponent)
-{
-	if (WorldContextActorComponent == nullptr || EnableMultiplayerWorldOriginRebasing <= 0)
-	{
-		return Location;
-	}
-
-	return RebaseOntoLocalOrigin(Location, WorldContextActorComponent->GetWorld()->OriginLocation);
-}
-
-/** Rebase local-origin position onto zero world origin value based on an actor component's world.*/
-FVector FRepMovement::RebaseOntoZeroOrigin(const FVector& Location, const UActorComponent* const WorldContextActorComponent)
-{
-	if (WorldContextActorComponent == nullptr || EnableMultiplayerWorldOriginRebasing <= 0)
-	{
-		return Location;
-	}
-
-	return RebaseOntoZeroOrigin(Location, WorldContextActorComponent->GetWorld()->OriginLocation);
-}
-
 const TCHAR* LexToString(const EWorldType::Type Value)
 {
 	switch (Value)
diff --git a/Engine/Source/Runtime/Engine/Private/GameplayStatics.cpp b/Engine/Source/Runtime/Engine/Private/GameplayStatics.cpp
--- a/Engine/Source/Runtime/Engine/Private/GameplayStatics.cpp
+++ b/Engine/Source/Runtime/Engine/Private/GameplayStatics.cpp
@@ -49,6 +49,7 @@
 #include "Sound/SoundCue.h"
 #include "Sound/SoundWave.h"
 #include "Audio/ActorSoundParameterInterface.h"
+#include "Engine/DamageEvents.h"
 
 #if WITH_ACCESSIBILITY
 #include "Framework/Application/SlateApplication.h"
diff --git a/Engine/Source/Runtime/Engine/Private/PainCausingVolume.cpp b/Engine/Source/Runtime/Engine/Private/PainCausingVolume.cpp
--- a/Engine/Source/Runtime/Engine/Private/PainCausingVolume.cpp
+++ b/Engine/Source/Runtime/Engine/Private/PainCausingVolume.cpp
@@ -3,6 +3,7 @@
 #include "GameFramework/PainCausingVolume.h"
 #include "TimerManager.h"
 #include "GameFramework/Pawn.h"
+#include "Engine/DamageEvents.h"
 
 APainCausingVolume::APainCausingVolume(const FObjectInitializer& ObjectInitializer)
 	: Super(ObjectInitializer)
diff --git a/Engine/Source/Runtime/Engine/Private/PhysicsVolume.cpp b/Engine/Source/Runtime/Engine/Private/PhysicsVolume.cpp
--- a/Engine/Source/Runtime/Engine/Private/PhysicsVolume.cpp
+++ b/Engine/Source/Runtime/Engine/Private/PhysicsVolume.cpp
@@ -5,6 +5,7 @@
 #include "Components/PrimitiveComponent.h"
 #include "Components/BrushComponent.h"
 #include "PhysicsEngine/PhysicsSettings.h"
+#include "Engine/OverlapInfo.h"
 
 APhysicsVolume::APhysicsVolume(const FObjectInitializer& ObjectInitializer)
 	: Super(ObjectInitializer)
diff --git a/Engine/Source/Runtime/Engine/Public/CharacterMovementComponentAsync.h b/Engine/Source/Runtime/Engine/Public/CharacterMovementComponentAsync.h
--- a/Engine/Source/Runtime/Engine/Public/CharacterMovementComponentAsync.h
+++ b/Engine/Source/Runtime/Engine/Public/CharacterMovementComponentAsync.h
@@ -3,6 +3,7 @@
 #pragma once
 #include "GameFramework/RootMotionSource.h"
 #include "Chaos/SimCallbackObject.h"
+#include "Engine/OverlapInfo.h"
 #include "CharacterMovementComponentAsync.generated.h"
 
 
diff --git a/Engine/Source/Runtime/Engine/Public/Engine/OverlapInfo.h b/Engine/Source/Runtime/Engine/Public/Engine/OverlapInfo.h
new file mode 100644
--- /dev/null
+++ b/Engine/Source/Runtime/Engine/Public/Engine/OverlapInfo.h
@@ -0,0 +1,67 @@
+// Copyright Epic Games, Inc. All Rights Reserved.
+
+#pragma once
+
+#include "CoreMinimal.h"
+#include "Engine/HitResult.h"
+
+class UPrimitiveComponent;
+
+/** Overlap info consisting of the primitive and the body that is overlapping */
+struct ENGINE_API FOverlapInfo
+{
+	FOverlapInfo()
+	{}
+
+	explicit FOverlapInfo(const FHitResult& InSweepResult)
+		: bFromSweep(true), OverlapInfo(InSweepResult)
+	{
+	}
+
+	explicit FOverlapInfo(UPrimitiveComponent* InComponent, int32 InBodyIndex = INDEX_NONE);
+
+	int32 GetBodyIndex() const { return OverlapInfo.Item; }
+
+	//This function completely ignores SweepResult information. It seems that places that use this function do not care, but it still seems risky
+	friend bool operator == (const FOverlapInfo& LHS, const FOverlapInfo& RHS) { return LHS.OverlapInfo.Component == RHS.OverlapInfo.Component && LHS.OverlapInfo.Item == RHS.OverlapInfo.Item; }
+	bool bFromSweep;
+
+	/** Information for both sweep and overlap queries. Different parts are valid depending on bFromSweep.
+	  * If bFromSweep is true then FHitResult is completely valid just like a regular sweep result.
+	  * If bFromSweep is false only FHitResult::Component, FHitResult::Actor, FHitResult::Item are valid as this is really just an FOverlapResult*/
+	FHitResult OverlapInfo;
+};
+
+/*
+ * Predicate for comparing FOverlapInfos when exact weak object pointer index/serial numbers should match, assuming one is not null and not invalid.
+ * Compare to operator== for WeakObjectPtr which does both HasSameIndexAndSerialNumber *and* IsValid() checks on both pointers.
+ */
+struct FFastOverlapInfoCompare
+{
+	FFastOverlapInfoCompare(const FOverlapInfo& BaseInfo)
+		: MyBaseInfo(BaseInfo)
+	{
+	}
+
+	bool operator() (const FOverlapInfo& Info)
+	{
+		return MyBaseInfo.OverlapInfo.Component.HasSameIndexAndSerialNumber(Info.OverlapInfo.Component)
+			&& MyBaseInfo.GetBodyIndex() == Info.GetBodyIndex();
+	}
+
+	bool operator() (const FOverlapInfo* Info)
+	{
+		return MyBaseInfo.OverlapInfo.Component.HasSameIndexAndSerialNumber(Info->OverlapInfo.Component)
+			&& MyBaseInfo.GetBodyIndex() == Info->GetBodyIndex();
+	}
+
+private:
+	const FOverlapInfo& MyBaseInfo;
+
+};
+
+// All added members of FOverlapInfo are PODs.
+template<> struct TIsPODType<FOverlapInfo> { enum { Value = TIsPODType<FHitResult>::Value }; };
+
+typedef TArray<FOverlapInfo, TInlineAllocator<3>> TInlineOverlapInfoArray;
+typedef TArrayView<const FOverlapInfo> TOverlapArrayView;
\ No newline at end of file
diff --git a/Engine/Source/Runtime/Engine/Public/Engine/ScopedMovementUpdate.h b/Engine/Source/Runtime/Engine/Public/Engine/ScopedMovementUpdate.h
new file mode 100644
--- /dev/null
+++ b/Engine/Source/Runtime/Engine/Public/Engine/ScopedMovementUpdate.h
@@ -0,0 +1,211 @@
+// Copyright Epic Games, Inc. All Rights Reserved.
+
+#pragma once
+
+#include "CoreMinimal.h"
+#include "Engine/HitResult.h"
+#include "Engine/OverlapInfo.h"
+
+/**
+ * Enum that controls the scoping behavior of FScopedMovementUpdate.
+ * Note that EScopedUpdate::ImmediateUpdates is not allowed within outer scopes that defer updates,
+ * and any attempt to do so will change the new inner scope to use deferred updates instead.
+ */
+namespace EScopedUpdate
+{
+	enum Type
+	{
+		ImmediateUpdates,
+		DeferredUpdates	
+	};
+}
+
+/**
+ * FScopedMovementUpdate creates a new movement scope, within which propagation of moves may be deferred until the end of the outermost scope that does not defer updates.
+ * Moves within this scope will avoid updates such as UpdateBounds(), OnUpdateTransform(), UpdatePhysicsVolume(), UpdateChildTransforms() etc
+ * until the move is committed (which happens when the last deferred scope goes out of context).
+ *
+ * Note that non-deferred scopes are not allowed within outer scopes that defer updates, and any attempt to use one will change the inner scope to use deferred updates.
+ */
+class ENGINE_API FScopedMovementUpdate : private FNoncopyable
+{
+public:
+	typedef TArray<FHitResult, TInlineAllocator<2>> TScopedBlockingHitArray;
+	typedef TArray<FOverlapInfo, TInlineAllocator<3>> TScopedOverlapInfoArray;
+
+	FScopedMovementUpdate( USceneComponent* Component, EScopedUpdate::Type ScopeBehavior = EScopedUpdate::DeferredUpdates, bool bRequireOverlapsEventFlagToQueueOverlaps = true );
+	~FScopedMovementUpdate();
+
+	enum class EHasMovedTransformOption
+	{
+		eTestTransform,
+		eIgnoreTransform
+	};
+
+	enum class EOverlapState
+	{
+		eUseParent,
+		eUnknown,
+		eIncludesOverlaps,
+		eForceUpdate,
+	};
+
+	/** Get the scope containing this scope. A scope only has an outer scope if they both defer updates. */
+	const FScopedMovementUpdate* GetOuterDeferredScope() const;
+
+	/** Return true if deferring updates, false if updates are applied immediately. */
+	bool IsDeferringUpdates() const;
+	
+	/** Revert movement to the initial location of the Component at the start of the scoped update. Also clears pending overlaps and sets bHasMoved to false. */
+	void RevertMove();
+
+	/** Returns whether movement has occurred at all during this scope, optionally checking if the transform is different (since changing scale does not go through a move). RevertMove() sets this back to false. */
+	bool HasMoved(EHasMovedTransformOption CheckTransform) const;
+
+	/** Returns true if the Component's transform differs from that at the start of the scoped update. */
+	bool IsTransformDirty() const;
+
+	/** Returns true if there are pending overlaps queued in this scope. */
+	bool HasPendingOverlaps() const;
+
+	/**
+	* Returns true if we require GetGenerateOverlapEvents() on both the moving object and the overlapped object to add them to the pending overlaps list.
+	* These flags will still be required when dispatching calls to UpdateOverlaps(), but this allows some custom processing of queued overlaps that would be otherwise missed along the way.
+	*/
+	bool RequiresOverlapsEventFlag() const;
+
+	/** Returns the pending overlaps within this scope. */
+	const TScopedOverlapInfoArray& GetPendingOverlaps() const;
+
+	/** Returns the list of pending blocking hits, which will be used for notifications once the move is committed. */
+	const TScopedBlockingHitArray& GetPendingBlockingHits() const;
+
+	//--------------------------------------------------------------------------------------------------------//
+	// These methods are intended only to be used by SceneComponent and derived classes.
+
+	/** Add overlaps to the queued overlaps array. This is intended for use only by SceneComponent and its derived classes whenever movement is performed. */
+	void AppendOverlapsAfterMove(const TOverlapArrayView& NewPendingOverlaps, bool bSweep, bool bIncludesOverlapsAtEnd);
+
+	/** Keep current pending overlaps after a move but make note that there was movement (just a symmetric rotation). */
+	void KeepCurrentOverlapsAfterRotation(bool bSweep);
+
+	/** Add blocking hit that will get processed once the move is committed. This is intended for use only by SceneComponent and its derived classes. */
+	void AppendBlockingHitAfterMove(const FHitResult& Hit);
+
+	/** Clear overlap state at current location, we don't know what it is. */
+	void InvalidateCurrentOverlaps();
+
+	/** Force full overlap update once this scope finishes. */
+	void ForceOverlapUpdate();
+
+	/** Registers that this move is a teleport */
+	void SetHasTeleported(ETeleportType InTeleportType);
+
+protected:
+	/** Fills in the list of overlaps at the end location (in EndOverlaps). Returns pointer to the list, or null if it can't be computed. */
+	TOptional<TOverlapArrayView> GetOverlapsAtEnd(class UPrimitiveComponent& PrimComponent, TInlineOverlapInfoArray& OutEndOverlaps, bool bTransformChanged) const;
+
+	bool SetWorldLocationAndRotation(FVector NewLocation, const FQuat& NewQuat, bool bNoPhysics = false, ETeleportType Teleport = ETeleportType::None);
+
+private:
+	/** Notify this scope that the given inner scope completed its update (ie is going out of scope). Only occurs for deferred updates. */
+	void OnInnerScopeComplete(const FScopedMovementUpdate& InnerScope);
+	
+	// This class can only be created on the stack, otherwise the ordering constraints
+	// of the constructor and destructor between encapsulated scopes could be violated.
+	void*	operator new		(size_t);
+	void*	operator new[]		(size_t);
+	void	operator delete		(void *);
+	void	operator delete[]	(void*);
+
+protected:
+	USceneComponent* Owner;
+	FScopedMovementUpdate* OuterDeferredScope;
+
+	EOverlapState CurrentOverlapState;
+	ETeleportType TeleportType;
+
+	FTransform InitialTransform;
+	FVector InitialRelativeLocation;
+	FRotator InitialRelativeRotation;
+	FVector InitialRelativeScale;
+
+	int32 FinalOverlapCandidatesIndex;		// If not INDEX_NONE, overlaps at this index and beyond in PendingOverlaps are at the final destination
+	TScopedOverlapInfoArray PendingOverlaps;	// All overlaps encountered during the scope of moves.
+	TScopedBlockingHitArray BlockingHits;		// All blocking hits encountered during the scope of moves.
+
+	uint8 bDeferUpdates:1;
+	uint8 bHasMoved:1;
+	uint8 bRequireOverlapsEventFlag:1;
+
+	friend class USceneComponent;
+};
+
+//////////////////////////////////////////////////////////////////////////
+// FScopedMovementUpdate inlines
+
+FORCEINLINE const FScopedMovementUpdate* FScopedMovementUpdate::GetOuterDeferredScope() const
+{
+	return OuterDeferredScope;
+}
+
+FORCEINLINE bool FScopedMovementUpdate::IsDeferringUpdates() const
+{
+	return bDeferUpdates;
+}
+
+FORCEINLINE bool FScopedMovementUpdate::HasMoved(EHasMovedTransformOption CheckTransform) const
+{
+	return bHasMoved || (CheckTransform == EHasMovedTransformOption::eTestTransform && IsTransformDirty());
+}
+
+FORCEINLINE bool FScopedMovementUpdate::HasPendingOverlaps() const
+{
+	return PendingOverlaps.Num() > 0;
+}
+
+FORCEINLINE bool FScopedMovementUpdate::RequiresOverlapsEventFlag() const
+{
+	return bRequireOverlapsEventFlag;
+}
+
+FORCEINLINE const FScopedMovementUpdate::TScopedOverlapInfoArray& FScopedMovementUpdate::GetPendingOverlaps() const
+{
+	return PendingOverlaps;
+}
+
+FORCEINLINE const FScopedMovementUpdate::TScopedBlockingHitArray& FScopedMovementUpdate::GetPendingBlockingHits() const
+{
+	return BlockingHits;
+}
+
+FORCEINLINE_DEBUGGABLE void FScopedMovementUpdate::AppendBlockingHitAfterMove(const FHitResult& Hit)
+{
+	BlockingHits.Add(Hit);
+}
+
+FORCEINLINE_DEBUGGABLE void FScopedMovementUpdate::KeepCurrentOverlapsAfterRotation(bool bSweep)
+{
+	bHasMoved = true;
+	// CurrentOverlapState is unchanged
+}
+
+FORCEINLINE_DEBUGGABLE void FScopedMovementUpdate::InvalidateCurrentOverlaps()
+{
+	bHasMoved = true;
+	CurrentOverlapState = EOverlapState::eUnknown;
+	FinalOverlapCandidatesIndex = INDEX_NONE;
+}
+
+FORCEINLINE_DEBUGGABLE void FScopedMovementUpdate::ForceOverlapUpdate()
+{
+	bHasMoved = true;
+	CurrentOverlapState = EOverlapState::eForceUpdate;
+	FinalOverlapCandidatesIndex = INDEX_NONE;
+}
+
+FORCEINLINE_DEBUGGABLE void FScopedMovementUpdate::SetHasTeleported(ETeleportType InTeleportType)
+{
+	// Request an initialization. Teleport type can only go higher - i.e. if we have requested a reset, then a teleport will still reset fully
+	TeleportType = ((InTeleportType > TeleportType) ? InTeleportType : TeleportType); 
+}
diff --git a/Engine/Source/Runtime/Engine/Public/PhysicsPublic.h b/Engine/Source/Runtime/Engine/Public/PhysicsPublic.h
--- a/Engine/Source/Runtime/Engine/Public/PhysicsPublic.h
+++ b/Engine/Source/Runtime/Engine/Public/PhysicsPublic.h
@@ -14,6 +14,11 @@
 #include "Misc/App.h"
 #include "EngineDefines.h"
 #include "PhysicsInterfaceDeclaresCore.h"
+#if UE_ENABLE_INCLUDE_ORDER_DEPRECATED_IN_5_1
+#include "RenderResource.h"
+#include "LocalVertexFactory.h"
+#include "DynamicMeshBuilder.h"
+#endif
 #include "PhysicsPublicCore.h"
 //#include "StaticMeshResources.h"
 
diff --git a/Engine/Source/Runtime/Engine/Public/WorldCollision.h b/Engine/Source/Runtime/Engine/Public/WorldCollision.h
--- a/Engine/Source/Runtime/Engine/Public/WorldCollision.h
+++ b/Engine/Source/Runtime/Engine/Public/WorldCollision.h
@@ -9,6 +9,7 @@
 #include "Async/TaskGraphInterfaces.h"
 #include "CollisionQueryParams.h"
 #include "CollisionShape.h"
+#include "Engine/OverlapResult.h"
 
 struct FOverlapDatum;
 struct FTraceDatum;
diff --git a/Engine/Source/Runtime/Foliage/Private/FoliageComponent.cpp b/Engine/Source/Runtime/Foliage/Private/FoliageComponent.cpp
--- a/Engine/Source/Runtime/Foliage/Private/FoliageComponent.cpp
+++ b/Engine/Source/Runtime/Foliage/Private/FoliageComponent.cpp
@@ -13,6 +13,7 @@
 #include "InteractiveFoliageActor.h"
 #include "InteractiveFoliageComponent.h"
 #include "Engine/StaticMesh.h"
+#include "Engine/DamageEvents.h"
 
 /** Scene proxy class for UInteractiveFoliageComponent. */
 class FInteractiveFoliageSceneProxy final : public FStaticMeshSceneProxy
diff --git a/Engine/Source/Runtime/Foliage/Private/InstancedFoliage.cpp b/Engine/Source/Runtime/Foliage/Private/InstancedFoliage.cpp
--- a/Engine/Source/Runtime/Foliage/Private/InstancedFoliage.cpp
+++ b/Engine/Source/Runtime/Foliage/Private/InstancedFoliage.cpp
@@ -55,6 +55,7 @@ InstancedFoliage.cpp: Instanced foliage implementation.
 #include "Algo/Transform.h"
 #include "ActorPartition/ActorPartitionSubsystem.h"
 #include "Misc/CoreMisc.h"
+#include "Engine/DamageEvents.h"
 
 #define LOCTEXT_NAMESPACE "InstancedFoliage"
 
diff --git a/Engine/Source/Runtime/MovieSceneTracks/Public/Systems/MovieSceneDeferredComponentMovementSystem.h b/Engine/Source/Runtime/MovieSceneTracks/Public/Systems/MovieSceneDeferredComponentMovementSystem.h
--- a/Engine/Source/Runtime/MovieSceneTracks/Public/Systems/MovieSceneDeferredComponentMovementSystem.h
+++ b/Engine/Source/Runtime/MovieSceneTracks/Public/Systems/MovieSceneDeferredComponentMovementSystem.h
@@ -5,9 +5,14 @@
 #include "EntitySystem/MovieSceneEntitySystem.h"
 #include "Containers/ChunkedArray.h"
 #include "Misc/Optional.h"
+#include "Engine/ScopedMovementUpdate.h"
+#if UE_ENABLE_INCLUDE_ORDER_DEPRECATED_IN_5_1
 #include "Components/SceneComponent.h"
+#endif
 #include "MovieSceneDeferredComponentMovementSystem.generated.h"
 
+class USceneComponent;
+
 /**
  * System that maintains a FScopedMovementUpdate for any USceneComponent that has an
  * animated transform or attachment for the duration of the evaluation to avoid repeatedly
