Fix thumbnail rendering.

FSlateTextureRenderTarget2DResource contained a separate texture resource and manually synced the TextureRHI member to it (TextureRHI is inherited from FRHITexture via FTextureRenderTargetResource). CL 19548181 broke this, but that member isn't needed at all after the texture type refactor, so this change removes it and just uses the inherited member.

#rnx

diff --git a/Engine/Source/Runtime/Engine/Private/Slate/SlateTextures.cpp b/Engine/Source/Runtime/Engine/Private/Slate/SlateTextures.cpp
--- a/Engine/Source/Runtime/Engine/Private/Slate/SlateTextures.cpp
+++ b/Engine/Source/Runtime/Engine/Private/Slate/SlateTextures.cpp
@@ -321,7 +321,7 @@ void FSlateTextureRenderTarget2DResource::InitDynamicRHI()
 			.SetFormat((EPixelFormat)Format)
 			.SetClearValue(FClearValueBinding(ClearColor));
 
-		RHICreateTargetableShaderResource(Desc, ETextureCreateFlags::RenderTargetable, RenderTargetTextureRHI, Texture2DRHI);
+		RHICreateTargetableShaderResource(Desc, ETextureCreateFlags::RenderTargetable, RenderTargetTextureRHI, TextureRHI);
 	}
 
 	// Create the sampler state RHI resource.
@@ -342,7 +342,6 @@ void FSlateTextureRenderTarget2DResource::ReleaseDynamicRHI()
 	// Release the FTexture RHI resources here as well
 	ReleaseRHI();
 
-	Texture2DRHI.SafeRelease();
 	RenderTargetTextureRHI.SafeRelease();	
 
 	// Remove from global list of deferred clears
diff --git a/Engine/Source/Runtime/Engine/Public/Slate/SlateTextures.h b/Engine/Source/Runtime/Engine/Public/Slate/SlateTextures.h
--- a/Engine/Source/Runtime/Engine/Public/Slate/SlateTextures.h
+++ b/Engine/Source/Runtime/Engine/Public/Slate/SlateTextures.h
@@ -154,7 +154,7 @@ public:
 	virtual void SetSize(int32 InSizeX,int32 InSizeY);
 
 	/** Gets the RHI resource for this render target */
-	FTexture2DRHIRef GetTextureRHI() { return Texture2DRHI; }
+	FTexture2DRHIRef GetTextureRHI() { return TextureRHI; }
 
 public:
 	// FTextureRenderTargetResource implementation
@@ -177,7 +177,6 @@ protected:
 
 
 private:
-	FTexture2DRHIRef Texture2DRHI;
 	FLinearColor ClearColor;
 	int32 TargetSizeX,TargetSizeY;
 
