Adding missing legal header.

diff --git a/Engine/Source/Runtime/Engine/Classes/Engine/OverlapResult.h b/Engine/Source/Runtime/Engine/Classes/Engine/OverlapResult.h
--- a/Engine/Source/Runtime/Engine/Classes/Engine/OverlapResult.h
+++ b/Engine/Source/Runtime/Engine/Classes/Engine/OverlapResult.h
@@ -1,3 +1,4 @@
+// Copyright Epic Games, Inc. All Rights Reserved.
 
 #pragma once
 
