Server crash in UAnimSequence::GetNumberOfSampledKeys
#fix ensure compressed bone data is valid before retrieving # sampled keys, otherwise recalculate with target frame rate and play length (curve only animation sequence)

diff --git a/Engine/Source/Runtime/Engine/Private/Animation/AnimSequence.cpp b/Engine/Source/Runtime/Engine/Private/Animation/AnimSequence.cpp
--- a/Engine/Source/Runtime/Engine/Private/Animation/AnimSequence.cpp
+++ b/Engine/Source/Runtime/Engine/Private/Animation/AnimSequence.cpp
@@ -1240,11 +1240,7 @@ void UAnimSequence::ExtractBoneTransform(const struct FRawAnimSequenceTrack& Raw
 
 void UAnimSequence::ExtractBoneTransform(const struct FRawAnimSequenceTrack& RawTrack, FTransform& OutAtom, float Time) const
 {
-#if WITH_EDITOR
-	FAnimationUtils::ExtractTransformFromTrack(Time, NumberOfSampledKeys, GetPlayLength(), RawTrack, Interpolation, OutAtom);
-#else
-	FAnimationUtils::ExtractTransformFromTrack(Time, CompressedData.CompressedDataStructure->CompressedNumberOfKeys, GetPlayLength(), RawTrack, Interpolation, OutAtom);
-#endif
+	FAnimationUtils::ExtractTransformFromTrack(Time, GetNumberOfSampledKeys(), GetPlayLength(), RawTrack, Interpolation, OutAtom);
 }
 
 void UAnimSequence::HandleAssetPlayerTickedInternal(FAnimAssetTickContext &Context, const float PreviousTime, const float MoveDelta, const FAnimTickRecord &Instance, struct FAnimNotifyQueue& NotifyQueue) const
@@ -4365,7 +4361,7 @@ int32 UAnimSequence::GetNumberOfSampledKeys() const
 #if WITH_EDITOR
 	return NumberOfSampledKeys;
 #else
-	return CompressedData.CompressedDataStructure->CompressedNumberOfKeys;
+	return CompressedData.CompressedDataStructure ? CompressedData.CompressedDataStructure->CompressedNumberOfKeys : TargetFrameRate.AsFrameNumber(GetPlayLength()).Value + 1;
 #endif
 }
 
