- WorldPartition Package Cache now shares a global cache.
- Fixes crashes when loading partitioned level instances.

diff --git a/Engine/Source/Runtime/Engine/Private/WorldPartition/WorldPartitionLevelHelper.cpp b/Engine/Source/Runtime/Engine/Private/WorldPartition/WorldPartitionLevelHelper.cpp
--- a/Engine/Source/Runtime/Engine/Private/WorldPartition/WorldPartitionLevelHelper.cpp
+++ b/Engine/Source/Runtime/Engine/Private/WorldPartition/WorldPartitionLevelHelper.cpp
@@ -16,6 +16,7 @@
 #include "WorldPartition/ActorDescContainer.h"
 #include "WorldPartition/WorldPartitionActorDesc.h"
 #include "WorldPartition/WorldPartition.h"
+#include "WorldPartition/WorldPartitionPackageHelper.h"
 #include "StaticMeshCompiler.h"
 #include "LevelUtils.h"
 #include "Templates/SharedPointer.h"
@@ -337,34 +338,12 @@ bool FWorldPartitionLevelHelper::LoadActors(ULevel* InDestLevel, TArrayView<FWor
 			if (LoadedPackage)
 			{
 				FName DuplicatePackageName(*FString::Printf(TEXT("%s_%s"), *LoadedPackage->GetName(), *ContainerInstanceID.ToString()));
-				UPackage* DuplicatedPackage = InPackageCache.FindPackage(DuplicatePackageName);
-				UWorld* DuplicatedWorld = nullptr;
-
-				// Check that existing Package contains all actors (might be a previously used Dup package when loading unloaded cells)
-				if (DuplicatedPackage)
-				{
-					DuplicatedWorld = UWorld::FindWorldInPackage(DuplicatedPackage);
-					for (auto Mapping : Mappings)
-					{
-						FString ActorName = FPaths::GetExtension(Mapping->Path.ToString());
-						if (!FindObject<AActor>(DuplicatedWorld->PersistentLevel, *ActorName))
-						{
-							InPackageCache.TrashPackage(DuplicatedPackage);
-							DuplicatedPackage = nullptr;
-							DuplicatedWorld = nullptr;
-							break;
-						}
-					}
-				}
-
-				// Package wasn't found or was missing actors
-				if (!DuplicatedPackage)
-				{
-					DuplicatedPackage = InPackageCache.DuplicateWorldPackage(LoadedPackage, DuplicatePackageName);
-					check(DuplicatedPackage);
-					DuplicatedWorld = UWorld::FindWorldInPackage(DuplicatedPackage);
-				}
-
+				check(!FindPackage(nullptr, *DuplicatePackageName.ToString()));
+				UPackage* DuplicatedPackage = FWorldPartitionPackageHelper::DuplicateWorldPackage(LoadedPackage, DuplicatePackageName);
+				check(DuplicatedPackage);
+				UWorld* DuplicatedWorld = UWorld::FindWorldInPackage(DuplicatedPackage);
+				check(DuplicatedWorld);
+				
 				// Get Source Path for this world (ReplaceFrom)
 				UWorld* LoadedWorld = UWorld::FindWorldInPackage(LoadedPackage);
 				FString SourceWorldPath, RemappedWorldPath;
@@ -408,6 +387,18 @@ bool FWorldPartitionLevelHelper::LoadActors(ULevel* InDestLevel, TArrayView<FWor
 						UE_LOG(LogEngine, Warning, TEXT("Failed to find actor %s in package %s"), *ActorName, *LoadedPackageName.ToString());
 					}
 				}
+
+				if (InDestLevel)
+				{
+					// In the case where we have a destination level, we don't need the duplicated package anymore
+					FWorldPartitionPackageHelper::UnloadPackage(DuplicatedPackage);
+				}
+				else
+				{
+					// Push the package in the cache so that package unloading will be handled by the cache destruction
+					// The only existing case is during cook (Actors are loaded with a null destination level, then MoveExternalActorsToLevel is called).
+					InPackageCache.CacheLoadedPackage(DuplicatedPackage);
+				}
 			}
 			else
 			{
@@ -431,4 +422,4 @@ bool FWorldPartitionLevelHelper::LoadActors(ULevel* InDestLevel, TArrayView<FWor
 	return (LoadProgress->NumPendingLoadRequests == 0);
 }
 
-#endif
+#endif
\ No newline at end of file
diff --git a/Engine/Source/Runtime/Engine/Private/WorldPartition/WorldPartitionPackageCache.cpp b/Engine/Source/Runtime/Engine/Private/WorldPartition/WorldPartitionPackageCache.cpp
--- a/Engine/Source/Runtime/Engine/Private/WorldPartition/WorldPartitionPackageCache.cpp
+++ b/Engine/Source/Runtime/Engine/Private/WorldPartition/WorldPartitionPackageCache.cpp
@@ -7,96 +7,78 @@
 #include "Engine/World.h"
 #include "PackageTools.h"
 #include "Misc/PackagePath.h"
+#include "WorldPartition/WorldPartition.h"
+#include "WorldPartition/WorldPartitionPackageHelper.h"
 
-FWorldPartitionPackageCache::FWorldPartitionPackageCache()
-{}
+/*
+ * FWorldPartitionPackageGlobalCache
+ */
 
-FWorldPartitionPackageCache::~FWorldPartitionPackageCache()
-{
-	check(!LoadingPackages.Num());
-	UnloadPackages();
-}
-
-void FWorldPartitionPackageCache::TrashPackage(UPackage* InPackage)
-{	
-	CachedPackages.Remove(InPackage->GetFName());
-	UnloadPackage(InPackage);
+FWorldPartitionPackageGlobalCache FWorldPartitionPackageGlobalCache::GlobalCache;
 
-	// Rename so it isn't found again
-	FName NewUniqueTrashName = MakeUniqueObjectName(nullptr, UPackage::StaticClass(), FName(*FString::Printf(TEXT("%s_Trashed"), *InPackage->GetName())));
-	InPackage->Rename(*NewUniqueTrashName.ToString(), nullptr, REN_DontCreateRedirectors | REN_ForceNoResetLoaders | REN_NonTransactional | REN_DoNotDirty);
-}
-
-void FWorldPartitionPackageCache::UnloadPackages()
+void FWorldPartitionPackageGlobalCache::UnloadPackages(FWorldPartitionPackageCache* InClient)
 {
-	if (CachedPackages.Num())
+	for (auto It = Cache.CreateIterator(); It; ++It)
 	{
-		for (const auto& Pair : CachedPackages)
+		FPackageCacheInfo& CacheInfo = It->Value;
+		CacheInfo.Referencers.Remove(InClient);
+		if (CacheInfo.Referencers.Num() == 0)
 		{
-			if (UPackage* CachedPackage = Pair.Value.Get())
+			check(!CacheInfo.HasPendingLoad());
+			if (UPackage* CachedPackage = CacheInfo.Package.Get())
 			{
-				UnloadPackage(CachedPackage);
+				FWorldPartitionPackageHelper::UnloadPackage(CachedPackage);
 			}
+			It.RemoveCurrent();
 		}
-				
-		CachedPackages.Empty();
 	}
 }
 
-void FWorldPartitionPackageCache::UnloadPackage(UPackage* InPackage)
+void FWorldPartitionPackageGlobalCache::LoadPackage(FWorldPartitionPackageCache* InClient, FName InPackageName, const TCHAR* InPackageToLoadFrom, FLoadPackageAsyncDelegate InCompletionDelegate, bool bLoadAsync, bool bInWorldPackage)
 {
-	check(InPackage);
-	ForEachObjectWithPackage(InPackage, [](UObject* Object)
-	{
-		Object->ClearFlags(RF_Standalone);
-		return true;
-	}, false);
+	FPackageCacheInfo& CacheInfo = Cache.FindOrAdd(InPackageName);
+	CacheInfo.Referencers.Add(InClient);
 
-	// World specific
-	if (UWorld* PackageWorld = UWorld::FindWorldInPackage(InPackage))
-	{
-		if (PackageWorld->PersistentLevel)
-		{
-			// Manual cleanup of level since world was not initialized
-			PackageWorld->PersistentLevel->CleanupLevel(/*bCleanupResources*/ true);
-		}
-	}
-}
-
-void FWorldPartitionPackageCache::LoadPackage(FName InPackageName, const TCHAR* InPackageToLoadFrom, FLoadPackageAsyncDelegate InCompletionDelegate, bool bLoadAsync, bool bInWorldPackage)
-{
-	if (UPackage* CachedPackage = FindPackage(InPackageName))
+	if (UPackage* CachedPackage = CacheInfo.Package.Get())
 	{
 		InCompletionDelegate.Execute(InPackageName, CachedPackage, EAsyncLoadingResult::Succeeded);
 		return;
 	}
+	check(!::FindPackage(nullptr, *InPackageName.ToString()));
 
-	// Find loading package
-	if (TArray<FLoadPackageAsyncDelegate>* CompletionDelegates = LoadingPackages.Find(InPackageName))
+	// If request already pending, simply push the delegate
+	if (CacheInfo.HasPendingLoad())
 	{
-		CompletionDelegates->Add(InCompletionDelegate);
+		CacheInfo.CompletionDelegates.Add(InCompletionDelegate);
 		return;
 	}
 
 	// Not found start loading
-	LoadingPackages.Add(InPackageName, { InCompletionDelegate });
+	CacheInfo.CompletionDelegates.Add(InCompletionDelegate);
 
-	FLoadPackageAsyncDelegate CompletionCallback = FLoadPackageAsyncDelegate::CreateLambda([this, bInWorldPackage](const FName& LoadedPackageName, UPackage* LoadedPackage, EAsyncLoadingResult::Type Result)
+	FLoadPackageAsyncDelegate CompletionCallback = FLoadPackageAsyncDelegate::CreateLambda([this, bInWorldPackage, InClient](const FName& LoadedPackageName, UPackage* LoadedPackage, EAsyncLoadingResult::Type Result)
 	{
+		FPackageCacheInfo& CacheInfo = Cache.FindChecked(LoadedPackageName);
+		check(CacheInfo.ScopedLoadAllExternalObjects.IsValid());
+		CacheInfo.ScopedLoadAllExternalObjects.Reset();
+
 		if (Result == EAsyncLoadingResult::Succeeded)
 		{
-			CachedPackages.Add(LoadedPackageName, LoadedPackage);
+			check(CacheInfo.Package.Get() == nullptr);
+			CacheInfo.Package = LoadedPackage;
 		}
 
 		if (bInWorldPackage)
 		{
 			if (UWorld* PackageWorld = UWorld::FindWorldInPackage(LoadedPackage))
 			{
+				check(PackageWorld->HasAnyFlags(RF_Standalone));
 				PackageWorld->PersistentLevel->OnLevelLoaded();
 			}
 		}
 	
-		TArray<FLoadPackageAsyncDelegate> CompletionDelegates = LoadingPackages.FindAndRemoveChecked(LoadedPackageName);
+		check(CacheInfo.HasPendingLoad());
+		TArray<FLoadPackageAsyncDelegate> CompletionDelegates = MoveTemp(CacheInfo.CompletionDelegates);
 		for (FLoadPackageAsyncDelegate& CompletionDelegate : CompletionDelegates)
 		{
 			CompletionDelegate.Execute(LoadedPackageName, LoadedPackage, Result);
@@ -109,6 +91,8 @@ void FWorldPartitionPackageCache::LoadPackage(FName InPackageName, const TCHAR*
 		UWorld::WorldTypePreLoadMap.FindOrAdd(InPackageName) = EWorldType::Editor;
 	}
 
+	// Cache doesn't initialize the world, so make sure external objects will be loaded in ULevel::PostLoad
+	CacheInfo.ScopedLoadAllExternalObjects = TUniquePtr<FScopedLoadAllExternalObjects>(new FScopedLoadAllExternalObjects(InPackageName));
 	if (bLoadAsync)
 	{
 		::LoadPackageAsync(FPackagePath::FromPackageNameChecked(InPackageToLoadFrom), InPackageName, CompletionCallback);
@@ -121,45 +105,32 @@ void FWorldPartitionPackageCache::LoadPackage(FName InPackageName, const TCHAR*
 	}
 }
 
-UPackage* FWorldPartitionPackageCache::FindPackage(FName InPackageName)
-{
-	if (TWeakObjectPtr<UPackage>* CachedPackagePtr = CachedPackages.Find(InPackageName))
-	{
-		if (UPackage* CachedPackage = CachedPackagePtr->Get())
-		{
-			return CachedPackage;
-		}
-		CachedPackages.Remove(InPackageName);
-	}
+/*
+ * FWorldPartitionPackageCache
+ */
 
-	// Might have been cached by other instance of a FWorldPartitionPackageCache. Happens when cooking WP Cells that each have their own FWorldPartitionPackageCache.
-	if(UPackage* Package = ::FindPackage(nullptr, *InPackageName.ToString()))
-	{
-		CachedPackages.Add(InPackageName, Package);
-		return Package;
-	}
+FWorldPartitionPackageCache::~FWorldPartitionPackageCache()
+{
+	UnloadPackages();
+	
+}
 
-	return nullptr;
+void FWorldPartitionPackageCache::LoadPackage(FName InPackageName, const TCHAR* InPackageToLoadFrom, FLoadPackageAsyncDelegate InCompletionDelegate, bool bLoadAsync, bool bInWorldPackage)
+{
+	FWorldPartitionPackageGlobalCache::Get().LoadPackage(this, InPackageName, InPackageToLoadFrom, InCompletionDelegate, bLoadAsync, bInWorldPackage);
 }
 
-UPackage* FWorldPartitionPackageCache::DuplicateWorldPackage(UPackage* InPackage, FName InDuplicatePackageName)
+void FWorldPartitionPackageCache::UnloadPackages()
 {
-	check(!CachedPackages.Contains(InDuplicatePackageName));
+	FWorldPartitionPackageGlobalCache::Get().UnloadPackages(this);
 
-	UPackage* DuplicatedPackage = nullptr;
-	if (UWorld* PackageWorld = UWorld::FindWorldInPackage(InPackage))
+	for (TWeakObjectPtr<UPackage>& Package : LocalCache)
 	{
-		DuplicatedPackage = CreatePackage(*InDuplicatePackageName.ToString());
-		FObjectDuplicationParameters DuplicationParameters(PackageWorld, DuplicatedPackage);
-		DuplicationParameters.bAssignExternalPackages = false;
-		DuplicationParameters.DuplicateMode = EDuplicateMode::World;
-
-		UWorld* DuplicatedWorld = Cast<UWorld>(StaticDuplicateObjectEx(DuplicationParameters));
-		check(DuplicatedWorld);
-		CachedPackages.Add(InDuplicatePackageName, DuplicatedPackage);
+		if (UPackage* CachedPackage = Package.Get())
+		{
+			FWorldPartitionPackageHelper::UnloadPackage(CachedPackage);
+		}
 	}
-	
-	return DuplicatedPackage;
 }
 
-#endif
+#endif
\ No newline at end of file
diff --git a/Engine/Source/Runtime/Engine/Private/WorldPartition/WorldPartitionPackageHelper.cpp b/Engine/Source/Runtime/Engine/Private/WorldPartition/WorldPartitionPackageHelper.cpp
new file mode 100644
--- /dev/null
+++ b/Engine/Source/Runtime/Engine/Private/WorldPartition/WorldPartitionPackageHelper.cpp
@@ -0,0 +1,66 @@
+// Copyright Epic Games, Inc. All Rights Reserved.
+
+#include "WorldPartition/WorldPartitionPackageHelper.h"
+
+#if WITH_EDITOR
+
+#include "Engine/World.h"
+
+void FWorldPartitionPackageHelper::UnloadPackage(UPackage* InPackage)
+{
+	auto TrashPackage = [](UPackage* InPackage)
+	{
+		// Clears RF_Standalone flag on objects in package (UMetaData)
+		ForEachObjectWithPackage(InPackage, [](UObject* Object)
+		{
+			Object->ClearFlags(RF_Standalone);
+			return true;
+		}, false);
+
+		// Rename so it isn't found again
+		FName NewUniqueTrashName = MakeUniqueObjectName(nullptr, UPackage::StaticClass(), FName(*FString::Printf(TEXT("%s_Trashed"), *InPackage->GetName())));
+		InPackage->Rename(*NewUniqueTrashName.ToString(), nullptr, REN_DontCreateRedirectors | REN_ForceNoResetLoaders | REN_NonTransactional | REN_DoNotDirty);
+	};
+
+	TrashPackage(InPackage);
+
+	// World specific
+	if (UWorld* PackageWorld = UWorld::FindWorldInPackage(InPackage))
+	{
+		if (PackageWorld->PersistentLevel)
+		{
+			// Manual cleanup of level since world was not initialized
+			PackageWorld->PersistentLevel->CleanupLevel(/*bCleanupResources*/ true, /*bUnloadFromEditor*/true);
+		}
+
+		if (PackageWorld->PersistentLevel->IsUsingExternalObjects())
+		{
+			ForEachObjectWithOuter(PackageWorld->PersistentLevel, [&TrashPackage](UObject* InObject)
+			{
+				if (UPackage* ExternalPackage = InObject->GetExternalPackage())
+				{
+					TrashPackage(ExternalPackage);
+				}
+			}, false);
+		}
+	}
+}
+
+UPackage* FWorldPartitionPackageHelper::DuplicateWorldPackage(UPackage* InPackage, FName InDuplicatePackageName)
+{
+	UPackage* DuplicatedPackage = nullptr;
+	UWorld* PackageWorld = UWorld::FindWorldInPackage(InPackage);
+	check(PackageWorld);
+
+	DuplicatedPackage = CreatePackage(*InDuplicatePackageName.ToString());
+	FObjectDuplicationParameters DuplicationParameters(PackageWorld, DuplicatedPackage);
+	DuplicationParameters.bAssignExternalPackages = false;
+	DuplicationParameters.DuplicateMode = EDuplicateMode::World;
+
+	UWorld* DuplicatedWorld = Cast<UWorld>(StaticDuplicateObjectEx(DuplicationParameters));
+	check(DuplicatedWorld);
+
+	return DuplicatedPackage;
+}
+
+#endif
\ No newline at end of file
diff --git a/Engine/Source/Runtime/Engine/Private/WorldPartition/WorldPartitionRuntimeLevelStreamingCell.cpp b/Engine/Source/Runtime/Engine/Private/WorldPartition/WorldPartitionRuntimeLevelStreamingCell.cpp
--- a/Engine/Source/Runtime/Engine/Private/WorldPartition/WorldPartitionRuntimeLevelStreamingCell.cpp
+++ b/Engine/Source/Runtime/Engine/Private/WorldPartition/WorldPartitionRuntimeLevelStreamingCell.cpp
@@ -142,19 +142,13 @@ UWorldPartitionLevelStreamingDynamic* UWorldPartitionRuntimeLevelStreamingCell::
 	return nullptr;
 }
 
-void UWorldPartitionRuntimeLevelStreamingCell::LoadActorsForCook()
-{
-	FWorldPartitionPackageCache PackageCache;
-	const bool bLoadAsync = false;
-	verify(FWorldPartitionLevelHelper::LoadActors(nullptr, Packages, PackageCache, [](bool){}, bLoadAsync));
-}
-
 void UWorldPartitionRuntimeLevelStreamingCell::MoveAlwaysLoadedContentToPersistentLevel()
 {
 	check(IsAlwaysLoaded());
 	if (GetActorCount() > 0)
 	{
-		LoadActorsForCook();
+		FWorldPartitionPackageCache PackageCache;
+		verify(FWorldPartitionLevelHelper::LoadActors(nullptr, Packages, PackageCache, [](bool) {}, /*bLoadAsync*/false));
 
 		UWorld* OuterWorld = GetOuterUWorldPartition()->GetTypedOuter<UWorld>();
 		FWorldPartitionLevelHelper::MoveExternalActorsToLevel(Packages, OuterWorld->PersistentLevel);
@@ -198,7 +192,8 @@ bool UWorldPartitionRuntimeLevelStreamingCell::PopulateGeneratedPackageForCook(U
 		}
 
 		// Load cell Actors
-		LoadActorsForCook();
+		FWorldPartitionPackageCache PackageCache;
+		verify(FWorldPartitionLevelHelper::LoadActors(nullptr, Packages, PackageCache, [](bool) {}, /*bLoadAsync*/false));
 
 		// Create a level and move these actors in it
 		UWorldPartition* WorldPartition = GetOuterUWorldPartition();
diff --git a/Engine/Source/Runtime/Engine/Public/WorldPartition/WorldPartitionPackageCache.h b/Engine/Source/Runtime/Engine/Public/WorldPartition/WorldPartitionPackageCache.h
--- a/Engine/Source/Runtime/Engine/Public/WorldPartition/WorldPartitionPackageCache.h
+++ b/Engine/Source/Runtime/Engine/Public/WorldPartition/WorldPartitionPackageCache.h
@@ -8,6 +8,7 @@
 #include "UObject/ObjectMacros.h"
 #include "UObject/UObjectGlobals.h"
 #include "UObject/WeakObjectPtr.h"
+#include "Engine/World.h"
 #include "Misc/Guid.h"
 
 class FLinkerInstancingContext;
@@ -15,21 +16,37 @@ class FLinkerInstancingContext;
 class FWorldPartitionPackageCache
 {
 public:
-	FWorldPartitionPackageCache();
 	~FWorldPartitionPackageCache();
-
 	void LoadPackage(FName InPackageName, const TCHAR* InPackageToLoadFrom, FLoadPackageAsyncDelegate InCompletionDelegate, bool bLoadAsync, bool bInWorldPackage);
-	UPackage* FindPackage(FName InPackageName);
-	UPackage* DuplicateWorldPackage(UPackage* InPackage, FName InPackageName);
 	void UnloadPackages();
-	void TrashPackage(UPackage* InPackage);
-	
+
+	// @todo_ow: This should be removed as soon as FWorldPartitionLevelHelper::LoadActors doesn't do package duplication
+	void CacheLoadedPackage(UPackage* InPackage) { LocalCache.Add(InPackage); }
 private:
+	TSet<TWeakObjectPtr<UPackage>> LocalCache;
+};
+
+class FWorldPartitionPackageGlobalCache
+{
+public:
+	static FWorldPartitionPackageGlobalCache& Get() { return GlobalCache; }
 
-	void UnloadPackage(UPackage* InPackage);
+	void LoadPackage(FWorldPartitionPackageCache* Client, FName InPackageName, const TCHAR* InPackageToLoadFrom, FLoadPackageAsyncDelegate InCompletionDelegate, bool bLoadAsync, bool bInWorldPackage);
+	void UnloadPackages(FWorldPartitionPackageCache* Client);
+
+private:
 
-	TMap<FName, TWeakObjectPtr<UPackage>> CachedPackages;
-	TMap<FName, TArray<FLoadPackageAsyncDelegate>> LoadingPackages;
+	struct FPackageCacheInfo
+	{
+		bool HasPendingLoad() const { return CompletionDelegates.Num() > 0; }
+		TSet<FWorldPartitionPackageCache*> Referencers;
+		TWeakObjectPtr<UPackage> Package;
+		TArray<FLoadPackageAsyncDelegate> CompletionDelegates;
+		TUniquePtr<FScopedLoadAllExternalObjects> ScopedLoadAllExternalObjects;
+	};
+	TMap<FName, FPackageCacheInfo> Cache;
+
+	static FWorldPartitionPackageGlobalCache GlobalCache;
 };
 
 #endif
\ No newline at end of file
diff --git a/Engine/Source/Runtime/Engine/Public/WorldPartition/WorldPartitionPackageHelper.h b/Engine/Source/Runtime/Engine/Public/WorldPartition/WorldPartitionPackageHelper.h
new file mode 100644
--- /dev/null
+++ b/Engine/Source/Runtime/Engine/Public/WorldPartition/WorldPartitionPackageHelper.h
@@ -0,0 +1,18 @@
+// Copyright Epic Games, Inc. All Rights Reserved.
+
+#pragma once
+
+#if WITH_EDITOR
+
+#include "UObject/NameTypes.h"
+
+class UPackage;
+
+class FWorldPartitionPackageHelper
+{
+public:
+	static void UnloadPackage(UPackage* InPackage);
+	static UPackage* DuplicateWorldPackage(UPackage* InPackage, FName InDuplicatePackageName);
+};
+
+#endif
\ No newline at end of file
diff --git a/Engine/Source/Runtime/Engine/Public/WorldPartition/WorldPartitionRuntimeLevelStreamingCell.h b/Engine/Source/Runtime/Engine/Public/WorldPartition/WorldPartitionRuntimeLevelStreamingCell.h
--- a/Engine/Source/Runtime/Engine/Public/WorldPartition/WorldPartitionRuntimeLevelStreamingCell.h
+++ b/Engine/Source/Runtime/Engine/Public/WorldPartition/WorldPartitionRuntimeLevelStreamingCell.h
@@ -54,7 +54,6 @@ private:
 	class UWorldPartitionLevelStreamingDynamic* GetOrCreateLevelStreaming() const;
 
 #if WITH_EDITOR
-	void LoadActorsForCook();
 	void MoveAlwaysLoadedContentToPersistentLevel();
 	class UWorldPartitionLevelStreamingDynamic* CreateLevelStreaming(const FString& InPackageName = FString()) const;
 #endif
