Modified and new Level Instances are now properly taken into account by world partition when running PIE.

diff --git a/Engine/Source/Runtime/Engine/Private/WorldPartition/LevelInstance/LevelInstanceActorDesc.cpp b/Engine/Source/Runtime/Engine/Private/WorldPartition/LevelInstance/LevelInstanceActorDesc.cpp
--- a/Engine/Source/Runtime/Engine/Private/WorldPartition/LevelInstance/LevelInstanceActorDesc.cpp
+++ b/Engine/Source/Runtime/Engine/Private/WorldPartition/LevelInstance/LevelInstanceActorDesc.cpp
@@ -18,13 +18,23 @@ static FAutoConsoleVariableRef CVarForceLevelStreaming(
 	GLevelInstanceDebugForceLevelStreaming,
 	TEXT("Set to 1 to force Level Instance to be streamed instead of embedded in World Partition grid."));
 
-struct FActorDescContainerInstance
+struct FActorDescContainerInstance : public FGCObject
 {
 	FActorDescContainerInstance()
 	: Container(nullptr)
 	, RefCount(0)
 	{}
 
+	virtual void AddReferencedObjects(FReferenceCollector& Collector) override
+	{
+		if (Container)
+		{
+			Collector.AddReferencedObject(Container);
+		}
+	}
+
+	virtual FString GetReferencerName() const override { return TEXT("FActorDescContainerInstance"); }
+
 	UActorDescContainer* Container;
 	uint32 RefCount;
 };
@@ -38,7 +48,8 @@ FLevelInstanceActorDesc::FLevelInstanceActorDesc()
 
 FLevelInstanceActorDesc::~FLevelInstanceActorDesc()
 {
-	check(!LevelInstanceContainer);
+	UnregisterContainerInstance();
+	check(!LevelInstanceContainer.IsValid());
 }
 
 void FLevelInstanceActorDesc::Init(const AActor* InActor)
@@ -75,38 +86,53 @@ bool FLevelInstanceActorDesc::Equals(const FWorldPartitionActorDesc* Other) cons
 	return false;
 }
 
+void FLevelInstanceActorDesc::RegisterContainerInstance(UWorld* InWorld)
+{
+	check(InWorld);
+	check(!LevelInstanceContainer.IsValid());
+	if (DesiredRuntimeBehavior == ELevelInstanceRuntimeBehavior::Partitioned && !GLevelInstanceDebugForceLevelStreaming)
+	{
+		if (!LevelPackage.IsNone() && ULevel::GetIsLevelUsingExternalActorsFromPackage(LevelPackage) && !ULevel::GetIsLevelPartitionedFromPackage(LevelPackage))
+		{
+			LevelInstanceContainer = FLevelInstanceActorDesc::RegisterActorDescContainer(LevelPackage, InWorld);
+			check(LevelInstanceContainer.IsValid());
+		}
+	}
+}
+
+void FLevelInstanceActorDesc::UnregisterContainerInstance()
+{
+	if (LevelInstanceContainer.IsValid())
+	{
+		FLevelInstanceActorDesc::UnregisterActorDescContainer(LevelInstanceContainer.Get());
+		LevelInstanceContainer.Reset();
+	}
+}
+
 void FLevelInstanceActorDesc::SetContainer(UActorDescContainer* InContainer)
 {
 	FWorldPartitionActorDesc::SetContainer(InContainer);
 
 	if (Container)
 	{
-		check(!LevelInstanceContainer);
-
-		if (DesiredRuntimeBehavior == ELevelInstanceRuntimeBehavior::Partitioned && !GLevelInstanceDebugForceLevelStreaming)
-		{
-			if (!LevelPackage.IsNone() && ULevel::GetIsLevelUsingExternalActorsFromPackage(LevelPackage) && !ULevel::GetIsLevelPartitionedFromPackage(LevelPackage))
-			{
-				LevelInstanceContainer = RegisterActorDescContainer(LevelPackage, Container->GetWorld());
-				check(LevelInstanceContainer);
-			}
-		}
+		RegisterContainerInstance(Container->GetWorld());
 	}
 	else
 	{
-		if (LevelInstanceContainer)
-		{
-			UnregisterActorDescContainer(LevelInstanceContainer);
-			LevelInstanceContainer = nullptr;
-		}
+		UnregisterContainerInstance();
 	}
 }
 
 bool FLevelInstanceActorDesc::GetContainerInstance(const UActorDescContainer*& OutLevelContainer, FTransform& OutLevelTransform, EContainerClusterMode& OutClusterMode) const
 {
-	if (LevelInstanceContainer)
+	if (!LevelInstanceContainer.IsValid() && IsLoaded())
+	{
+		// Lazy initialization of LevelInstanceContainer (used by ModifiedActorsDescList)
+		const_cast<FLevelInstanceActorDesc*>(this)->RegisterContainerInstance(GetActor()->GetWorld());
+	}
+	if (LevelInstanceContainer.IsValid())
 	{
-		OutLevelContainer = LevelInstanceContainer;
+		OutLevelContainer = LevelInstanceContainer.Get();
 		OutLevelTransform = LevelInstanceTransform;
 		OutClusterMode = EContainerClusterMode::Partitioned;
 		return true;
@@ -171,11 +197,6 @@ void FLevelInstanceActorDesc::Serialize(FArchive& Ar)
 	}
 }
 
-void FLevelInstanceActorDesc::AddReferencedObjects(FReferenceCollector& Collector)
-{
-	Collector.AddReferencedObject(LevelInstanceContainer);
-}
-
 UActorDescContainer* FLevelInstanceActorDesc::RegisterActorDescContainer(FName PackageName, UWorld* InWorld)
 {
 	FActorDescContainerInstance& ExistingContainerInstance = ActorDescContainers.FindOrAdd(PackageName);
@@ -190,6 +211,7 @@ UActorDescContainer* FLevelInstanceActorDesc::RegisterActorDescContainer(FName P
 		ActorDescContainer->Initialize(InWorld, PackageName);
 	}
 
+	check(ActorDescContainer->GetWorld() == InWorld);
 	return ActorDescContainer;
 }
 
diff --git a/Engine/Source/Runtime/Engine/Public/WorldPartition/LevelInstance/LevelInstanceActorDesc.h b/Engine/Source/Runtime/Engine/Public/WorldPartition/LevelInstance/LevelInstanceActorDesc.h
--- a/Engine/Source/Runtime/Engine/Public/WorldPartition/LevelInstance/LevelInstanceActorDesc.h
+++ b/Engine/Source/Runtime/Engine/Public/WorldPartition/LevelInstance/LevelInstanceActorDesc.h
@@ -16,7 +16,7 @@ enum class ELevelInstanceRuntimeBehavior : uint8;
 /**
  * ActorDesc for Actors that are part of a LevelInstanceActor Level.
  */
-class ENGINE_API FLevelInstanceActorDesc : public FWorldPartitionActorDesc, public FGCObject
+class ENGINE_API FLevelInstanceActorDesc : public FWorldPartitionActorDesc
 {
 #if WITH_EDITOR
 	friend class ALevelInstance;
@@ -35,17 +35,18 @@ protected:
 	virtual bool Equals(const FWorldPartitionActorDesc* Other) const override;
 	virtual void TransferFrom(const FWorldPartitionActorDesc* From) override;
 	virtual void Serialize(FArchive& Ar) override;
-	virtual void AddReferencedObjects(FReferenceCollector& Collector);
-	virtual FString GetReferencerName() const override { return TEXT("FLevelInstanceActorDesc"); }
 	virtual void SetContainer(UActorDescContainer* InContainer) override;
 
 	FName LevelPackage;
 	FTransform LevelInstanceTransform;
 	ELevelInstanceRuntimeBehavior DesiredRuntimeBehavior;
 
-	TObjectPtr<UActorDescContainer> LevelInstanceContainer;
+	TWeakObjectPtr<UActorDescContainer> LevelInstanceContainer;
 
 private:
+	void RegisterContainerInstance(UWorld* InWorld);
+	void UnregisterContainerInstance();
+
 	static UActorDescContainer* RegisterActorDescContainer(FName PackageName, UWorld* InWorld);
 	static void UnregisterActorDescContainer(UActorDescContainer* Container);
 #endif
