Add IsDataValid override to UAnimationAsset that forwards the call to any UAssetUserDatas on the object.  This allows any UAssetUserData stored on an animation to be part of the validation when saving the asset.

diff --git a/Engine/Source/Runtime/Engine/Classes/Animation/AnimationAsset.h b/Engine/Source/Runtime/Engine/Classes/Animation/AnimationAsset.h
--- a/Engine/Source/Runtime/Engine/Classes/Animation/AnimationAsset.h
+++ b/Engine/Source/Runtime/Engine/Classes/Animation/AnimationAsset.h
@@ -1087,6 +1087,7 @@ public:
 	//~ Begin UObject Interface.
 #if WITH_EDITOR
 	virtual void PostEditChangeProperty(FPropertyChangedEvent& PropertyChangedEvent) override;
+	virtual EDataValidationResult IsDataValid(TArray<FText>& ValidationErrors) override;
 #endif // WITH_EDITOR
 
 	/**
diff --git a/Engine/Source/Runtime/Engine/Private/AnimationAsset.cpp b/Engine/Source/Runtime/Engine/Private/AnimationAsset.cpp
--- a/Engine/Source/Runtime/Engine/Private/AnimationAsset.cpp
+++ b/Engine/Source/Runtime/Engine/Private/AnimationAsset.cpp
@@ -698,6 +698,19 @@ void UAnimationAsset::PostEditChangeProperty(FPropertyChangedEvent& PropertyChan
 
 	Super::PostEditChangeProperty(PropertyChangedEvent);
 }
+
+EDataValidationResult UAnimationAsset::IsDataValid(TArray<FText>& ValidationErrors)
+{
+	EDataValidationResult Result = UObject::IsDataValid(ValidationErrors);
+	for (UAssetUserData* Datum : AssetUserData)
+	{
+		if(Datum != nullptr && Datum->IsDataValid(ValidationErrors) == EDataValidationResult::Invalid)
+		{
+			Result = EDataValidationResult::Invalid;
+		}
+	}
+	return Result;
+}
 #endif
 
 ///////////////////////////////////////////////////////////////////////////////////////
