Adding convenience function to UAnimSequenceBase for removing all notifies.

diff --git a/Engine/Source/Runtime/Engine/Classes/Animation/AnimSequenceBase.h b/Engine/Source/Runtime/Engine/Classes/Animation/AnimSequenceBase.h
--- a/Engine/Source/Runtime/Engine/Classes/Animation/AnimSequenceBase.h
+++ b/Engine/Source/Runtime/Engine/Classes/Animation/AnimSequenceBase.h
@@ -80,6 +80,9 @@ public:
 
 	/** Remove the notifies specified */
 	bool RemoveNotifies(const TArray<FName>& NotifiesToRemove);
+	
+	/** Remove all notifies */
+	void RemoveNotifies();
 
 	/** 
 	 * Retrieves AnimNotifies given a StartTime and a DeltaTime.
diff --git a/Engine/Source/Runtime/Engine/Private/Animation/AnimSequenceBase.cpp b/Engine/Source/Runtime/Engine/Private/Animation/AnimSequenceBase.cpp
--- a/Engine/Source/Runtime/Engine/Private/Animation/AnimSequenceBase.cpp
+++ b/Engine/Source/Runtime/Engine/Private/Animation/AnimSequenceBase.cpp
@@ -271,6 +271,18 @@ bool UAnimSequenceBase::RemoveNotifies(const TArray<FName>& NotifiesToRemove)
 	return bSequenceModified;
 }
 
+void UAnimSequenceBase::RemoveNotifies()
+{
+	if (Notifies.Num() == 0)
+	{
+		return;
+	}
+	Modify();
+	Notifies.Reset();
+	MarkPackageDirty();
+	RefreshCacheData();
+}
+
 bool UAnimSequenceBase::IsNotifyAvailable() const
 {
 	return (Notifies.Num() != 0) && (GetPlayLength() > 0.f);
