Set bInTick=false after calling TickFlush from UDemoNetDriver::TickFlushAsyncEndOfFrame

diff --git a/Engine/Source/Runtime/Engine/Classes/Engine/NetDriver.h b/Engine/Source/Runtime/Engine/Classes/Engine/NetDriver.h
--- a/Engine/Source/Runtime/Engine/Classes/Engine/NetDriver.h
+++ b/Engine/Source/Runtime/Engine/Classes/Engine/NetDriver.h
@@ -904,7 +904,7 @@ public:
 
 private:
 	double						ElapsedTime;
-	
+
 	/** Whether or not the NetDriver is ticking */
 	bool bInTick;
 
@@ -1691,6 +1691,8 @@ public:
 
 protected:
 
+	void SetIsInTick(bool bIsInTick) { bInTick = bIsInTick; }
+
 	/** Register all TickDispatch, TickFlush, PostTickFlush to tick in World */
 	ENGINE_API void RegisterTickEvents(class UWorld* InWorld);
 	/** Unregister all TickDispatch, TickFlush, PostTickFlush to tick in World */
diff --git a/Engine/Source/Runtime/Engine/Private/DemoNetDriver.cpp b/Engine/Source/Runtime/Engine/Private/DemoNetDriver.cpp
--- a/Engine/Source/Runtime/Engine/Private/DemoNetDriver.cpp
+++ b/Engine/Source/Runtime/Engine/Private/DemoNetDriver.cpp
@@ -1246,6 +1246,7 @@ void UDemoNetDriver::TickFlushAsyncEndOfFrame(float DeltaSeconds)
 	if (ShouldTickFlushAsyncEndOfFrame())
 	{
 		TickFlushInternal(DeltaSeconds);
+		SetIsInTick(false); //PostTickFlush isn't called after the async TickFlush, so set bInTick to false here
 	}
 }
 
