Fixed incorrect guarding of GetNaniteResourceInfo with OptimizedWPO

diff --git a/Engine/Source/Runtime/Engine/Private/PrimitiveUniformShaderParameters.cpp b/Engine/Source/Runtime/Engine/Private/PrimitiveUniformShaderParameters.cpp
--- a/Engine/Source/Runtime/Engine/Private/PrimitiveUniformShaderParameters.cpp
+++ b/Engine/Source/Runtime/Engine/Private/PrimitiveUniformShaderParameters.cpp
@@ -141,10 +141,13 @@ FPrimitiveSceneShaderData::FPrimitiveSceneShaderData(const FPrimitiveSceneProxy*
 	bool bHasNaniteImposterData = false;
 	bool bEvaluateWorldPositionOffset = !OptimizedWPO;
 
-	if (OptimizedWPO && Proxy->IsNaniteMesh())
+	if (Proxy->IsNaniteMesh())
 	{
 		Proxy->GetNaniteResourceInfo(NaniteResourceID, NaniteHierarchyOffset, NaniteImposterIndex);
-		bEvaluateWorldPositionOffset = static_cast<const Nanite::FSceneProxyBase*>(Proxy)->EvaluateWorldPositionOffset();
+		if (OptimizedWPO)
+		{
+			bEvaluateWorldPositionOffset = static_cast<const Nanite::FSceneProxyBase*>(Proxy)->EvaluateWorldPositionOffset();
+		}
 	}
 
 	FPrimitiveUniformShaderParametersBuilder Builder = FPrimitiveUniformShaderParametersBuilder{}
