Handle actors destroyed in the middle of an amortized checkpoint save
Unify the logging of pending checkpoint actors with the other amortized steps
Caching deleted actors at the beginning of the checkpoint record
Move shared actor destroy logic from demo driver and replay connection into replay helper
Changed pending checkpoint actor list to not pop/realloc for each actor

diff --git a/Engine/Source/Runtime/Engine/Private/DemoNetDriver.cpp b/Engine/Source/Runtime/Engine/Private/DemoNetDriver.cpp
--- a/Engine/Source/Runtime/Engine/Private/DemoNetDriver.cpp
+++ b/Engine/Source/Runtime/Engine/Private/DemoNetDriver.cpp
@@ -4981,7 +4981,6 @@ void UDemoNetDriver::NotifyActorDestroyed(AActor* Actor, bool IsSeamlessTravel)
 	const bool bIsRecording = IsRecording();
 	const bool bNetStartup = Actor->IsNetStartupActor();
 	const bool bActorRewindable = Actor->bReplayRewindable;
-	const bool bDeltaCheckpoint = HasDeltaCheckpoints();
 
 	if (bActorRewindable && !IsSeamlessTravel && !bIsRecording)
 	{
@@ -4994,15 +4993,15 @@ void UDemoNetDriver::NotifyActorDestroyed(AActor* Actor, bool IsSeamlessTravel)
 		}
 	}
 
-	if (bIsRecording && bNetStartup)
+	if (bIsRecording)
 	{
 		// We don't want to send any destruction info in this case, because the actor should stick around.
 		// The Replay will manage deleting this when it performs streaming or travel behavior.
-		if (IsSeamlessTravel)
+		if (bNetStartup && IsSeamlessTravel)
 		{
 			// This is a stripped down version of UNetDriver::NotifyActorDestroy and UActorChannel::Close
 			// combined, and should be kept up to date with those methods.
-		
+
 			// Remove the actor from the property tracker map
 			RepChangedPropertyTrackerMap.Remove(Actor);
 
@@ -5021,38 +5020,18 @@ void UDemoNetDriver::NotifyActorDestroyed(AActor* Actor, bool IsSeamlessTravel)
 					Channel->Actor = nullptr;
 					Channel->CleanupReplicators(false);
 				}
-	
+
 				Connection->DormantReplicatorMap.Remove(Actor);
 			}
-	
+
 			GetNetworkObjectList().Remove(Actor);
 			RenamedStartupActors.Remove(Actor->GetFName());
 			return;
 		}
-		else
-		{
-			const FString FullName = Actor->GetFullName();
-
-			// This was deleted due to a game interaction, which isn't supported for Rewindable actors (while recording).
-			// However, since the actor is going to be deleted imminently, we need to track it.
-			UE_CLOG(bActorRewindable, LogDemo, Warning, TEXT("Replay Rewindable Actor destroyed during recording. Replay may show artifacts (%s)"), *FullName);
-
-			UE_LOG(LogDemo, VeryVerbose, TEXT("NotifyActyorDestroyed: adding actor to deleted startup list: %s"), *FullName);
-			ReplayHelper.RecordingDeletedNetStartupActors.Add(FullName);
-
-			if (bDeltaCheckpoint)
-			{
-				ReplayHelper.RecordingDeltaCheckpointData.RecordingDeletedNetStartupActors.Add(FullName);
-			}
-		}
-	}
 
-	if (bIsRecording && !bNetStartup && bDeltaCheckpoint)
-	{
-		FNetworkGUID NetGUID = GuidCache->NetGUIDLookup.FindRef(Actor);
-		if (NetGUID.IsValid())
+		if (!IsSeamlessTravel)
 		{
-			ReplayHelper.RecordingDeltaCheckpointData.DestroyedDynamicActors.Add(NetGUID);
+			ReplayHelper.NotifyActorDestroyed(ClientConnections[0], Actor);
 		}
 	}
 
diff --git a/Engine/Source/Runtime/Engine/Private/ReplayHelper.cpp b/Engine/Source/Runtime/Engine/Private/ReplayHelper.cpp
--- a/Engine/Source/Runtime/Engine/Private/ReplayHelper.cpp
+++ b/Engine/Source/Runtime/Engine/Private/ReplayHelper.cpp
@@ -613,6 +613,7 @@ void FReplayHelper::SaveCheckpoint(UNetConnection* Connection)
 		}
 
 		CheckpointSaveContext.PendingCheckpointActors.Reserve(ActorArray.Num());
+		CheckpointSaveContext.PendingActorToIndex.Reserve(ActorArray.Num());
 
 		uint32 LevelIt = 0;
 		for (int32 CurrentIt = 0, EndIt = ActorArray.Num(); CurrentIt != EndIt; ++LevelIt)
@@ -626,7 +627,8 @@ void FReplayHelper::SaveCheckpoint(UNetConnection* Connection)
 #endif
 			while (CurrentIt < EndIt && (CurrentLevelToIndex == ActorArray[CurrentIt].Level))
 			{
-				CheckpointSaveContext.PendingCheckpointActors.Add({ ActorArray[CurrentIt].Actor, LevelStatus.LevelIndex });
+				int32 PendingIndex = CheckpointSaveContext.PendingCheckpointActors.Add({ ActorArray[CurrentIt].Actor, LevelStatus.LevelIndex });
+				CheckpointSaveContext.PendingActorToIndex.Add(ActorArray[CurrentIt].Actor, PendingIndex);
 				++CurrentIt;
 			};
 		}
@@ -670,10 +672,12 @@ void FReplayHelper::SaveCheckpoint(UNetConnection* Connection)
 
 	// We are now processing checkpoint actors	
 	CheckpointSaveContext.CheckpointSaveState = ECheckpointSaveState::ProcessCheckpointActors;
+	CheckpointSaveContext.NextAmortizedItem = 0;
 	CheckpointSaveContext.TotalCheckpointSaveTimeSeconds = 0;
 	CheckpointSaveContext.TotalCheckpointReplicationTimeSeconds = 0;
 	CheckpointSaveContext.TotalCheckpointSaveFrames = 0;
 	CheckpointSaveContext.TotalCheckpointActors = CheckpointSaveContext.PendingCheckpointActors.Num();
+	CheckpointSaveContext.CheckpointDeletedNetStartupActors.Reset();
 
 	LastCheckpointTime = DemoCurrentTime;
 
@@ -686,6 +690,7 @@ void FReplayHelper::SaveCheckpoint(UNetConnection* Connection)
 	else
 	{
 		CheckpointSaveContext.NameTableMap.Empty();
+		CheckpointSaveContext.CheckpointDeletedNetStartupActors.Append(RecordingDeletedNetStartupActors);
 	}
 
 	UE_LOG(LogDemo, Log, TEXT("Starting checkpoint. Networked Actors: %i"), NetworkObjectList.GetAllObjects().Num());
@@ -709,6 +714,65 @@ static bool inline ShouldExecuteState(const FRepActorsCheckpointParams& Params,
 	return (1.0 - ((CurrentTime - Params.StartCheckpointTime) / Params.CheckpointMaxUploadTimePerFrame)) > RequiredRatioToStart;
 }
 
+void FReplayHelper::ProcessCheckpointActors(UNetConnection* Connection, TArrayView<FPendingCheckPointActor> PendingActors, int32& NextIndex, FRepActorsCheckpointParams& Params)
+{
+	UPackageMapClient* PackageMapClient = Cast<UPackageMapClient>(Connection->PackageMap);
+
+	// Save package map ack status in case we export stuff during the checkpoint (so we can restore the connection back to what it was before we saved the checkpoint)
+	PackageMapClient->OverridePackageMapExportAckStatus(&CheckpointSaveContext.CheckpointAckState);
+
+	Connection->SetReserveDestroyedChannels(false);
+
+	// Save the replicated server time so we can restore it after the checkpoint has been serialized.
+	// This preserves the existing behavior and prevents clients from receiving updated server time
+	// more often than the normal update rate.
+	AGameStateBase* const GameState = World != nullptr ? World->GetGameState() : nullptr;
+
+	const float SavedReplicatedServerTimeSeconds = GameState ? GameState->ReplicatedWorldTimeSeconds : -1.0f;
+
+	// Normally AGameStateBase::ReplicatedWorldTimeSeconds is only updated periodically,
+	// but we want to make sure it's accurate for the checkpoint.
+	if (GameState)
+	{
+		GameState->UpdateServerTimeSeconds();
+	}
+
+	{
+		const bool bDeltaCheckpoint = HasDeltaCheckpoints();
+
+		// Re-use the existing connection to record all properties that have changed since channels were first opened
+		TGuardValue<EResendAllDataState> ResendAllData(Connection->ResendAllDataState, bDeltaCheckpoint ? EResendAllDataState::SinceCheckpoint : EResendAllDataState::SinceOpen);
+
+		bool bContinue = true;
+
+		while (NextIndex < PendingActors.Num())
+		{
+			const FPendingCheckPointActor& Current = PendingActors[NextIndex];
+
+			AActor* Actor = Current.Actor.Get();
+
+			++NextIndex;
+
+			if (!ReplicateCheckpointActor(Actor, Connection, Params))
+			{
+				break;
+			}
+		}
+
+		if (GameState)
+		{
+			// Restore the game state's replicated world time
+			GameState->ReplicatedWorldTimeSeconds = SavedReplicatedServerTimeSeconds;
+		}
+
+		FlushNetChecked(*Connection);
+
+		PackageMapClient->OverridePackageMapExportAckStatus(nullptr);
+	}
+
+	Connection->SetReserveDestroyedChannels(true);
+}
+
 void FReplayHelper::TickCheckpoint(UNetConnection* Connection)
 {
 	CSV_SCOPED_TIMING_STAT(Demo, DemoRecordCheckpointTime);
@@ -761,78 +825,24 @@ void FReplayHelper::TickCheckpoint(UNetConnection* Connection)
 			{
 				SCOPED_NAMED_EVENT(FReplayHelper_ProcessCheckpointActors, FColor::Green);
 
-				Connection->SetReserveDestroyedChannels(false);
-
-				// Save the replicated server time so we can restore it after the checkpoint has been serialized.
-				// This preserves the existing behavior and prevents clients from receiving updated server time
-				// more often than the normal update rate.
-				AGameStateBase* const GameState = World != nullptr ? World->GetGameState() : nullptr;
-
-				const float SavedReplicatedServerTimeSeconds = GameState ? GameState->ReplicatedWorldTimeSeconds : -1.0f;
-
-				// Normally AGameStateBase::ReplicatedWorldTimeSeconds is only updated periodically,
-				// but we want to make sure it's accurate for the checkpoint.
-				if (GameState)
 				{
-					GameState->UpdateServerTimeSeconds();
-				}
-
-				{
-					// Re-use the existing connection to record all properties that have changed since channels were first opened
-					TGuardValue<EResendAllDataState> ResendAllData(Connection->ResendAllDataState, bDeltaCheckpoint ? EResendAllDataState::SinceCheckpoint : EResendAllDataState::SinceOpen);
-
-					bool bContinue = true;
-
-					int32 NumActorsToReplicate = CheckpointSaveContext.PendingCheckpointActors.Num();
-
-					do
-					{
-						const FPendingCheckPointActor Current = CheckpointSaveContext.PendingCheckpointActors.Pop();
-
-						AActor* Actor = Current.Actor.Get();
-
-						bContinue = ReplicateCheckpointActor(Actor, Connection, Params);
-					} 
-					while (--NumActorsToReplicate && bContinue);
-
-					if (GameState)
-					{
-						// Restore the game state's replicated world time
-						GameState->ReplicatedWorldTimeSeconds = SavedReplicatedServerTimeSeconds;
-					}
-
-					FlushNetChecked(*Connection);
-
-					PackageMapClient->OverridePackageMapExportAckStatus(nullptr);
+					FCheckpointStepHelper StepHelper(ECheckpointSaveState::ProcessCheckpointActors, Params.StartCheckpointTime, &CheckpointSaveContext.NextAmortizedItem, CheckpointSaveContext.PendingCheckpointActors.Num());
+					ProcessCheckpointActors(Connection, CheckpointSaveContext.PendingCheckpointActors, CheckpointSaveContext.NextAmortizedItem, Params);
 				}
 
-				Connection->SetReserveDestroyedChannels(true);
-
-				// We are done processing for this frame so  store the TotalCheckpointSave time here to be true to the old behavior which did not account for the	actual saving time of the check point
+				// We are done processing for this frame so store the TotalCheckpointSave time here to be true to the old behavior which did not account for the	actual saving time of the check point
 				CheckpointSaveContext.TotalCheckpointReplicationTimeSeconds += (FPlatformTime::Seconds() - Params.StartCheckpointTime);
 
-				// if we have replicated all checkpointactors, move on to the next state
-				if (CheckpointSaveContext.PendingCheckpointActors.Num() == 0)
+				// if we have replicated all checkpoint actors, move on to the next state
+				if (CheckpointSaveContext.NextAmortizedItem == CheckpointSaveContext.PendingCheckpointActors.Num())
 				{
-					CheckpointSaveContext.CheckpointSaveState = ECheckpointSaveState::CacheDeletedActors;
-
+					CheckpointSaveContext.PendingCheckpointActors.Empty();
+					CheckpointSaveContext.PendingActorToIndex.Empty();
+					
 					Connection->SetReserveDestroyedChannels(false);
 					Connection->SetIgnoreReservedChannels(false);
-				}
-			}
-			break;
-
-			case ECheckpointSaveState::CacheDeletedActors:
-			{
-				// Postpone execution of this state if we have used too much of our alloted time, this value can be tweaked based on profiling
-				const double RequiredRatioFor_CacheDeletedActors = 0.6;
-				if ((bExecuteNextState = ShouldExecuteState(Params, CurrentTime, RequiredRatioFor_CacheDeletedActors)) == true)
-				{
-					SCOPED_NAMED_EVENT(FReplayHelper_CacheDeletedActors, FColor::Green);
 
-					//
-					// We're done saving this checkpoint, now we need to write out all data for it.
-					//
+					CheckpointSaveContext.CheckpointSaveState = ECheckpointSaveState::SerializeDeletedStartupActors;
 
 					CheckpointSaveContext.bWriteCheckpointOffset = HasLevelStreamingFixes();
 					if (HasLevelStreamingFixes())
@@ -844,7 +854,7 @@ void FReplayHelper::TickCheckpoint(UNetConnection* Connection)
 
 					*CheckpointArchive << CurrentLevelIndex;
 
-					CacheDeletedActors(Connection);
+					CheckpointSaveContext.NextAmortizedItem = 0;
 					CheckpointSaveContext.CheckpointSaveState = ECheckpointSaveState::SerializeDeletedStartupActors;
 				}
 			}
@@ -1498,28 +1508,6 @@ void FReplayHelper::CacheNetGuids(UNetConnection* Connection)
 	}
 }
 
-void FReplayHelper::CacheDeletedActors(UNetConnection* Connection)
-{
-	if (Connection && Connection->Driver)
-	{
-		const double StartTime = FPlatformTime::Seconds();
-
-		// initialize serialization
-		CheckpointSaveContext.CheckpointDeletedNetStartupActors.Reset();
-		CheckpointSaveContext.NextAmortizedItem = 0;
-
-		const bool bDeltaCheckpoint = HasDeltaCheckpoints();
-
-		// delta entries already copied into CheckpointSaveContext.DeltaCheckpointData
-		if (!bDeltaCheckpoint)
-		{
-			CheckpointSaveContext.CheckpointDeletedNetStartupActors.Append(RecordingDeletedNetStartupActors);
-		}
-
-		UE_LOG(LogDemo, Verbose, TEXT("CacheDeletedActors: %d, %.1f ms"), bDeltaCheckpoint ? CheckpointSaveContext.DeltaCheckpointData.RecordingDeletedNetStartupActors.Num() : CheckpointSaveContext.CheckpointDeletedNetStartupActors.Num(), (FPlatformTime::Seconds() - StartTime) * 1000);
-	}
-}
-
 bool FReplayHelper::ReplicateCheckpointActor(AActor* ToReplicate, UNetConnection* Connection, class FRepActorsCheckpointParams& Params)
 {
 	// Early out if the actor has been destroyed or the world is streamed out.
@@ -1695,6 +1683,7 @@ void FReplayHelper::FCheckpointSaveStateContext::CountBytes(FArchive& Ar) const
 
 	GRANULAR_NETWORK_MEMORY_TRACKING_TRACK("CheckpointAckState", CheckpointAckState.CountBytes(Ar));
 	GRANULAR_NETWORK_MEMORY_TRACKING_TRACK("PendingCheckpointActors", PendingCheckpointActors.CountBytes(Ar));
+	GRANULAR_NETWORK_MEMORY_TRACKING_TRACK("PendingActorToIndex", PendingActorToIndex.CountBytes(Ar));
 	GRANULAR_NETWORK_MEMORY_TRACKING_TRACK("DeltaCheckpointData", DeltaCheckpointData.CountBytes(Ar));
 	GRANULAR_NETWORK_MEMORY_TRACKING_TRACK("DeltaChannelCloseKeys", DeltaChannelCloseKeys.CountBytes(Ar));
 	GRANULAR_NETWORK_MEMORY_TRACKING_TRACK("NetGuidCacheSnapshot", NetGuidCacheSnapshot.CountBytes(Ar));
@@ -2298,4 +2287,79 @@ const TCHAR* LexToString(EReplayHeaderFlags Flag)
 void FReplayHelper::RequestCheckpoint()
 {
 	bPendingCheckpointRequest = true;
-}
\ No newline at end of file
+}
+
+void FReplayHelper::NotifyActorDestroyed(UNetConnection* Connection, AActor* Actor)
+{
+	check(Actor);
+	check(Connection);
+
+	// if we're recording a checkpoint, and we have not yet passed the actor recording phase
+	if (CheckpointSaveContext.CheckpointSaveState == ECheckpointSaveState::ProcessCheckpointActors)
+	{
+		// if there's already a channel open for this actor
+		if (UActorChannel* Channel = Connection->FindActorChannelRef(Actor))
+		{
+			if (Channel->ChIndex != INDEX_NONE)
+			{
+				// if this actor is in the pending checkpoint actor list and we have not already recorded it
+				if (int32* PendingIndex = CheckpointSaveContext.PendingActorToIndex.Find(Actor))
+				{
+					if (CheckpointSaveContext.PendingCheckpointActors.IsValidIndex(*PendingIndex) 
+						&& (*PendingIndex >= CheckpointSaveContext.NextAmortizedItem)
+						&& (CheckpointSaveContext.PendingCheckpointActors[*PendingIndex].Actor.Get() == Actor))
+					{
+						UE_LOG(LogDemo, Verbose, TEXT("Destroying actor while it is still in the PendingCheckpointActors list: %s"), *GetNameSafe(Actor));
+
+						FRepActorsCheckpointParams Params
+						{
+							FPlatformTime::Seconds(),
+							(double)GetCheckpointSaveMaxMSPerFrame() / 1000
+						};
+
+						// force record it to the checkpoint now
+						TArrayView<FPendingCheckPointActor> PendingView(CheckpointSaveContext.PendingCheckpointActors);
+						int32 ActorIndex = 0;
+
+						ProcessCheckpointActors(Connection, PendingView.Slice(*PendingIndex, 1), ActorIndex, Params);
+
+						// don't allow the channel index to be reused until we're done with the checkpoint
+						Connection->AddReservedChannel(Channel->ChIndex);
+					}
+				}
+			}
+		}
+	}
+
+	const bool bNetStartup = Actor->IsNetStartupActor();
+	const bool bActorRewindable = Actor->bReplayRewindable;
+	const bool bDeltaCheckpoint = HasDeltaCheckpoints();
+
+	if (bNetStartup)
+	{
+		const FString FullName = Actor->GetFullName();
+
+		// This was deleted due to a game interaction, which isn't supported for Rewindable actors (while recording).
+		// However, since the actor is going to be deleted imminently, we need to track it.
+		UE_CLOG(bActorRewindable, LogDemo, Warning, TEXT("Replay Rewindable Actor destroyed during recording. Replay may show artifacts (%s)"), *FullName);
+
+		UE_LOG(LogDemo, VeryVerbose, TEXT("NotifyActorDestroyed: adding actor to deleted startup list: %s"), *FullName);
+		RecordingDeletedNetStartupActors.Add(FullName);
+
+		if (bDeltaCheckpoint)
+		{
+			RecordingDeltaCheckpointData.RecordingDeletedNetStartupActors.Add(FullName);
+		}
+	}
+	else
+	{
+		if (bDeltaCheckpoint)
+		{
+			FNetworkGUID NetGUID = Connection->Driver->GuidCache->NetGUIDLookup.FindRef(Actor);
+			if (NetGUID.IsValid())
+			{
+				RecordingDeltaCheckpointData.DestroyedDynamicActors.Add(NetGUID);
+			}
+		}
+	}
+}
diff --git a/Engine/Source/Runtime/Engine/Private/ReplayNetConnection.cpp b/Engine/Source/Runtime/Engine/Private/ReplayNetConnection.cpp
--- a/Engine/Source/Runtime/Engine/Private/ReplayNetConnection.cpp
+++ b/Engine/Source/Runtime/Engine/Private/ReplayNetConnection.cpp
@@ -255,42 +255,12 @@ void UReplayNetConnection::OnSeamlessTravelStart(UWorld* CurrentWorld, const FSt
 
 void UReplayNetConnection::NotifyActorDestroyed(AActor* Actor, bool IsSeamlessTravel /* = false */)
 {
-	Super::NotifyActorDestroyed(Actor, IsSeamlessTravel);
-
-	check(Actor != nullptr);
-
-	const bool bNetStartup = Actor->IsNetStartupActor();
-	const bool bActorRewindable = Actor->bReplayRewindable;
-	const bool bDeltaCheckpoint = ReplayHelper.HasDeltaCheckpoints();
-
-	if (bNetStartup)
+	if (!IsSeamlessTravel)
 	{
-		if (!IsSeamlessTravel)
-		{
-			const FString FullName = Actor->GetFullName();
-
-			// This was deleted due to a game interaction, which isn't supported for Rewindable actors (while recording).
-			// However, since the actor is going to be deleted imminently, we need to track it.
-			UE_CLOG(bActorRewindable, LogDemo, Warning, TEXT("Replay Rewindable Actor destroyed during recording. Replay may show artifacts (%s)"), *FullName);
-
-			UE_LOG(LogDemo, VeryVerbose, TEXT("NotifyActorDestroyed: adding actor to deleted startup list: %s"), *FullName);
-			ReplayHelper.RecordingDeletedNetStartupActors.Add(FullName);
-
-			if (bDeltaCheckpoint)
-			{
-				ReplayHelper.RecordingDeltaCheckpointData.RecordingDeletedNetStartupActors.Add(FullName);
-			}
-		}
+		ReplayHelper.NotifyActorDestroyed(this, Actor);
 	}
 
-	if (!bNetStartup && bDeltaCheckpoint)
-	{
-		FNetworkGUID NetGUID = Driver->GuidCache->NetGUIDLookup.FindRef(Actor);
-		if (NetGUID.IsValid())
-		{
-			ReplayHelper.RecordingDeltaCheckpointData.DestroyedDynamicActors.Add(NetGUID);
-		}
-	}
+	Super::NotifyActorDestroyed(Actor, IsSeamlessTravel);
 }
 
 void UReplayNetConnection::SetAnalyticsProvider(TSharedPtr<IAnalyticsProvider> InProvider)
diff --git a/Engine/Source/Runtime/Engine/Public/ReplayHelper.h b/Engine/Source/Runtime/Engine/Public/ReplayHelper.h
--- a/Engine/Source/Runtime/Engine/Public/ReplayHelper.h
+++ b/Engine/Source/Runtime/Engine/Public/ReplayHelper.h
@@ -7,6 +7,7 @@
 #include "Engine/PackageMapClient.h"
 #include "NetworkReplayStreaming.h"
 #include "ReplayTypes.h"
+#include "Containers/ArrayView.h"
 
 class APlayerController;
 class UNetConnection;
@@ -105,6 +106,8 @@ private:
 
 	void RequestCheckpoint();
 
+	void NotifyActorDestroyed(UNetConnection* Connection, AActor* Actor);
+
 private:
 	// Hooks used to determine when levels are streamed in, streamed out, or if there's a map change.
 	void OnLevelAddedToWorld(ULevel* Level, UWorld* World);
@@ -143,7 +146,6 @@ private:
 	static const EReadPacketState ReadPacket(FArchive& Archive, TArray<uint8>& OutBuffer, const EReadPacketMode Mode);
 
 	void CacheNetGuids(UNetConnection* Connection);
-	void CacheDeletedActors(UNetConnection* Connection);
 
 	bool SerializeGuidCache(UNetConnection* Connection, const FRepActorsCheckpointParams& Params, FArchive* CheckpointArchive);
 	bool SerializeDeletedStartupActors(UNetConnection* Connection, const FRepActorsCheckpointParams& Params, FArchive* CheckpointArchive);
@@ -238,7 +240,6 @@ private:
 	{
 		Idle,
 		ProcessCheckpointActors,
-		CacheDeletedActors,
 		SerializeDeletedStartupActors,
 		SerializeDeltaDynamicDestroyed,
 		SerializeDeltaClosedChannels,
@@ -272,6 +273,9 @@ private:
 
 			switch (CheckpointState)
 			{
+			case ECheckpointSaveState::ProcessCheckpointActors:
+				StateStr = TEXT("ProcessCheckpointActors");
+				break;
 			case ECheckpointSaveState::SerializeDeletedStartupActors:
 				StateStr = TEXT("SerializeDeletedStartupActors");
 				break;
@@ -333,6 +337,7 @@ private:
 		ECheckpointSaveState CheckpointSaveState;						// Current state of checkpoint SaveState
 		FPackageMapAckState CheckpointAckState;							// Current ack state of packagemap for the current checkpoint being saved
 		TArray<FPendingCheckPointActor> PendingCheckpointActors;		// Actors to be serialized by pending checkpoint
+		TMap<TWeakObjectPtr<AActor>, int32, FDefaultSetAllocator, TWeakObjectPtrMapKeyFuncs<TWeakObjectPtr<AActor>, int32>> PendingActorToIndex;
 		double				TotalCheckpointSaveTimeSeconds;				// Total time it took to save checkpoint including the finaling part across all frames
 		double				TotalCheckpointReplicationTimeSeconds;		// Total time it took to write all replicated objects across all frames
 		bool				bWriteCheckpointOffset;
@@ -530,6 +535,8 @@ private:
 
 	ECheckpointSaveState GetCheckpointSaveState() const { return CheckpointSaveContext.CheckpointSaveState; }
 
+	void ProcessCheckpointActors(UNetConnection* Connection, TArrayView<FPendingCheckPointActor> PendingActors, int32& NextIndex, FRepActorsCheckpointParams& Params);
+
 	static constexpr int32 MAX_DEMO_READ_WRITE_BUFFER = 1024 * 2;
 	static constexpr int32 MAX_DEMO_STRING_SERIALIZATION_SIZE = 16 * 1024 * 1024;
 };
\ No newline at end of file
