Fix static analysis warnings

diff --git a/Engine/Source/Runtime/Engine/Private/AudioDevice.cpp b/Engine/Source/Runtime/Engine/Private/AudioDevice.cpp
--- a/Engine/Source/Runtime/Engine/Private/AudioDevice.cpp
+++ b/Engine/Source/Runtime/Engine/Private/AudioDevice.cpp
@@ -4210,9 +4210,15 @@ void FAudioDevice::StartSources(TArray<FWaveInstance*>& WaveInstances, int32 Fir
 	{
 		FWaveInstance* WaveInstance = WaveInstances[InstanceIndex];
 
+		USoundWave* WaveData = WaveInstance->WaveData;
+		if (!WaveData)
+		{
+			continue;
+		}
+
 		// Make sure we've finished precaching the wave instance's wave data before trying to create a source for it
-		ESoundWavePrecacheState PrecacheState = WaveInstance->WaveData->GetPrecacheState();
-		const bool bIsSoundWaveStillLoading = WaveInstance->WaveData->HasAnyFlags(RF_NeedLoad);
+		ESoundWavePrecacheState PrecacheState = WaveData->GetPrecacheState();
+		const bool bIsSoundWaveStillLoading = WaveData->HasAnyFlags(RF_NeedLoad);
 		if (PrecacheState == ESoundWavePrecacheState::InProgress || (WaitForSoundWaveToLoadCvar && bIsSoundWaveStillLoading))
 		{
 			continue;
@@ -4236,10 +4242,7 @@ void FAudioDevice::StartSources(TArray<FWaveInstance*>& WaveInstances, int32 Fir
 				Source = FreeSources.Pop();
 				check(Source);
 
-				if (WaveInstance->WaveData)
-				{
-					StartingSoundWaves.AddUnique(WaveInstance->WaveData);
-				}
+				StartingSoundWaves.AddUnique(WaveData);
 
 				// Prepare for initialization...
 				bool bSuccess = false;
@@ -6662,7 +6665,7 @@ void FAudioDevice::StopSoundsUsingResource(USoundWave* SoundWave, TArray<UAudioC
 		{
 			// If anything the ActiveSound uses the wave then we stop the sound
 			FWaveInstance* WaveInstance = WaveInstancePair.Value;
-			if (WaveInstance->WaveData == SoundWave)
+			if (WaveInstance && WaveInstance->WaveData == SoundWave)
 			{
 				if (StoppedComponents)
 				{
