Converting Quartz audio frame counting to be double instead of integer - the latter causes clock drift over time and causes issue when synchronizing with external sources or long audio

diff --git a/Engine/Source/Runtime/AudioMixer/Private/Quartz/AudioMixerClock.cpp b/Engine/Source/Runtime/AudioMixer/Private/Quartz/AudioMixerClock.cpp
--- a/Engine/Source/Runtime/AudioMixer/Private/Quartz/AudioMixerClock.cpp
+++ b/Engine/Source/Runtime/AudioMixer/Private/Quartz/AudioMixerClock.cpp
@@ -57,7 +57,7 @@ namespace Audio
 		FQuartzClockTickRate CurrentTickRate = Metronome.GetTickRate();
 
 		// ratio between new and old rates
-		const float Ratio = static_cast<float>(InNewTickRate.GetFramesPerTick()) / static_cast<float>(CurrentTickRate.GetFramesPerTick());
+		const double Ratio = InNewTickRate.GetFramesPerTick() / CurrentTickRate.GetFramesPerTick();
 
 		// adjust time-till-fire for existing commands
 		for (auto& Command : PendingCommands)
@@ -321,7 +321,8 @@ namespace Audio
 		}
 
 		// get number of frames until event (assuming we are at frame 0)
-		FramesUntilExec = FMath::Max(0, InNewEvent->OverrideFramesUntilExec(Metronome.GetFramesUntilBoundary(InQuantizationBondary)));
+		FramesUntilExec = FMath::RoundToInt(Metronome.GetFramesUntilBoundary(InQuantizationBondary)); // query metronome (round result to int)
+		FramesUntilExec = FMath::Max(0, InNewEvent->OverrideFramesUntilExec(FramesUntilExec)); // allow command to override the deadline (clamp result)
 
 		// if this is going to execute on the next tick, warn Game Thread Subscribers as soon as possible
 		if (FramesUntilExec == 0)
@@ -377,12 +378,12 @@ namespace Audio
 		FQuartzClockTickRate TickRate = Metronome.GetTickRate();
 
 		// get number of frames until the relevant quantization event
-		int64 FramesUntilExec = TickRate.GetFramesPerDuration(QuantizationType);
+		double FramesUntilExec = TickRate.GetFramesPerDuration(QuantizationType);
 
 		//Translate frames to seconds
-		float SampleRate = TickRate.GetSampleRate();
+		double SampleRate = TickRate.GetSampleRate();
 
-		if (SampleRate != 0)
+		if (!FMath::IsNearlyZero(SampleRate))
 		{
 			return (FramesUntilExec * Multiplier) / SampleRate;
 		}
diff --git a/Engine/Source/Runtime/AudioMixer/Private/Quartz/AudioMixerQuantizedCommands.cpp b/Engine/Source/Runtime/AudioMixer/Private/Quartz/AudioMixerQuantizedCommands.cpp
--- a/Engine/Source/Runtime/AudioMixer/Private/Quartz/AudioMixerQuantizedCommands.cpp
+++ b/Engine/Source/Runtime/AudioMixer/Private/Quartz/AudioMixerQuantizedCommands.cpp
@@ -119,7 +119,7 @@ namespace Audio
 	int32 FQuantizedQueueCommand::OverrideFramesUntilExec(int32 NumFramesUntilExec)
 	{
 		// Calculate the amount of time before taking up a voice slot
-		int32 NumFramesBeforeVoiceSlot = NumFramesUntilExec - OwningClockPtr->GetTickRate().GetFramesPerDuration(AudioComponentData.AnticapatoryBoundary.Quantization);
+		int32 NumFramesBeforeVoiceSlot = NumFramesUntilExec - static_cast<int32>(OwningClockPtr->GetTickRate().GetFramesPerDuration(AudioComponentData.AnticapatoryBoundary.Quantization));
 
 		//If NumFramesBeforeVoiceSlot is less than 0, change the boundary back to the original, and mark this command as having 0 frames till exec
 		if (NumFramesBeforeVoiceSlot < 0)
diff --git a/Engine/Source/Runtime/AudioMixer/Private/Quartz/QuartzMetronome.cpp b/Engine/Source/Runtime/AudioMixer/Private/Quartz/QuartzMetronome.cpp
--- a/Engine/Source/Runtime/AudioMixer/Private/Quartz/QuartzMetronome.cpp
+++ b/Engine/Source/Runtime/AudioMixer/Private/Quartz/QuartzMetronome.cpp
@@ -99,7 +99,7 @@ namespace Audio
 	void FQuartzMetronome::SetTickRate(FQuartzClockTickRate InNewTickRate, int32 NumFramesLeft)
 	{
 		// early exit?
-		const bool bSameAsOldTickRate = (InNewTickRate.GetFramesPerTick() == CurrentTickRate.GetFramesPerTick());
+		const bool bSameAsOldTickRate = FMath::IsNearlyEqual(InNewTickRate.GetFramesPerTick(), CurrentTickRate.GetFramesPerTick());
 		const bool bIsInitialized = (MusicalDurationsInFrames[0] > 0);
 
 		if (bSameAsOldTickRate && bIsInitialized)
@@ -108,11 +108,11 @@ namespace Audio
 		}
 
 		// ratio between new and old rates
-		const float Ratio = static_cast<float>(InNewTickRate.GetFramesPerTick()) / static_cast<float>(CurrentTickRate.GetFramesPerTick());
+		const double Ratio = InNewTickRate.GetFramesPerTick() / CurrentTickRate.GetFramesPerTick();
 
 		if (NumFramesLeft)
 		{
-			for (int32& Value : FramesLeftInMusicalDuration.FramesInTimeValueInternal)
+			for (double& Value : FramesLeftInMusicalDuration.FramesInTimeValueInternal)
 			{
 				Value = NumFramesLeft + Ratio * (Value - NumFramesLeft);
 			}
@@ -134,7 +134,7 @@ namespace Audio
 		RecalculateDurations();
 	}
 
-	int32 FQuartzMetronome::GetFramesUntilBoundary(FQuartzQuantizationBoundary InQuantizationBoundary) const
+	double FQuartzMetronome::GetFramesUntilBoundary(FQuartzQuantizationBoundary InQuantizationBoundary) const
 	{
 		if (!ensure(InQuantizationBoundary.Quantization != EQuartzCommandQuantization::None))
 		{
@@ -148,7 +148,7 @@ namespace Audio
 		}
 
 		// number of frames until the next occurrence of this boundary
-		int32 FramesUntilBoundary = FramesLeftInMusicalDuration[InQuantizationBoundary.Quantization];
+		double FramesUntilBoundary = FramesLeftInMusicalDuration[InQuantizationBoundary.Quantization];
 
 		// how many multiples actually exist until the boundary we care about?
 		int32 NumDurationsLeft = static_cast<int32>(InQuantizationBoundary.Multiplier) - 1;
@@ -207,7 +207,7 @@ namespace Audio
 			}
 		}
 
-		const float FrationalPortion = FMath::Fractional(InQuantizationBoundary.Multiplier);
+		const double FrationalPortion = FMath::Fractional(InQuantizationBoundary.Multiplier);
 
 		// for Beats, the lengths are not uniform for complex meters
 		if ((InQuantizationBoundary.Quantization == EQuartzCommandQuantization::Beat) && PulseDurations.Num())
@@ -238,7 +238,7 @@ namespace Audio
 		{
 			const float Multiplier = NumDurationsLeft + FrationalPortion;
 			const float Duration = static_cast<float>(MusicalDurationsInFrames[InQuantizationBoundary.Quantization]);
-			FramesUntilBoundary += FMath::RoundToInt(Multiplier * Duration);
+			FramesUntilBoundary += Multiplier * Duration;
 		}
 
 		return FramesUntilBoundary;
@@ -334,12 +334,12 @@ namespace Audio
 	{
 		CurrentTimeStamp.Reset();
 
-		for (int32& FrameCount : FramesLeftInMusicalDuration.FramesInTimeValueInternal)
+		for (double& FrameCount : FramesLeftInMusicalDuration.FramesInTimeValueInternal)
 		{
-			FrameCount = 0;
+			FrameCount = 0.0;
 		}
 
-		TimeSinceStart = 0;
+		TimeSinceStart = 0.0;
 		PulseDurationIndex = -1;
 	}
 
@@ -354,7 +354,7 @@ namespace Audio
 		}
 
 		// determine actual length of a bar
-		const int32 BarLength = CurrentTimeSignature.NumBeats * CurrentTickRate.GetFramesPerDuration(CurrentTimeSignature.BeatType);
+		const double BarLength = CurrentTimeSignature.NumBeats * CurrentTickRate.GetFramesPerDuration(CurrentTimeSignature.BeatType);
 		MusicalDurationsInFrames[EQuartzCommandQuantization::Bar] = BarLength;
 
 		// default beat value to the denominator of our time signature
@@ -364,8 +364,8 @@ namespace Audio
 		if (CurrentTimeSignature.OptionalPulseOverride.Num() != 0)
 		{
 			// determine the length of each beat
-			int32 LengthCounter = 0;
-			int32 StepLength = 0;
+			double LengthCounter = 0.0;
+			double StepLength = 0.0;
 
 			for (const FQuartzPulseOverrideStep& PulseStep : CurrentTimeSignature.OptionalPulseOverride)
 			{
@@ -394,12 +394,12 @@ namespace Audio
 			}
 
 			// check to see if all our pulses are the same length
-			const int32 FirstValue = PulseDurations[0];
+			const double FirstValue = PulseDurations[0];
 			bool bBeatDurationsAreConstant = true;
 
-			for (const int32& Values : PulseDurations)
+			for (const double& Values : PulseDurations)
 			{
-				if (Values != FirstValue)
+				if (!FMath::IsNearlyEqual(Values, FirstValue))
 				{
 					bBeatDurationsAreConstant = false;
 					break;
diff --git a/Engine/Source/Runtime/AudioMixer/Public/Quartz/QuartzMetronome.h b/Engine/Source/Runtime/AudioMixer/Public/Quartz/QuartzMetronome.h
--- a/Engine/Source/Runtime/AudioMixer/Public/Quartz/QuartzMetronome.h
+++ b/Engine/Source/Runtime/AudioMixer/Public/Quartz/QuartzMetronome.h
@@ -32,7 +32,7 @@ namespace Audio
 		// affects bars/beats values we send back to the game thread
 		void SetTimeSignature(const FQuartzTimeSignature& InNewTimeSignature);
 
-		int32 GetFramesUntilBoundary(FQuartzQuantizationBoundary InQuantizationBoundary) const;
+		double GetFramesUntilBoundary(FQuartzQuantizationBoundary InQuantizationBoundary) const;
 
 		const FQuartzTimeSignature & GetTimeSignature() const { return CurrentTimeSignature; }
 
@@ -77,28 +77,28 @@ namespace Audio
 		{
 		public:
 			// index operators for EQuartzCommandQuantization
-			int32& operator[](EQuartzCommandQuantization InTimeValue)
+			double& operator[](EQuartzCommandQuantization InTimeValue)
 			{
 				return FramesInTimeValueInternal[static_cast<int32>(InTimeValue)];
 			}
 
-			const int32& operator[](EQuartzCommandQuantization InTimeValue) const
+			const double& operator[](EQuartzCommandQuantization InTimeValue) const
 			{
 				return FramesInTimeValueInternal[static_cast<int32>(InTimeValue)];
 			}
 
 			// index operators for int32
-			int32& operator[](int32 Index)
+			double& operator[](int32 Index)
 			{
 				return FramesInTimeValueInternal[Index];
 			}
 
-			const int32& operator[](int32 Index) const
+			const double& operator[](int32 Index) const
 			{
 				return FramesInTimeValueInternal[Index];
 			}
 
-			int32 FramesInTimeValueInternal[static_cast<int32>(EQuartzCommandQuantization::Count)]{ 0 };
+			double FramesInTimeValueInternal[static_cast<int32>(EQuartzCommandQuantization::Count)]{ 0.0 };
 		};
 
 		// array of lengths of musical durations (in audio frames)
@@ -108,7 +108,7 @@ namespace Audio
 		FramesInTimeValue FramesLeftInMusicalDuration;
 
 		// optional array of pulse duration overrides (for odd meters)
-		TArray<int32> PulseDurations;
+		TArray<double> PulseDurations;
 
 		// the index of the active pulse duration override
 		int32 PulseDurationIndex{ -1 };
diff --git a/Engine/Source/Runtime/Engine/Classes/Sound/QuartzQuantizationUtilities.cpp b/Engine/Source/Runtime/Engine/Classes/Sound/QuartzQuantizationUtilities.cpp
--- a/Engine/Source/Runtime/Engine/Classes/Sound/QuartzQuantizationUtilities.cpp
+++ b/Engine/Source/Runtime/Engine/Classes/Sound/QuartzQuantizationUtilities.cpp
@@ -144,8 +144,8 @@ void FQuartLatencyTracker::PushSingleResult(const double& InResult)
 		LifetimeAverage = (LifetimeAverage * (NumEntries - 1) + InResult) / NumEntries;
 	}
 
-	Min = FMath::Min(Min, static_cast<float>(InResult));
-	Max = FMath::Max(Max, static_cast<float>(InResult));
+	Min = FMath::Min(Min, (float)InResult);
+	Max = FMath::Max(Max, (float)InResult);
 }
 
 void FQuartLatencyTracker::DigestQueue()
@@ -165,24 +165,24 @@ namespace Audio
 {
 	FQuartzClockTickRate::FQuartzClockTickRate()
 	{
-		SetBeatsPerMinute(60.f);
+		SetBeatsPerMinute(60.0);
 	}
 
 	void FQuartzClockTickRate::SetFramesPerTick(int32 InNewFramesPerTick)
 	{
-		if (InNewFramesPerTick < 1)
+		if (InNewFramesPerTick < 1.0)
 		{
 			UE_LOG(LogAudioQuartz, Warning, TEXT("Quartz Metronme requires at least 1 frame per tick, clamping request"));
-			InNewFramesPerTick = 1;
+			InNewFramesPerTick = 1.0;
 		}
 
-		FramesPerTick = InNewFramesPerTick;
+		FramesPerTick = (double)InNewFramesPerTick;
 		RecalculateDurationsBasedOnFramesPerTick();
 	}
 
 	void FQuartzClockTickRate::SetMillisecondsPerTick(float InNewMillisecondsPerTick)
 	{
-		FramesPerTick = FMath::Max(1.0f, (InNewMillisecondsPerTick * SampleRate) / 1000.f);
+		FramesPerTick = FMath::Max(1.0, ((double)InNewMillisecondsPerTick * SampleRate) / 1000.0);
 		RecalculateDurationsBasedOnFramesPerTick();
 	}
 
@@ -190,7 +190,7 @@ namespace Audio
 	{
 		check(InNewThirtySecondNotesPerMinute > 0);
 
-		FramesPerTick = FMath::Max(1.0f, (60.f * SampleRate ) / InNewThirtySecondNotesPerMinute);
+		FramesPerTick = FMath::Max(1.0, (60. * SampleRate ) / (double)InNewThirtySecondNotesPerMinute);
 		RecalculateDurationsBasedOnFramesPerTick();
 	}
 
@@ -203,7 +203,7 @@ namespace Audio
 		// FramesPerTick = 1/8 * (60.f / (InNewBeatsPerMinute)) * SampleRate;
 		// (60.0 / 8.0) = 7.5f
 
-		FramesPerTick = FMath::Max(1.0f, (7.5f * SampleRate) / InNewBeatsPerMinute);
+		FramesPerTick = FMath::Max(1.0, (7.5 * SampleRate) / (double)InNewBeatsPerMinute);
 		RecalculateDurationsBasedOnFramesPerTick();
 	}
 
@@ -211,16 +211,16 @@ namespace Audio
 	{
 		check(InNewSampleRate >= 0);
 
-		FramesPerTick = FMath::Max(1.0f, (InNewSampleRate / SampleRate) * static_cast<float>(FramesPerTick));
+		FramesPerTick = FMath::Max(1.0, ((double)InNewSampleRate / SampleRate) * FramesPerTick);
 		SampleRate = InNewSampleRate;
 
 		RecalculateDurationsBasedOnFramesPerTick();
 	}
 
-	int64 FQuartzClockTickRate::GetFramesPerDuration(EQuartzCommandQuantization InDuration) const
+	double FQuartzClockTickRate::GetFramesPerDuration(EQuartzCommandQuantization InDuration) const
 	{
-		const int64 FramesPerDotted16th = FramesPerTick * 3;
-		const int64 FramesPer16thTriplet = 4.f * FramesPerTick / 3.f;
+		const double FramesPerDotted16th = FramesPerTick * 3.0;
+		const double FramesPer16thTriplet = 4.0 * FramesPerTick / 3.0;
 
 		switch (InDuration)
 		{
@@ -233,37 +233,37 @@ namespace Audio
 			return FramesPerTick; // same as 1/32nd note
 
 		case EQuartzCommandQuantization::SixteenthNote:
-			return (int64)FramesPerTick << 1;
+			return FramesPerTick * 2.0;
 
 		case EQuartzCommandQuantization::EighthNote:
-			return (int64)FramesPerTick << 2;
+			return FramesPerTick * 4.0;
 
 		case EQuartzCommandQuantization::Beat: // default to quarter note (should be overridden for non-basic meters)
 		case EQuartzCommandQuantization::QuarterNote:
-			return (int64)FramesPerTick << 3;
+			return FramesPerTick * 8.0;
 
 		case EQuartzCommandQuantization::HalfNote:
-			return (int64)FramesPerTick << 4;
+			return FramesPerTick * 16.0;
 
 		case EQuartzCommandQuantization::Bar: // default to whole note (should be overridden for non-4/4 meters)
 		case EQuartzCommandQuantization::WholeNote:
-			return (int64)FramesPerTick << 5;
+			return FramesPerTick * 32.0;
 
 			// DOTTED
 		case EQuartzCommandQuantization::DottedSixteenthNote:
 			return FramesPerDotted16th;
 
 		case EQuartzCommandQuantization::DottedEighthNote:
-			return FramesPerDotted16th << 1;
+			return FramesPerDotted16th * 2.0;
 
 		case EQuartzCommandQuantization::DottedQuarterNote:
-			return FramesPerDotted16th << 2;
+			return FramesPerDotted16th * 4.0;
 
 		case EQuartzCommandQuantization::DottedHalfNote:
-			return FramesPerDotted16th << 3;
+			return FramesPerDotted16th * 8.0;
 
 		case EQuartzCommandQuantization::DottedWholeNote:
-			return FramesPerDotted16th << 4;
+			return FramesPerDotted16th * 16.0;
 
 
 			// TRIPLETS
@@ -271,14 +271,13 @@ namespace Audio
 			return FramesPer16thTriplet;
 
 		case EQuartzCommandQuantization::EighthNoteTriplet:
-			return FramesPer16thTriplet << 1;
+			return FramesPer16thTriplet * 2.0;
 
 		case EQuartzCommandQuantization::QuarterNoteTriplet:
-			return FramesPer16thTriplet << 2;
+			return FramesPer16thTriplet * 4.0;
 
 		case EQuartzCommandQuantization::HalfNoteTriplet:
-			return FramesPer16thTriplet << 3;
-
+			return FramesPer16thTriplet * 8.0;
 
 
 		default:
@@ -289,7 +288,7 @@ namespace Audio
 		return INVALID_DURATION;
 	}
 
-	int64 FQuartzClockTickRate::GetFramesPerDuration(EQuartzTimeSignatureQuantization InDuration) const
+	double FQuartzClockTickRate::GetFramesPerDuration(EQuartzTimeSignatureQuantization InDuration) const
 	{
 		switch (InDuration)
 		{
@@ -366,20 +365,19 @@ namespace Audio
 
 	void FQuartzClockTickRate::RecalculateDurationsBasedOnFramesPerTick()
 	{
-		check(FramesPerTick > 0);
-		check(SampleRate > 0);
-		const float FloatFramesPerTick = static_cast<float>(FramesPerTick);
+		check(FramesPerTick > 0.0);
+		check(SampleRate > 0.0);
 
-		SecondsPerTick = (FloatFramesPerTick / SampleRate);
-		MillisecondsPerTick = SecondsPerTick * 1000.f;
-		ThirtySecondNotesPerMinute = (60.f * SampleRate) / FloatFramesPerTick;
-		BeatsPerMinute = ThirtySecondNotesPerMinute / 8.0f;
+		SecondsPerTick = FramesPerTick / SampleRate;
+		MillisecondsPerTick = SecondsPerTick * 1000.0;
+		ThirtySecondNotesPerMinute = (60.0 * SampleRate) / FramesPerTick;
+		BeatsPerMinute = ThirtySecondNotesPerMinute / 8.0;
 	}
 
 
 	void FQuartzClockTickRate::SetSecondsPerTick(float InNewSecondsPerTick)
 	{
-		SetMillisecondsPerTick(InNewSecondsPerTick * 1000.f);
+		SetMillisecondsPerTick(InNewSecondsPerTick * 1000.0);
 	}
 
 
diff --git a/Engine/Source/Runtime/Engine/Classes/Sound/QuartzQuantizationUtilities.h b/Engine/Source/Runtime/Engine/Classes/Sound/QuartzQuantizationUtilities.h
--- a/Engine/Source/Runtime/Engine/Classes/Sound/QuartzQuantizationUtilities.h
+++ b/Engine/Source/Runtime/Engine/Classes/Sound/QuartzQuantizationUtilities.h
@@ -346,7 +346,7 @@ namespace Audio
 		void SetSampleRate(float InNewSampleRate);
 
 		// Getters
-		int32 GetFramesPerTick() const { return FramesPerTick; }
+		double GetFramesPerTick() const { return FramesPerTick; }
 
 		float GetMillisecondsPerTick() const { return MillisecondsPerTick; }
 
@@ -358,9 +358,9 @@ namespace Audio
 
 		float GetSampleRate() const { return SampleRate; }
 
-		int64 GetFramesPerDuration(EQuartzCommandQuantization InDuration) const;
+		double GetFramesPerDuration(EQuartzCommandQuantization InDuration) const;
 
-		int64 GetFramesPerDuration(EQuartzTimeSignatureQuantization InDuration) const;
+		double GetFramesPerDuration(EQuartzTimeSignatureQuantization InDuration) const;
 
 		bool IsValid(int32 InEventResolutionThreshold = 1) const;
 
@@ -371,12 +371,12 @@ namespace Audio
 	private:
 		// FramesPerTick is our ground truth 
 		// update FramesPerTick and call RecalculateDurationsBasedOnFramesPerTick() to update other members
-		int32 FramesPerTick{ 1 };
-		float MillisecondsPerTick{ 1.f };
-		float SecondsPerTick{ 1.f };
-		float ThirtySecondNotesPerMinute{ 1.f };
-		float BeatsPerMinute{ 1 };
-		float SampleRate{ 44100.f };
+		double FramesPerTick{ 1.0 };
+		double MillisecondsPerTick{ 1.0 };
+		double SecondsPerTick{ 1.0 };
+		double ThirtySecondNotesPerMinute{ 1.0 };
+		double BeatsPerMinute{ 0.0 };
+		double SampleRate{ 44100.0 };
 
 		void RecalculateDurationsBasedOnFramesPerTick();
 
diff --git a/Engine/Source/Runtime/Engine/Private/Components/AudioComponent.cpp b/Engine/Source/Runtime/Engine/Private/Components/AudioComponent.cpp
--- a/Engine/Source/Runtime/Engine/Private/Components/AudioComponent.cpp
+++ b/Engine/Source/Runtime/Engine/Private/Components/AudioComponent.cpp
@@ -497,8 +497,8 @@ void UAudioComponent::PlayQuantized(
 	Audio::FQuartzClockTickRate OutTickRate;
 	InClockHandle->GetCurrentTickRate(WorldContextObject, OutTickRate);
 
-	int32 NumFramesBeforeMinTime = OutTickRate.GetFramesPerDuration(MinimumQuantization);
-	int32 NumFramesForDesiredTime = OutTickRate.GetFramesPerDuration(InQuantizationBoundary.Quantization) * InQuantizationBoundary.Multiplier;
+	double NumFramesBeforeMinTime = OutTickRate.GetFramesPerDuration(MinimumQuantization);
+	double NumFramesForDesiredTime = OutTickRate.GetFramesPerDuration(InQuantizationBoundary.Quantization) * InQuantizationBoundary.Multiplier;
 
 	// If the desired quantization time is less than our min time, just execute immediately
 	const bool bStealVoiceSlot = NumFramesForDesiredTime <= NumFramesBeforeMinTime;
