Fixed issue where saved character moves could contain uobject ptrs that get GC'd out from under them. Implemented in a way that project-specific extensions can do the same with any references they introduce. Adapted from PR 8794.

diff --git a/Engine/Source/Runtime/Engine/Classes/GameFramework/CharacterMovementComponent.h b/Engine/Source/Runtime/Engine/Classes/GameFramework/CharacterMovementComponent.h
--- a/Engine/Source/Runtime/Engine/Classes/GameFramework/CharacterMovementComponent.h
+++ b/Engine/Source/Runtime/Engine/Classes/GameFramework/CharacterMovementComponent.h
@@ -1219,6 +1219,7 @@ public:
 
 	// Begin UObject Interface
 	virtual void Serialize(FArchive& Archive) override;
+	static void AddReferencedObjects(UObject* InThis, FReferenceCollector& Collector);
 	// End UObject Interface
 
 	//Begin UActorComponent Interface
@@ -2875,6 +2876,9 @@ public:
 	/** Packs control rotation for network transport */
 	virtual void GetPackedAngles(uint32& YawAndPitchPack, uint8& RollPack) const;
 
+	/** Allows references to be considered during GC */
+	virtual void AddStructReferencedObjects(FReferenceCollector& Collector) const;
+
 	// Bit masks used by GetCompressedFlags() to encode movement information.
 	enum CompressedFlags
 	{
@@ -2918,11 +2922,16 @@ public:
 
 class ENGINE_API FNetworkPredictionData_Client_Character : public FNetworkPredictionData_Client, protected FNoncopyable
 {
+	using Super = FNetworkPredictionData_Client;
+
 public:
 
 	FNetworkPredictionData_Client_Character(const UCharacterMovementComponent& ClientMovement);
 	virtual ~FNetworkPredictionData_Client_Character();
 
+	/** Allows references to be considered during GC */
+	void AddStructReferencedObjects(FReferenceCollector& Collector) const override;
+
 	/** Client timestamp of last time it sent a servermove() to the server. This is an increasing timestamp from the owning UWorld. Used for holding off on sending movement updates to save bandwidth. */
 	float ClientUpdateTime;
 
diff --git a/Engine/Source/Runtime/Engine/Classes/GameFramework/RootMotionSource.h b/Engine/Source/Runtime/Engine/Classes/GameFramework/RootMotionSource.h
--- a/Engine/Source/Runtime/Engine/Classes/GameFramework/RootMotionSource.h
+++ b/Engine/Source/Runtime/Engine/Classes/GameFramework/RootMotionSource.h
@@ -878,7 +878,7 @@ struct ENGINE_API FRootMotionSourceGroup
 	bool operator!=(const FRootMotionSourceGroup& Other) const;
 
 	/** Exposes references to GC system */
-	void AddStructReferencedObjects(class FReferenceCollector& Collector);
+	void AddStructReferencedObjects(FReferenceCollector& Collector) const;
 
 protected:
 
diff --git a/Engine/Source/Runtime/Engine/Classes/Interfaces/NetworkPredictionInterface.h b/Engine/Source/Runtime/Engine/Classes/Interfaces/NetworkPredictionInterface.h
--- a/Engine/Source/Runtime/Engine/Classes/Interfaces/NetworkPredictionInterface.h
+++ b/Engine/Source/Runtime/Engine/Classes/Interfaces/NetworkPredictionInterface.h
@@ -73,6 +73,9 @@ public:
 	}
 
 	virtual ~FNetworkPredictionData_Client() {}
+
+	/** Can be overridden to allow references to be accounted for during GC */
+	virtual void AddStructReferencedObjects(FReferenceCollector& Collector) const {};
 };
 
 
@@ -101,6 +104,9 @@ public:
 		bForcedUpdateDurationExceeded = false;
 	}
 
+	/** Can be overridden to allow references to be accounted for during GC */
+	virtual void AddStructReferencedObjects(FReferenceCollector& Collector) const {};
+
 	/** Server clock time when last server move was received or movement was forced to be processed */
 	float ServerTimeStamp;
 
diff --git a/Engine/Source/Runtime/Engine/Private/Components/CharacterMovementComponent.cpp b/Engine/Source/Runtime/Engine/Private/Components/CharacterMovementComponent.cpp
--- a/Engine/Source/Runtime/Engine/Private/Components/CharacterMovementComponent.cpp
+++ b/Engine/Source/Runtime/Engine/Private/Components/CharacterMovementComponent.cpp
@@ -1394,6 +1394,28 @@ void UCharacterMovementComponent::Serialize(FArchive& Archive)
 	}
 }
 
+void UCharacterMovementComponent::AddReferencedObjects(UObject* InThis, FReferenceCollector& Collector)
+{
+	Super::AddReferencedObjects(InThis, Collector);
+
+	UCharacterMovementComponent* This = CastChecked<UCharacterMovementComponent>(InThis);
+
+	if (This->HasPredictionData_Client())
+	{
+		if (const FNetworkPredictionData_Client* ClientData = This->GetPredictionData_Client())
+		{
+			ClientData->AddStructReferencedObjects(Collector);
+		}
+	}
+	else if (This->HasPredictionData_Server())
+	{
+		if (const FNetworkPredictionData_Server* ServerData = This->GetPredictionData_Server())
+		{
+			ServerData->AddStructReferencedObjects(Collector);
+		}
+	}
+}
+
 void UCharacterMovementComponent::TickComponent(float DeltaTime, enum ELevelTick TickType, FActorComponentTickFunction *ThisTickFunction)
 {
 	SCOPED_NAMED_EVENT(UCharacterMovementComponent_TickComponent, FColor::Yellow);
@@ -11466,6 +11488,29 @@ FNetworkPredictionData_Client_Character::~FNetworkPredictionData_Client_Characte
 	LastAckedMove = NULL;
 }
 
+void FNetworkPredictionData_Client_Character::AddStructReferencedObjects(FReferenceCollector& Collector) const
+{
+	Super::AddStructReferencedObjects(Collector);
+
+	for (FSavedMovePtr SavedMove : SavedMoves)
+	{
+		if (const FSavedMove_Character* SavedMovePtr = SavedMove.Get())
+		{
+			SavedMovePtr->AddStructReferencedObjects(Collector);
+		}
+	}
+
+	if (const FSavedMove_Character* PendingMovePtr = PendingMove.Get())
+	{
+		PendingMovePtr->AddStructReferencedObjects(Collector);
+	}
+
+	if (const FSavedMove_Character* LastAckedMovePtr = LastAckedMove.Get())
+	{
+		LastAckedMovePtr->AddStructReferencedObjects(Collector);
+	}
+}
+
 
 FSavedMovePtr FNetworkPredictionData_Client_Character::CreateSavedMove()
 {
@@ -12334,6 +12379,18 @@ uint8 FSavedMove_Character::GetCompressedFlags() const
 	return Result;
 }
 
+
+void FSavedMove_Character::AddStructReferencedObjects(FReferenceCollector& Collector) const
+{
+	if (const UAnimMontage* RootMotionMontageRawPtr = RootMotionMontage.Get())
+	{
+		Collector.AddReferencedObject(RootMotionMontageRawPtr);
+	}
+
+	SavedRootMotion.AddStructReferencedObjects(Collector);
+}
+
+
 void UCharacterMovementComponent::UpdateFromCompressedFlags(uint8 Flags)
 {
 	if (!CharacterOwner)
diff --git a/Engine/Source/Runtime/Engine/Private/GameFramework/RootMotionSource.cpp b/Engine/Source/Runtime/Engine/Private/GameFramework/RootMotionSource.cpp
--- a/Engine/Source/Runtime/Engine/Private/GameFramework/RootMotionSource.cpp
+++ b/Engine/Source/Runtime/Engine/Private/GameFramework/RootMotionSource.cpp
@@ -2270,9 +2270,9 @@ bool FRootMotionSourceGroup::operator!=(const FRootMotionSourceGroup& Other) con
 	return !(FRootMotionSourceGroup::operator==(Other));
 }
 
-void FRootMotionSourceGroup::AddStructReferencedObjects(class FReferenceCollector& Collector)
+void FRootMotionSourceGroup::AddStructReferencedObjects(FReferenceCollector& Collector) const
 {
-	for (TSharedPtr<FRootMotionSource>& RootMotionSource : RootMotionSources)
+	for (const TSharedPtr<FRootMotionSource>& RootMotionSource : RootMotionSources)
 	{
 		if (RootMotionSource.IsValid())
 		{
@@ -2280,7 +2280,7 @@ void FRootMotionSourceGroup::AddStructReferencedObjects(class FReferenceCollecto
 		}
 	}
 
-	for (TSharedPtr<FRootMotionSource>& RootMotionSource : PendingAddRootMotionSources)
+	for (const TSharedPtr<FRootMotionSource>& RootMotionSource : PendingAddRootMotionSources)
 	{
 		if (RootMotionSource.IsValid())
 		{
