PR #8557:  Allow resize of Render Targets via Blueprints (during runtime)

diff --git a/Engine/Source/Runtime/Engine/Classes/Kismet/KismetRenderingLibrary.h b/Engine/Source/Runtime/Engine/Classes/Kismet/KismetRenderingLibrary.h
--- a/Engine/Source/Runtime/Engine/Classes/Kismet/KismetRenderingLibrary.h
+++ b/Engine/Source/Runtime/Engine/Classes/Kismet/KismetRenderingLibrary.h
@@ -67,6 +67,13 @@ class UKismetRenderingLibrary : public UBlueprintFunctionLibrary
 	UFUNCTION(BlueprintCallable, Category="Rendering")
 	static ENGINE_API void ReleaseRenderTarget2D(UTextureRenderTarget2D* TextureRenderTarget);
 
+	/**
+	 * Changes the resolution of a render target. This is useful for when you need to resize the game viewport or change the in-game resolution during runtime
+	 * and thus need to update the sizes of all the render targets in the game accordingly.
+	 */
+	UFUNCTION(BlueprintCallable, Category="Rendering")
+	static ENGINE_API void ResizeRenderTarget2D(UTextureRenderTarget2D* TextureRenderTarget, int32 Width = 256, int32 Height = 256);
+
 	/** 
 	 * Renders a quad with the material applied to the specified render target.   
 	 * This sets the render target even if it is already set, which is an expensive operation. 
diff --git a/Engine/Source/Runtime/Engine/Private/KismetRenderingLibrary.cpp b/Engine/Source/Runtime/Engine/Private/KismetRenderingLibrary.cpp
--- a/Engine/Source/Runtime/Engine/Private/KismetRenderingLibrary.cpp
+++ b/Engine/Source/Runtime/Engine/Private/KismetRenderingLibrary.cpp
@@ -129,6 +129,24 @@ void UKismetRenderingLibrary::ReleaseRenderTarget2D(UTextureRenderTarget2D* Text
 	TextureRenderTarget->ReleaseResource();
 }
 
+void UKismetRenderingLibrary::ResizeRenderTarget2D(UTextureRenderTarget2D* TextureRenderTarget, int32 Width, int32 Height)
+{
+	if (!TextureRenderTarget)
+	{
+		return;
+	}
+
+	// Resize function silently fails if either dimension isn't positive, so check for that here so we can warn the caller
+	if ((Width > 0) && (Height > 0))
+	{
+		TextureRenderTarget->ResizeTarget(Width, Height);
+	}
+	else
+	{
+		FMessageLog("Blueprint").Warning(LOCTEXT("ResizeRenderTarget2D_InvalidDimensions", "ResizeRenderTarget2D: Dimensions must be positive."));
+	}
+}
+
 void UKismetRenderingLibrary::DrawMaterialToRenderTarget(UObject* WorldContextObject, UTextureRenderTarget2D* TextureRenderTarget, UMaterialInterface* Material)
 {
 	if (!FApp::CanEverRender())
