UE5_MAIN:  Editor bug fix for mouse cursor location when tracking a cinematic camera, breaking various operations, such as the focus distance eyedropper tool or moving actors, among many other things.

Cinematic camera tracking sets a viewport that is a subset of the render target in order to achieve an aspect ratio that matches the cinematic, with letter boxing added as needed.  The mouse cursor calculations weren't taking this letterboxing into account, because they were only using the size of the viewport, and ignoring the offset of the viewport.  To minimize the risk of affecting existing code using PixelToScreen, a new utility function CursorToScreen was added, and the FViewportCursorLocation class alone was updated to use it.

#rnx

diff --git a/Engine/Source/Editor/UnrealEd/Private/EditorViewportClient.cpp b/Engine/Source/Editor/UnrealEd/Private/EditorViewportClient.cpp
--- a/Engine/Source/Editor/UnrealEd/Private/EditorViewportClient.cpp
+++ b/Engine/Source/Editor/UnrealEd/Private/EditorViewportClient.cpp
@@ -179,7 +179,7 @@ FViewportCursorLocation::FViewportCursorLocation(const FSceneView* View, FEditor
 	: Origin(ForceInit), Direction(ForceInit), CursorPos(X, Y)
 {
 
-	FVector4 ScreenPos = View->PixelToScreen(X, Y, 0);
+	FVector4 ScreenPos = View->CursorToScreen(X, Y, 0);
 
 	const FMatrix InvViewMatrix = View->ViewMatrices.GetInvViewMatrix();
 	const FMatrix InvProjMatrix = View->ViewMatrices.GetInvProjectionMatrix();
diff --git a/Engine/Source/Runtime/Engine/Private/SceneView.cpp b/Engine/Source/Runtime/Engine/Private/SceneView.cpp
--- a/Engine/Source/Runtime/Engine/Private/SceneView.cpp
+++ b/Engine/Source/Runtime/Engine/Private/SceneView.cpp
@@ -1170,6 +1170,17 @@ FVector4 FSceneView::PixelToScreen(float InX,float InY,float Z) const
 	}
 }
 
+// Similar to PixelToScreen, but handles cases where the viewport isn't the whole render target, such as a mouse cursor location in
+// render target space.  If outside the bounds of the viewport, will clamp to the nearest edge, so a valid result is returned.  The
+// caller can check the cursor against UnscaledViewRect if they want to reject out of bounds cursor locations.
+FVector4 FSceneView::CursorToScreen(float InX, float InY, float Z) const
+{
+	float ClampedX = FMath::Clamp(InX - UnscaledViewRect.Min.X, 0.0f, UnscaledViewRect.Width());
+	float ClampedY = FMath::Clamp(InY - UnscaledViewRect.Min.Y, 0.0f, UnscaledViewRect.Height());
+
+	return PixelToScreen(ClampedX, ClampedY, Z);
+}
+
 /** Transforms a point from the view's world-space into pixel coordinates relative to the view's X,Y. */
 bool FSceneView::WorldToPixel(const FVector& WorldPoint,FVector2D& OutPixelLocation) const
 {
diff --git a/Engine/Source/Runtime/Engine/Public/SceneView.h b/Engine/Source/Runtime/Engine/Public/SceneView.h
--- a/Engine/Source/Runtime/Engine/Public/SceneView.h
+++ b/Engine/Source/Runtime/Engine/Public/SceneView.h
@@ -1298,6 +1298,9 @@ public:
 	/** Transforms a point from pixel coordinates relative to the view's X,Y (left, top) into the view's screen-space. */
 	FVector4 PixelToScreen(float X,float Y,float Z) const;
 
+	/** Transforms a cursor location in render target pixel coordinates into the view's screen-space, taking into account the viewport rectangle. */
+	FVector4 CursorToScreen(float X, float Y, float Z) const;
+
 	/** Transforms a point from the view's world-space into pixel coordinates relative to the view's X,Y (left, top). */
 	bool WorldToPixel(const FVector& WorldPoint,FVector2D& OutPixelLocation) const;
 
