Add additionnal info to check if the actor we loaded is from the same level as we're expecting to help diagnose the issue.

diff --git a/Engine/Source/Runtime/Engine/Private/WorldPartition/WorldPartitionLevelHelper.cpp b/Engine/Source/Runtime/Engine/Private/WorldPartition/WorldPartitionLevelHelper.cpp
--- a/Engine/Source/Runtime/Engine/Private/WorldPartition/WorldPartitionLevelHelper.cpp
+++ b/Engine/Source/Runtime/Engine/Private/WorldPartition/WorldPartitionLevelHelper.cpp
@@ -289,7 +289,7 @@ bool FWorldPartitionLevelHelper::LoadActors(ULevel* InDestLevel, TArrayView<FWor
 				{
 					check(Actor->IsPackageExternal());
 					InDestLevel->Actors.Add(Actor);
-					check(Actor->GetLevel() == InDestLevel);
+					checkf(Actor->GetLevel() == InDestLevel, TEXT("Levels mismatch, got : %s, expected: %s\nActor: %s\nActorFullName: %s\nActorPackage: %s"), *InDestLevel->GetFullName(), *Actor->GetLevel()->GetFullName(), *Actor->GetActorNameOrLabel(), *Actor->GetFullName(), *Actor->GetPackage()->GetFullName());
 				}
 								
 				UE_LOG(LogEngine, Verbose, TEXT(" ==> Loaded %s (remaining: %d)"), *Actor->GetFullName(), LoadProgress->NumPendingLoadRequests);
