Fix for Non-Unity Test

diff --git a/Engine/Source/Runtime/Engine/Public/Engine/ScopedMovementUpdate.h b/Engine/Source/Runtime/Engine/Public/Engine/ScopedMovementUpdate.h
--- a/Engine/Source/Runtime/Engine/Public/Engine/ScopedMovementUpdate.h
+++ b/Engine/Source/Runtime/Engine/Public/Engine/ScopedMovementUpdate.h
@@ -5,6 +5,7 @@
 #include "CoreMinimal.h"
 #include "Engine/HitResult.h"
 #include "Engine/OverlapInfo.h"
+#include "Engine/EngineTypes.h"
 
 /**
  * Enum that controls the scoping behavior of FScopedMovementUpdate.
