Invalid data fix-up for FRichCurveKey::TangentWeightMode
#fix added fix-up behaviour during load - and ensure to track if any more invalid data creeps up afterwards

diff --git a/Engine/Source/Runtime/Core/Public/UObject/UE5MainStreamObjectVersion.h b/Engine/Source/Runtime/Core/Public/UObject/UE5MainStreamObjectVersion.h
--- a/Engine/Source/Runtime/Core/Public/UObject/UE5MainStreamObjectVersion.h
+++ b/Engine/Source/Runtime/Core/Public/UObject/UE5MainStreamObjectVersion.h
@@ -203,6 +203,9 @@ struct CORE_API FUE5MainStreamObjectVersion
 		// Switching FK control naming scheme to incorporate _CURVE for curve controls
 		FKControlNamingScheme,
 
+		// Fix-up for FRichCurveKey::TangentWeightMode, which were found to contain invalid value w.r.t the enum-type
+		RichCurveKeyInvalidTangentMode,
+
 		// -----<new versions can be added above this line>-------------------------------------------------
 		VersionPlusOne,
 		LatestVersion = VersionPlusOne - 1
diff --git a/Engine/Source/Runtime/Engine/Private/Curves/RichCurve.cpp b/Engine/Source/Runtime/Engine/Private/Curves/RichCurve.cpp
--- a/Engine/Source/Runtime/Engine/Private/Curves/RichCurve.cpp
+++ b/Engine/Source/Runtime/Engine/Private/Curves/RichCurve.cpp
@@ -3,6 +3,7 @@
 #include "Curves/RichCurve.h"
 #include "Templates/Function.h"
 #include "Containers/ArrayView.h"
+#include "UObject/UE5MainStreamObjectVersion.h"
 
 // Broken - do not turn on! 
 #define MIXEDKEY_STRIPS_TANGENTS 0
@@ -165,6 +166,8 @@ bool FRichCurveKey::Serialize(FArchive& Ar)
 		return false;
 	}
 
+	Ar.UsingCustomVersion(FUE5MainStreamObjectVersion::GUID);
+
 	// Serialization is handled manually to avoid the extra size overhead of FProperty tagging.
 	// Otherwise with many keys in a rich curve the size can become quite large.
 	Ar << InterpMode;
@@ -177,10 +180,44 @@ bool FRichCurveKey::Serialize(FArchive& Ar)
 	Ar << LeaveTangent;
 	Ar << LeaveTangentWeight;
 
+	if (Ar.IsLoading())
+	{
+		if (TangentWeightMode != RCTWM_WeightedNone &&
+			TangentWeightMode != RCTWM_WeightedArrive &&
+			TangentWeightMode != RCTWM_WeightedLeave &&
+			TangentWeightMode != RCTWM_WeightedBoth)
+		{
+			if(Ar.CustomVer(FUE5MainStreamObjectVersion::GUID) < FUE5MainStreamObjectVersion::RichCurveKeyInvalidTangentMode)
+			{
+				TangentWeightMode = RCTWM_WeightedNone;	
+
+				// No valid weight mode - find correct one one
+				const bool bHasArriveWeight = !FMath::IsNearlyZero(ArriveTangentWeight);
+				const bool bHasLeaveWeight = !FMath::IsNearlyZero(LeaveTangentWeight);
+
+				if (bHasArriveWeight && bHasLeaveWeight)
+				{
+					TangentWeightMode = RCTWM_WeightedBoth;				
+				}
+				else if (bHasArriveWeight)
+				{
+					TangentWeightMode = RCTWM_WeightedArrive;
+				}
+				else if (bHasLeaveWeight)
+				{
+					TangentWeightMode = RCTWM_WeightedLeave;
+				}
+			}
+			else
+			{
+				ensureMsgf(false, TEXT("FRichCurveKey found with invalid TangentWeightMode value of %i"), static_cast<int32>(TangentWeightMode));
+			}			
+		}
+	}
+
 	return true;
 }
 
-
 bool FRichCurveKey::operator==( const FRichCurveKey& Curve ) const
 {
 	return (Time == Curve.Time) && (Value == Curve.Value) && (InterpMode == Curve.InterpMode) &&
