Fix FActorDescContainerInstance crash (FGCObject doesn't support being trivially relocated when resizing a container)

#rnx

diff --git a/Engine/Source/Runtime/Engine/Private/WorldPartition/LevelInstance/LevelInstanceActorDesc.cpp b/Engine/Source/Runtime/Engine/Private/WorldPartition/LevelInstance/LevelInstanceActorDesc.cpp
--- a/Engine/Source/Runtime/Engine/Private/WorldPartition/LevelInstance/LevelInstanceActorDesc.cpp
+++ b/Engine/Source/Runtime/Engine/Private/WorldPartition/LevelInstance/LevelInstanceActorDesc.cpp
@@ -39,7 +39,7 @@ struct FActorDescContainerInstance : public FGCObject
 	uint32 RefCount;
 };
 
-static TMap<FName, FActorDescContainerInstance> ActorDescContainers;
+static TMap<FName, TUniquePtr<FActorDescContainerInstance>> ActorDescContainers;
 
 FLevelInstanceActorDesc::FLevelInstanceActorDesc()
 	: DesiredRuntimeBehavior(ELevelInstanceRuntimeBehavior::Partitioned)
@@ -199,17 +199,23 @@ void FLevelInstanceActorDesc::Serialize(FArchive& Ar)
 
 UActorDescContainer* FLevelInstanceActorDesc::RegisterActorDescContainer(FName PackageName, UWorld* InWorld)
 {
-	FActorDescContainerInstance& ExistingContainerInstance = ActorDescContainers.FindOrAdd(PackageName);
-	UActorDescContainer* ActorDescContainer = ExistingContainerInstance.Container;
-	
-	if (ExistingContainerInstance.RefCount++ == 0)
+	TUniquePtr<FActorDescContainerInstance>& ExistingContainerInstance = ActorDescContainers.FindOrAdd(PackageName);
+	UActorDescContainer* ActorDescContainer = nullptr;
+	if (!ExistingContainerInstance)
 	{
+		ExistingContainerInstance.Reset(new FActorDescContainerInstance());
 		ActorDescContainer = NewObject<UActorDescContainer>(GetTransientPackage());
-		ExistingContainerInstance.Container = ActorDescContainer;
+		ExistingContainerInstance->Container = ActorDescContainer;
+		ExistingContainerInstance->RefCount++;
 		
 		// This will potentially invalidate ExistingContainerInstance due to ActorDescContainers reallocation
 		ActorDescContainer->Initialize(InWorld, PackageName);
 	}
+	else
+	{
+		ExistingContainerInstance->RefCount++;
+		ActorDescContainer = ExistingContainerInstance->Container;
+	}
 
 	check(ActorDescContainer->GetWorld() == InWorld);
 	return ActorDescContainer;
@@ -218,11 +224,11 @@ UActorDescContainer* FLevelInstanceActorDesc::RegisterActorDescContainer(FName P
 void FLevelInstanceActorDesc::UnregisterActorDescContainer(UActorDescContainer* Container)
 {
 	FName PackageName = Container->GetContainerPackage();
-	FActorDescContainerInstance& ExistingContainerInstance = ActorDescContainers.FindChecked(PackageName);
+	TUniquePtr<FActorDescContainerInstance>& ExistingContainerInstance = ActorDescContainers.FindChecked(PackageName);
 
-	if (ExistingContainerInstance.RefCount-- == 1)
+	if (ExistingContainerInstance->RefCount-- == 1)
 	{
-		ExistingContainerInstance.Container->Uninitialize();
+		ExistingContainerInstance->Container->Uninitialize();
 		ActorDescContainers.FindAndRemoveChecked(PackageName);
 	}
 }
