[UE][FEATURE] Create straight capsule geometry for tapered capsules in ChaosInterface::CreateGeometry
- Currently tapered capsules defined in physics assets are ignored. With this change a normal capsule with a radius equal to the mean of the tapered capsule's radi will be created at runtime.

diff --git a/Engine/Source/Runtime/Engine/Private/PhysicsEngine/Experimental/ChaosInterfaceUtils.cpp b/Engine/Source/Runtime/Engine/Private/PhysicsEngine/Experimental/ChaosInterfaceUtils.cpp
--- a/Engine/Source/Runtime/Engine/Private/PhysicsEngine/Experimental/ChaosInterfaceUtils.cpp
+++ b/Engine/Source/Runtime/Engine/Private/PhysicsEngine/Experimental/ChaosInterfaceUtils.cpp
@@ -30,6 +30,7 @@
 #endif
 
 #define FORCE_ANALYTICS 0
+#define CREATE_STRAIGHT_CAPSULE_GEOMETRY_FOR_TAPERED_CAPSULES
 
 namespace ChaosInterface
 {
@@ -276,7 +277,6 @@ namespace ChaosInterface
 					TUniquePtr<Chaos::FPerShapeData> NewShape = NewShapeHelper(MakeSerializable(ImplicitSphere),Shapes.Num(), (void*)UnscaledSphyl.GetUserData(), UnscaledSphyl.GetCollisionEnabled());
 					Shapes.Emplace(MoveTemp(NewShape));
 					Geoms.Add(MoveTemp(ImplicitSphere));
-
 				}
 				else
 				{
@@ -288,7 +288,37 @@ namespace ChaosInterface
 					Geoms.Add(MoveTemp(ImplicitCapsule));
 				}
 			}
-#if 0
+#ifdef CREATE_STRAIGHT_CAPSULE_GEOMETRY_FOR_TAPERED_CAPSULES
+			for (uint32 i = 0; i < static_cast<uint32>(InParams.Geometry->TaperedCapsuleElems.Num()); ++i)
+			{
+				const FKTaperedCapsuleElem& UnscaledTaperedCapsule = InParams.Geometry->TaperedCapsuleElems[i];
+
+				const FKTaperedCapsuleElem ScaledTaperedCapsule = UnscaledTaperedCapsule.GetFinalScaled(Scale, InParams.LocalTransform);
+				Chaos::FReal HalfHeight = FMath::Max(ScaledTaperedCapsule.Length * 0.5f, KINDA_SMALL_NUMBER);
+				const Chaos::FReal Radius0 = FMath::Max(ScaledTaperedCapsule.Radius0, KINDA_SMALL_NUMBER);
+				const Chaos::FReal Radius1 = FMath::Max(ScaledTaperedCapsule.Radius1, KINDA_SMALL_NUMBER);
+
+				if (HalfHeight < KINDA_SMALL_NUMBER)
+				{
+					//not a capsule just use a sphere
+					const Chaos::FReal MaxRadius = FMath::Max(Radius0, Radius1);
+					auto ImplicitSphere = MakeUnique<Chaos::TSphere<Chaos::FReal, 3>>(ScaledTaperedCapsule.Center, MaxRadius);
+					TUniquePtr<Chaos::FPerShapeData> NewShape = NewShapeHelper(MakeSerializable(ImplicitSphere), Shapes.Num(), (void*)UnscaledTaperedCapsule.GetUserData(), UnscaledTaperedCapsule.GetCollisionEnabled());
+					Shapes.Emplace(MoveTemp(NewShape));
+					Geoms.Add(MoveTemp(ImplicitSphere));
+				}
+				else
+				{
+					Chaos::FVec3 HalfExtents = ScaledTaperedCapsule.Rotation.RotateVector(Chaos::FVec3(0, 0, HalfHeight));
+					const Chaos::FReal MeanRadius = 0.5f * (Radius0 + Radius1);
+
+					auto ImplicitCapsule = MakeUnique<Chaos::FCapsule>(ScaledTaperedCapsule.Center - HalfExtents, ScaledTaperedCapsule.Center + HalfExtents, MeanRadius);
+					TUniquePtr<Chaos::FPerShapeData> NewShape = NewShapeHelper(MakeSerializable(ImplicitCapsule), Shapes.Num(), (void*)UnscaledTaperedCapsule.GetUserData(), UnscaledTaperedCapsule.GetCollisionEnabled());
+					Shapes.Emplace(MoveTemp(NewShape));
+					Geoms.Add(MoveTemp(ImplicitCapsule));
+				}
+			}
+#elif 0
 			for (uint32 i = 0; i < static_cast<uint32>(InParams.Geometry->TaperedCapsuleElems.Num()); ++i)
 			{
 				ensure(FMath::IsNearlyEqual(Scale[0], Scale[1]) && FMath::IsNearlyEqual(Scale[1], Scale[2]));
