Fixed ILevelInstanceInterface::GetLoadedLevel always returning null.

diff --git a/Engine/Source/Runtime/Engine/Private/LevelInstance/LevelInstanceInterface.cpp b/Engine/Source/Runtime/Engine/Private/LevelInstance/LevelInstanceInterface.cpp
--- a/Engine/Source/Runtime/Engine/Private/LevelInstance/LevelInstanceInterface.cpp
+++ b/Engine/Source/Runtime/Engine/Private/LevelInstance/LevelInstanceInterface.cpp
@@ -64,7 +64,7 @@ ULevel* ILevelInstanceInterface::GetLoadedLevel() const
 {
 	if (ULevelInstanceSubsystem* LevelInstanceSubsystem = GetLevelInstanceSubsystem())
 	{
-		LevelInstanceSubsystem->GetLevelInstanceLevel(this);
+		return LevelInstanceSubsystem->GetLevelInstanceLevel(this);
 	}
 
 	return nullptr;
