FWorldPartitionActorDesc::Load now uses LoadPackageAsync to load actor packages and flushes right away.
This is a workaround needed for level instance streaming of partitioned worlds to work. The proper fix will be done by Core.

diff --git a/Engine/Source/Runtime/Engine/Private/WorldPartition/WorldPartitionActorDesc.cpp b/Engine/Source/Runtime/Engine/Private/WorldPartition/WorldPartitionActorDesc.cpp
--- a/Engine/Source/Runtime/Engine/Private/WorldPartition/WorldPartitionActorDesc.cpp
+++ b/Engine/Source/Runtime/Engine/Private/WorldPartition/WorldPartitionActorDesc.cpp
@@ -377,17 +377,19 @@ AActor* FWorldPartitionActorDesc::Load() const
 			Container->GetInstancingContext(InstancingContext, SoftObjectPathFixupArchive);
 		}
 
-		UPackage* Package = nullptr;
-
+		const FName PackageNameToLoad = *ActorPackage.ToString();
+		FName PackageNameToCreate = PackageNameToLoad;
 		if (InstancingContext)
 		{
-			FName RemappedPackageName = InstancingContext->Remap(ActorPackage);
-			check(RemappedPackageName != ActorPath);
-
-			Package = CreatePackage(*RemappedPackageName.ToString());
+			PackageNameToCreate = InstancingContext->Remap(ActorPackage);
+			check(PackageNameToCreate != ActorPath);
 		}
 
-		Package = LoadPackage(Package, *ActorPackage.ToString(), LOAD_None, nullptr, InstancingContext);
+		// Here we are not using LoadPackage because we could be inside an async load and it's currently not properly handled.
+		// Instead, use the async version and flush.
+		LoadPackageAsync(FPackagePath::FromPackageNameChecked(PackageNameToLoad), PackageNameToCreate, FLoadPackageAsyncDelegate(), PKG_None, INDEX_NONE, 0, InstancingContext);
+		FlushAsyncLoading();
+		UPackage* Package = FindPackage(nullptr, *PackageNameToCreate.ToString());
 
 		if (Package)
 		{
