World Partition: don't apply instance transform if the actor is already in the world.

#rnx

diff --git a/Engine/Source/Runtime/Engine/Classes/Engine/Level.h b/Engine/Source/Runtime/Engine/Classes/Engine/Level.h
--- a/Engine/Source/Runtime/Engine/Classes/Engine/Level.h
+++ b/Engine/Source/Runtime/Engine/Classes/Engine/Level.h
@@ -984,13 +984,19 @@ public:
 #if WITH_EDITOR
 	/**
 	 * Add a dynamically loaded actor to this level, as if it was part of the original map load process.
+	 *
+	 * @param Actor				The actor to add to the level
+	 * @param TransformToApply	The transform to apply to this actor if it's not already in the level
 	 */
-	ENGINE_API void AddLoadedActor(AActor* Actor);
+	ENGINE_API void AddLoadedActor(AActor* Actor, const FTransform* TransformToApply = nullptr);
 
 	/**
 	 * Remove a dynamically loaded actor from this level.
+	 *
+	 * @param Actor				The actor to remove from the level
+	 * @param TransformToRemove	The transform that was applied to this actor when it was added to the level
 	 */
-	ENGINE_API void RemoveLoadedActor(AActor* Actor);
+	ENGINE_API void RemoveLoadedActor(AActor* Actor, const FTransform* TransformToRemove = nullptr);
 
 	/** Called when dynamically loaded actor is added to this level */
 	DECLARE_EVENT_OneParam(ULevel, FLoadedActorAddedToLevelEvent, AActor&);
diff --git a/Engine/Source/Runtime/Engine/Private/Level.cpp b/Engine/Source/Runtime/Engine/Private/Level.cpp
--- a/Engine/Source/Runtime/Engine/Private/Level.cpp
+++ b/Engine/Source/Runtime/Engine/Private/Level.cpp
@@ -752,7 +752,7 @@ bool ULevel::IsNetActor(const AActor* Actor)
 }
 
 #if WITH_EDITOR
-void ULevel::AddLoadedActor(AActor* Actor)
+void ULevel::AddLoadedActor(AActor* Actor, const FTransform* TransformToApply)
 {
 	check(Actor->GetLevel() == this);
 	check(IsValidChecked(Actor));
@@ -763,6 +763,14 @@ void ULevel::AddLoadedActor(AActor* Actor)
 		Actors.Add(Actor);
 		ActorsForGC.Add(Actor);
 
+		if (TransformToApply)
+		{
+			FLevelUtils::FApplyLevelTransformParams TransformParams(this, *TransformToApply);
+			TransformParams.Actor = Actor;
+			TransformParams.bDoPostEditMove = true;
+			FLevelUtils::ApplyLevelTransform(TransformParams);
+		}
+
 		// If components from actors belonging to this level were already registered, do the same for the newly loaded actor.
 		// Otherwise, the new actor components will be registered later once the world initialization is completed, through UpdateWorldComponents()
 		if (bAreComponentsCurrentlyRegistered)
@@ -782,7 +790,7 @@ void ULevel::AddLoadedActor(AActor* Actor)
 	}
 }
 
-void ULevel::RemoveLoadedActor(AActor* Actor)
+void ULevel::RemoveLoadedActor(AActor* Actor, const FTransform* TransformToRemove)
 {
 	check(Actor);
 	check(Actor->GetLevel() == this);
@@ -797,6 +805,14 @@ void ULevel::RemoveLoadedActor(AActor* Actor)
 	Actors[ActorIndex] = nullptr;
 	ActorsForGC.Remove(Actor);
 
+	if (TransformToRemove)
+	{
+		FLevelUtils::FApplyLevelTransformParams TransformParams(this, TransformToRemove->Inverse());
+		TransformParams.Actor = Actor;
+		TransformParams.bDoPostEditMove = true;
+		FLevelUtils::ApplyLevelTransform(TransformParams);
+	}
+
 	OnLoadedActorRemovedFromLevelEvent.Broadcast(*Actor);
 }
 #endif
diff --git a/Engine/Source/Runtime/Engine/Private/WorldPartition/WorldPartition.cpp b/Engine/Source/Runtime/Engine/Private/WorldPartition/WorldPartition.cpp
--- a/Engine/Source/Runtime/Engine/Private/WorldPartition/WorldPartition.cpp
+++ b/Engine/Source/Runtime/Engine/Private/WorldPartition/WorldPartition.cpp
@@ -1112,23 +1112,11 @@ void UWorldPartition::OnActorDescUpdated(FWorldPartitionActorDesc* ActorDesc)
 	}
 }
 
-void UWorldPartition::ApplyActorTransform(AActor* Actor, const FTransform& InTransform)
-{
-	if (!InTransform.Equals(FTransform::Identity))
-	{
-		FLevelUtils::FApplyLevelTransformParams TransformParams(Actor->GetLevel(), InTransform);
-		TransformParams.Actor = Actor;
-		TransformParams.bDoPostEditMove = true;
-		FLevelUtils::ApplyLevelTransform(TransformParams);
-	}
-}
-
 void UWorldPartition::OnActorDescRegistered(const FWorldPartitionActorDesc& ActorDesc) 
 {
 	AActor* Actor = ActorDesc.GetActor();
 	check(Actor);
-	ApplyActorTransform(Actor, InstanceTransform);
-	Actor->GetLevel()->AddLoadedActor(Actor);
+	Actor->GetLevel()->AddLoadedActor(Actor, !InstanceTransform.Equals(FTransform::Identity) ? &InstanceTransform : nullptr);
 }
 
 void UWorldPartition::OnActorDescUnregistered(const FWorldPartitionActorDesc& ActorDesc) 
@@ -1136,8 +1124,7 @@ void UWorldPartition::OnActorDescUnregistered(const FWorldPartitionActorDesc& Ac
 	AActor* Actor = ActorDesc.GetActor();
 	if (IsValidChecked(Actor))
 	{
-		Actor->GetLevel()->RemoveLoadedActor(Actor);
-		ApplyActorTransform(Actor, InstanceTransform.Inverse());
+		Actor->GetLevel()->RemoveLoadedActor(Actor, !InstanceTransform.Equals(FTransform::Identity) ? &InstanceTransform : nullptr);
 	}
 }
 
diff --git a/Engine/Source/Runtime/Engine/Public/WorldPartition/WorldPartition.h b/Engine/Source/Runtime/Engine/Public/WorldPartition/WorldPartition.h
--- a/Engine/Source/Runtime/Engine/Public/WorldPartition/WorldPartition.h
+++ b/Engine/Source/Runtime/Engine/Public/WorldPartition/WorldPartition.h
@@ -279,8 +279,6 @@ private:
 	bool ShouldActorBeLoadedByEditorCells(const FWorldPartitionActorDesc* ActorDesc) const;
 	bool UpdateEditorCells(TFunctionRef<bool(TArray<UWorldPartitionEditorCell*>&)> GetCellsToProcess, bool bIsCellShouldBeLoaded, bool bIsFromUserChange);
 
-	void ApplyActorTransform(AActor* Actor, const FTransform& InTransform);
-
 	TMap<FGuid, FWorldPartitionReference> PinnedActors;
 	TMap<FGuid, TMap<FGuid, FWorldPartitionReference>> PinnedActorRefs;
 #endif
