Make current saved move being replayed on the client accessible to any functions that might need it

diff --git a/Engine/Source/Runtime/Engine/Classes/GameFramework/CharacterMovementComponent.h b/Engine/Source/Runtime/Engine/Classes/GameFramework/CharacterMovementComponent.h
--- a/Engine/Source/Runtime/Engine/Classes/GameFramework/CharacterMovementComponent.h
+++ b/Engine/Source/Runtime/Engine/Classes/GameFramework/CharacterMovementComponent.h
@@ -2515,8 +2515,25 @@ protected:
 	 */
 	FCharacterMoveResponseDataContainer& GetMoveResponseDataContainer() const { return *MoveResponseDataContainerPtr; }
 
+	/**
+	 * Used internally to save the SavedMove currently being replayed on the client so it is accessible to any functions that might need it.
+	 * @see: ClientUpdatePositionAfterServerUpdate 
+	 */
+	void SetCurrentReplayedSavedMove(FSavedMove_Character* SavedMove) { CurrentReplayedSavedMove = SavedMove; }
+	
+public:
+
+	/**
+	 * Gets the SavedMove being replayed on the client after a correction is received.
+	 * @see: ClientUpdatePositionAfterServerUpdate
+	 */
+	const FSavedMove_Character* GetCurrentReplayedSavedMove() const { return CurrentReplayedSavedMove; }
+
 private:
 
+	/** Current SavedMove being replayed on the client after a correction is received */
+	FSavedMove_Character* CurrentReplayedSavedMove = nullptr;
+
 	//////////////////////////////////////////////////////////////////////////
 	// Server move data
 
diff --git a/Engine/Source/Runtime/Engine/Private/Components/CharacterMovementComponent.cpp b/Engine/Source/Runtime/Engine/Private/Components/CharacterMovementComponent.cpp
--- a/Engine/Source/Runtime/Engine/Private/Components/CharacterMovementComponent.cpp
+++ b/Engine/Source/Runtime/Engine/Private/Components/CharacterMovementComponent.cpp
@@ -8091,6 +8091,10 @@ bool UCharacterMovementComponent::ClientUpdatePositionAfterServerUpdate()
 	{
 		FSavedMove_Character* const CurrentMove = ClientData->SavedMoves[i].Get();
 		checkSlow(CurrentMove != nullptr);
+
+		// Make current SavedMove accessible to any functions that might need it.
+		SetCurrentReplayedSavedMove(CurrentMove);
+
 		CurrentMove->PrepMoveFor(CharacterOwner);
 
 		if (ShouldUsePackedMovementRPCs())
@@ -8107,6 +8111,7 @@ bool UCharacterMovementComponent::ClientUpdatePositionAfterServerUpdate()
 
 		CurrentMove->PostUpdate(CharacterOwner, FSavedMove_Character::PostUpdate_Replay);
 		SetCurrentNetworkMoveData(nullptr);
+		SetCurrentReplayedSavedMove(nullptr);
 	}
 	const bool bPostReplayPressedJump = CharacterOwner->bPressedJump;
 
