AWorldDataLayers::Create no longer sets the UWorld::WorldDataLayers.
Pass an optional name to AWorldDataLayer::Create
Remove UDataLayerEditorSubsystem::GetWorldDataLayers(bool). UWorld::WorldDataLayers are always created now.
Preparation for multiple AWorldDataLayers within a single world

diff --git a/Engine/Source/Editor/DataLayerEditor/Private/DataLayer/DataLayerEditorSubsystem.cpp b/Engine/Source/Editor/DataLayerEditor/Private/DataLayer/DataLayerEditorSubsystem.cpp
--- a/Engine/Source/Editor/DataLayerEditor/Private/DataLayer/DataLayerEditorSubsystem.cpp
+++ b/Engine/Source/Editor/DataLayerEditor/Private/DataLayer/DataLayerEditorSubsystem.cpp
@@ -1007,23 +1007,11 @@ UDataLayerInstance* UDataLayerEditorSubsystem::GetDataLayerInstance(const UDataL
 	return WorldDataLayers ? const_cast<UDataLayerInstance*>(WorldDataLayers->GetDataLayerInstance(DataLayerAsset)) : nullptr;
 }
 
-const AWorldDataLayers* UDataLayerEditorSubsystem::GetWorldDataLayers() const
+AWorldDataLayers* UDataLayerEditorSubsystem::GetWorldDataLayers() const
 {
 	return GetWorld()->GetWorldDataLayers();
 }
 
-AWorldDataLayers* UDataLayerEditorSubsystem::GetWorldDataLayers(bool bCreateIfNotFound)
-{
-	AWorldDataLayers* WorldDataLayers = GetWorld()->GetWorldDataLayers();
-
-	if (!WorldDataLayers && bCreateIfNotFound)
-	{
-		WorldDataLayers = AWorldDataLayers::Create(GetWorld());
-	}
-
-	return WorldDataLayers;
-}
-
 void UDataLayerEditorSubsystem::AddAllDataLayersTo(TArray<TWeakObjectPtr<UDataLayerInstance>>& OutDataLayers) const
 {
 	if (const AWorldDataLayers* WorldDataLayers = GetWorldDataLayers())
@@ -1040,8 +1028,9 @@ UDataLayerInstance* UDataLayerEditorSubsystem::CreateDataLayerInstance(const FDa
 {
 	UDataLayerInstance* NewDataLayer = nullptr;
 
-	AWorldDataLayers* WorldDataLayers = Parameters.WorlDataLayers != nullptr ? Parameters.WorlDataLayers.Get() : GetWorld()->GetWorldDataLayers();
-	if (!HasDeprecatedDataLayers())
+
+	AWorldDataLayers* WorldDataLayers = GetWorldDataLayers();
+	if (!WorldDataLayers->HasDeprecatedDataLayers())
 	{
 		NewDataLayer = WorldDataLayers->CreateDataLayer<UDataLayerInstanceWithAsset>(Parameters.DataLayerAsset);
 	}
diff --git a/Engine/Source/Editor/DataLayerEditor/Public/DataLayer/DataLayerEditorSubsystem.h b/Engine/Source/Editor/DataLayerEditor/Public/DataLayer/DataLayerEditorSubsystem.h
--- a/Engine/Source/Editor/DataLayerEditor/Public/DataLayer/DataLayerEditorSubsystem.h
+++ b/Engine/Source/Editor/DataLayerEditor/Public/DataLayer/DataLayerEditorSubsystem.h
@@ -632,17 +632,9 @@ private:
 	/**
 	 * Gets the WorldDataLayers for this world
 	 *
-	 * @param	bCreateIfNotFound	If true, created the AWorldDataLayers object if not found.
 	 * @return	The AWorldDataLayers Object
 	 */
-	class AWorldDataLayers* GetWorldDataLayers(bool bCreateIfNotFound = false);
-
-	/**
-	 * Gets the WorldDataLayers for this world
-	 *
-	 * @return	The AWorldDataLayers Object
-	 */
-	const class AWorldDataLayers* GetWorldDataLayers() const;
+	class AWorldDataLayers* GetWorldDataLayers() const;
 
 	/**
 	 * Get the current UWorld object.
diff --git a/Engine/Source/Runtime/Engine/Private/WorldPartition/DataLayer/WorldDataLayers.cpp b/Engine/Source/Runtime/Engine/Private/WorldPartition/DataLayer/WorldDataLayers.cpp
--- a/Engine/Source/Runtime/Engine/Private/WorldPartition/DataLayer/WorldDataLayers.cpp
+++ b/Engine/Source/Runtime/Engine/Private/WorldPartition/DataLayer/WorldDataLayers.cpp
@@ -451,14 +451,14 @@ void AWorldDataLayers::GetUserLoadedInEditorStates(TArray<FName>& OutDataLayersL
 	});
 }
 
-AWorldDataLayers* AWorldDataLayers::Create(UWorld* World)
+
+AWorldDataLayers* AWorldDataLayers::Create(UWorld* World, FName InWorldDataLayerName)
 {
 	check(World);
-	check(!World->GetWorldDataLayers());
-
+	
 	AWorldDataLayers* WorldDataLayers = nullptr;
-
-	static FName WorldDataLayersName = AWorldDataLayers::StaticClass()->GetFName();
+	
+	FName WorldDataLayersName = InWorldDataLayerName.IsNone() ? AWorldDataLayers::StaticClass()->GetFName() : InWorldDataLayerName;
 	if (UObject* ExistingObject = StaticFindObject(nullptr, World->PersistentLevel, *WorldDataLayersName.ToString()))
 	{
 		WorldDataLayers = CastChecked<AWorldDataLayers>(ExistingObject);
@@ -469,7 +469,7 @@ AWorldDataLayers* AWorldDataLayers::Create(UWorld* World)
 			WorldDataLayers = nullptr;
 		}
 	}
-
+	
 	if (!WorldDataLayers)
 	{
 		FActorSpawnParameters SpawnParams;
@@ -478,12 +478,13 @@ AWorldDataLayers* AWorldDataLayers::Create(UWorld* World)
 		SpawnParams.NameMode = FActorSpawnParameters::ESpawnActorNameMode::Required_Fatal;
 		WorldDataLayers = World->SpawnActor<AWorldDataLayers>(AWorldDataLayers::StaticClass(), SpawnParams);
 	}
-
+	else
+	{
+		UE_LOG(LogWorldPartition, Error, TEXT("Failed to create WorldDataLayers Actor. There is already a WorldDataLayer Actor named \"%s\" "), *WorldDataLayersName.ToString())
+	}
+	
 	check(WorldDataLayers);
 
-	World->Modify();
-	World->SetWorldDataLayers(WorldDataLayers);
-
 	return WorldDataLayers;
 }
 
diff --git a/Engine/Source/Runtime/Engine/Private/WorldPartition/WorldPartition.cpp b/Engine/Source/Runtime/Engine/Private/WorldPartition/WorldPartition.cpp
--- a/Engine/Source/Runtime/Engine/Private/WorldPartition/WorldPartition.cpp
+++ b/Engine/Source/Runtime/Engine/Private/WorldPartition/WorldPartition.cpp
@@ -738,6 +738,7 @@ UWorldPartition* UWorldPartition::CreateOrRepairWorldPartition(AWorldSettings* W
 		if (!WorldDataLayers)
 		{
 			WorldDataLayers = AWorldDataLayers::Create(World);
+			World->SetWorldDataLayers(WorldDataLayers);
 		}
 
 		FWorldPartitionMiniMapHelper::GetWorldPartitionMiniMap(World, true);
diff --git a/Engine/Source/Runtime/Engine/Public/WorldPartition/DataLayer/WorldDataLayers.h b/Engine/Source/Runtime/Engine/Public/WorldPartition/DataLayer/WorldDataLayers.h
--- a/Engine/Source/Runtime/Engine/Public/WorldPartition/DataLayer/WorldDataLayers.h
+++ b/Engine/Source/Runtime/Engine/Public/WorldPartition/DataLayer/WorldDataLayers.h
@@ -52,7 +52,7 @@ public:
 	virtual bool IsUserManaged() const override { return false; }
 	virtual TUniquePtr<class FWorldPartitionActorDesc> CreateClassActorDesc() const override;
 
-	static AWorldDataLayers* Create(UWorld* World);
+	static AWorldDataLayers* Create(UWorld* World, FName InWorldDataLayerName = NAME_None);
 
 	bool IsEmpty() const { return DataLayerInstances.IsEmpty(); }
 	bool HasDeprecatedDataLayers() const { return bHasDeprecatedDataLayers; }
