World Partition: avoid an unnecessary call to FindPackage when loading an external actor.

#rnx

diff --git a/Engine/Source/Runtime/Engine/Private/WorldPartition/WorldPartitionActorDesc.cpp b/Engine/Source/Runtime/Engine/Private/WorldPartition/WorldPartitionActorDesc.cpp
--- a/Engine/Source/Runtime/Engine/Private/WorldPartition/WorldPartitionActorDesc.cpp
+++ b/Engine/Source/Runtime/Engine/Private/WorldPartition/WorldPartitionActorDesc.cpp
@@ -387,9 +387,16 @@ AActor* FWorldPartitionActorDesc::Load() const
 
 		// Here we are not using LoadPackage because we could be inside an async load and it's currently not properly handled.
 		// Instead, use the async version and flush.
-		LoadPackageAsync(FPackagePath::FromPackageNameChecked(PackageNameToLoad), PackageNameToCreate, FLoadPackageAsyncDelegate(), PKG_None, INDEX_NONE, 0, InstancingContext);
+		UPackage* Package = nullptr;
+		LoadPackageAsync(FPackagePath::FromPackageNameChecked(PackageNameToLoad), PackageNameToCreate, FLoadPackageAsyncDelegate::CreateLambda([&Package](const FName&, UPackage* LoadedPackage, EAsyncLoadingResult::Type Result)
+		{
+			if (Result == EAsyncLoadingResult::Succeeded)
+			{
+				Package = LoadedPackage;
+			}
+		}), PKG_None, INDEX_NONE, 0, InstancingContext);
+
 		FlushAsyncLoading();
-		UPackage* Package = FindPackage(nullptr, *PackageNameToCreate.ToString());
 
 		if (Package)
 		{
