World Partition: make the instance transform optional to avoid checking for identity transform when loading/unloading actors.

#rnx

diff --git a/Engine/Source/Runtime/Engine/Private/WorldPartition/WorldPartition.cpp b/Engine/Source/Runtime/Engine/Private/WorldPartition/WorldPartition.cpp
--- a/Engine/Source/Runtime/Engine/Private/WorldPartition/WorldPartition.cpp
+++ b/Engine/Source/Runtime/Engine/Private/WorldPartition/WorldPartition.cpp
@@ -266,7 +266,6 @@ UWorldPartition::UWorldPartition(const FObjectInitializer& ObjectInitializer)
 	, bIsPIE(false)
 #endif
 	, InitState(EWorldPartitionInitState::Uninitialized)
-	, InstanceTransform(FTransform::Identity)
 	, StreamingPolicy(nullptr)
 #if !UE_BUILD_SHIPPING
 	, Replay(nullptr)
@@ -366,7 +365,7 @@ void UWorldPartition::Initialize(UWorld* InWorld, const FTransform& InTransform)
 {
 	UE_SCOPED_TIMER(TEXT("WorldPartition initialize"), LogWorldPartition, Display);
 	TRACE_CPUPROFILER_EVENT_SCOPE(UWorldPartition::Initialize);
-	
+
 	check(!World || (World == InWorld));
 	if (!ensure(!IsInitialized()))
 	{
@@ -380,7 +379,11 @@ void UWorldPartition::Initialize(UWorld* InWorld, const FTransform& InTransform)
 
 	check(InWorld);
 	World = InWorld;
-	InstanceTransform = InTransform;
+
+	if (!InTransform.Equals(FTransform::Identity))
+	{
+		InstanceTransform = InTransform;
+	}
 
 	check(InitState == EWorldPartitionInitState::Uninitialized);
 	InitState = EWorldPartitionInitState::Initializing;
@@ -463,7 +466,7 @@ void UWorldPartition::Initialize(UWorld* InWorld, const FTransform& InTransform)
 
 					InstancingContext.AddMapping(*LongActorPackageName, *InstancedName);
 
-					ActorDescIterator->TransformInstance(SourceWorldPath, RemappedWorldPath, InstanceTransform);
+					ActorDescIterator->TransformInstance(SourceWorldPath, RemappedWorldPath, GetInstanceTransform());
 				}
 
 				ActorDescIterator->bIsForcedNonSpatiallyLoaded = !IsStreamingEnabled();
@@ -615,6 +618,20 @@ bool UWorldPartition::IsMainWorldPartition() const
 	return World == GetTypedOuter<UWorld>();
 }
 
+const FTransform& UWorldPartition::GetInstanceTransform() const
+{
+	if (const FTransform* TransformPtr = GetInstanceTransformPtr())
+	{
+		return *TransformPtr;
+	}
+	return FTransform::Identity;
+}
+
+const FTransform* UWorldPartition::GetInstanceTransformPtr() const
+{
+	return InstanceTransform.IsSet() ? &InstanceTransform.GetValue() : nullptr;
+}
+
 void UWorldPartition::OnPostBugItGoCalled(const FVector& Loc, const FRotator& Rot)
 {
 #if WITH_EDITOR
@@ -1117,7 +1134,7 @@ void UWorldPartition::OnActorDescRegistered(const FWorldPartitionActorDesc& Acto
 {
 	AActor* Actor = ActorDesc.GetActor();
 	check(Actor);
-	Actor->GetLevel()->AddLoadedActor(Actor, !InstanceTransform.Equals(FTransform::Identity) ? &InstanceTransform : nullptr);
+	Actor->GetLevel()->AddLoadedActor(Actor, GetInstanceTransformPtr());
 }
 
 void UWorldPartition::OnActorDescUnregistered(const FWorldPartitionActorDesc& ActorDesc) 
@@ -1125,7 +1142,7 @@ void UWorldPartition::OnActorDescUnregistered(const FWorldPartitionActorDesc& Ac
 	AActor* Actor = ActorDesc.GetActor();
 	if (IsValidChecked(Actor))
 	{
-		Actor->GetLevel()->RemoveLoadedActor(Actor, !InstanceTransform.Equals(FTransform::Identity) ? &InstanceTransform : nullptr);
+		Actor->GetLevel()->RemoveLoadedActor(Actor, GetInstanceTransformPtr());
 	}
 }
 
diff --git a/Engine/Source/Runtime/Engine/Public/WorldPartition/WorldPartition.h b/Engine/Source/Runtime/Engine/Public/WorldPartition/WorldPartition.h
--- a/Engine/Source/Runtime/Engine/Public/WorldPartition/WorldPartition.h
+++ b/Engine/Source/Runtime/Engine/Public/WorldPartition/WorldPartition.h
@@ -5,6 +5,7 @@
 #include "UObject/ObjectMacros.h"
 #include "UObject/Object.h"
 #include "Misc/CoreDelegates.h"
+#include "Misc/Optional.h"
 #include "GameFramework/Actor.h"
 #include "Templates/SubclassOf.h"
 #include "WorldPartition/WorldPartitionLog.h"
@@ -184,7 +185,8 @@ public:
 	virtual void Uninitialize() override;
 
 	bool IsMainWorldPartition() const;
-	const FTransform& GetInstanceTransform() const { return InstanceTransform; }
+	const FTransform& GetInstanceTransform() const;
+	const FTransform* GetInstanceTransformPtr() const;
 
 	void Tick(float DeltaSeconds);
 	void UpdateStreamingState();
@@ -209,11 +211,11 @@ public:
 
 	IWorldPartitionEditor* WorldPartitionEditor;
 
+private:
 	/** Class of WorldPartitionStreamingPolicy to be used to manage world partition streaming. */
 	UPROPERTY()
 	TSubclassOf<UWorldPartitionStreamingPolicy> WorldPartitionStreamingPolicyClass;
 
-private:
 	/** Enables streaming for this world. */
 	UPROPERTY()
 	bool bEnableStreaming;
@@ -252,7 +254,7 @@ public:
 
 private:
 	EWorldPartitionInitState InitState;
-	FTransform InstanceTransform;
+	TOptional<FTransform> InstanceTransform;
 
 	UPROPERTY()
 	mutable TObjectPtr<UWorldPartitionStreamingPolicy> StreamingPolicy;
