Texture fix PNGCompressed version check was not correct
even though the bug was only made in old versions it could be persisted to new versions
continue to do the fixup on load to newer versions
also add failsafe detection on uncompressed data in TryDecompress

diff --git a/Engine/Source/Runtime/Engine/Classes/Engine/Texture.h b/Engine/Source/Runtime/Engine/Classes/Engine/Texture.h
--- a/Engine/Source/Runtime/Engine/Classes/Engine/Texture.h
+++ b/Engine/Source/Runtime/Engine/Classes/Engine/Texture.h
@@ -1310,6 +1310,9 @@ public:
 	/** Texture color management settings: source encoding and color space. */
 	UPROPERTY(EditAnywhere, BlueprintReadWrite, Category = Texture, AdvancedDisplay)
 	FTextureSourceColorSettings SourceColorSettings;
+	
+	/* Store the FUE5MainStreamObjectVersion of the Texture uasset loaded for debugging (== LatestVersion if not loaded) */
+	int LoadedMainStreamObjectVersion;
 
 	/** Indicates we're currently importing the object (set in PostEditImport, unset in the subsequent PostEditChange) */
 	uint8 bIsImporting : 1;
diff --git a/Engine/Source/Runtime/Engine/Private/Texture.cpp b/Engine/Source/Runtime/Engine/Private/Texture.cpp
--- a/Engine/Source/Runtime/Engine/Private/Texture.cpp
+++ b/Engine/Source/Runtime/Engine/Private/Texture.cpp
@@ -111,6 +111,7 @@ UTexture::UTexture(const FObjectInitializer& ObjectInitializer)
 	Filter = TF_Default;
 	MipLoadOptions = ETextureMipLoadOptions::Default;
 #if WITH_EDITORONLY_DATA
+	LoadedMainStreamObjectVersion = FUE5MainStreamObjectVersion::LatestVersion;
 	SourceColorSettings = FTextureSourceColorSettings();
 	AdjustBrightness = 1.0f;
 	AdjustBrightnessCurve = 1.0f;
@@ -622,10 +623,15 @@ void UTexture::Serialize(FArchive& Ar)
 	Super::Serialize(Ar);
 	
 	FStripDataFlags StripFlags(Ar);
-
+	
 	/** Legacy serialization. */
 #if WITH_EDITORONLY_DATA
 
+	if (Ar.IsLoading())
+	{
+		LoadedMainStreamObjectVersion = Ar.CustomVer(FUE5MainStreamObjectVersion::GUID);
+	}
+
 	if (Ar.IsLoading() && Ar.CustomVer(FUE5MainStreamObjectVersion::GUID) < FUE5MainStreamObjectVersion::TextureDoScaleMipsForAlphaCoverage)
 	{
 		// bDoScaleMipsForAlphaCoverage was not transmitted in old versions
@@ -716,11 +722,13 @@ void UTexture::Serialize(FArchive& Ar)
 		//   is the authoritative source on whether something is a PNG or not.
 		// - In between, for a while after CompressionFormat was introduced, a bug meant that textures that were flagged 
 		//   as !bPNGCompressed, had their compression format set to PNG, but did not actually contain compressed data. Fix these up.
+		//   this bug only existed up to version TextureSourceVirtualization, but it could be carried forward to later asset versions
+		//   until this fixup was added
 		//
 		// Now, the separate bPNGCompressed is gone (to avoid further desyncs like this) and we make sure
 		// that CompressionFormat always matches the contents.	
 
-		if (Ar.CustomVer(FUE5MainStreamObjectVersion::GUID) < FUE5MainStreamObjectVersion::TextureSourceVirtualization
+		if (Ar.CustomVer(FUE5MainStreamObjectVersion::GUID) < FUE5MainStreamObjectVersion::VolumetricCloudReflectionSampleCountDefaultUpdate
 			&& !Source.bPNGCompressed_DEPRECATED
 			&& Source.CompressionFormat == TSCF_PNG)
 		{
@@ -1718,6 +1726,8 @@ FSharedBuffer FTextureSource::Decompress(IImageWrapperModule* ) const
 	
 	// ImageWrapperModule argument ignored, not drilled through DecompressImage
 
+	// @@!! todo: validate the size of the FSharedBuffer
+
 	if (CompressionFormat != TSCF_None )
 	{
 		return TryDecompressData();
@@ -1833,6 +1843,8 @@ uint8* FTextureSource::LockMipInternal(int32 BlockIndex, int32 LayerIndex, int32
 			return nullptr;
 		}
 		
+		// @@!! pointer advanced without checking size of buffer !?
+
 		MipData += CalcMipOffset(BlockIndex, LayerIndex, MipIndex);
 
 		if (NumLockedMips == 0)
@@ -2145,8 +2157,30 @@ FSharedBuffer FTextureSource::TryDecompressData() const
 		FImage Image;
 		if ( ! FImageUtils::DecompressImage(Payload.GetData(), Payload.GetSize(), Image) )
 		{
-			UE_LOG(LogTexture, Warning, TEXT("TryDecompressData failed to return uncompressed data"));
-			return FSharedBuffer();
+			int64 LayerSize = CalcLayerSize(0,0);
+			int64 PayloadSize = Payload.GetSize();
+			
+			UE_LOG(LogTexture, Display, TEXT("TryDecompressData failed: LayerSize = %lld PayloadSize = %lld"),LayerSize,PayloadSize);
+			UE_LOG(LogTexture, Display, TEXT("TryDecompressData failed: LoadedMainStreamObjectVersion = %d, TextureSourceVirtualization = %d, VolumetricCloudReflectionSampleCountDefaultUpdate = %d"),
+				Owner->LoadedMainStreamObjectVersion,
+				FUE5MainStreamObjectVersion::TextureSourceVirtualization,
+				FUE5MainStreamObjectVersion::VolumetricCloudReflectionSampleCountDefaultUpdate);
+
+			if ( LayerSize == PayloadSize )
+			{
+				// this is most likely from the bug where data is marked TSCF_PNG but is actually uncompressed
+				// fix CompressionFormat for the future :
+				check( CompressionFormat == TSCF_PNG );
+				const_cast<FTextureSource *>(this)->CompressionFormat = TSCF_None;
+
+				UE_LOG(LogTexture, Warning, TEXT("TryDecompressData data marked compressed appears to be uncompressed?"));
+				return Payload;
+			}
+			else
+			{
+				UE_LOG(LogTexture, Error, TEXT("TryDecompressData failed to return uncompressed data"));
+				return FSharedBuffer();
+			}
 		}
 
 		// we got data in Image.Format
