Another fix to non-unity build

diff --git a/Engine/Source/Runtime/Engine/Private/Animation/MotionTrajectoryTypes.cpp b/Engine/Source/Runtime/Engine/Private/Animation/MotionTrajectoryTypes.cpp
--- a/Engine/Source/Runtime/Engine/Private/Animation/MotionTrajectoryTypes.cpp
+++ b/Engine/Source/Runtime/Engine/Private/Animation/MotionTrajectoryTypes.cpp
@@ -5,6 +5,7 @@
 #include "Animation/AnimTypes.h"
 #include "DrawDebugHelpers.h"
 #include "Animation/AnimTypes.h"
+#include "HAL/IConsoleManager.h"
 
 #if ENABLE_ANIM_DEBUG
 static constexpr int32 DebugTrajectorySampleDisable = 0;
