Fix serialization bug that cause COTF to ignore what shader platform we were requesting.

- Add CPU markers to various ODSC related functions.

diff --git a/Engine/Source/Editor/UnrealEd/Private/Cooker/IoStoreCookOnTheFlyRequestManager.cpp b/Engine/Source/Editor/UnrealEd/Private/Cooker/IoStoreCookOnTheFlyRequestManager.cpp
--- a/Engine/Source/Editor/UnrealEd/Private/Cooker/IoStoreCookOnTheFlyRequestManager.cpp
+++ b/Engine/Source/Editor/UnrealEd/Private/Cooker/IoStoreCookOnTheFlyRequestManager.cpp
@@ -1068,6 +1068,7 @@ private:
 			TUniquePtr<FArchive> Ar = Request.ReadBody();
 			*Ar << RecompileData.MaterialsToLoad;
 			*Ar << iShaderPlatform;
+			RecompileData.ShaderPlatform = static_cast<EShaderPlatform>(iShaderPlatform);
 			*Ar << RecompileData.CommandType;
 			*Ar << RecompileData.ShadersToRecompile;
 		}
diff --git a/Engine/Source/Runtime/Engine/Private/Materials/Material.cpp b/Engine/Source/Runtime/Engine/Private/Materials/Material.cpp
--- a/Engine/Source/Runtime/Engine/Private/Materials/Material.cpp
+++ b/Engine/Source/Runtime/Engine/Private/Materials/Material.cpp
@@ -4618,6 +4618,8 @@ void UMaterial::CompileMaterialsForRemoteRecompile(
 	ITargetPlatform* TargetPlatform,
 	TMap<FString, TArray<TRefCountPtr<FMaterialShaderMap> > >& OutShaderMaps)
 {
+	TRACE_CPUPROFILER_EVENT_SCOPE(UMaterial::CompileMaterialsForRemoteRecompile);
+
 	// Build a map from UMaterial / UMaterialInstance to the resources which are being compiled
 	TMap<FString, TArray<FMaterialResource*> > CompilingResources;
 
diff --git a/Engine/Source/Runtime/Engine/Private/ShaderCompiler/ShaderCompiler.cpp b/Engine/Source/Runtime/Engine/Private/ShaderCompiler/ShaderCompiler.cpp
--- a/Engine/Source/Runtime/Engine/Private/ShaderCompiler/ShaderCompiler.cpp
+++ b/Engine/Source/Runtime/Engine/Private/ShaderCompiler/ShaderCompiler.cpp
@@ -6081,6 +6081,7 @@ namespace
 
 void ProcessCookOnTheFlyShaders(bool bReloadGlobalShaders, const TArray<uint8>& MeshMaterialMaps, const TArray<FString>& MaterialsToLoad, const TArray<uint8>& GlobalShaderMap)
 {
+	TRACE_CPUPROFILER_EVENT_SCOPE(ProcessCookOnTheFlyShaders);
 	check(IsInGameThread());
 
 	bool bHasFlushed = false;
@@ -6089,7 +6090,7 @@ void ProcessCookOnTheFlyShaders(bool bReloadGlobalShaders, const TArray<uint8>&
 		if (!bHasFlushed )
 		{
 			// now we need to refresh the RHI resources
-			FlushRenderingCommands();
+	FlushRenderingCommands();
 			bHasFlushed = true;
 		}
 	};
@@ -7192,8 +7193,6 @@ bool RecompileChangedShadersForPlatform(const FString& PlatformName)
 	TArray<FName> DesiredShaderFormats;
 	TargetPlatform->GetAllTargetedShaderFormats(DesiredShaderFormats);
 
-
-
 	// figure out which shaders are out of date
 	TArray<const FShaderType*> OutdatedShaderTypes;
 	TArray<const FVertexFactoryType*> OutdatedFactoryTypes;
@@ -7272,6 +7271,8 @@ void RecompileShadersForRemote(
 	FShaderRecompileData& Args,
 	const FString& OutputDirectory)
 {
+	TRACE_CPUPROFILER_EVENT_SCOPE(RecompileShadersForRemote);
+
 	// figure out what shader platforms to recompile
 	ITargetPlatformManagerModule* TPM = GetTargetPlatformManager();
 	ITargetPlatform* TargetPlatform = TPM->FindTargetPlatform(Args.PlatformName);
