Fix non-unity
#rnx

diff --git a/Engine/Source/Runtime/Engine/Private/EdGraph/EdGraph.cpp b/Engine/Source/Runtime/Engine/Private/EdGraph/EdGraph.cpp
--- a/Engine/Source/Runtime/Engine/Private/EdGraph/EdGraph.cpp
+++ b/Engine/Source/Runtime/Engine/Private/EdGraph/EdGraph.cpp
@@ -4,6 +4,7 @@
 #include "Engine/Blueprint.h"
 #include "GraphEditAction.h"
 #include "EdGraph/EdGraphSchema.h"
+#include "EngineLogs.h"
 #if WITH_EDITOR
 #include "UObject/UObjectHash.h"
 #endif
