Adding more descriptive StatID for AudioComponents when applicable. It will still show as an anonymous AudioComponent when the sound is passed explicitly, but most cases have one pre-set.

diff --git a/Engine/Source/Runtime/Engine/Classes/Components/AudioComponent.h b/Engine/Source/Runtime/Engine/Classes/Components/AudioComponent.h
--- a/Engine/Source/Runtime/Engine/Classes/Components/AudioComponent.h
+++ b/Engine/Source/Runtime/Engine/Classes/Components/AudioComponent.h
@@ -797,6 +797,8 @@ public:
 	uint64 GetInstanceOwnerID() const override { return AudioComponentID; }
 	USoundBase* GetSound() override { return Sound; }
 
+	virtual FName GetFNameForStatID() const override;
+
 public:
 
 	/**
diff --git a/Engine/Source/Runtime/Engine/Private/Components/AudioComponent.cpp b/Engine/Source/Runtime/Engine/Private/Components/AudioComponent.cpp
--- a/Engine/Source/Runtime/Engine/Private/Components/AudioComponent.cpp
+++ b/Engine/Source/Runtime/Engine/Private/Components/AudioComponent.cpp
@@ -856,6 +856,12 @@ FAudioDevice* UAudioComponent::GetAudioDevice() const
 	return AudioDevice;
 }
 
+FName UAudioComponent::GetFNameForStatID() const
+{
+	const USoundBase* SoundObject = Sound.Get();
+	return SoundObject ? SoundObject->GetFNameForStatID() : Super::GetFNameForStatID();
+}
+
 void UAudioComponent::FadeIn(float FadeInDuration, float FadeVolumeLevel, float StartTime, const EAudioFaderCurve FadeCurve)
 {
 	PlayInternalRequestData Data;
