Full set of fixes for non-unity non-pch builds after EngineTypes cleanup. This includes adjusting previous build fixes that introduced new includes that weren't actually necessary. Restored a few includes and wrapped them in UE_ENABLE_INCLUDE_ORDER_DEPRECATED_IN_5_1 checks to maintain compatability.
Also includes fixes for other build issues as it's hard to separate them all.

diff --git a/Engine/Plugins/AI/MassAI/Source/MassAIBehavior/Public/MassComponentHitSubsystem.h b/Engine/Plugins/AI/MassAI/Source/MassAIBehavior/Public/MassComponentHitSubsystem.h
--- a/Engine/Plugins/AI/MassAI/Source/MassAIBehavior/Public/MassComponentHitSubsystem.h
+++ b/Engine/Plugins/AI/MassAI/Source/MassAIBehavior/Public/MassComponentHitSubsystem.h
@@ -5,6 +5,7 @@
 #include "MassEntityTypes.h"
 #include "MassComponentHitTypes.h"
 #include "Subsystems/WorldSubsystem.h"
+#include "Engine/HitResult.h"
 
 #include "MassComponentHitSubsystem.generated.h"
 
diff --git a/Engine/Plugins/Experimental/ChaosVehiclesPlugin/Source/ChaosVehicles/Public/ChaosVehicleWheel.h b/Engine/Plugins/Experimental/ChaosVehiclesPlugin/Source/ChaosVehicles/Public/ChaosVehicleWheel.h
--- a/Engine/Plugins/Experimental/ChaosVehiclesPlugin/Source/ChaosVehicles/Public/ChaosVehicleWheel.h
+++ b/Engine/Plugins/Experimental/ChaosVehiclesPlugin/Source/ChaosVehicles/Public/ChaosVehicleWheel.h
@@ -13,7 +13,10 @@
 #include "EngineDefines.h"
 #include "SimpleVehicle.h"
 #include "VehicleUtility.h"
+#if UE_ENABLE_INCLUDE_ORDER_DEPRECATED_IN_5_1
 #include "Engine/EngineTypes.h"
+#endif
+#include "Engine/HitResult.h"
 #include "Curves/CurveFloat.h"
 
 #include "ChaosVehicleWheel.generated.h"
diff --git a/Engine/Plugins/Experimental/VirtualCamera/Source/VirtualCamera/Public/IVirtualCameraController.h b/Engine/Plugins/Experimental/VirtualCamera/Source/VirtualCamera/Public/IVirtualCameraController.h
--- a/Engine/Plugins/Experimental/VirtualCamera/Source/VirtualCamera/Public/IVirtualCameraController.h
+++ b/Engine/Plugins/Experimental/VirtualCamera/Source/VirtualCamera/Public/IVirtualCameraController.h
@@ -4,6 +4,9 @@
 
 #include "CoreMinimal.h"
 #include "Delegates/DelegateCombinations.h"
+#if UE_ENABLE_INCLUDE_ORDER_DEPRECATED_IN_5_1
+#include "Engine/EngineTypes.h"
+#endif
 #include "Engine/HitResult.h"
 #include "LiveLinkRole.h"
 #include "Misc/CoreMiscDefines.h"
diff --git a/Engine/Plugins/Runtime/GeoReferencing/Source/GeoReferencingEditor/Public/GeoReferencingEditorBPLibrary.h b/Engine/Plugins/Runtime/GeoReferencing/Source/GeoReferencingEditor/Public/GeoReferencingEditorBPLibrary.h
--- a/Engine/Plugins/Runtime/GeoReferencing/Source/GeoReferencingEditor/Public/GeoReferencingEditorBPLibrary.h
+++ b/Engine/Plugins/Runtime/GeoReferencing/Source/GeoReferencingEditor/Public/GeoReferencingEditorBPLibrary.h
@@ -3,7 +3,10 @@
 #pragma once
 
 #include "CoreMinimal.h"
+#if UE_ENABLE_INCLUDE_ORDER_DEPRECATED_IN_5_1
 #include "Engine/EngineTypes.h"
+#endif
+#include "Engine/HitResult.h"
 #include "Kismet/BlueprintFunctionLibrary.h"
 
 #include "GeoReferencingEditorBPLibrary.generated.h"
diff --git a/Engine/Plugins/Runtime/MassGameplay/Source/MassCommon/Private/MassCommonTypes.cpp b/Engine/Plugins/Runtime/MassGameplay/Source/MassCommon/Private/MassCommonTypes.cpp
--- a/Engine/Plugins/Runtime/MassGameplay/Source/MassCommon/Private/MassCommonTypes.cpp
+++ b/Engine/Plugins/Runtime/MassGameplay/Source/MassCommon/Private/MassCommonTypes.cpp
@@ -3,6 +3,7 @@
 #include "MassCommonTypes.h"
 #include "MassEntitySubsystem.h"
 #include "Math/ColorList.h"
+#include "HAL/IConsoleManager.h"
 
 #if WITH_MASSGAMEPLAY_DEBUG
 
diff --git a/Engine/Plugins/Runtime/MassGameplay/Source/MassLOD/Private/MassLODUtils.cpp b/Engine/Plugins/Runtime/MassGameplay/Source/MassLOD/Private/MassLODUtils.cpp
--- a/Engine/Plugins/Runtime/MassGameplay/Source/MassLOD/Private/MassLODUtils.cpp
+++ b/Engine/Plugins/Runtime/MassGameplay/Source/MassLOD/Private/MassLODUtils.cpp
@@ -2,6 +2,7 @@
 
 #include "MassLODUtils.h"
 #include "MassCommandBuffer.h"
+#include "HAL/IConsoleManager.h"
 
 DEFINE_LOG_CATEGORY(LogMassLOD);
 #if WITH_MASSGAMEPLAY_DEBUG
diff --git a/Engine/Plugins/Runtime/MeshModelingToolset/Source/ModelingComponents/Public/Mechanics/DragAlignmentMechanic.h b/Engine/Plugins/Runtime/MeshModelingToolset/Source/ModelingComponents/Public/Mechanics/DragAlignmentMechanic.h
--- a/Engine/Plugins/Runtime/MeshModelingToolset/Source/ModelingComponents/Public/Mechanics/DragAlignmentMechanic.h
+++ b/Engine/Plugins/Runtime/MeshModelingToolset/Source/ModelingComponents/Public/Mechanics/DragAlignmentMechanic.h
@@ -7,6 +7,9 @@
 #include "BaseBehaviors/BehaviorTargetInterfaces.h"
 #include "DynamicMesh/DynamicMeshAABBTree3.h"
 #include "ToolContextInterfaces.h" // FViewCameraState
+#if UE_ENABLE_INCLUDE_ORDER_DEPRECATED_IN_5_1
+#include "Engine/EngineTypes.h" //FHitResult
+#endif
 #include "Engine/HitResult.h"
 #include "InteractionMechanic.h"
 #include "TransformTypes.h"
diff --git a/Engine/Plugins/Runtime/MeshModelingToolset/Source/ModelingComponents/Public/Mechanics/PlaneDistanceFromHitMechanic.h b/Engine/Plugins/Runtime/MeshModelingToolset/Source/ModelingComponents/Public/Mechanics/PlaneDistanceFromHitMechanic.h
--- a/Engine/Plugins/Runtime/MeshModelingToolset/Source/ModelingComponents/Public/Mechanics/PlaneDistanceFromHitMechanic.h
+++ b/Engine/Plugins/Runtime/MeshModelingToolset/Source/ModelingComponents/Public/Mechanics/PlaneDistanceFromHitMechanic.h
@@ -3,7 +3,10 @@
 #pragma once
 
 #include "CoreMinimal.h"
+#if UE_ENABLE_INCLUDE_ORDER_DEPRECATED_IN_5_1
 #include "Engine/EngineTypes.h"
+#endif
+#include "Engine/HitResult.h"
 #include "InteractionMechanic.h"
 #include "DynamicMesh/DynamicMesh3.h"
 #include "DynamicMesh/DynamicMeshAABBTree3.h"
diff --git a/Engine/Plugins/Runtime/MeshModelingToolset/Source/ModelingComponents/Public/Selection/SelectClickedAction.h b/Engine/Plugins/Runtime/MeshModelingToolset/Source/ModelingComponents/Public/Selection/SelectClickedAction.h
--- a/Engine/Plugins/Runtime/MeshModelingToolset/Source/ModelingComponents/Public/Selection/SelectClickedAction.h
+++ b/Engine/Plugins/Runtime/MeshModelingToolset/Source/ModelingComponents/Public/Selection/SelectClickedAction.h
@@ -5,6 +5,7 @@
 #include "CoreMinimal.h"
 #include "BaseBehaviors/BehaviorTargetInterfaces.h"
 #include "ToolSceneQueriesUtil.h"
+#include "Engine/HitResult.h"
 
 /**
  * BehaviorTarget to do world raycast selection from a click
diff --git a/Engine/Plugins/Runtime/MeshModelingToolset/Source/ModelingComponentsEditorOnly/Public/ToolTargets/VolumeComponentToolTarget.h b/Engine/Plugins/Runtime/MeshModelingToolset/Source/ModelingComponentsEditorOnly/Public/ToolTargets/VolumeComponentToolTarget.h
--- a/Engine/Plugins/Runtime/MeshModelingToolset/Source/ModelingComponentsEditorOnly/Public/ToolTargets/VolumeComponentToolTarget.h
+++ b/Engine/Plugins/Runtime/MeshModelingToolset/Source/ModelingComponentsEditorOnly/Public/ToolTargets/VolumeComponentToolTarget.h
@@ -12,6 +12,7 @@
 #include "TargetInterfaces/MeshDescriptionProvider.h"
 #include "TargetInterfaces/PhysicsDataSource.h"
 #include "ToolTargets/PrimitiveComponentToolTarget.h"
+#include "HAL/IConsoleManager.h"
 
 #include "VolumeComponentToolTarget.generated.h"
 
diff --git a/Engine/Plugins/Runtime/NetworkPredictionExtras/Source/NetworkPredictionExtras/Public/BaseMovementComponent.h b/Engine/Plugins/Runtime/NetworkPredictionExtras/Source/NetworkPredictionExtras/Public/BaseMovementComponent.h
--- a/Engine/Plugins/Runtime/NetworkPredictionExtras/Source/NetworkPredictionExtras/Public/BaseMovementComponent.h
+++ b/Engine/Plugins/Runtime/NetworkPredictionExtras/Source/NetworkPredictionExtras/Public/BaseMovementComponent.h
@@ -3,6 +3,7 @@
 #pragma once
 
 #include "NetworkPredictionComponent.h"
+#include "Engine/HitResult.h"
 
 #include "BaseMovementComponent.generated.h"
 
diff --git a/Engine/Plugins/Tests/EditorTests/Source/EditorTests/Private/UnrealEd/CollisionAutomationTests.h b/Engine/Plugins/Tests/EditorTests/Source/EditorTests/Private/UnrealEd/CollisionAutomationTests.h
--- a/Engine/Plugins/Tests/EditorTests/Source/EditorTests/Private/UnrealEd/CollisionAutomationTests.h
+++ b/Engine/Plugins/Tests/EditorTests/Source/EditorTests/Private/UnrealEd/CollisionAutomationTests.h
@@ -5,7 +5,9 @@
 #include "CoreMinimal.h"
 #include "UObject/ObjectMacros.h"
 #include "UObject/Object.h"
+#if UE_ENABLE_INCLUDE_ORDER_DEPRECATED_IN_5_1
 #include "Engine/EngineTypes.h"
+#endif
 #include "Engine/HitResult.h"
 #include "CollisionAutomationTests.generated.h"
 
diff --git a/Engine/Source/Developer/CollisionAnalyzer/Private/CollisionAnalyzer.cpp b/Engine/Source/Developer/CollisionAnalyzer/Private/CollisionAnalyzer.cpp
--- a/Engine/Source/Developer/CollisionAnalyzer/Private/CollisionAnalyzer.cpp
+++ b/Engine/Source/Developer/CollisionAnalyzer/Private/CollisionAnalyzer.cpp
@@ -4,6 +4,7 @@
 #include "HAL/FileManager.h"
 #include "Widgets/DeclarativeSyntaxSupport.h"
 #include "Engine/GameViewportClient.h"
+#include "Engine/HitResult.h"
 #include "DrawDebugHelpers.h"
 #include "SCollisionAnalyzer.h"
 #include "CollisionAnalyzerLog.h"
@@ -15,6 +16,12 @@
 /** Version of collision analyzer. Incremented on serialization changes.			*/
 #define COLLISION_ANALYZER_VERSION					0
 
+FCAQuery::FCAQuery() :
+	Params(NAME_None, FCollisionQueryParams::GetUnknownStatId())
+{
+}
+FCAQuery::~FCAQuery() = default;
+
 /** Util for serializing an FHitResult struct */
 void SerializeHitResult(FArchive& Ar, FHitResult& Result)
 {
diff --git a/Engine/Source/Developer/CollisionAnalyzer/Private/CollisionAnalyzer.h b/Engine/Source/Developer/CollisionAnalyzer/Private/CollisionAnalyzer.h
--- a/Engine/Source/Developer/CollisionAnalyzer/Private/CollisionAnalyzer.h
+++ b/Engine/Source/Developer/CollisionAnalyzer/Private/CollisionAnalyzer.h
@@ -13,10 +13,8 @@ class SWidget;
 /** Stores information about one collision query */
 struct FCAQuery
 {
-	FCAQuery():
-		Params(NAME_None, FCollisionQueryParams::GetUnknownStatId())
-	{
-	}
+	FCAQuery();
+	~FCAQuery();
 
 	FVector						Start;
 	FVector						End;
diff --git a/Engine/Source/Developer/CollisionAnalyzer/Private/SCAQueryDetails.h b/Engine/Source/Developer/CollisionAnalyzer/Private/SCAQueryDetails.h
--- a/Engine/Source/Developer/CollisionAnalyzer/Private/SCAQueryDetails.h
+++ b/Engine/Source/Developer/CollisionAnalyzer/Private/SCAQueryDetails.h
@@ -3,7 +3,7 @@
 #pragma once
 
 #include "CoreMinimal.h"
-#include "Engine/EngineTypes.h"
+#include "Engine/HitResult.h"
 #include "Widgets/DeclarativeSyntaxSupport.h"
 #include "Widgets/SCompoundWidget.h"
 #include "Widgets/Views/STableViewBase.h"
diff --git a/Engine/Source/Developer/CollisionAnalyzer/Private/SCAQueryTableRow.cpp b/Engine/Source/Developer/CollisionAnalyzer/Private/SCAQueryTableRow.cpp
--- a/Engine/Source/Developer/CollisionAnalyzer/Private/SCAQueryTableRow.cpp
+++ b/Engine/Source/Developer/CollisionAnalyzer/Private/SCAQueryTableRow.cpp
@@ -2,6 +2,7 @@
 
 #include "SCAQueryTableRow.h"
 #include "CollisionAnalyzerStyle.h"
+#include "Engine/HitResult.h"
 
 #define LOCTEXT_NAMESPACE "CollisionAnalyzer"
 
diff --git a/Engine/Source/Developer/CollisionAnalyzer/Public/ICollisionAnalyzer.h b/Engine/Source/Developer/CollisionAnalyzer/Public/ICollisionAnalyzer.h
--- a/Engine/Source/Developer/CollisionAnalyzer/Public/ICollisionAnalyzer.h
+++ b/Engine/Source/Developer/CollisionAnalyzer/Public/ICollisionAnalyzer.h
@@ -8,6 +8,7 @@
 class SWidget;
 struct FCollisionObjectQueryParams;
 struct FCollisionResponseParams;
+struct FHitResult;
 
 /** Enum to describe type of the query */
 namespace ECAQueryType
diff --git a/Engine/Source/Developer/FunctionalTesting/Classes/TraceQueryTestResults.h b/Engine/Source/Developer/FunctionalTesting/Classes/TraceQueryTestResults.h
--- a/Engine/Source/Developer/FunctionalTesting/Classes/TraceQueryTestResults.h
+++ b/Engine/Source/Developer/FunctionalTesting/Classes/TraceQueryTestResults.h
@@ -5,7 +5,9 @@
 #include "CoreMinimal.h"
 #include "UObject/ObjectMacros.h"
 #include "UObject/UObjectGlobals.h"
+#if UE_ENABLE_INCLUDE_ORDER_DEPRECATED_IN_5_1
 #include "Engine/EngineTypes.h"
+#endif
 #include "Engine/HitResult.h"
 #include "TraceQueryTestResults.generated.h"
 
diff --git a/Engine/Source/Editor/ClothPainter/Private/ClothingAssetExporter.cpp b/Engine/Source/Editor/ClothPainter/Private/ClothingAssetExporter.cpp
--- a/Engine/Source/Editor/ClothPainter/Private/ClothingAssetExporter.cpp
+++ b/Engine/Source/Editor/ClothPainter/Private/ClothingAssetExporter.cpp
@@ -7,6 +7,8 @@
 #include "AssetRegistryModule.h"
 #include "IContentBrowserSingleton.h"
 #include "ContentBrowserModule.h"
+#include "Features/IModularFeatures.h"
+#include "Misc/MessageDialog.h"
 
 #define LOCTEXT_NAMESPACE "ClothingAssetExporter"
 
diff --git a/Engine/Source/Editor/ClothPainter/Private/ClothingAssetListCommands.cpp b/Engine/Source/Editor/ClothPainter/Private/ClothingAssetListCommands.cpp
--- a/Engine/Source/Editor/ClothPainter/Private/ClothingAssetListCommands.cpp
+++ b/Engine/Source/Editor/ClothPainter/Private/ClothingAssetListCommands.cpp
@@ -2,6 +2,7 @@
 
 #include "ClothingAssetListCommands.h"
 #include "ClothingAssetExporter.h"
+#include "Features/IModularFeatures.h"
 
 #define LOCTEXT_NAMESPACE "ClothingAssetListCommands"
 
diff --git a/Engine/Source/Editor/MeshPaint/Public/IMeshPainter.h b/Engine/Source/Editor/MeshPaint/Public/IMeshPainter.h
--- a/Engine/Source/Editor/MeshPaint/Public/IMeshPainter.h
+++ b/Engine/Source/Editor/MeshPaint/Public/IMeshPainter.h
@@ -19,6 +19,8 @@ class FEditorViewportClient;
 class FSceneView;
 class AActor;
 
+struct FHitResult;
+
 enum class EMeshPaintAction : uint8;
 
 /** Base class for creating a mesh painter, has basic functionality combined with IMeshPaintMode and requires painting/situation-specific information to do actual painting */
diff --git a/Engine/Source/Runtime/Engine/Classes/AI/Navigation/PathFollowingAgentInterface.h b/Engine/Source/Runtime/Engine/Classes/AI/Navigation/PathFollowingAgentInterface.h
--- a/Engine/Source/Runtime/Engine/Classes/AI/Navigation/PathFollowingAgentInterface.h
+++ b/Engine/Source/Runtime/Engine/Classes/AI/Navigation/PathFollowingAgentInterface.h
@@ -3,11 +3,14 @@
 #pragma once
 
 #include "CoreMinimal.h"
+#if UE_ENABLE_INCLUDE_ORDER_DEPRECATED_IN_5_1
 #include "Engine/EngineTypes.h"
+#endif
 #include "UObject/ObjectMacros.h"
 #include "UObject/Interface.h"
 #include "PathFollowingAgentInterface.generated.h"
 
+struct FHitResult;
 
 UINTERFACE(MinimalAPI, meta = (CannotImplementInterfaceInBlueprint))
 class UPathFollowingAgentInterface : public UInterface
diff --git a/Engine/Source/Runtime/Engine/Classes/Components/InterpToMovementComponent.h b/Engine/Source/Runtime/Engine/Classes/Components/InterpToMovementComponent.h
--- a/Engine/Source/Runtime/Engine/Classes/Components/InterpToMovementComponent.h
+++ b/Engine/Source/Runtime/Engine/Classes/Components/InterpToMovementComponent.h
@@ -5,9 +5,10 @@
 #include "CoreMinimal.h"
 #include "UObject/ObjectMacros.h"
 #include "GameFramework/MovementComponent.h"
-#include "Engine/HitResult.h"
 #include "InterpToMovementComponent.generated.h"
 
+struct FHitResult;
+
 /**
  * Controls the movement behaviour
  */
diff --git a/Engine/Source/Runtime/Engine/Classes/Components/SceneComponent.h b/Engine/Source/Runtime/Engine/Classes/Components/SceneComponent.h
--- a/Engine/Source/Runtime/Engine/Classes/Components/SceneComponent.h
+++ b/Engine/Source/Runtime/Engine/Classes/Components/SceneComponent.h
@@ -11,6 +11,7 @@
 #include "Engine/OverlapInfo.h"
 #include "Engine/ScopedMovementUpdate.h"
 #endif
+#include "Engine/HitResult.h"
 #include "ComponentInstanceDataCache.h"
 #include "Components/ActorComponent.h"
 #include "RHIDefinitions.h"
diff --git a/Engine/Source/Runtime/Engine/Classes/GameFramework/LightWeightInstanceBlueprintFunctionLibrary.h b/Engine/Source/Runtime/Engine/Classes/GameFramework/LightWeightInstanceBlueprintFunctionLibrary.h
--- a/Engine/Source/Runtime/Engine/Classes/GameFramework/LightWeightInstanceBlueprintFunctionLibrary.h
+++ b/Engine/Source/Runtime/Engine/Classes/GameFramework/LightWeightInstanceBlueprintFunctionLibrary.h
@@ -3,7 +3,10 @@
 #pragma once
 
 #include "Kismet/BlueprintFunctionLibrary.h"
+#if UE_ENABLE_INCLUDE_ORDER_DEPRECATED_IN_5_1
 #include "Engine/EngineTypes.h"
+#endif
+#include "Engine/ActorInstanceHandle.h"
 
 #include "LightWeightInstanceBlueprintFunctionLibrary.generated.h"
 
diff --git a/Engine/Source/Runtime/Engine/Classes/GameFramework/ProjectileMovementComponent.h b/Engine/Source/Runtime/Engine/Classes/GameFramework/ProjectileMovementComponent.h
--- a/Engine/Source/Runtime/Engine/Classes/GameFramework/ProjectileMovementComponent.h
+++ b/Engine/Source/Runtime/Engine/Classes/GameFramework/ProjectileMovementComponent.h
@@ -5,9 +5,10 @@
 #include "CoreMinimal.h"
 #include "UObject/ObjectMacros.h"
 #include "GameFramework/MovementComponent.h"
-#include "Engine/HitResult.h"
 #include "ProjectileMovementComponent.generated.h"
 
+struct FHitResult;
+
 /**
  * ProjectileMovementComponent updates the position of another component during its tick.
  *
diff --git a/Engine/Source/Runtime/Engine/Classes/Kismet/KismetSystemLibrary.h b/Engine/Source/Runtime/Engine/Classes/Kismet/KismetSystemLibrary.h
--- a/Engine/Source/Runtime/Engine/Classes/Kismet/KismetSystemLibrary.h
+++ b/Engine/Source/Runtime/Engine/Classes/Kismet/KismetSystemLibrary.h
@@ -6,7 +6,10 @@
 #include "UObject/ObjectMacros.h"
 #include "UObject/Object.h"
 #include "Templates/SubclassOf.h"
+#if UE_ENABLE_INCLUDE_ORDER_DEPRECATED_IN_5_1
 #include "Engine/EngineTypes.h"
+#endif
+#include "Engine/HitResult.h"
 #include "UObject/UnrealType.h"
 #include "UObject/ScriptMacros.h"
 #include "UObject/Interface.h"
diff --git a/Engine/Source/Runtime/Engine/Private/Collision/CollisionDebugDrawing.h b/Engine/Source/Runtime/Engine/Private/Collision/CollisionDebugDrawing.h
--- a/Engine/Source/Runtime/Engine/Private/Collision/CollisionDebugDrawing.h
+++ b/Engine/Source/Runtime/Engine/Private/Collision/CollisionDebugDrawing.h
@@ -5,11 +5,15 @@
 #pragma once 
 
 #include "CoreMinimal.h"
+#if UE_ENABLE_INCLUDE_ORDER_DEPRECATED_IN_5_1
 #include "Engine/EngineTypes.h"
+#endif
 #include "EngineDefines.h"
 #include "PhysicsPublic.h"
 #include "PhysXIncludes.h"
 
+struct FHitResult;
+
 namespace Chaos
 {
 	class FImplicitObject;
diff --git a/Engine/Source/Runtime/Engine/Private/EdGraph/EdGraphNode.cpp b/Engine/Source/Runtime/Engine/Private/EdGraph/EdGraphNode.cpp
--- a/Engine/Source/Runtime/Engine/Private/EdGraph/EdGraphNode.cpp
+++ b/Engine/Source/Runtime/Engine/Private/EdGraph/EdGraphNode.cpp
@@ -11,6 +11,7 @@
 #include "EdGraph/EdGraphPin.h"
 #include "Textures/SlateIcon.h"
 #include "EdGraph/EdGraph.h"
+#include "EngineLogs.h"
 #if WITH_EDITOR
 #include "CookerSettings.h"
 #include "Kismet2/BlueprintEditorUtils.h"
diff --git a/Engine/Source/Runtime/Engine/Private/Engine/ScopedMovementUpdate.cpp b/Engine/Source/Runtime/Engine/Private/Engine/ScopedMovementUpdate.cpp
--- a/Engine/Source/Runtime/Engine/Private/Engine/ScopedMovementUpdate.cpp
+++ b/Engine/Source/Runtime/Engine/Private/Engine/ScopedMovementUpdate.cpp
@@ -182,7 +182,7 @@ void FScopedMovementUpdate::OnInnerScopeComplete(const FScopedMovementUpdate& In
 	}	
 }
 
-bool FScopedMovementUpdate::SetWorldLocationAndRotation(FVector NewLocation, const FQuat& NewQuat, bool bNoPhysics /*= false*/, ETeleportType Teleport /*= ETeleportType::None*/)
+bool FScopedMovementUpdate::SetWorldLocationAndRotation(FVector NewLocation, const FQuat& NewQuat, bool bNoPhysics, ETeleportType Teleport)
 {
 	if (Owner)
 	{
diff --git a/Engine/Source/Runtime/Engine/Public/Collision.h b/Engine/Source/Runtime/Engine/Public/Collision.h
--- a/Engine/Source/Runtime/Engine/Public/Collision.h
+++ b/Engine/Source/Runtime/Engine/Public/Collision.h
@@ -8,7 +8,10 @@
 
 #include "CoreMinimal.h"
 #include "Stats/Stats.h"
+#if UE_ENABLE_INCLUDE_ORDER_DEPRECATED_IN_5_1
 #include "Engine/EngineTypes.h"
+#endif
+#include "Engine/HitResult.h"
 #include "EngineDefines.h"
 
 /**
diff --git a/Engine/Source/Runtime/Engine/Public/Engine/ScopedMovementUpdate.h b/Engine/Source/Runtime/Engine/Public/Engine/ScopedMovementUpdate.h
--- a/Engine/Source/Runtime/Engine/Public/Engine/ScopedMovementUpdate.h
+++ b/Engine/Source/Runtime/Engine/Public/Engine/ScopedMovementUpdate.h
@@ -5,7 +5,8 @@
 #include "CoreMinimal.h"
 #include "Engine/HitResult.h"
 #include "Engine/OverlapInfo.h"
-#include "Engine/EngineTypes.h"
+
+enum class ETeleportType : uint8;
 
 /**
  * Enum that controls the scoping behavior of FScopedMovementUpdate.
@@ -106,7 +107,7 @@ protected:
 	/** Fills in the list of overlaps at the end location (in EndOverlaps). Returns pointer to the list, or null if it can't be computed. */
 	TOptional<TOverlapArrayView> GetOverlapsAtEnd(class UPrimitiveComponent& PrimComponent, TInlineOverlapInfoArray& OutEndOverlaps, bool bTransformChanged) const;
 
-	bool SetWorldLocationAndRotation(FVector NewLocation, const FQuat& NewQuat, bool bNoPhysics = false, ETeleportType Teleport = ETeleportType::None);
+	bool SetWorldLocationAndRotation(FVector NewLocation, const FQuat& NewQuat, bool bNoPhysics, ETeleportType Teleport);
 
 private:
 	/** Notify this scope that the given inner scope completed its update (ie is going out of scope). Only occurs for deferred updates. */
diff --git a/Engine/Source/Runtime/Engine/Public/PhysicsReplication.h b/Engine/Source/Runtime/Engine/Public/PhysicsReplication.h
--- a/Engine/Source/Runtime/Engine/Public/PhysicsReplication.h
+++ b/Engine/Source/Runtime/Engine/Public/PhysicsReplication.h
@@ -8,6 +8,7 @@
 #pragma once
 
 #include "Engine/EngineTypes.h"
+#include "Engine/ReplicatedState.h"
 #include "Physics/PhysicsInterfaceDeclares.h"
 #include "PhysicsProxy/SingleParticlePhysicsProxyFwd.h"
 #if WITH_CHAOS
diff --git a/Engine/Source/Runtime/Engine/Public/WorldCollision.h b/Engine/Source/Runtime/Engine/Public/WorldCollision.h
--- a/Engine/Source/Runtime/Engine/Public/WorldCollision.h
+++ b/Engine/Source/Runtime/Engine/Public/WorldCollision.h
@@ -10,6 +10,7 @@
 #include "CollisionQueryParams.h"
 #include "CollisionShape.h"
 #include "Engine/OverlapResult.h"
+#include "Engine/HitResult.h"
 
 struct FOverlapDatum;
 struct FTraceDatum;
diff --git a/Engine/Source/Runtime/Experimental/ChaosSolverEngine/Public/Chaos/ChaosNotifyHandlerInterface.h b/Engine/Source/Runtime/Experimental/ChaosSolverEngine/Public/Chaos/ChaosNotifyHandlerInterface.h
--- a/Engine/Source/Runtime/Experimental/ChaosSolverEngine/Public/Chaos/ChaosNotifyHandlerInterface.h
+++ b/Engine/Source/Runtime/Experimental/ChaosSolverEngine/Public/Chaos/ChaosNotifyHandlerInterface.h
@@ -4,7 +4,10 @@
 
 #include "Kismet/BlueprintFunctionLibrary.h"
 #include "UObject/Interface.h"
+#if UE_ENABLE_INCLUDE_ORDER_DEPRECATED_IN_5_1
 #include "Engine/EngineTypes.h"
+#endif
+#include "Engine/HitResult.h"
 #include "ChaosNotifyHandlerInterface.generated.h"
 
 USTRUCT(BlueprintType)
diff --git a/Engine/Source/Runtime/InteractiveToolsFramework/Private/BaseTools/MeshSurfacePointTool.cpp b/Engine/Source/Runtime/InteractiveToolsFramework/Private/BaseTools/MeshSurfacePointTool.cpp
--- a/Engine/Source/Runtime/InteractiveToolsFramework/Private/BaseTools/MeshSurfacePointTool.cpp
+++ b/Engine/Source/Runtime/InteractiveToolsFramework/Private/BaseTools/MeshSurfacePointTool.cpp
@@ -9,6 +9,8 @@
 #include "TargetInterfaces/PrimitiveComponentBackedTarget.h"
 #include "ToolTargetManager.h"
 
+#include "Engine/World.h"
+
 #define LOCTEXT_NAMESPACE "UMeshSurfacePointTool"
 
 
diff --git a/Engine/Source/Runtime/InteractiveToolsFramework/Public/BaseTools/BaseBrushTool.h b/Engine/Source/Runtime/InteractiveToolsFramework/Public/BaseTools/BaseBrushTool.h
--- a/Engine/Source/Runtime/InteractiveToolsFramework/Public/BaseTools/BaseBrushTool.h
+++ b/Engine/Source/Runtime/InteractiveToolsFramework/Public/BaseTools/BaseBrushTool.h
@@ -5,6 +5,9 @@
 #include "CoreMinimal.h"
 #include "BaseTools/MeshSurfacePointTool.h"
 #include "InteractiveTool.h"
+#if UE_ENABLE_INCLUDE_ORDER_DEPRECATED_IN_5_1
+#include "Engine/EngineTypes.h"
+#endif
 #include "Engine/HitResult.h"
 #include "BaseBrushTool.generated.h"
 
diff --git a/Engine/Source/Runtime/InteractiveToolsFramework/Public/SceneQueries/SceneSnappingManager.h b/Engine/Source/Runtime/InteractiveToolsFramework/Public/SceneQueries/SceneSnappingManager.h
--- a/Engine/Source/Runtime/InteractiveToolsFramework/Public/SceneQueries/SceneSnappingManager.h
+++ b/Engine/Source/Runtime/InteractiveToolsFramework/Public/SceneQueries/SceneSnappingManager.h
@@ -6,7 +6,10 @@
 #include "Misc/Optional.h"
 #include "Math/Vector.h"
 #include "Math/Ray.h"
+#if UE_ENABLE_INCLUDE_ORDER_DEPRECATED_IN_5_1
 #include "Engine/EngineTypes.h"
+#endif
+#include "Engine/HitResult.h"
 #include "SceneSnappingManager.generated.h"
 
 class AActor;
diff --git a/Engine/Source/Runtime/InteractiveToolsFramework/Public/TargetInterfaces/PrimitiveComponentBackedTarget.h b/Engine/Source/Runtime/InteractiveToolsFramework/Public/TargetInterfaces/PrimitiveComponentBackedTarget.h
--- a/Engine/Source/Runtime/InteractiveToolsFramework/Public/TargetInterfaces/PrimitiveComponentBackedTarget.h
+++ b/Engine/Source/Runtime/InteractiveToolsFramework/Public/TargetInterfaces/PrimitiveComponentBackedTarget.h
@@ -3,7 +3,9 @@
 #pragma once
 
 #include "CoreMinimal.h"
-#include "Engine/HitResult.h"
+#if UE_ENABLE_INCLUDE_ORDER_DEPRECATED_IN_5_1
+#include "Engine/EngineTypes.h"    // FHitResult
+#endif
 #include "UObject/Interface.h"
 
 #include "PrimitiveComponentBackedTarget.generated.h"
@@ -11,6 +13,8 @@
 class UPrimitiveComponent;
 class AActor;
 
+struct FHitResult;
+
 UINTERFACE()
 class INTERACTIVETOOLSFRAMEWORK_API UPrimitiveComponentBackedTarget : public UInterface
 {
diff --git a/Engine/Source/Runtime/UMG/Public/Components/WidgetInteractionComponent.h b/Engine/Source/Runtime/UMG/Public/Components/WidgetInteractionComponent.h
--- a/Engine/Source/Runtime/UMG/Public/Components/WidgetInteractionComponent.h
+++ b/Engine/Source/Runtime/UMG/Public/Components/WidgetInteractionComponent.h
@@ -6,6 +6,9 @@
 #include "UObject/ObjectMacros.h"
 #include "UObject/UObjectGlobals.h"
 #include "InputCoreTypes.h"
+#if UE_ENABLE_INCLUDE_ORDER_DEPRECATED_IN_5_1
+#include "Engine/EngineTypes.h"
+#endif
 #include "Engine/HitResult.h"
 #include "Components/SceneComponent.h"
 #include "GenericPlatform/GenericApplication.h"
