Optimize GetMaterialShaderMapKeyString

215s -> 11s together with previous DDC cache key optimization when cooking an internal project

* Moved to Append methods and changed printf to BytesToHex (added some new core append / BytseToHex methods in 19585530)
* Added local lookup cache around FShaderParametersMetadata::GetStructList()
* Added local lookup cache around Freeze::HashLayout()
* Fixed bug in FUniformBufferNameSortOrder that caused unnecessarily long keys, which means all material keys will change
* Optimized ReferencedUniformBuffers hash map merging

diff --git a/Engine/Source/Runtime/Engine/Private/Materials/MaterialShader.cpp b/Engine/Source/Runtime/Engine/Private/Materials/MaterialShader.cpp
--- a/Engine/Source/Runtime/Engine/Private/Materials/MaterialShader.cpp
+++ b/Engine/Source/Runtime/Engine/Private/Materials/MaterialShader.cpp
@@ -1,6 +1,7 @@
 // Copyright Epic Games, Inc. All Rights Reserved.
 
 #include "MaterialShader.h"
+#include "Algo/Find.h"
 #include "Containers/StringConv.h"
 #include "DerivedDataCache.h"
 #include "DerivedDataRequestOwner.h"
@@ -172,14 +173,24 @@ FString GetBlendModeString(EBlendMode BlendMode)
 static FString GetMaterialShaderMapKeyString(const FMaterialShaderMapId& ShaderMapId, EShaderPlatform Platform)
 {
 	FName Format = LegacyShaderPlatformToShaderFormat(Platform);
-	FString ShaderMapKeyString = Format.ToString() + TEXT("_") + FString(FString::FromInt(GetTargetPlatformManagerRef().ShaderFormatVersion(Format))) + TEXT("_");
+	FString ShaderMapKeyString;
+	ShaderMapKeyString.Reserve(16384);
+	
+	ShaderMapKeyString.Append(TEXTVIEW("MATSM_"));
+	ShaderMapKeyString.Append(GetMaterialShaderMapDDCKey());
+	ShaderMapKeyString.AppendChar('_');
+
+	Format.AppendString(ShaderMapKeyString);
+	ShaderMapKeyString.AppendChar('_');
+	ShaderMapKeyString.AppendInt(GetTargetPlatformManagerRef().ShaderFormatVersion(Format));
+	ShaderMapKeyString.AppendChar('_');
 
 	ShaderMapAppendKeyString(Platform, ShaderMapKeyString);
 	ShaderMapId.AppendKeyString(ShaderMapKeyString);
 	FMaterialAttributeDefinitionMap::AppendDDCKeyString(ShaderMapKeyString);
 	FShaderCompileUtilities::AppendGBufferDDCKeyString(Platform, ShaderMapKeyString);
 
-	return FString::Printf(TEXT("%s_%s_%s"), TEXT("MATSM"), *GetMaterialShaderMapDDCKey(), *ShaderMapKeyString);
+	return ShaderMapKeyString;
 }
 
 static UE::DerivedData::FCacheKey GetMaterialShaderMapKey(const FStringView MaterialShaderMapKey)
@@ -911,19 +922,114 @@ void FMaterialShaderMapId::UpdateFromParameterSet(const FStaticParameterSet& Sta
 #endif // WITH_EDITOR
 
 #if WITH_EDITOR
+
+static TSortedMap<const TCHAR*, FCachedUniformBufferDeclaration, FDefaultAllocator, FUniformBufferNameSortOrder> MergeBufferMaps(TArrayView<const TMap<const TCHAR*,FCachedUniformBufferDeclaration>*> Maps)
+{
+	int32 Num = 0;
+	for (const TMap<const TCHAR*, FCachedUniformBufferDeclaration>* Map : Maps)
+	{
+		Num += Map->Num();
+	}
+
+	// Deduplicate identical pointers (the common case)
+	using PairType = TPair<const TCHAR*, FCachedUniformBufferDeclaration>;
+	TMap<const void*, FCachedUniformBufferDeclaration> UniquePointerPairs;
+	UniquePointerPairs.Reserve(Num);
+	for (const TMap<const TCHAR*, FCachedUniformBufferDeclaration>* Map : Maps)
+	{
+		for (const PairType& Pair : *Map)
+		{
+			UniquePointerPairs.Add(Pair);
+		}
+	}
+
+	// Presort pairs so that they can be efficiently inserted into the TSortedMap
+	TArray<PairType> PresortedPairs;
+	PresortedPairs.Reserve(UniquePointerPairs.Num());
+	for (TPair<const void*, FCachedUniformBufferDeclaration>& UniquePointerPair : UniquePointerPairs)
+	{
+		PresortedPairs.Add(PairType(static_cast<const TCHAR*>(UniquePointerPair.Key), MoveTemp(UniquePointerPair.Value)));
+	}
+
+	Algo::SortBy(PresortedPairs, &PairType::Key, FUniformBufferNameSortOrder());
+
+	TSortedMap<const TCHAR*, FCachedUniformBufferDeclaration, FDefaultAllocator, FUniformBufferNameSortOrder> Out;
+	Out.Reserve(PresortedPairs.Num());
+	for (PairType& PresortedPair : PresortedPairs)
+	{
+		Out.Add(PresortedPair.Key, MoveTemp(PresortedPair.Value));
+	}
+
+	return Out;
+}
+
+
+uint32 GetTypeHash(FPlatformTypeLayoutParameters Params) { return HashCombine(Params.Flags, Params.MaxFieldAlignment); }
+
+class FFrozenMaterialLayoutHashCache
+{
+public:
+	FSHAHash Get(const FTypeLayoutDesc& TypeDesc, FPlatformTypeLayoutParameters LayoutParams)
+	{
+		{
+			FReadScopeLock ReadScope(Lock);
+
+			if (const FPlatformCache* Platform = Algo::FindBy(Platforms, LayoutParams, &FPlatformCache::Parameters))
+			{
+				if (const FSHAHash* Hash = Platform->Cache.Find(&TypeDesc))
+				{
+					return *Hash;
+				}
+			}
+		}
+
+		FSHAHash Hash = Freeze::HashLayout(TypeDesc, LayoutParams);
+
+		{
+			FWriteScopeLock WriteScope(Lock);
+
+			FPlatformCache* Platform = Algo::FindBy(Platforms, LayoutParams, &FPlatformCache::Parameters);
+			if (!Platform)
+			{
+				Platform = &Platforms.AddDefaulted_GetRef();
+				Platform->Parameters = LayoutParams;
+			}
+
+
+			Platform->Cache.FindOrAdd(&TypeDesc, Hash);
+		}
+
+		return Hash;
+	}
+
+private:
+	struct FPlatformCache
+	{
+		FPlatformTypeLayoutParameters Parameters;
+		TMap<const FTypeLayoutDesc*, FSHAHash> Cache;
+	};
+
+	FRWLock Lock;
+	TArray<FPlatformCache, TInlineAllocator<8>> Platforms;
+};
+
+static FFrozenMaterialLayoutHashCache GFrozenMaterialLayoutHashes;
+
 void FMaterialShaderMapId::AppendKeyString(FString& KeyString) const
 {
 	check(IsContentValid());
-	KeyString += BaseMaterialId.ToString();
-	KeyString += TEXT("_");
+	BaseMaterialId.AppendString(KeyString);
+	KeyString.AppendChar('_');
 
 	FString QualityLevelName;
 	GetMaterialQualityLevelName(QualityLevel, QualityLevelName);
-	KeyString += QualityLevelName + TEXT("_");
+	KeyString += QualityLevelName;
+	KeyString.AppendChar('_');
 
 	FString FeatureLevelString;
 	GetFeatureLevelName(FeatureLevel, FeatureLevelString);
-	KeyString += FeatureLevelString + TEXT("_");
+	KeyString += FeatureLevelString;
+	KeyString.AppendChar('_');
 
 	LayoutParams.AppendKeyString(KeyString);
 
@@ -944,37 +1050,37 @@ void FMaterialShaderMapId::AppendKeyString(FString& KeyString) const
 		MaterialLayersId->AppendKeyString(KeyString);
 	}
 
-	KeyString += TEXT("_");
-	KeyString += FString::FromInt(Usage);
-	KeyString += TEXT("_");
+	KeyString.AppendChar('_');
+	KeyString.AppendInt(Usage);
+	KeyString.AppendChar('_');
 
 	if (Usage == EMaterialShaderMapUsage::MaterialExportCustomOutput)
 	{
 		KeyString += UsageCustomOutput;
-		KeyString += TEXT("_");
+		KeyString.AppendChar('_');
 	}
 
 	// Add any referenced functions to the key so that we will recompile when they are changed
 	for (int32 FunctionIndex = 0; FunctionIndex < ReferencedFunctions.Num(); FunctionIndex++)
 	{
-		KeyString += ReferencedFunctions[FunctionIndex].ToString();
+		ReferencedFunctions[FunctionIndex].AppendString(KeyString);
 	}
 
 	{
-		const FSHAHash LayoutHash = Freeze::HashLayout(StaticGetTypeLayoutDesc<FMaterialShaderMapContent>(), LayoutParams);
-		KeyString += TEXT("_");
-		KeyString += LayoutHash.ToString();
-		KeyString += TEXT("_");
+		const FSHAHash LayoutHash = GFrozenMaterialLayoutHashes.Get(StaticGetTypeLayoutDesc<FMaterialShaderMapContent>(), LayoutParams);
+		KeyString.AppendChar('_');
+		LayoutHash.AppendString(KeyString);
+		KeyString.AppendChar('_');
 	}
 
-	KeyString += TEXT("_");
+	KeyString.AppendChar('_');
 
 	for (int32 CollectionIndex = 0; CollectionIndex < ReferencedParameterCollections.Num(); CollectionIndex++)
 	{
-		KeyString += ReferencedParameterCollections[CollectionIndex].ToString();
+		ReferencedParameterCollections[CollectionIndex].AppendString(KeyString);
 	}
 
-	TSortedMap<const TCHAR*, FCachedUniformBufferDeclaration, FDefaultAllocator, FUniformBufferNameSortOrder> ReferencedUniformBuffers;
+	TArray<const TMap<const TCHAR*,FCachedUniformBufferDeclaration>*> ReferencedUniformBufferMaps;
 
 	// Add the inputs for any shaders that are stored inline in the shader map
 	for (int32 ShaderIndex = 0; ShaderIndex < ShaderTypeDependencies.Num(); ShaderIndex++)
@@ -985,24 +1091,19 @@ void FMaterialShaderMapId::AppendKeyString(FString& KeyString) const
 		checkf(ShaderType != nullptr, TEXT("Failed to find FShaderType for dependency %s (total in the NameToTypeMap: %d)"), ShaderTypeDependency.ShaderTypeName.GetDebugString().String.Get(), FShaderType::GetNameToTypeMap().Num());
 #endif
 
-		KeyString += TEXT("_");
+		KeyString.AppendChar('_');
 		KeyString += ShaderType->GetName();
-		KeyString += ShaderTypeDependency.SourceHash.ToString();
+		ShaderTypeDependency.SourceHash.AppendString(KeyString);
 
 		if (const FShaderParametersMetadata* ParameterStructMetadata = ShaderType->GetRootParametersMetadata())
 		{
-			KeyString += FString::Printf(TEXT("%08x"), ParameterStructMetadata->GetLayoutHash());
+			KeyString.Appendf(TEXT("%08x"), ParameterStructMetadata->GetLayoutHash());
 		}
 
-		const FSHAHash LayoutHash = Freeze::HashLayout(ShaderType->GetLayout(), LayoutParams);
-		KeyString += LayoutHash.ToString();
+		const FSHAHash LayoutHash = GFrozenMaterialLayoutHashes.Get(ShaderType->GetLayout(), LayoutParams);
+		LayoutHash.AppendString(KeyString);
 
-		const TMap<const TCHAR*,FCachedUniformBufferDeclaration>& ReferencedUniformBufferStructsCache = ShaderType->GetReferencedUniformBufferStructsCache();
-
-		for (TMap<const TCHAR*,FCachedUniformBufferDeclaration>::TConstIterator It(ReferencedUniformBufferStructsCache); It; ++It)
-		{
-			ReferencedUniformBuffers.Add(It.Key(), It.Value());
-		}
+		ReferencedUniformBufferMaps.Add(&ShaderType->GetReferencedUniformBufferStructsCache());
 	}
 
 	// Add the inputs for any shader pipelines that are stored inline in the shader map
@@ -1014,54 +1115,43 @@ void FMaterialShaderMapId::AppendKeyString(FString& KeyString) const
 		checkf(ShaderPipelineType != nullptr, TEXT("Failed to find FShaderPipelineType for dependency %s (total in the NameToTypeMap: %d)"), Dependency.ShaderPipelineTypeName.GetDebugString().String.Get(), FShaderType::GetNameToTypeMap().Num());
 #endif
 
-		KeyString += TEXT("_");
+		KeyString.AppendChar('_');
 		KeyString += ShaderPipelineType->GetName();
-		KeyString += Dependency.StagesSourceHash.ToString();
+		Dependency.StagesSourceHash.AppendString(KeyString);
 
 		for (const FShaderType* ShaderType : ShaderPipelineType->GetStages())
 		{
 			if (const FShaderParametersMetadata* ParameterStructMetadata = ShaderType->GetRootParametersMetadata())
 			{
-				KeyString += FString::Printf(TEXT("%08x"), ParameterStructMetadata->GetLayoutHash());
+				KeyString.Appendf(TEXT("%08x"), ParameterStructMetadata->GetLayoutHash());
 			}
 
-			const TMap<const TCHAR*, FCachedUniformBufferDeclaration>& ReferencedUniformBufferStructsCache = ShaderType->GetReferencedUniformBufferStructsCache();
-
-			// Gather referenced uniform buffers
-			for (TMap<const TCHAR*, FCachedUniformBufferDeclaration>::TConstIterator It(ReferencedUniformBufferStructsCache); It; ++It)
-			{
-				ReferencedUniformBuffers.Add(It.Key(), It.Value());
-			}
+			ReferencedUniformBufferMaps.Add(&ShaderType->GetReferencedUniformBufferStructsCache());
 		}
 	}
 
 	// Add the inputs for any shaders that are stored inline in the shader map
 	for (int32 VFIndex = 0; VFIndex < VertexFactoryTypeDependencies.Num(); VFIndex++)
 	{
-		KeyString += TEXT("_");
+		KeyString.AppendChar('_');
 
 		const FVertexFactoryTypeDependency& VFDependency = VertexFactoryTypeDependencies[VFIndex];
 		const FVertexFactoryType* VertexFactoryType = FVertexFactoryType::GetVFByName(VFDependency.VertexFactoryTypeName);
 
 		KeyString += VertexFactoryType->GetName();
-		KeyString += VFDependency.VFSourceHash.ToString();
+		VFDependency.VFSourceHash.AppendString(KeyString);
 
 		for (int32 Frequency = 0; Frequency < SF_NumFrequencies; Frequency++)
 		{
 			const FTypeLayoutDesc* ParameterLayout = VertexFactoryType->GetShaderParameterLayout((EShaderFrequency)Frequency);
 			if (ParameterLayout)
 			{
-				const FSHAHash LayoutHash = Freeze::HashLayout(*ParameterLayout, LayoutParams);
-				KeyString += LayoutHash.ToString();
+				const FSHAHash LayoutHash = GFrozenMaterialLayoutHashes.Get(*ParameterLayout, LayoutParams);
+				LayoutHash.AppendString(KeyString);
 			}
 		}
 
-		const TMap<const TCHAR*,FCachedUniformBufferDeclaration>& ReferencedUniformBufferStructsCache = VertexFactoryType->GetReferencedUniformBufferStructsCache();
-
-		for (TMap<const TCHAR*,FCachedUniformBufferDeclaration>::TConstIterator It(ReferencedUniformBufferStructsCache); It; ++It)
-		{
-			ReferencedUniformBuffers.Add(It.Key(), It.Value());
-		}
+		ReferencedUniformBufferMaps.Add(&VertexFactoryType->GetReferencedUniformBufferStructsCache());
 	}
 
 	{
@@ -1069,6 +1159,7 @@ void FMaterialShaderMapId::AppendKeyString(FString& KeyString) const
 		FSerializationHistory SerializationHistory;
 		FMemoryWriter Ar(TempData, true);
 		FShaderSaveArchive SaveArchive(Ar, SerializationHistory);
+		TSortedMap<const TCHAR*, FCachedUniformBufferDeclaration, FDefaultAllocator, FUniformBufferNameSortOrder> ReferencedUniformBuffers = MergeBufferMaps(ReferencedUniformBufferMaps);
 
 		// Save uniform buffer member info so we can detect when layout has changed
 		SerializeUniformBufferInfo(SaveArchive, ReferencedUniformBuffers);
@@ -1076,9 +1167,9 @@ void FMaterialShaderMapId::AppendKeyString(FString& KeyString) const
 		SerializationHistory.AppendKeyString(KeyString);
 	}
 
-	KeyString += BytesToHex(&TextureReferencesHash.Hash[0], sizeof(TextureReferencesHash.Hash));
+	BytesToHex(&TextureReferencesHash.Hash[0], sizeof(TextureReferencesHash.Hash), KeyString);
 
-	KeyString += BytesToHex(&BasePropertyOverridesHash.Hash[0], sizeof(BasePropertyOverridesHash.Hash));
+	BytesToHex(&BasePropertyOverridesHash.Hash[0], sizeof(BasePropertyOverridesHash.Hash), KeyString);
 
 	if (bUsingNewHLSLGenerator)
 	{
diff --git a/Engine/Source/Runtime/Engine/Public/MaterialTypes.h b/Engine/Source/Runtime/Engine/Public/MaterialTypes.h
--- a/Engine/Source/Runtime/Engine/Public/MaterialTypes.h
+++ b/Engine/Source/Runtime/Engine/Public/MaterialTypes.h
@@ -55,9 +55,17 @@ struct ENGINE_API FMaterialParameterInfo
 
 	explicit FMaterialParameterInfo(const struct FMemoryImageMaterialParameterInfo& Rhs);
 
+	void AppendString(FString& Out) const
+	{
+		Name.AppendString(Out);
+		Out.AppendInt(Association);
+		Out.AppendInt(Index);
+	}
 	FString ToString() const
 	{
-		return *Name.ToString() + FString::FromInt(Association) + FString::FromInt(Index);
+		FString Out;
+		AppendString(Out);
+		return Out;
 	}
 
 	friend FArchive& operator<<(FArchive& Ar, FMaterialParameterInfo& Ref)
diff --git a/Engine/Source/Runtime/Engine/Public/StaticParameterSet.h b/Engine/Source/Runtime/Engine/Public/StaticParameterSet.h
--- a/Engine/Source/Runtime/Engine/Public/StaticParameterSet.h
+++ b/Engine/Source/Runtime/Engine/Public/StaticParameterSet.h
@@ -77,7 +77,9 @@ struct FStaticParameterBase
 
 	void AppendKeyString(FString& KeyString) const
 	{
-		KeyString += ParameterInfo.ToString() + FString::FromInt(bOverride) + ExpressionGUID.ToString();
+		ParameterInfo.AppendString(KeyString);
+		KeyString.AppendInt(bOverride);
+		ExpressionGUID.AppendString(KeyString);
 	}
 };
 
@@ -133,7 +135,7 @@ struct FStaticSwitchParameter : public FStaticParameterBase
 	void AppendKeyString(FString& KeyString) const
 	{
 		FStaticParameterBase::AppendKeyString(KeyString);
-		KeyString += FString::FromInt(Value);
+		KeyString.AppendInt(Value);
 	}
 
 	void GetValue(FMaterialParameterMetadata& OutResult) const
diff --git a/Engine/Source/Runtime/RenderCore/Private/ShaderCore.cpp b/Engine/Source/Runtime/RenderCore/Private/ShaderCore.cpp
--- a/Engine/Source/Runtime/RenderCore/Private/ShaderCore.cpp
+++ b/Engine/Source/Runtime/RenderCore/Private/ShaderCore.cpp
@@ -11,6 +11,7 @@
 #include "Misc/PathViews.h"
 #include "Misc/ScopeLock.h"
 #include "Stats/StatsMisc.h"
+#include "Misc/ScopeRWLock.h"
 #include "Misc/ScopedSlowTask.h"
 #include "Misc/StringBuilder.h"
 #include "Shader.h"
@@ -1488,33 +1489,82 @@ void GenerateReferencedUniformBuffers(
 	}
 }
 
-void SerializeUniformBufferInfo(FShaderSaveArchive& Ar, const TSortedMap<const TCHAR*, FCachedUniformBufferDeclaration, FDefaultAllocator, FUniformBufferNameSortOrder>& UniformBufferEntries)
+
+/** Efficient lookup to find FShaderParametersMetadata members by name pointer */
+class FShaderParameterMemberLookup 
 {
-	for (TSortedMap<const TCHAR*, FCachedUniformBufferDeclaration, FDefaultAllocator, FUniformBufferNameSortOrder>::TConstIterator It(UniformBufferEntries); It; ++It)
+public:
+	FShaderParameterMemberLookup(TLinkedList<FShaderParametersMetadata*>& ShaderParameters)
 	{
-		for (TLinkedList<FShaderParametersMetadata*>::TIterator StructIt(FShaderParametersMetadata::GetStructList()); StructIt; StructIt.Next())
+		for (FShaderParametersMetadata* Struct : ShaderParameters)
+		{
+			Map.Add(Struct->GetShaderVariableName(), Struct->GetMembers());
+		}
+	}
+
+	const TConstArrayView<FShaderParametersMetadata::FMember>* FindMembersByPointer(const TCHAR* ShaderVariableName) const
+	{
+		return Map.Find(ShaderVariableName);
+	}
+
+private:
+	TMap<const void*, TConstArrayView<FShaderParametersMetadata::FMember>> Map;
+};
+
+/** Cache providing a FShaderParameterMemberLookup for the current FShaderParametersMetadata::GetStructList */
+class FFShaderParameterPointerLookupCache
+{
+public:
+	TSharedPtr<const FShaderParameterMemberLookup> Get()
+	{
+		TLinkedList<FShaderParametersMetadata*>* CurrentHead = FShaderParametersMetadata::GetStructList();
+		check(CurrentHead);
+
 		{
-			if (It.Key() == StructIt->GetShaderVariableName())
+			FReadScopeLock ReadScope(Lock);
+			if (CurrentHead == CachedHead)
 			{
-				// Serialize information about the struct layout so we can detect when it changes
-				const FShaderParametersMetadata& Struct = **StructIt;
-				const TArray<FShaderParametersMetadata::FMember>& Members = Struct.GetMembers();
+				return CachedLookup;
+			}
+		}
 
-				int32 NumMembers = Members.Num();
-				// Serializing with NULL so that FShaderSaveArchive will record the length without causing an actual data serialization
-				Ar.Serialize(NULL, NumMembers);
+		TSharedPtr<FShaderParameterMemberLookup> NewLookup = MakeShared<FShaderParameterMemberLookup>(*CurrentHead);
 
-				for (int32 MemberIndex = 0; MemberIndex < Members.Num(); MemberIndex++)
-				{
-					const FShaderParametersMetadata::FMember& Member = Members[MemberIndex];
+		FWriteScopeLock WriteScope(Lock);
+		CachedHead = CurrentHead;
+		CachedLookup = NewLookup;
 
-					// Note: Only comparing number of floats used by each member and type, so this can be tricked (eg. swapping two equal size and type members)
-					int32 MemberSize = Member.GetNumColumns() * Member.GetNumRows();
-					Ar.Serialize(NULL, MemberSize);
-					int32 MemberType = (int32)Member.GetBaseType();
-					Ar.Serialize(NULL, MemberType);
-				}
-				break;
+		return NewLookup;
+	}
+
+private:
+	FRWLock Lock;
+	TLinkedList<FShaderParametersMetadata*>* CachedHead = nullptr;
+	TSharedPtr<const FShaderParameterMemberLookup> CachedLookup;
+};
+
+static FFShaderParameterPointerLookupCache GShaderParameterMemberLookupCache;
+
+void SerializeUniformBufferInfo(FShaderSaveArchive& Ar, const TSortedMap<const TCHAR*, FCachedUniformBufferDeclaration, FDefaultAllocator, FUniformBufferNameSortOrder>& UniformBufferEntries)
+{
+	TSharedPtr<const FShaderParameterMemberLookup> ShaderParameterMembers = GShaderParameterMemberLookupCache.Get();
+
+	for (const TPair<const TCHAR*, FCachedUniformBufferDeclaration>& Entry : UniformBufferEntries)
+	{
+		if (const TConstArrayView<FShaderParametersMetadata::FMember>* Members = ShaderParameterMembers->FindMembersByPointer(Entry.Key))
+		{
+			// Serialize information about the struct layout so we can detect when it changes
+			int32 NumMembers = Members->Num();
+			// Serializing with NULL so that FShaderSaveArchive will record the length without causing an actual data serialization
+			Ar.Serialize(nullptr, NumMembers);
+
+			for (const FShaderParametersMetadata::FMember& Member : *Members)
+			{
+				// Note: Only comparing number of floats used by each member and type, so this can be tricked (eg. swapping two equal size and type members)
+				int32 MemberSize = Member.GetNumColumns() * Member.GetNumRows();
+				Ar.Serialize(nullptr, MemberSize);
+				int32 MemberType = (int32)Member.GetBaseType();
+				Ar.Serialize(nullptr, MemberType);
 			}
 		}
 	}
diff --git a/Engine/Source/Runtime/RenderCore/Public/Shader.h b/Engine/Source/Runtime/RenderCore/Public/Shader.h
--- a/Engine/Source/Runtime/RenderCore/Public/Shader.h
+++ b/Engine/Source/Runtime/RenderCore/Public/Shader.h
@@ -427,7 +427,7 @@ public:
 		NumTokens(0)
 	{}
 
-	void AddValue(uint32 InValue)
+	FORCEINLINE void AddValue(uint32 InValue)
 	{
 		const int32 UIntIndex = NumTokens / 8; 
 
@@ -457,11 +457,11 @@ public:
 		NumTokens++;
 	}
 
-	uint8 GetToken(int32 Index) const
+	FORCEINLINE uint8 GetToken(int32 Index) const
 	{
-		check(Index < NumTokens);
+		checkSlow(Index < NumTokens);
 		const int32 UIntIndex = Index / 8; 
-		check(UIntIndex < TokenBits.Num());
+		checkSlow(UIntIndex < TokenBits.Num());
 		const uint32 Shift = (Index % 8) * 4;
 		const uint8 Token = (TokenBits[UIntIndex] >> Shift) & 0xF;
 		return Token;
@@ -2202,18 +2202,18 @@ public:
 		return Value;
 	}
 
-	void StepForward()
+	FORCEINLINE void StepForward()
 	{
 		const int8 Token = History.GetToken(NextTokenIndex);
 
 		if (Token == 0)
 		{
-			check(NextFullLengthIndex - 1 < History.FullLengths.Num());
+			checkSlow(NextFullLengthIndex - 1 < History.FullLengths.Num());
 			NextFullLengthIndex++;
 		}
 
 		// Not supporting seeking past the front most serialization in the history
-		check(NextTokenIndex - 1 < History.NumTokens);
+		checkSlow(NextTokenIndex - 1 < History.NumTokens);
 		NextTokenIndex++;
 	}
 
@@ -2234,7 +2234,7 @@ public:
 };
 
 /** Archive used when saving shaders, which generates data used to detect serialization mismatches on load. */
-class FShaderSaveArchive : public FArchiveProxy
+class FShaderSaveArchive final : public FArchiveProxy
 {
 public:
 
diff --git a/Engine/Source/Runtime/RenderCore/Public/ShaderCore.h b/Engine/Source/Runtime/RenderCore/Public/ShaderCore.h
--- a/Engine/Source/Runtime/RenderCore/Public/ShaderCore.h
+++ b/Engine/Source/Runtime/RenderCore/Public/ShaderCore.h
@@ -1052,9 +1052,9 @@ extern void GenerateReferencedUniformBuffers(
 
 struct FUniformBufferNameSortOrder
 {
-	FORCEINLINE bool operator()(const TCHAR* Name1, const TCHAR* Name2)
+	FORCEINLINE bool operator()(const TCHAR* Name1, const TCHAR* Name2) const
 	{
-		return FCString::Strcmp(Name1, Name2) <= 0;
+		return FCString::Strcmp(Name1, Name2) < 0;
 	}
 };
 
