Use a conservative vertex count to initialize ray tracing geometries with multiple segments.

- Segment.MaxVertices value must be at least as large as the maximum value in the index buffer + 1.
Previously this was initialized to per-segment vertex count, which can be too small.

diff --git a/Engine/Source/Runtime/Engine/Private/SkeletalMesh.cpp b/Engine/Source/Runtime/Engine/Private/SkeletalMesh.cpp
--- a/Engine/Source/Runtime/Engine/Private/SkeletalMesh.cpp
+++ b/Engine/Source/Runtime/Engine/Private/SkeletalMesh.cpp
@@ -6800,6 +6800,8 @@ void FSkeletalMeshSceneProxy::GetDynamicRayTracingInstances(FRayTracingMaterialG
 			MeshObject->MaterialIndexPreview = INDEX_NONE;
 		#endif
 
+			uint32 TotalNumVertices = 0;
+
 			for (FSkeletalMeshSectionIter Iter(LODIndex, *MeshObject, LODData, LODSection); Iter; ++Iter)
 			{
 				const FSkelMeshRenderSection& Section = Iter.GetSection();
@@ -6810,6 +6812,8 @@ void FSkeletalMeshSceneProxy::GetDynamicRayTracingInstances(FRayTracingMaterialG
 				CreateBaseMeshBatch(Context.ReferenceView, LODData, LODIndex, SectionIndex, SectionElementInfo, MeshBatch, ESkinVertexFactoryMode::RayTracing);
 
 				RayTracingInstance.Materials.Add(MeshBatch);
+
+				TotalNumVertices += Section.GetNumVertices();
 			}
 
 			RayTracingInstance.InstanceTransforms.Add(GetLocalToWorld());
@@ -6832,7 +6836,7 @@ void FSkeletalMeshSceneProxy::GetDynamicRayTracingInstances(FRayTracingMaterialG
 
 					FRayTracingGeometrySegment Segment;
 					Segment.VertexBufferStride = VertexBufferStride;
-					Segment.MaxVertices = Section.GetNumVertices();
+					Segment.MaxVertices = TotalNumVertices;
 					Segment.FirstPrimitive = Section.BaseIndex / 3;
 					Segment.NumPrimitives = Section.NumTriangles;
 					Segment.bEnabled = !MeshObject->IsMaterialHidden(LODIndex, SectionElementInfo.UseMaterialIndex) && !Section.bDisabled && Section.bVisibleInRayTracing;
diff --git a/Engine/Source/Runtime/Engine/Private/SkeletalRenderCPUSkin.cpp b/Engine/Source/Runtime/Engine/Private/SkeletalRenderCPUSkin.cpp
--- a/Engine/Source/Runtime/Engine/Private/SkeletalRenderCPUSkin.cpp
+++ b/Engine/Source/Runtime/Engine/Private/SkeletalRenderCPUSkin.cpp
@@ -423,6 +423,13 @@ void FSkeletalMeshObjectCPUSkin::FSkeletalMeshObjectLOD::InitResources(FSkelMesh
 
 				TArray<FRayTracingGeometrySegment> GeometrySections;
 				GeometrySections.Reserve(RenderSections->Num());
+
+				uint32 TotalNumVertices = 0;
+				for (const FSkelMeshRenderSection& Section : *RenderSections)
+				{
+					TotalNumVertices += Section.GetNumVertices();
+				}
+
 				for (const FSkelMeshRenderSection& Section : *RenderSections)
 				{
 					FRayTracingGeometrySegment Segment;
@@ -430,7 +437,7 @@ void FSkeletalMeshObjectCPUSkin::FSkeletalMeshObjectLOD::InitResources(FSkelMesh
 					Segment.VertexBufferStride = VertexBufferStride;
 					Segment.VertexBufferOffset = 0;
 					Segment.VertexBufferElementType = VET_Float3;
-					Segment.MaxVertices = Section.GetNumVertices();
+					Segment.MaxVertices = TotalNumVertices;
 					Segment.FirstPrimitive = Section.BaseIndex / 3;
 					Segment.NumPrimitives = Section.NumTriangles;
 					Segment.bEnabled = !Section.bDisabled && Section.bVisibleInRayTracing;
diff --git a/Engine/Source/Runtime/Engine/Private/SkeletalRenderGPUSkin.cpp b/Engine/Source/Runtime/Engine/Private/SkeletalRenderGPUSkin.cpp
--- a/Engine/Source/Runtime/Engine/Private/SkeletalRenderGPUSkin.cpp
+++ b/Engine/Source/Runtime/Engine/Private/SkeletalRenderGPUSkin.cpp
@@ -811,11 +811,13 @@ void FSkeletalMeshObjectGPUSkin::UpdateRayTracingGeometry(FSkeletalMeshLODRender
 			MemoryEstimation += LODModel.StaticVertexBuffers.PositionVertexBuffer.VertexBufferRHI->GetSize();
 
 			//#dxr_todo: do we need support for separate sections in FRayTracingGeometryData?
-			uint32 TrianglesCount = 0;
+			uint32 TotalNumTriangles = 0;
+			uint32 TotalNumVertices = 0;
 			for (int32 SectionIndex = 0; SectionIndex < LODModel.RenderSections.Num(); SectionIndex++)
 			{
 				const FSkelMeshRenderSection& Section = LODModel.RenderSections[SectionIndex];
-				TrianglesCount += Section.NumTriangles;
+				TotalNumTriangles += Section.NumTriangles;
+				TotalNumVertices += Section.GetNumVertices();
 			}
 
 			FRayTracingGeometryInitializer Initializer;
@@ -834,7 +836,7 @@ void FSkeletalMeshObjectGPUSkin::UpdateRayTracingGeometry(FSkeletalMeshLODRender
 			}
 
 			Initializer.IndexBuffer = IndexBufferRHI;
-			Initializer.TotalPrimitiveCount = TrianglesCount;
+			Initializer.TotalPrimitiveCount = TotalNumTriangles;
 			Initializer.GeometryType = RTGT_Triangles;
 			Initializer.bFastBuild = true;
 			Initializer.bAllowUpdate = true;
@@ -850,7 +852,7 @@ void FSkeletalMeshObjectGPUSkin::UpdateRayTracingGeometry(FSkeletalMeshLODRender
 				Segment.VertexBufferElementType = VET_Float3;
 				Segment.VertexBufferStride = VertexBufferStride;
 				Segment.VertexBufferOffset = 0;
-				Segment.MaxVertices = Section.GetNumVertices();
+				Segment.MaxVertices = TotalNumVertices;
 				Segment.FirstPrimitive = Section.BaseIndex / 3;
 				Segment.NumPrimitives = Section.NumTriangles;
 
diff --git a/Engine/Source/Runtime/Engine/Private/SkeletalRenderStatic.cpp b/Engine/Source/Runtime/Engine/Private/SkeletalRenderStatic.cpp
--- a/Engine/Source/Runtime/Engine/Private/SkeletalRenderStatic.cpp
+++ b/Engine/Source/Runtime/Engine/Private/SkeletalRenderStatic.cpp
@@ -83,6 +83,13 @@ void FSkeletalMeshObjectStatic::InitResources(USkinnedMeshComponent* InMeshCompo
 
 							TArray<FRayTracingGeometrySegment> GeometrySections;
 							GeometrySections.Reserve(RenderSections->Num());
+
+							uint32 TotalNumVertices = 0;
+							for (const FSkelMeshRenderSection& Section : *RenderSections)
+							{
+								TotalNumVertices += Section.GetNumVertices();
+							}
+
 							for (const FSkelMeshRenderSection& Section : *RenderSections)
 							{
 								FRayTracingGeometrySegment Segment;
@@ -90,7 +97,7 @@ void FSkeletalMeshObjectStatic::InitResources(USkinnedMeshComponent* InMeshCompo
 								Segment.VertexBufferElementType = VET_Float3;
 								Segment.VertexBufferOffset = 0;
 								Segment.VertexBufferStride = VertexBufferStride;
-								Segment.MaxVertices = Section.GetNumVertices();
+								Segment.MaxVertices = TotalNumVertices;
 								Segment.FirstPrimitive = Section.BaseIndex / 3;
 								Segment.NumPrimitives = Section.NumTriangles;
 								Segment.bEnabled = !Section.bDisabled && Section.bVisibleInRayTracing;
