Level Instance: implement FActorDescContainerInstanceManager on top of FActorDescContainerInstance, to reduce the amount of FGCObject to a single one.

#rnx

diff --git a/Engine/Source/Runtime/Engine/Private/WorldPartition/LevelInstance/LevelInstanceActorDesc.cpp b/Engine/Source/Runtime/Engine/Private/WorldPartition/LevelInstance/LevelInstanceActorDesc.cpp
--- a/Engine/Source/Runtime/Engine/Private/WorldPartition/LevelInstance/LevelInstanceActorDesc.cpp
+++ b/Engine/Source/Runtime/Engine/Private/WorldPartition/LevelInstance/LevelInstanceActorDesc.cpp
@@ -18,28 +18,79 @@ static FAutoConsoleVariableRef CVarForceLevelStreaming(
 	GLevelInstanceDebugForceLevelStreaming,
 	TEXT("Set to 1 to force Level Instance to be streamed instead of embedded in World Partition grid."));
 
-struct FActorDescContainerInstance : public FGCObject
+class FActorDescContainerInstanceManager : public FGCObject
 {
-	FActorDescContainerInstance()
-	: Container(nullptr)
-	, RefCount(0)
-	{}
+	struct FActorDescContainerInstance
+	{
+		FActorDescContainerInstance()
+		: Container(nullptr)
+		, RefCount(0)
+		{}
+
+		void AddReferencedObjects(FReferenceCollector& Collector)
+		{
+			Collector.AddReferencedObject(Container);
+		}
 
+		UActorDescContainer* Container;
+		uint32 RefCount;
+	};
+
+	//~ Begin FGCObject Interface
 	virtual void AddReferencedObjects(FReferenceCollector& Collector) override
 	{
-		if (Container)
+		for (auto& It : ActorDescContainers)
 		{
-			Collector.AddReferencedObject(Container);
+			It.Value.AddReferencedObjects(Collector);
 		}
 	}
 
-	virtual FString GetReferencerName() const override { return TEXT("FActorDescContainerInstance"); }
+	virtual FString GetReferencerName() const override
+	{
+		return TEXT("FActorDescContainerInstanceManager");
+	}
+	//~ End FGCObject Interface
 
-	UActorDescContainer* Container;
-	uint32 RefCount;
-};
+public:
+	static FActorDescContainerInstanceManager& Get()
+	{
+		static FActorDescContainerInstanceManager Instance;
+		return Instance;
+	}
+
+	UActorDescContainer* RegisterContainer(FName PackageName, UWorld* InWorld)
+	{
+		FActorDescContainerInstance& ExistingContainerInstance = ActorDescContainers.FindOrAdd(PackageName);
+		UActorDescContainer* ActorDescContainer = ExistingContainerInstance.Container;
+	
+		if (ExistingContainerInstance.RefCount++ == 0)
+		{
+			ActorDescContainer = NewObject<UActorDescContainer>(GetTransientPackage());
+			ExistingContainerInstance.Container = ActorDescContainer;
+		
+			// This will potentially invalidate ExistingContainerInstance due to ActorDescContainers reallocation
+			ActorDescContainer->Initialize(InWorld, PackageName);
+		}
+
+		check(ActorDescContainer->GetWorld() == InWorld);
+		return ActorDescContainer;
+	}
+
+	void UnregisterContainer(UActorDescContainer* Container)
+	{
+		FName PackageName = Container->GetContainerPackage();
+		FActorDescContainerInstance& ExistingContainerInstance = ActorDescContainers.FindChecked(PackageName);
+
+		if (ExistingContainerInstance.RefCount-- == 1)
+		{
+			ExistingContainerInstance.Container->Uninitialize();
+			ActorDescContainers.FindAndRemoveChecked(PackageName);
+		}
+	}
 
-static TMap<FName, TUniquePtr<FActorDescContainerInstance>> ActorDescContainers;
+private:
+	TMap<FName, FActorDescContainerInstance> ActorDescContainers;
+};
 
 FLevelInstanceActorDesc::FLevelInstanceActorDesc()
 	: DesiredRuntimeBehavior(ELevelInstanceRuntimeBehavior::Partitioned)
@@ -94,7 +145,7 @@ void FLevelInstanceActorDesc::RegisterContainerInstance(UWorld* InWorld)
 	{
 		if (!LevelPackage.IsNone() && ULevel::GetIsLevelUsingExternalActorsFromPackage(LevelPackage) && ULevelInstanceSubsystem::CanUsePackage(LevelPackage))
 		{
-			LevelInstanceContainer = FLevelInstanceActorDesc::RegisterActorDescContainer(LevelPackage, InWorld);
+			LevelInstanceContainer = FActorDescContainerInstanceManager::Get().RegisterContainer(LevelPackage, InWorld);
 			check(LevelInstanceContainer.IsValid());
 		}
 	}
@@ -104,7 +155,7 @@ void FLevelInstanceActorDesc::UnregisterContainerInstance()
 {
 	if (LevelInstanceContainer.IsValid())
 	{
-		FLevelInstanceActorDesc::UnregisterActorDescContainer(LevelInstanceContainer.Get());
+		FActorDescContainerInstanceManager::Get().UnregisterContainer(LevelInstanceContainer.Get());
 		LevelInstanceContainer.Reset();
 	}
 }
@@ -197,39 +248,4 @@ void FLevelInstanceActorDesc::Serialize(FArchive& Ar)
 	}
 }
 
-UActorDescContainer* FLevelInstanceActorDesc::RegisterActorDescContainer(FName PackageName, UWorld* InWorld)
-{
-	TUniquePtr<FActorDescContainerInstance>& ExistingContainerInstance = ActorDescContainers.FindOrAdd(PackageName);
-	UActorDescContainer* ActorDescContainer = nullptr;
-	if (!ExistingContainerInstance)
-	{
-		ExistingContainerInstance.Reset(new FActorDescContainerInstance());
-		ActorDescContainer = NewObject<UActorDescContainer>(GetTransientPackage());
-		ExistingContainerInstance->Container = ActorDescContainer;
-		ExistingContainerInstance->RefCount++;
-		
-		// This will potentially invalidate ExistingContainerInstance due to ActorDescContainers reallocation
-		ActorDescContainer->Initialize(InWorld, PackageName);
-	}
-	else
-	{
-		ExistingContainerInstance->RefCount++;
-		ActorDescContainer = ExistingContainerInstance->Container;
-	}
-
-	check(ActorDescContainer->GetWorld() == InWorld);
-	return ActorDescContainer;
-}
-
-void FLevelInstanceActorDesc::UnregisterActorDescContainer(UActorDescContainer* Container)
-{
-	FName PackageName = Container->GetContainerPackage();
-	TUniquePtr<FActorDescContainerInstance>& ExistingContainerInstance = ActorDescContainers.FindChecked(PackageName);
-
-	if (ExistingContainerInstance->RefCount-- == 1)
-	{
-		ExistingContainerInstance->Container->Uninitialize();
-		ActorDescContainers.FindAndRemoveChecked(PackageName);
-	}
-}
 #endif
\ No newline at end of file
diff --git a/Engine/Source/Runtime/Engine/Public/WorldPartition/LevelInstance/LevelInstanceActorDesc.h b/Engine/Source/Runtime/Engine/Public/WorldPartition/LevelInstance/LevelInstanceActorDesc.h
--- a/Engine/Source/Runtime/Engine/Public/WorldPartition/LevelInstance/LevelInstanceActorDesc.h
+++ b/Engine/Source/Runtime/Engine/Public/WorldPartition/LevelInstance/LevelInstanceActorDesc.h
@@ -46,8 +46,5 @@ protected:
 private:
 	void RegisterContainerInstance(UWorld* InWorld);
 	void UnregisterContainerInstance();
-
-	static UActorDescContainer* RegisterActorDescContainer(FName PackageName, UWorld* InWorld);
-	static void UnregisterActorDescContainer(UActorDescContainer* Container);
 #endif
 };
