Adjusting log verbosity

diff --git a/Engine/Source/Runtime/Engine/Public/ReplayHelper.h b/Engine/Source/Runtime/Engine/Public/ReplayHelper.h
--- a/Engine/Source/Runtime/Engine/Public/ReplayHelper.h
+++ b/Engine/Source/Runtime/Engine/Public/ReplayHelper.h
@@ -293,7 +293,7 @@ private:
 				break;
 			}
 
-			UE_LOG(LogDemo, Log, TEXT("Checkpoint. %s: %i/%i, took %.2fms (Total this frame: %.2fms)"), StateStr, *CurrentIndex, TotalCount, StepTimeInMS, TotalTimeInMS);
+			UE_LOG(LogDemo, Verbose, TEXT("Checkpoint. %s: %i/%i, took %.2fms (Total this frame: %.2fms)"), StateStr, *CurrentIndex, TotalCount, StepTimeInMS, TotalTimeInMS);
 		}
 
 	private:
