fix for cvar read triggering render thread ensure

diff --git a/Engine/Source/Runtime/Engine/Private/PrimitiveUniformShaderParameters.cpp b/Engine/Source/Runtime/Engine/Private/PrimitiveUniformShaderParameters.cpp
--- a/Engine/Source/Runtime/Engine/Private/PrimitiveUniformShaderParameters.cpp
+++ b/Engine/Source/Runtime/Engine/Private/PrimitiveUniformShaderParameters.cpp
@@ -130,7 +130,7 @@ FPrimitiveSceneShaderData::FPrimitiveSceneShaderData(const FPrimitiveSceneProxy*
 	FBoxSphereBounds PreSkinnedLocalBounds;
 	Proxy->GetPreSkinnedLocalBounds(PreSkinnedLocalBounds);
 
-	const bool OptimizedWPO = CVarOptimizedWPO.GetValueOnRenderThread();
+	const bool OptimizedWPO = CVarOptimizedWPO.GetValueOnAnyThread();
 
 	FPrimitiveSceneInfo* PrimitiveSceneInfo = Proxy->GetPrimitiveSceneInfo();
 
