Strata - Fix for some cases where custom output nodes would not be taken into account when using material attribute workflow.

diff --git a/Engine/Source/Runtime/Engine/Private/Materials/Material.cpp b/Engine/Source/Runtime/Engine/Private/Materials/Material.cpp
--- a/Engine/Source/Runtime/Engine/Private/Materials/Material.cpp
+++ b/Engine/Source/Runtime/Engine/Private/Materials/Material.cpp
@@ -2836,19 +2836,68 @@ void UMaterial::ConvertMaterialToStrataMaterial()
 		}
 	};
 
+	bool bCustomNodesGathered = false;
+	UMaterialExpressionThinTranslucentMaterialOutput* ThinTranslucentOutput = nullptr;
+	UMaterialExpressionSingleLayerWaterMaterialOutput* SingleLayerWaterOutput = nullptr;
+	UMaterialExpressionClearCoatNormalCustomOutput* ClearCoatOutput = nullptr;
+	UMaterialExpressionTangentOutput* TangentOutput = nullptr;
+	auto GatherCustomNodes = [&]()
+	{
+		if (!bCustomNodesGathered)
+		{
+			bCustomNodesGathered = true;
+
+			TArray<class UMaterialExpressionCustomOutput*> CustomOutputExpressions;
+			GetAllCustomOutputExpressions(CustomOutputExpressions);
+			for (UMaterialExpressionCustomOutput* Expression : CustomOutputExpressions)
+			{
+				// Gather custom output for thin translucency
+				if (ThinTranslucentOutput == nullptr && Cast<UMaterialExpressionThinTranslucentMaterialOutput>(Expression))
+				{
+					ThinTranslucentOutput = Cast<UMaterialExpressionThinTranslucentMaterialOutput>(Expression);
+				}
+
+				// Gather custom output for single layer water
+				if (SingleLayerWaterOutput == nullptr && Cast<UMaterialExpressionSingleLayerWaterMaterialOutput>(Expression))
+				{
+					SingleLayerWaterOutput = Cast<UMaterialExpressionSingleLayerWaterMaterialOutput>(Expression);
+				}
+
+				// Gather custom output for clear coat
+				if (ClearCoatOutput == nullptr && Cast<UMaterialExpressionClearCoatNormalCustomOutput>(Expression))
+				{
+					ClearCoatOutput = Cast<UMaterialExpressionClearCoatNormalCustomOutput>(Expression);
+				}
+
+				// Gather custom output for tangent (unused atm)
+				if (TangentOutput == nullptr && Cast<UMaterialExpressionTangentOutput>(Expression))
+				{
+					TangentOutput = Cast<UMaterialExpressionTangentOutput>(Expression);
+				}
+
+				if (ThinTranslucentOutput && SingleLayerWaterOutput && ClearCoatOutput && TangentOutput)
+				{
+					break;
+				}
+			}
+		}
+	};
+
 	// SSS Profile
 	const bool bHasShadingModelMixture		= ShadingModels.CountShadingModels() > 1;
 	const bool bRequireSubsurfacePasses		= ShadingModels.HasShadingModel(MSM_SubsurfaceProfile) || ShadingModels.HasShadingModel(MSM_Subsurface) || ShadingModels.HasShadingModel(MSM_PreintegratedSkin) || ShadingModels.HasShadingModel(MSM_Eye);
 	const bool bRequireNoSubsurfaceProfile	= !bHasShadingModelMixture && (ShadingModel == MSM_Subsurface || ShadingModel == MSM_PreintegratedSkin); // Insure there is no profile, as this would take priority otherwise
 
 	bool bInvalidateShader = false;
+	bool bRelinkCustomOutputNodes = false;
+	UMaterialExpressionStrataLegacyConversion* ConvertNode = nullptr;
 	// Connect all the legacy pin into the conversion node
 	if (bUseMaterialAttributes && MaterialAttributes.Expression && !MaterialAttributes.Expression->IsResultStrataMaterial(MaterialAttributes.OutputIndex)) // M_Rifle cause issues there
 	{
 		UMaterialExpressionBreakMaterialAttributes* BreakMatAtt = NewObject<UMaterialExpressionBreakMaterialAttributes>(this);
 		MoveConnectionTo(MaterialAttributes, BreakMatAtt, 0);
 
-		UMaterialExpressionStrataLegacyConversion* ConvertNode = NewObject<UMaterialExpressionStrataLegacyConversion>(this);
+		ConvertNode = NewObject<UMaterialExpressionStrataLegacyConversion>(this);
 		ConvertNode->BaseColor.Connect(0, BreakMatAtt);
 		ConvertNode->Metallic.Connect(1, BreakMatAtt);
 		ConvertNode->Specular.Connect(2, BreakMatAtt);
@@ -2863,6 +2912,7 @@ void UMaterial::ConvertMaterialToStrataMaterial()
 		ConvertNode->Opacity.Connect(6, BreakMatAtt);
 		ConvertNode->ShadingModel.Connect(25, BreakMatAtt);
 		ConvertNode->SubsurfaceProfile = bRequireNoSubsurfaceProfile ? nullptr : SubsurfaceProfile;
+		bRelinkCustomOutputNodes = true;
 
 		// * Remove support for material attribute
 		// * explicitely connect the Strata node to the root node
@@ -2911,75 +2961,21 @@ void UMaterial::ConvertMaterialToStrataMaterial()
 
 		if (MaterialDomain == MD_Surface)
 		{
-			TArray<class UMaterialExpressionCustomOutput*> CustomOutputExpressions;
-			GetAllCustomOutputExpressions(CustomOutputExpressions);
-
-			UMaterialExpressionThinTranslucentMaterialOutput* ThinTranslucentOutput = nullptr;
-			UMaterialExpressionSingleLayerWaterMaterialOutput* SingleLayerWaterOutput = nullptr;
-			UMaterialExpressionClearCoatNormalCustomOutput* ClearCoatOutput = nullptr;
-			UMaterialExpressionTangentOutput* TangentOutput = nullptr;
-
-			for (UMaterialExpressionCustomOutput* Expression : CustomOutputExpressions)
-			{
-				// Gather custom output for thin translucency
-				if (ThinTranslucentOutput == nullptr && Cast<UMaterialExpressionThinTranslucentMaterialOutput>(Expression))
-				{
-					ThinTranslucentOutput = Cast<UMaterialExpressionThinTranslucentMaterialOutput>(Expression);
-				}
-
-				// Gather custom output for single layer water
-				if (SingleLayerWaterOutput == nullptr && Cast<UMaterialExpressionSingleLayerWaterMaterialOutput>(Expression))
-				{
-					SingleLayerWaterOutput = Cast<UMaterialExpressionSingleLayerWaterMaterialOutput>(Expression);
-				}
-
-				// Gather custom output for clear coat
-				if (ClearCoatOutput == nullptr && Cast<UMaterialExpressionClearCoatNormalCustomOutput>(Expression))
-				{
-					ClearCoatOutput = Cast<UMaterialExpressionClearCoatNormalCustomOutput>(Expression);
-				}
-
-				// Gather custom output for tangent (unused atm)
-				if (TangentOutput == nullptr && Cast<UMaterialExpressionTangentOutput>(Expression))
-				{
-					TangentOutput = Cast<UMaterialExpressionTangentOutput>(Expression);
-				}
-
-				if (ThinTranslucentOutput && SingleLayerWaterOutput && ClearCoatOutput && TangentOutput)
-				{
-					break;
-				}
-			}
-
-			UMaterialExpressionStrataLegacyConversion* ConvertNode = NewObject<UMaterialExpressionStrataLegacyConversion>(this);
+			ConvertNode = NewObject<UMaterialExpressionStrataLegacyConversion>(this);
 			ConvertNode->SubsurfaceProfile = bRequireNoSubsurfaceProfile ? nullptr : SubsurfaceProfile;
-			MoveConnectionTo(BaseColor, ConvertNode, 0);	
-			MoveConnectionTo(Metallic, ConvertNode, 1);		
-			MoveConnectionTo(Specular, ConvertNode, 2);		 
+			MoveConnectionTo(BaseColor, ConvertNode, 0);
+			MoveConnectionTo(Metallic, ConvertNode, 1);
+			MoveConnectionTo(Specular, ConvertNode, 2);
 			MoveConnectionTo(Roughness, ConvertNode, 3);
-			MoveConnectionTo(Anisotropy, ConvertNode, 4);	
+			MoveConnectionTo(Anisotropy, ConvertNode, 4);
 			MoveConnectionTo(EmissiveColor, ConvertNode, 5);
 			CopyConnectionTo(Normal, ConvertNode, 6);
 			MoveConnectionTo(Tangent, ConvertNode, 7);
-			MoveConnectionTo(SubsurfaceColor, ConvertNode, 8);	
-			MoveConnectionTo(ClearCoat, ConvertNode, 9);		
-			MoveConnectionTo(ClearCoatRoughness,ConvertNode,10);
+			MoveConnectionTo(SubsurfaceColor, ConvertNode, 8);
+			MoveConnectionTo(ClearCoat, ConvertNode, 9);
+			MoveConnectionTo(ClearCoatRoughness, ConvertNode, 10);
 			MoveConnectionTo(Opacity, ConvertNode, 11);
-			if (ThinTranslucentOutput)
-			{
-				MoveConnectionTo(*ThinTranslucentOutput->GetInput(0), ConvertNode, 12);	 // TransmittanceColor
-			}
-			if (SingleLayerWaterOutput)
-			{
-				MoveConnectionTo(*SingleLayerWaterOutput->GetInput(0), ConvertNode, 13); // WaterScatteringCoefficients
-				MoveConnectionTo(*SingleLayerWaterOutput->GetInput(1), ConvertNode, 14); // WaterAbsorptionCoefficients
-				MoveConnectionTo(*SingleLayerWaterOutput->GetInput(2), ConvertNode, 15); // WaterPhaseG
-				MoveConnectionTo(*SingleLayerWaterOutput->GetInput(3), ConvertNode, 16); // ColorScaleBehindWater
-			}
-			if (ClearCoatOutput)
-			{
-				MoveConnectionTo(*ClearCoatOutput->GetInput(0), ConvertNode, 17);		 // ClearCoatNormal
-			}
+			bRelinkCustomOutputNodes = true;
 			
 			// Shading Model
 			// * either use the shader graph expression 
@@ -2995,6 +2991,7 @@ void UMaterial::ConvertMaterialToStrataMaterial()
 				MoveConnectionTo(ShadingModelFromMaterialExpression, ConvertNode, 18);
 
 				// Store strata shading model of the converted material. 
+				GatherCustomNodes();
 				if (SingleLayerWaterOutput)
 				{
 					ConvertNode->ConvertedStrataMaterialInfo.AddShadingModel(SSM_SingleLayerWater);
@@ -3090,7 +3087,7 @@ void UMaterial::ConvertMaterialToStrataMaterial()
 			ShadingModels.ClearShadingModels();
 			ShadingModels.AddShadingModel(MSM_DefaultLit);
 
-			UMaterialExpressionStrataLegacyConversion* ConvertNode = NewObject<UMaterialExpressionStrataLegacyConversion>(this);
+			ConvertNode = NewObject<UMaterialExpressionStrataLegacyConversion>(this);
 			MoveConnectionTo(BaseColor, ConvertNode, 0);
 			MoveConnectionTo(Metallic, ConvertNode, 1);
 			MoveConnectionTo(Specular, ConvertNode, 2);
@@ -3124,6 +3121,28 @@ void UMaterial::ConvertMaterialToStrataMaterial()
 		ConvertLegacyToStrataBlendMode();
 	}
 
+	if (bRelinkCustomOutputNodes)
+	{
+		check(ConvertNode);
+		GatherCustomNodes();
+
+		if (ThinTranslucentOutput)
+		{
+			MoveConnectionTo(*ThinTranslucentOutput->GetInput(0), ConvertNode, 12);	 // TransmittanceColor
+		}
+		if (SingleLayerWaterOutput)
+		{
+			MoveConnectionTo(*SingleLayerWaterOutput->GetInput(0), ConvertNode, 13); // WaterScatteringCoefficients
+			MoveConnectionTo(*SingleLayerWaterOutput->GetInput(1), ConvertNode, 14); // WaterAbsorptionCoefficients
+			MoveConnectionTo(*SingleLayerWaterOutput->GetInput(2), ConvertNode, 15); // WaterPhaseG
+			MoveConnectionTo(*SingleLayerWaterOutput->GetInput(3), ConvertNode, 16); // ColorScaleBehindWater
+		}
+		if (ClearCoatOutput)
+		{
+			MoveConnectionTo(*ClearCoatOutput->GetInput(0), ConvertNode, 17);		 // ClearCoatNormal
+		}
+	}
+
 	if (bInvalidateShader)
 	{
 		// Now force the material to recompile and we use a hash of the original StateId.
