Changed r.OptimizedWPO so it can be toggled at runtime (by forcing a scene proxy recreate in order for the behavior change to be properly pushed to GPU Scene)

diff --git a/Engine/Source/Runtime/Engine/Private/PrimitiveUniformShaderParameters.cpp b/Engine/Source/Runtime/Engine/Private/PrimitiveUniformShaderParameters.cpp
--- a/Engine/Source/Runtime/Engine/Private/PrimitiveUniformShaderParameters.cpp
+++ b/Engine/Source/Runtime/Engine/Private/PrimitiveUniformShaderParameters.cpp
@@ -6,6 +6,7 @@
 #include "PrimitiveSceneInfo.h"
 #include "NaniteSceneProxy.h"
 #include "ProfilingDebugging/LoadTimeTracker.h"
+#include "ComponentRecreateRenderStateContext.h"
 
 static TAutoConsoleVariable<bool> CVarOptimizedWPO(
 	TEXT("r.OptimizedWPO"),
@@ -13,7 +14,11 @@ static TAutoConsoleVariable<bool> CVarOptimizedWPO(
 	TEXT("Special mode where primitives can explicitly indicate if WPO should be evaluated or not as an optimization.\n")
 	TEXT(" False ( 0): Ignore WPO evaluation flag, and always evaluate WPO.\n")
 	TEXT(" True  ( 1): Only evaluate WPO on primitives with explicit activation."),
-	ECVF_RenderThreadSafe | ECVF_ReadOnly
+	FConsoleVariableDelegate::CreateLambda([](IConsoleVariable* InVariable)
+	{
+		FGlobalComponentRecreateRenderStateContext Context;
+	}),
+	ECVF_RenderThreadSafe
 );
 
 void FSinglePrimitiveStructured::InitRHI() 
