Expose mobile specific "Enable GPUScene" option to project settings

diff --git a/Engine/Source/Runtime/Engine/Classes/Engine/RendererSettings.h b/Engine/Source/Runtime/Engine/Classes/Engine/RendererSettings.h
--- a/Engine/Source/Runtime/Engine/Classes/Engine/RendererSettings.h
+++ b/Engine/Source/Runtime/Engine/Classes/Engine/RendererSettings.h
@@ -318,6 +318,12 @@ class ENGINE_API URendererSettings : public UDeveloperSettings
 		ToolTip = "Whether to use the Reflection Capture Compression or not for mobile. It will use ETC2 format to do the compression."))
 	uint32 bReflectionCaptureCompression : 1;
 
+	UPROPERTY(config, EditAnywhere, Category = Mobile, meta = (
+		ConsoleVariable = "r.Mobile.SupportGPUScene", DisplayName = "Enable GPUScene on Mobile",
+		ToolTip = "Whether to enable GPUScene on mobile. GPUScene is required for mesh auto-instancing. Changing this setting requires restarting the editor.",
+		ConfigRestartRequired = true))
+	uint32 bMobileSupportGPUScene : 1;
+
 	UPROPERTY(config, EditAnywhere, Category = Materials, meta = (
 		ConsoleVariable = "r.DiscardUnusedQuality", DisplayName = "Game Discards Unused Material Quality Levels",
 		ToolTip = "When running in game mode, whether to keep shaders for all quality levels in memory or only those needed for the current quality level.\nUnchecked: Keep all quality levels in memory allowing a runtime quality level change. (default)\nChecked: Discard unused quality levels when loading content for the game, saving some memory."))
