StrataLegacyConversion node is not longer placeable in material graph because it should only be used for conversion purposes.

diff --git a/Engine/Source/Runtime/Engine/Private/Materials/MaterialExpressions.cpp b/Engine/Source/Runtime/Engine/Private/Materials/MaterialExpressions.cpp
--- a/Engine/Source/Runtime/Engine/Private/Materials/MaterialExpressions.cpp
+++ b/Engine/Source/Runtime/Engine/Private/Materials/MaterialExpressions.cpp
@@ -314,7 +314,8 @@ bool IsAllowedExpressionType(const UClass* const Class, const bool bMaterialFunc
 	const bool bSharedAllowed = Class != UMaterialExpressionComment::StaticClass() 
 		&& Class != UMaterialExpressionPinBase::StaticClass()
 		&& Class != UMaterialExpressionParameter::StaticClass()
-		&& (Class != UMaterialExpressionTextureSampleParameter2DArray::StaticClass() || AllowTextureArrayAssetCreationVar->GetValueOnGameThread() != 0);
+		&& (Class != UMaterialExpressionTextureSampleParameter2DArray::StaticClass() || AllowTextureArrayAssetCreationVar->GetValueOnGameThread() != 0)
+		&& Class != UMaterialExpressionStrataLegacyConversion::StaticClass();
 
 	if (bMaterialFunction)
 	{
