Expose mobile shading path to project settings UI

diff --git a/Engine/Source/Runtime/Engine/Classes/Engine/RendererSettings.h b/Engine/Source/Runtime/Engine/Classes/Engine/RendererSettings.h
--- a/Engine/Source/Runtime/Engine/Classes/Engine/RendererSettings.h
+++ b/Engine/Source/Runtime/Engine/Classes/Engine/RendererSettings.h
@@ -212,6 +212,17 @@ namespace EMobileFloatPrecisionMode
 	};
 }
 
+UENUM()
+namespace EMobileShadingPath
+{
+	enum Type
+	{
+		/** The default shading path for mobile, supported on all devices. */
+		Forward = 0 UMETA(DisplayName = "Forward Shading"),
+		/** Use deferred shading. This path supports additional features compared to the Forward option but requires more modern devices. Features include: IES light profiles, light functions, lit deferred decals. Does not support MSAA. */
+		Deferred = 1 UMETA(DisplayName = "Deferred Shading"),
+	};
+}
 
 namespace EDefaultBackBufferPixelFormat
 {
@@ -278,16 +289,22 @@ class ENGINE_API URendererSettings : public UDeveloperSettings
 	GENERATED_UCLASS_BODY()
 
 	UPROPERTY(config, EditAnywhere, Category = Mobile, meta = (
-		ConsoleVariable = "r.Mobile.DisableVertexFog", DisplayName = "Disable vertex fogging in mobile shaders",
-		ToolTip = "If true, vertex fog will be omitted from the most of the mobile base pass shaders. Instead, fog will be applied in a separate pass and only when scene has a fog component.",
+		ConsoleVariable = "r.Mobile.ShadingPath", DisplayName = "Mobile Shading",
+		ToolTip = "The shading path to use on mobile platforms. Changing this setting requires restarting the editor.",
 		ConfigRestartRequired = true))
-		uint32 bMobileDisableVertexFog : 1;
+	TEnumAsByte<EMobileShadingPath::Type> MobileShadingPath;
+
+	UPROPERTY(config, EditAnywhere, Category = Mobile, meta = (
+		ConsoleVariable = "r.Mobile.SupportGPUScene", DisplayName = "Enable GPUScene on Mobile",
+		ToolTip = "Whether to enable GPUScene on mobile. GPUScene is required for mesh auto-instancing. Changing this setting requires restarting the editor.",
+		ConfigRestartRequired = true))
+	uint32 bMobileSupportGPUScene : 1;
 
 	UPROPERTY(config, EditAnywhere, Category = Mobile, meta = (
 		ConsoleVariable = "r.Shadow.CSM.MaxMobileCascades", DisplayName = "Maximum number of CSM cascades to render", ClampMin = 1, ClampMax = 4,
 		ToolTip = "The maximum number of cascades with which to render dynamic directional light shadows when using the mobile renderer.",
 		ConfigRestartRequired = true))
-		int32 MaxMobileCascades;
+	int32 MaxMobileCascades;
 
 	UPROPERTY(config, EditAnywhere, Category = Mobile, meta = (
 		ConsoleVariable = "r.Mobile.AntiAliasing", DisplayName = "Mobile Anti-Aliasing Method",
@@ -319,10 +336,10 @@ class ENGINE_API URendererSettings : public UDeveloperSettings
 	uint32 bReflectionCaptureCompression : 1;
 
 	UPROPERTY(config, EditAnywhere, Category = Mobile, meta = (
-		ConsoleVariable = "r.Mobile.SupportGPUScene", DisplayName = "Enable GPUScene on Mobile",
-		ToolTip = "Whether to enable GPUScene on mobile. GPUScene is required for mesh auto-instancing. Changing this setting requires restarting the editor.",
+		ConsoleVariable = "r.Mobile.DisableVertexFog", DisplayName = "Disable vertex fogging in mobile shaders",
+		ToolTip = "If true, vertex fog will be omitted from the most of the mobile base pass shaders. Instead, fog will be applied in a separate pass and only when scene has a fog component.",
 		ConfigRestartRequired = true))
-	uint32 bMobileSupportGPUScene : 1;
+	uint32 bMobileDisableVertexFog : 1;
 
 	UPROPERTY(config, EditAnywhere, Category = Materials, meta = (
 		ConsoleVariable = "r.DiscardUnusedQuality", DisplayName = "Game Discards Unused Material Quality Levels",
@@ -1082,7 +1099,7 @@ class ENGINE_API URendererSettings : public UDeveloperSettings
 		TEnumAsByte<EMobilePlanarReflectionMode::Type> MobilePlanarReflectionMode;
 
 	UPROPERTY(config, EditAnywhere, Category = Mobile, meta = (
-		ConsoleVariable = "r.Mobile.SupportsGen4TAA", DisplayName = "Supports desktop Gen4 TAA on mobile",
+		ConsoleVariable = "r.Mobile.SupportsGen4TAA", DisplayName = "Support desktop Gen4 TAA on mobile",
 		ToolTip = "Support desktop Gen4 TAA with mobile rendering. Changing this setting requires restarting the editor.",
 		ConfigRestartRequired = true))
 		uint32 bMobileSupportsGen4TAA : 1;
