Strata - Fix for crash with parameter blending and normal resulting in INDEX_NONE code chunks.

diff --git a/Engine/Source/Runtime/Engine/Private/Materials/HLSLMaterialTranslator.cpp b/Engine/Source/Runtime/Engine/Private/Materials/HLSLMaterialTranslator.cpp
--- a/Engine/Source/Runtime/Engine/Private/Materials/HLSLMaterialTranslator.cpp
+++ b/Engine/Source/Runtime/Engine/Private/Materials/HLSLMaterialTranslator.cpp
@@ -10209,16 +10209,28 @@ FHLSLMaterialTranslator::FStrataSharedLocalBasesInfo FHLSLMaterialTranslator::St
 
 int32 FHLSLMaterialTranslator::StrataAddParameterBlendingBSDFCoverageToNormalMixCodeChunk(int32 ACodeChunk, int32 BCodeChunk)
 {
+	if (ACodeChunk == INDEX_NONE || BCodeChunk == INDEX_NONE)
+	{
+		return INDEX_NONE;
+	}
 	return AddCodeChunk(MCT_Float, TEXT("AddParameterBlendingBSDFCoverageToNormalMix(%s, %s)"), *GetParameterCode(ACodeChunk), *GetParameterCode(BCodeChunk));
 }
 
 int32 FHLSLMaterialTranslator::StrataVerticalLayeringParameterBlendingBSDFCoverageToNormalMixCodeChunk(int32 TopCodeChunk)
 {
+	if (TopCodeChunk == INDEX_NONE)
+	{
+		return INDEX_NONE;
+	}
 	return AddCodeChunk(MCT_Float, TEXT("VerticalLayeringParameterBlendingBSDFCoverageToNormalMix(%s)"), *GetParameterCode(TopCodeChunk));
 }
 
 int32 FHLSLMaterialTranslator::StrataHorizontalMixingParameterBlendingBSDFCoverageToNormalMixCodeChunk(int32 BackgroundCodeChunk, int32 ForegroundCodeChunk, int32 HorizontalMixCodeChunk)
 {
+	if (BackgroundCodeChunk == INDEX_NONE || ForegroundCodeChunk == INDEX_NONE || HorizontalMixCodeChunk == INDEX_NONE)
+	{
+		return INDEX_NONE;
+	}
 	return AddCodeChunk(MCT_Float, TEXT("HorizontalMixingParameterBlendingBSDFCoverageToNormalMix(%s, %s, %s)"), *GetParameterCode(BackgroundCodeChunk), *GetParameterCode(ForegroundCodeChunk), *GetParameterCode(HorizontalMixCodeChunk));
 }
 
diff --git a/Engine/Source/Runtime/Engine/Private/Materials/MaterialExpressions.cpp b/Engine/Source/Runtime/Engine/Private/Materials/MaterialExpressions.cpp
--- a/Engine/Source/Runtime/Engine/Private/Materials/MaterialExpressions.cpp
+++ b/Engine/Source/Runtime/Engine/Private/Materials/MaterialExpressions.cpp
@@ -21795,8 +21795,6 @@ int32 UMaterialExpressionStrataConvertToDecal::Compile(class FMaterialCompiler*
 	{
 		return Compiler->Errorf(TEXT("Missing DecalMaterial input"));
 	}
-
-	int32 DecalMaterialCodeChunk = DecalMaterial.Compile(Compiler);
 	int32 WeightOfOneCodeChunk = Compiler->Constant(1.0f);
 
 	FStrataOperator& StrataOperator = Compiler->StrataCompilationGetOperator(this);
@@ -21813,6 +21811,7 @@ int32 UMaterialExpressionStrataConvertToDecal::Compile(class FMaterialCompiler*
 	int32 OutputCodeChunk = INDEX_NONE;
 	if (StrataOperator.bUseParameterBlending)
 	{
+		int32 DecalMaterialCodeChunk = DecalMaterial.Compile(Compiler);
 		OutputCodeChunk = Compiler->StrataWeightParameterBlending(
 			DecalMaterialCodeChunk, WeightOfOneCodeChunk,
 			StrataOperator.bRootOfParameterBlendingSubTree ? &StrataOperator : nullptr);
@@ -22188,11 +22187,18 @@ int32 UMaterialExpressionStrataHorizontalMixing::Compile(class FMaterialCompiler
 
 	const int32 HorizontalMixCodeChunk = CompileWithDefaultFloat1(Compiler, Mix, 0.5f);
 
-
 	int32 OutputCodeChunk = INDEX_NONE;
 	FStrataOperator& StrataOperator = Compiler->StrataCompilationGetOperator(this);
 	if (StrataOperator.bUseParameterBlending)
 	{
+		if (ForegroundCodeChunk == INDEX_NONE)
+		{
+			return Compiler->Errorf(TEXT("Foreground input graphs could not be evaluated for parameter blending."));
+		}
+		if (BackgroundCodeChunk == INDEX_NONE)
+		{
+			return Compiler->Errorf(TEXT("Background input graphs could not be evaluated for parameter blending."));
+		}
 		const int32 NormalMixCodeChunk = Compiler->StrataHorizontalMixingParameterBlendingBSDFCoverageToNormalMixCodeChunk(BackgroundCodeChunk, ForegroundCodeChunk, HorizontalMixCodeChunk);
 
 		FStrataOperator* BackgroundBSDFOperator = Compiler->StrataCompilationGetOperatorFromIndex(StrataOperator.LeftIndex);
@@ -22331,8 +22337,6 @@ int32 UMaterialExpressionStrataVerticalLayering::Compile(class FMaterialCompiler
 	int32 TopCodeChunk = Top.Compile(Compiler);
 	int32 BaseCodeChunk = Base.Compile(Compiler);
 
-	const int32 TopNormalMixCodeChunk = Compiler->StrataVerticalLayeringParameterBlendingBSDFCoverageToNormalMixCodeChunk(TopCodeChunk);
-
 	int32 OutputCodeChunk = INDEX_NONE;
 	FStrataOperator& StrataOperator = Compiler->StrataCompilationGetOperator(this);
 	if (StrataOperator.bUseParameterBlending)
@@ -22343,6 +22347,16 @@ int32 UMaterialExpressionStrataVerticalLayering::Compile(class FMaterialCompiler
 		{
 			return Compiler->Errorf(TEXT("Missing input on vertical layering node."));
 		}
+		if (TopCodeChunk == INDEX_NONE)
+		{
+			return Compiler->Errorf(TEXT("Top input graph could not be evaluated for parameter blending."));
+		}
+		if (BaseCodeChunk == INDEX_NONE)
+		{
+			return Compiler->Errorf(TEXT("Base input graph could not be evaluated for parameter blending."));
+		}
+
+		const int32 TopNormalMixCodeChunk = Compiler->StrataVerticalLayeringParameterBlendingBSDFCoverageToNormalMixCodeChunk(TopCodeChunk);
 
 		// Compute the new Normal and Tangent resulting from the blending using code chunk
 		const int32 NewNormalCodeChunk = StrataBlendNormal(Compiler, BaseBSDFOperator->BSDFRegisteredSharedLocalBasis.NormalCodeChunk, TopBSDFOperator->BSDFRegisteredSharedLocalBasis.NormalCodeChunk, TopNormalMixCodeChunk);
@@ -22469,8 +22483,6 @@ int32 UMaterialExpressionStrataAdd::Compile(class FMaterialCompiler* Compiler, i
 	int32 ACodeChunk = A.Compile(Compiler);
 	int32 BCodeChunk = B.Compile(Compiler);
 
-	const int32 ANormalMixCodeChunk = Compiler->StrataAddParameterBlendingBSDFCoverageToNormalMixCodeChunk(ACodeChunk, BCodeChunk);
-
 	int32 OutputCodeChunk = INDEX_NONE;
 	FStrataOperator& StrataOperator = Compiler->StrataCompilationGetOperator(this);
 	if (StrataOperator.bUseParameterBlending)
@@ -22481,6 +22493,16 @@ int32 UMaterialExpressionStrataAdd::Compile(class FMaterialCompiler* Compiler, i
 		{
 			return Compiler->Errorf(TEXT("Missing input on add node."));
 		}
+		if (ACodeChunk == INDEX_NONE)
+		{
+			return Compiler->Errorf(TEXT("A input graph could not be evaluated for parameter blending."));
+		}
+		if (BCodeChunk == INDEX_NONE)
+		{
+			return Compiler->Errorf(TEXT("B input graph could not be evaluated for parameter blending."));
+		}
+
+		const int32 ANormalMixCodeChunk = Compiler->StrataAddParameterBlendingBSDFCoverageToNormalMixCodeChunk(ACodeChunk, BCodeChunk);
 
 		// Compute the new Normal and Tangent resulting from the blending using code chunk
 		const int32 NewNormalCodeChunk = StrataBlendNormal(Compiler, BBSDFOperator->BSDFRegisteredSharedLocalBasis.NormalCodeChunk, ABSDFOperator->BSDFRegisteredSharedLocalBasis.NormalCodeChunk, ANormalMixCodeChunk);
@@ -22611,6 +22633,15 @@ int32 UMaterialExpressionStrataWeight::Compile(class FMaterialCompiler* Compiler
 		{
 			return Compiler->Errorf(TEXT("Missing input on weight node."));
 		}
+		if (ACodeChunk == INDEX_NONE)
+		{
+			return Compiler->Errorf(TEXT("A input graph could not be evaluated for parameter blending."));
+		}
+		if (WeightCodeChunk == INDEX_NONE)
+		{
+			return Compiler->Errorf(TEXT("Weight input graph could not be evaluated for parameter blending."));
+		}
+
 		StrataOperator.BSDFRegisteredSharedLocalBasis = Operator->BSDFRegisteredSharedLocalBasis;
 		StrataOperator.CopyFlagsForParameterBlending(*Operator);
 
