Forcefully auto set tangents for AnimDataModel on load, this is to ensure any auto-tangents for Key0 and Key1 are set to be flat (which is the new standard for importing/AutoSetTangents()) see UE-98071

diff --git a/Engine/Source/Runtime/Core/Public/UObject/UE5MainStreamObjectVersion.h b/Engine/Source/Runtime/Core/Public/UObject/UE5MainStreamObjectVersion.h
--- a/Engine/Source/Runtime/Core/Public/UObject/UE5MainStreamObjectVersion.h
+++ b/Engine/Source/Runtime/Core/Public/UObject/UE5MainStreamObjectVersion.h
@@ -206,6 +206,9 @@ struct CORE_API FUE5MainStreamObjectVersion
 		// Fix-up for FRichCurveKey::TangentWeightMode, which were found to contain invalid value w.r.t the enum-type
 		RichCurveKeyInvalidTangentMode,
 
+		// Enforcing new automatic tangent behaviour, enforcing auto-tangents for Key0 and KeyN to be flat, for Animation Assets.
+		ForceUpdateAnimationAssetCurveTangents,
+
 		// -----<new versions can be added above this line>-------------------------------------------------
 		VersionPlusOne,
 		LatestVersion = VersionPlusOne - 1
diff --git a/Engine/Source/Runtime/Engine/Private/Animation/AnimDataModel.cpp b/Engine/Source/Runtime/Engine/Private/Animation/AnimDataModel.cpp
--- a/Engine/Source/Runtime/Engine/Private/Animation/AnimDataModel.cpp
+++ b/Engine/Source/Runtime/Engine/Private/Animation/AnimDataModel.cpp
@@ -8,11 +8,26 @@
 #include "Algo/Transform.h"
 #include "Algo/Accumulate.h"
 #include "Animation/SmartName.h"
+#include "UObject/UE5MainStreamObjectVersion.h"
 
 
 void UAnimDataModel::PostLoad()
 {
 	UObject::PostLoad();
+
+	if (GetLinkerCustomVersion(FUE5MainStreamObjectVersion::GUID) < FUE5MainStreamObjectVersion::ForceUpdateAnimationAssetCurveTangents)
+	{
+		// Forcefully AutoSetTangents to fix-up any imported sequences pre the fix for flattening first/last key leave/arrive tangents
+		Notify(EAnimDataModelNotifyType::BracketOpened);
+		for (FFloatCurve& FloatCurve : CurveData.FloatCurves)
+		{
+			FloatCurve.FloatCurve.AutoSetTangents();
+			FCurvePayload Payload;
+			Payload.Identifier = FAnimationCurveIdentifier(FloatCurve.Name.UID, ERawCurveTrackTypes::RCT_Float);
+			Notify(EAnimDataModelNotifyType::CurveChanged, Payload);
+		}
+		Notify(EAnimDataModelNotifyType::BracketClosed);
+	}
 }
 
 void UAnimDataModel::PostDuplicate(bool bDuplicateForPIE)
