Optimus: Fix motion vectors for morph target and cloth.
#rnx

diff --git a/Engine/Source/Runtime/Engine/Private/GPUSkinVertexFactory.cpp b/Engine/Source/Runtime/Engine/Private/GPUSkinVertexFactory.cpp
--- a/Engine/Source/Runtime/Engine/Private/GPUSkinVertexFactory.cpp
+++ b/Engine/Source/Runtime/Engine/Private/GPUSkinVertexFactory.cpp
@@ -913,7 +913,7 @@ void FGPUSkinPassthroughVertexFactory::ReleaseRHI()
 	ColorVBAlias.ReleaseRHI();
 }
 
-void FGPUSkinPassthroughVertexFactory::InternalUpdateVertexDeclaration(FGPUBaseSkinVertexFactory* SourceVertexFactory)
+void FGPUSkinPassthroughVertexFactory::InternalUpdateVertexDeclaration(FGPUBaseSkinVertexFactory const* SourceVertexFactory)
 {
 	if (PositionVBAlias.VertexBufferRHI != nullptr)
 	{
@@ -992,7 +992,7 @@ void FGPUSkinPassthroughVertexFactory::InternalUpdateVertexDeclaration(FGPUBaseS
 }
 
 void FGPUSkinPassthroughVertexFactory::InternalUpdateVertexDeclaration(
-	FGPUBaseSkinVertexFactory* SourceVertexFactory, 
+	FGPUBaseSkinVertexFactory const* SourceVertexFactory, 
 	struct FRWBuffer* PositionRWBuffer, 
 	struct FRWBuffer* TangentRWBuffer)
 {
@@ -1015,7 +1015,7 @@ void FGPUSkinPassthroughVertexFactory::InternalUpdateVertexDeclaration(
 
 void FGPUSkinPassthroughVertexFactory::InternalUpdateVertexDeclaration(
 	EOverrideFlags OverrideFlags,
-	FGPUBaseSkinVertexFactory* SourceVertexFactory, 
+	FGPUBaseSkinVertexFactory const* SourceVertexFactory, 
 	TRefCountPtr<FRDGPooledBuffer> const& PositionBuffer, 
 	TRefCountPtr<FRDGPooledBuffer> const& TangentBuffer,
 	TRefCountPtr<FRDGPooledBuffer> const& ColorBuffer)
diff --git a/Engine/Source/Runtime/Engine/Private/SkeletalMeshDeformerHelpers.cpp b/Engine/Source/Runtime/Engine/Private/SkeletalMeshDeformerHelpers.cpp
--- a/Engine/Source/Runtime/Engine/Private/SkeletalMeshDeformerHelpers.cpp
+++ b/Engine/Source/Runtime/Engine/Private/SkeletalMeshDeformerHelpers.cpp
@@ -106,7 +106,7 @@ void FSkeletalMeshDeformerHelpers::SetVertexFactoryBufferOverrides(
 	const int32 NumSections = MeshObject->GetRenderSections(LODIndex).Num();
 	for (int32 SectionIndex = 0; SectionIndex < NumSections; ++SectionIndex)
 	{
-		FGPUBaseSkinVertexFactory* BaseVertexFactory = LOD.GPUSkinVertexFactories.VertexFactories[SectionIndex].Get();
+		FGPUBaseSkinVertexFactory const* BaseVertexFactory = MeshObjectGPU->GetBaseSkinVertexFactory(LODIndex, SectionIndex);
 		FGPUSkinPassthroughVertexFactory* TargetVertexFactory = LOD.GPUSkinVertexFactories.PassthroughVertexFactories[SectionIndex].Get();
 		TargetVertexFactory->InvalidateStreams();
 		TargetVertexFactory->UpdateVertexDeclaration(OverrideFlags, BaseVertexFactory, PositionBuffer, TangentBuffer, ColorBuffer);
diff --git a/Engine/Source/Runtime/Engine/Public/GPUSkinVertexFactory.h b/Engine/Source/Runtime/Engine/Public/GPUSkinVertexFactory.h
--- a/Engine/Source/Runtime/Engine/Public/GPUSkinVertexFactory.h
+++ b/Engine/Source/Runtime/Engine/Public/GPUSkinVertexFactory.h
@@ -381,13 +381,13 @@ public:
 	virtual uint32 GetNumTexCoords() const										{ return Data.IsValid() ? Data->NumTexCoords : 0; }
 	virtual const uint32 GetColorIndexMask() const								{ return Data.IsValid() ? Data->ColorIndexMask : 0; }
 
-	inline const FVertexStreamComponent& GetPositionStreamComponent()
+	inline const FVertexStreamComponent& GetPositionStreamComponent() const
 	{
 		check(Data.IsValid() && Data->PositionComponent.VertexBuffer != nullptr);
 		return Data->PositionComponent;
 	}
 	
-	inline const FVertexStreamComponent& GetTangentStreamComponent(int Index)
+	inline const FVertexStreamComponent& GetTangentStreamComponent(int Index) const
 	{
 		check(Data.IsValid() && Data->TangentBasisComponents[Index].VertexBuffer != nullptr);
 		return Data->TangentBasisComponents[Index];
@@ -500,7 +500,7 @@ public:
 
 	/** Poke values into the vertex factory. */
 	inline void UpdateVertexDeclaration(
-		FGPUBaseSkinVertexFactory* SourceVertexFactory, 
+		FGPUBaseSkinVertexFactory const* SourceVertexFactory, 
 		struct FRWBuffer* PositionRWBuffer, 
 		struct FRWBuffer* TangentRWBuffer)
 	{
@@ -524,7 +524,7 @@ public:
 	/** Poke values into the vertex factory. */
 	inline void UpdateVertexDeclaration(
 		EOverrideFlags OverrideFlags,
-		FGPUBaseSkinVertexFactory* SourceVertexFactory,
+		FGPUBaseSkinVertexFactory const* SourceVertexFactory,
 		TRefCountPtr<FRDGPooledBuffer> const& PositionBuffer,
 		TRefCountPtr<FRDGPooledBuffer> const& TangentBuffer,
 		TRefCountPtr<FRDGPooledBuffer> const& ColorBuffer)
@@ -560,14 +560,14 @@ protected:
 	uint32 UpdatedFrameNumber = 0;
 
 	void InternalUpdateVertexDeclaration(
-		FGPUBaseSkinVertexFactory* SourceVertexFactory);
+		FGPUBaseSkinVertexFactory const* SourceVertexFactory);
 	void InternalUpdateVertexDeclaration(
-		FGPUBaseSkinVertexFactory* SourceVertexFactory, 
+		FGPUBaseSkinVertexFactory const* SourceVertexFactory, 
 		struct FRWBuffer* PositionRWBuffer, 
 		struct FRWBuffer* TangentRWBuffer);
 	void InternalUpdateVertexDeclaration(
 		EOverrideFlags OverrideFlags,
-		FGPUBaseSkinVertexFactory* SourceVertexFactory,
+		FGPUBaseSkinVertexFactory const* SourceVertexFactory,
 		TRefCountPtr<FRDGPooledBuffer> const& PositionBuffer,
 		TRefCountPtr<FRDGPooledBuffer> const& TangentBuffer,
 		TRefCountPtr<FRDGPooledBuffer> const& ColorBuffer);
