[UE][FIX] Handle case where straight capsule physics bodies are created in place of tapered capsules in FBodyInstance::UpdateBodyScale()

diff --git a/Engine/Source/Runtime/Engine/Classes/PhysicsEngine/ShapeElem.h b/Engine/Source/Runtime/Engine/Classes/PhysicsEngine/ShapeElem.h
--- a/Engine/Source/Runtime/Engine/Classes/PhysicsEngine/ShapeElem.h
+++ b/Engine/Source/Runtime/Engine/Classes/PhysicsEngine/ShapeElem.h
@@ -90,6 +90,9 @@ struct FKShapeElem
 	/** Set the user-defined name for this shape */
 	ENGINE_API void SetName(const FName& InName) { Name = InName; }
 
+	/** Get the type of this shape */
+	ENGINE_API EAggCollisionShape::Type GetShapeType() const { return ShapeType; }
+
 	/** Get whether this shape contributes to the mass of the body */
 	ENGINE_API bool GetContributeToMass() const { return bContributeToMass; }
 
diff --git a/Engine/Source/Runtime/Engine/Private/PhysicsEngine/BodyInstance.cpp b/Engine/Source/Runtime/Engine/Private/PhysicsEngine/BodyInstance.cpp
--- a/Engine/Source/Runtime/Engine/Private/PhysicsEngine/BodyInstance.cpp
+++ b/Engine/Source/Runtime/Engine/Private/PhysicsEngine/BodyInstance.cpp
@@ -2068,20 +2068,45 @@ bool FBodyInstance::UpdateBodyScale(const FVector& InScale3D, bool bForceUpdate)
 
 					FReal ScaleRadius = FMath::Max(AdjustedScale3DAbs.X, AdjustedScale3DAbs.Y);
 					FReal ScaleLength = AdjustedScale3DAbs.Z;
-					FKSphylElem* SphylElem = ShapeElem->GetShapeCheck<FKSphylElem>();
 
-					const FReal InitialHeight = SphylElem->Radius * 2.0f + SphylElem->Length;
-					FReal Radius = FMath::Max(SphylElem->Radius * ScaleRadius, (FReal)0.1);
-					FReal HalfHeight = (SphylElem->Length * 0.5f + SphylElem->Radius) * ScaleLength;
+					FReal Radius = 0.0f;
+					FReal HalfHeight = 0.0f;
+					FVec3 Center;
+					FVec3 Axis;
+
+					if (ShapeElem->GetShapeType() == EAggCollisionShape::TaperedCapsule)
+					{
+						// Handle the case where standard capsules are generated in place of tapered capsules, which are not fully supported.
+						FKTaperedCapsuleElem* TaperedCapsuleElem = ShapeElem->GetShapeCheck<FKTaperedCapsuleElem>();
+
+						const FReal MeanRadius = 0.5f * (TaperedCapsuleElem->Radius0 + TaperedCapsuleElem->Radius1);
+
+						const FReal InitialHeight = MeanRadius * 2.0f + TaperedCapsuleElem->Length;
+						Radius = FMath::Max(MeanRadius * ScaleRadius, (FReal)0.1);
+						HalfHeight = (TaperedCapsuleElem->Length * 0.5f + MeanRadius) * ScaleLength;
+
+						// TODO: For Transformed implicit, do not bake this in. Set Transform instead.
+						Center = RelativeTM.TransformPosition(TaperedCapsuleElem->Center) * InScale3D;
+						Axis = TaperedCapsuleElem->Rotation.RotateVector(Chaos::FVec3(0, 0, 1));
+					}
+					else
+					{
+						FKSphylElem* SphylElem = ShapeElem->GetShapeCheck<FKSphylElem>();
+
+						const FReal InitialHeight = SphylElem->Radius * 2.0f + SphylElem->Length;
+						Radius = FMath::Max(SphylElem->Radius * ScaleRadius, (FReal)0.1);
+						HalfHeight = (SphylElem->Length * 0.5f + SphylElem->Radius) * ScaleLength;
+
+						// TODO: For Transformed implicit, do not bake this in. Set Transform instead.
+						Center = RelativeTM.TransformPosition(SphylElem->Center) * InScale3D;
+						Axis = SphylElem->Rotation.RotateVector(Chaos::FVec3(0, 0, 1));
+					}
+
 					Radius = FMath::Min(Radius, HalfHeight);	//radius is capped by half length
 					Radius = FMath::Max(Radius, (FReal)FCollisionShape::MinCapsuleRadius());
 					FReal HalfLength = HalfHeight - Radius;
 					HalfLength = FMath::Max((FReal)FCollisionShape::MinCapsuleAxisHalfHeight(), HalfLength);
 
-					// TODO: For Transformed implicit, do not bake this in. Set Transform instead.
-					FVec3 Center = RelativeTM.TransformPosition(SphylElem->Center) * InScale3D;
-					const FVec3 Axis = SphylElem->Rotation.RotateVector(Chaos::FVec3(0, 0, 1));
-
 					const FVec3 X1 = Center - HalfLength * Axis;
 					const FVec3 X2 = Center + HalfLength * Axis;
 
@@ -2315,7 +2340,6 @@ bool FBodyInstance::UpdateBodyScale(const FVector& InScale3D, bool bForceUpdate)
 				}
 				case ECollisionShapeType::Capsule:
 				{
-					FKSphylElem* SphylElem = ShapeElem->GetShapeCheck<FKSphylElem>();
 					ensure(ScaleMode == EScaleMode::LockedXY || ScaleMode == EScaleMode::LockedXYZ);
 
 					float ScaleRadius = FMath::Max(AdjustedScale3DAbs.X, AdjustedScale3DAbs.Y);
@@ -2323,10 +2347,39 @@ bool FBodyInstance::UpdateBodyScale(const FVector& InScale3D, bool bForceUpdate)
 
 					PxCapsuleGeometry& PCapsuleGeom = GeoCollection.GetCapsuleGeometry();
 
-					// this is a bit confusing since radius and height is scaled
-					// first apply the scale first 
-					float Radius = FMath::Max(SphylElem->Radius * ScaleRadius, 0.1f);
-					float Length = SphylElem->Length + SphylElem->Radius * 2.f;
+					float Radius = 0.0f;
+					float Length = 0.0f;
+
+					if (ShapeElem->GetShapeType() == EAggCollisionShape::TaperedCapsule)
+					{
+						// Handle the case where standard capsules are generated in place of tapered capsules, which are not fully supported.
+						FKTaperedCapsuleElem* TaperedCapsuleElem = ShapeElem->GetShapeCheck<FKTaperedCapsuleElem>();
+
+						const FReal MeanRadius = 0.5f * (TaperedCapsuleElem->Radius0 + TaperedCapsuleElem->Radius1);
+
+						// this is a bit confusing since radius and height is scaled
+						// first apply the scale first 
+						Radius = FMath::Max(MeanRadius * ScaleRadius, 0.1f);
+						Length = TaperedCapsuleElem->Length + MeanRadius * 2.f;
+
+						LocalTransform.SetTranslation(RelativeTM.TransformPosition(TaperedCapsuleElem->Center));
+						LocalTransform.SetRotation(TaperedCapsuleElem->Rotation.Quaternion()* U2PSphylBasis_UE); // #PHYS2 we probably want to put this behind the interface?
+					}
+					else
+					{
+						FKSphylElem* SphylElem = ShapeElem->GetShapeCheck<FKSphylElem>();
+
+						// this is a bit confusing since radius and height is scaled
+						// first apply the scale first 
+						Radius = FMath::Max(SphylElem->Radius * ScaleRadius, 0.1f);
+						Length = SphylElem->Length + SphylElem->Radius * 2.f;
+
+						LocalTransform.SetTranslation(RelativeTM.TransformPosition(SphylElem->Center));
+						LocalTransform.SetRotation(SphylElem->Rotation.Quaternion() * U2PSphylBasis_UE); // #PHYS2 we probably want to put this behind the interface?
+					}
+
+					LocalTransform.ScaleTranslation(AdjustedScale3D);
+
 					float HalfLength = Length * ScaleLength * 0.5f;
 					Radius = FMath::Min(Radius, HalfLength);	//radius is capped by half length
 					Radius = FMath::Max(Radius, FCollisionShape::MinCapsuleRadius()); // bounded by minimum limit.
@@ -2336,10 +2389,6 @@ bool FBodyInstance::UpdateBodyScale(const FVector& InScale3D, bool bForceUpdate)
 					PCapsuleGeom.halfHeight = FMath::Max(HalfHeight, KINDA_SMALL_NUMBER);
 					PCapsuleGeom.radius = FMath::Max(Radius, KINDA_SMALL_NUMBER);
 
-					LocalTransform.SetTranslation(RelativeTM.TransformPosition(SphylElem->Center));
-					LocalTransform.SetRotation(SphylElem->Rotation.Quaternion() * U2PSphylBasis_UE); // #PHYS2 we probably want to put this behind the interface?
-					LocalTransform.ScaleTranslation(AdjustedScale3D);
-
 					if (PCapsuleGeom.isValid())
 					{
 						UpdatedGeometry = &PCapsuleGeom;
