World Partition: added LWC support for very large grids.

#rnx

diff --git a/Engine/Source/Runtime/Engine/Private/WorldPartition/RuntimeSpatialHash/RuntimeSpatialHashGridHelper.cpp b/Engine/Source/Runtime/Engine/Private/WorldPartition/RuntimeSpatialHash/RuntimeSpatialHashGridHelper.cpp
--- a/Engine/Source/Runtime/Engine/Private/WorldPartition/RuntimeSpatialHash/RuntimeSpatialHashGridHelper.cpp
+++ b/Engine/Source/Runtime/Engine/Private/WorldPartition/RuntimeSpatialHash/RuntimeSpatialHashGridHelper.cpp
@@ -314,7 +314,7 @@ FORCEINLINE static bool IsPointInsideSector(const FVector2D& TestPoint, const FV
 
 bool FSquare2DGridHelper::FGrid2D::DoesCircleSectorIntersectsCell(const FIntVector2& Coords, const FVector2D& SectorCenter, float SectorRadiusSquared, const FVector2D& SectorStartVector, const FVector2D& SectorEndVector, float SectorAngle) const
 {
-	const int32 CellIndex = Coords.Y * GridSize + Coords.X;
+	const int64 CellIndex = (int64)Coords.Y * (int64)GridSize + (int64)Coords.X;
 	FBox2D CellBounds;
 	GetCellBounds(CellIndex, CellBounds);
 
diff --git a/Engine/Source/Runtime/Engine/Private/WorldPartition/RuntimeSpatialHash/RuntimeSpatialHashGridHelper.h b/Engine/Source/Runtime/Engine/Private/WorldPartition/RuntimeSpatialHash/RuntimeSpatialHashGridHelper.h
--- a/Engine/Source/Runtime/Engine/Private/WorldPartition/RuntimeSpatialHash/RuntimeSpatialHashGridHelper.h
+++ b/Engine/Source/Runtime/Engine/Private/WorldPartition/RuntimeSpatialHash/RuntimeSpatialHashGridHelper.h
@@ -46,11 +46,11 @@ struct FSquare2DGridHelper
 		 *
 		 * @return true if the specified index was valid
 		 */
-		inline bool GetCellBounds(int32 InIndex, FBox2D& OutBounds) const
+		inline bool GetCellBounds(int64 InIndex, FBox2D& OutBounds) const
 		{
-			if (InIndex >= 0 && InIndex <= (GridSize * GridSize))
+			if (InIndex >= 0 && InIndex <= ((int64)GridSize * (int64)GridSize))
 			{
-				const FIntVector2 Coords(InIndex % GridSize, InIndex / GridSize);
+				const FIntVector2 Coords(InIndex % (int64)GridSize, InIndex / (int64)GridSize);
 				return GetCellBounds(Coords, OutBounds);
 			}
 
@@ -122,11 +122,11 @@ struct FSquare2DGridHelper
 		 *
 		 * @return true if the coords was inside the grid
 		 */
-		inline bool GetCellIndex(const FIntVector2& InCoords, uint32& OutIndex) const
+		inline bool GetCellIndex(const FIntVector2& InCoords, uint64& OutIndex) const
 		{
 			if (IsValidCoords(InCoords))
 			{
-				OutIndex = (InCoords.Y * GridSize) + InCoords.X;
+				OutIndex = ((int64)InCoords.Y * (int64)GridSize) + (int64)InCoords.X;
 				return true;
 			}
 
@@ -138,7 +138,7 @@ struct FSquare2DGridHelper
 		 *
 		 * @return true if the position was inside the grid
 		 */
-		inline bool GetCellIndex(const FVector& InPos, uint32& OutIndex) const
+		inline bool GetCellIndex(const FVector& InPos, uint64& OutIndex) const
 		{
 			FIntVector2 Coords = FIntVector2(
 				FMath::FloorToInt(((InPos.X - Origin.X) / CellSize) + GridSize * 0.5f),
@@ -229,7 +229,7 @@ struct FSquare2DGridHelper
 
 			ForEachIntersectingCells(Box, [this, &InSphere, &InOperation, &NumCells](const FIntVector2& Coords)
 			{
-				const int32 CellIndex = Coords.Y * GridSize + Coords.X;
+				const int64 CellIndex = (int64)Coords.Y * (int64)GridSize + (int64)Coords.X;
 
 				FBox2D CellBounds;
 				GetCellBounds(CellIndex, CellBounds);
@@ -344,11 +344,11 @@ struct FSquare2DGridHelper
 		{
 			check(IsValidCoords(InCoords));
 
-			uint32 CellIndex;
+			uint64 CellIndex;
 			GetCellIndex(InCoords, CellIndex);
 
-			int32 CellIndexMapping;
-			int32* CellIndexMappingPtr = CellsMapping.Find(CellIndex);
+			int64 CellIndexMapping;
+			int64* CellIndexMappingPtr = CellsMapping.Find(CellIndex);
 			if (CellIndexMappingPtr)
 			{
 				CellIndexMapping = *CellIndexMappingPtr;
@@ -371,10 +371,10 @@ struct FSquare2DGridHelper
 		{
 			check(IsValidCoords(InCoords));
 
-			uint32 CellIndex;
+			uint64 CellIndex;
 			GetCellIndex(InCoords, CellIndex);
 
-			int32 CellIndexMapping = CellsMapping.FindChecked(CellIndex);
+			int64 CellIndexMapping = CellsMapping.FindChecked(CellIndex);
 
 			const FGridCell& Cell = Cells[CellIndexMapping];
 			check(Cell.GetCoords() == FIntVector(InCoords.X, InCoords.Y, Level));
@@ -383,7 +383,7 @@ struct FSquare2DGridHelper
 
 		int32 Level;
 		TArray<FGridCell> Cells;
-		TMap<int32, int32> CellsMapping;
+		TMap<int64, int64> CellsMapping;
 #endif
 	};
 
