Adding cook tags for Texture2D

diff --git a/Engine/Source/Runtime/Engine/Private/TextureDerivedData.cpp b/Engine/Source/Runtime/Engine/Private/TextureDerivedData.cpp
--- a/Engine/Source/Runtime/Engine/Private/TextureDerivedData.cpp
+++ b/Engine/Source/Runtime/Engine/Private/TextureDerivedData.cpp
@@ -45,6 +45,7 @@
 #include "Interfaces/ITextureFormat.h"
 #include "Misc/StringBuilder.h"
 #include "ProfilingDebugging/CookStats.h"
+#include "UObject/ArchiveCookContext.h"
 #include "VT/VirtualTextureDataBuilder.h"
 #include "VT/LightmapVirtualTexture.h"
 #include "TextureCompiler.h"
@@ -2438,6 +2439,36 @@ static void SerializePlatformData(
 				PlatformData->VTData->Chunks[ChunkIndex].BulkData.SetBulkDataFlags(BULKDATA_Force_NOT_InlinePayload);
 			}	
 		}
+
+		// Save cook tags
+#if WITH_EDITORONLY_DATA
+		if (Ar.GetCookContext() && Ar.GetCookContext()->GetCookTagList())
+		{
+			FCookTagList* CookTags = Ar.GetCookContext()->GetCookTagList();
+
+			if (bIsVirtual)
+			{
+				FVirtualTextureBuiltData* VTData = PlatformData->VTData;
+				CookTags->Add(Texture, "Size", FString::Printf(TEXT("%dx%d"), VTData->Width, VTData->Height));
+			}
+			else
+			{
+				FString DimensionsStr;
+				FTexture2DMipMap& TopMip = PlatformData->Mips[FirstMipToSerialize];
+				if (TopMip.SizeZ != 1)
+				{
+					DimensionsStr = FString::Printf(TEXT("%dx%dx%d"), TopMip.SizeX, TopMip.SizeY, TopMip.SizeZ);
+				}
+				else
+				{
+					DimensionsStr = FString::Printf(TEXT("%dx%d"), TopMip.SizeX, TopMip.SizeY);
+				}
+				CookTags->Add(Texture, "Size", MoveTemp(DimensionsStr));	
+			}
+
+			CookTags->Add(Texture, "Format", FString(GPixelFormats[PlatformData->PixelFormat].Name));
+		}
+#endif
 	}
 	Ar << NumMips;
 	check(NumMips >= (int32)PlatformData->GetNumMipsInTail());
