Fix AncientGame systematically re-creating an AWorldDataLayers actor when booting in the editor

diff --git a/Engine/Source/Runtime/Engine/Private/WorldPartition/WorldPartition.cpp b/Engine/Source/Runtime/Engine/Private/WorldPartition/WorldPartition.cpp
--- a/Engine/Source/Runtime/Engine/Private/WorldPartition/WorldPartition.cpp
+++ b/Engine/Source/Runtime/Engine/Private/WorldPartition/WorldPartition.cpp
@@ -492,6 +492,22 @@ void UWorldPartition::Initialize(UWorld* InWorld, const FTransform& InTransform)
 				break;
 			}
 		}
+		
+		// Repair existing maps with no UWorld::WorldDataLayers actor.
+		AWorldDataLayers* WorldDataLayers = World->GetWorldDataLayers();
+		if (!WorldDataLayers)
+		{
+			if (WorldDataLayersActor.IsValid())
+			{
+				WorldDataLayers = CastChecked<AWorldDataLayers>((*WorldDataLayersActor)->GetActor());
+			}
+			else
+			{
+				WorldDataLayers = AWorldDataLayers::Create(World);
+			}
+
+			World->SetWorldDataLayers(WorldDataLayers);
+		}
 	}
 
 	if (bIsEditor && !bIsCooking)
@@ -784,14 +800,14 @@ UWorldPartition* UWorldPartition::CreateOrRepairWorldPartition(AWorldSettings* W
 
 		WorldPartition->DefaultHLODLayer = UHLODLayer::GetEngineDefaultHLODLayersSetup();
 
-		FWorldPartitionMiniMapHelper::GetWorldPartitionMiniMap(World, true);
-	}
+		AWorldDataLayers* WorldDataLayers = World->GetWorldDataLayers();
+		if (!WorldDataLayers)
+		{
+			WorldDataLayers = AWorldDataLayers::Create(World);
+			World->SetWorldDataLayers(WorldDataLayers);
+		}
 
-	AWorldDataLayers* WorldDataLayers = World->GetWorldDataLayers();
-	if (!WorldDataLayers)
-	{
-		WorldDataLayers = AWorldDataLayers::Create(World);
-		World->SetWorldDataLayers(WorldDataLayers);
+		FWorldPartitionMiniMapHelper::GetWorldPartitionMiniMap(World, true);
 	}
 
 	if (!WorldPartition->EditorHash)
