Implemented a new flag on static mesh component called "Disallow Nanite", that even if the static mesh asset has Nanite enabled, you will always get the fallback mesh and regular SM scene proxy created for that SMC in the world.

diff --git a/Engine/Source/Runtime/Engine/Classes/Components/StaticMeshComponent.h b/Engine/Source/Runtime/Engine/Classes/Components/StaticMeshComponent.h
--- a/Engine/Source/Runtime/Engine/Classes/Components/StaticMeshComponent.h
+++ b/Engine/Source/Runtime/Engine/Classes/Components/StaticMeshComponent.h
@@ -223,6 +223,9 @@ public:
 	UPROPERTY(EditAnywhere, AdvancedDisplay, BlueprintReadOnly, Category=Rendering, meta=(editcondition = "bOverrideWireframeColor"))
 	FColor WireframeColorOverride;
 
+	UPROPERTY(EditAnywhere, AdvancedDisplay, BlueprintReadWrite, Category = Rendering)
+	uint8 bDisallowNanite : 1;
+
 	UPROPERTY(EditAnywhere, AdvancedDisplay, BlueprintReadWrite, Category=Rendering)
 	uint8 bEvaluateWorldPositionOffset : 1;
 
diff --git a/Engine/Source/Runtime/Engine/Private/Components/StaticMeshComponent.cpp b/Engine/Source/Runtime/Engine/Private/Components/StaticMeshComponent.cpp
--- a/Engine/Source/Runtime/Engine/Private/Components/StaticMeshComponent.cpp
+++ b/Engine/Source/Runtime/Engine/Private/Components/StaticMeshComponent.cpp
@@ -200,6 +200,7 @@ UStaticMeshComponent::UStaticMeshComponent(const FObjectInitializer& ObjectIniti
 	bOverrideNavigationExport = false;
 	bForceNavigationObstacle = true;
 	bDisallowMeshPaintPerInstance = false;
+	bDisallowNanite = false;
 	bEvaluateWorldPositionOffset = false;
 	DistanceFieldIndirectShadowMinVisibility = .1f;
 	GetBodyInstance()->bAutoWeld = true;	//static mesh by default has auto welding
@@ -1622,6 +1623,12 @@ void UStaticMeshComponent::UpdatePreCulledData(int32 LODIndex, const TArray<uint
 
 bool UStaticMeshComponent::ShouldCreateNaniteProxy() const
 {
+	if (bDisallowNanite)
+	{
+		// Regardless of the static mesh asset supporting Nanite, this component does not want Nanite to be used
+		return false;
+	}
+
 	// Whether or not to allow Nanite for this component
 #if WITH_EDITORONLY_DATA
 	const bool bAllowNanite = !bDisplayNaniteFallbackMesh;
