Moving STAT_TickLatentActions to the actual busy part of the function and separating out removal scope. This profiling overhead currently takes 85% of the function's samples.

diff --git a/Engine/Source/Runtime/Engine/Private/LatentActionManager.cpp b/Engine/Source/Runtime/Engine/Private/LatentActionManager.cpp
--- a/Engine/Source/Runtime/Engine/Private/LatentActionManager.cpp
+++ b/Engine/Source/Runtime/Engine/Private/LatentActionManager.cpp
@@ -70,6 +70,7 @@ int32 FLatentActionManager::GetNumActionsForObject(TWeakObjectPtr<UObject> InObj
 
 
 DECLARE_CYCLE_STAT(TEXT("Blueprint Latent Actions"), STAT_TickLatentActions, STATGROUP_Game);
+DECLARE_CYCLE_STAT(TEXT("Remove Latent Actions"), STAT_RemoveLatentActions, STATGROUP_Game);
 
 void FLatentActionManager::BeginFrame()
 {
@@ -153,8 +154,6 @@ struct FScopedLatentActionTimer
 
 void FLatentActionManager::ProcessLatentActions(UObject* InObject, float DeltaTime)
 {
-	SCOPE_CYCLE_COUNTER(STAT_TickLatentActions);
-
 	if (InObject && !InObject->GetClass()->HasAnyClassFlags(CLASS_CompiledFromBlueprint))
 	{
 		return;
@@ -165,29 +164,32 @@ void FLatentActionManager::ProcessLatentActions(UObject* InObject, float DeltaTi
 	const double StartTime = FPlatformTime::Seconds();
 #endif // LATENT_ACTION_PROFILING_ENABLED
 
-	for (FActionsForObject::TIterator It(ActionsToRemoveMap); It; ++It)
+	if (!ActionsToRemoveMap.IsEmpty())
 	{
-		FObjectActions* ObjectActions = GetActionsForObject(It->Key);
-		TSharedPtr<TArray<FUuidAndAction>> ActionToRemoveListPtr = It->Value;
-		if (ActionToRemoveListPtr.IsValid() && ObjectActions)
+		SCOPE_CYCLE_COUNTER(STAT_RemoveLatentActions);
+		for (FActionsForObject::TIterator It(ActionsToRemoveMap); It; ++It)
 		{
-			for (const FUuidAndAction& PendingActionToKill : *ActionToRemoveListPtr)
+			FObjectActions* ObjectActions = GetActionsForObject(It->Key);
+			TSharedPtr<TArray<FUuidAndAction>> ActionToRemoveListPtr = It->Value;
+			if (ActionToRemoveListPtr.IsValid() && ObjectActions)
 			{
-				FPendingLatentAction* Action = PendingActionToKill.Value;
-				const int32 RemovedNum = ObjectActions->ActionList.RemoveSingle(PendingActionToKill.Key, Action);
-				if (RemovedNum && Action)
+				for (const FUuidAndAction& PendingActionToKill : *ActionToRemoveListPtr)
 				{
-					Action->NotifyActionAborted();
-					delete Action;
+					FPendingLatentAction* Action = PendingActionToKill.Value;
+					const int32 RemovedNum = ObjectActions->ActionList.RemoveSingle(PendingActionToKill.Key, Action);
+					if (RemovedNum && Action)
+					{
+						Action->NotifyActionAborted();
+						delete Action;
+					}
 				}
-			}
 
-			// Notify listeners that latent actions for this object were removed
-			LatentActionsChangedDelegate.Broadcast(It->Key.Get(), ELatentActionChangeType::ActionsRemoved);
+				// Notify listeners that latent actions for this object were removed
+				LatentActionsChangedDelegate.Broadcast(It->Key.Get(), ELatentActionChangeType::ActionsRemoved);
+			}
 		}
-
+		ActionsToRemoveMap.Reset();
 	}
-	ActionsToRemoveMap.Reset();
 
 	if (InObject)
 	{
@@ -195,6 +197,7 @@ void FLatentActionManager::ProcessLatentActions(UObject* InObject, float DeltaTi
 		{
 			if (!ObjectActions->bProcessedThisFrame)
 			{
+				SCOPE_CYCLE_COUNTER(STAT_TickLatentActions);
 #if LATENT_ACTION_PROFILING_ENABLED
 				FScopedLatentActionTimer Timer(InObject);
 #endif // LATENT_ACTION_PROFILING_ENABLED
@@ -204,8 +207,9 @@ void FLatentActionManager::ProcessLatentActions(UObject* InObject, float DeltaTi
 			}
 		}
 	}
-	else 
+	else if (!ObjectToActionListMap.IsEmpty())
 	{
+		SCOPE_CYCLE_COUNTER(STAT_TickLatentActions);
 		for (FObjectToActionListMap::TIterator ObjIt(ObjectToActionListMap); ObjIt; ++ObjIt)
 		{	
 			TWeakObjectPtr<UObject> WeakPtr = ObjIt.Key();
