ZenCOTF: Fix ODSC related shutdown crashes and deadlocks caused by issueing shader compilation requests after cotf server connection has been closed by StorageServerPlatformFile.

#rnx

diff --git a/Engine/Source/Runtime/CookOnTheFly/Private/CookOnTheFlyServerConnection.cpp b/Engine/Source/Runtime/CookOnTheFly/Private/CookOnTheFlyServerConnection.cpp
--- a/Engine/Source/Runtime/CookOnTheFly/Private/CookOnTheFlyServerConnection.cpp
+++ b/Engine/Source/Runtime/CookOnTheFly/Private/CookOnTheFlyServerConnection.cpp
@@ -191,16 +191,13 @@ public:
 
 	virtual void Disconnect() override
 	{
-		if (IsConnected())
+		if (!bStopRequested.Exchange(true))
 		{
-			if (!bStopRequested.Exchange(true))
-			{
-				Socket->Close();
-				Thread->Join();
-				Thread.Reset();
-				Socket.Reset();
-				ClientId  = 0;
-			}
+			Socket->Close();
+			Thread->Join();
+			Thread.Reset();
+			Socket.Reset();
+			ClientId  = 0;
 		}
 	}
 
diff --git a/Engine/Source/Runtime/Engine/Private/ODSC/ODSCManager.cpp b/Engine/Source/Runtime/Engine/Private/ODSC/ODSCManager.cpp
--- a/Engine/Source/Runtime/Engine/Private/ODSC/ODSCManager.cpp
+++ b/Engine/Source/Runtime/Engine/Private/ODSC/ODSCManager.cpp
@@ -18,17 +18,30 @@ FODSCManager::FODSCManager()
 {
 	if (IsRunningCookOnTheFly())
 	{
+		FCoreDelegates::OnEnginePreExit.AddRaw(this, &FODSCManager::OnEnginePreExit);
 		Thread = new FODSCThread();
 		Thread->StartThread();
 	}
 }
 
 FODSCManager::~FODSCManager()
+{
+	FCoreDelegates::OnEnginePreExit.RemoveAll(this);
+	StopThread();
+}
+
+void FODSCManager::OnEnginePreExit()
+{
+	StopThread();
+}
+
+void FODSCManager::StopThread()
 {
 	if (Thread)
 	{
 		Thread->StopThread();
 		delete Thread;
+		Thread = nullptr;
 	}
 }
 
@@ -50,26 +63,25 @@ bool FODSCManager::Tick(float DeltaSeconds)
 			ProcessCookOnTheFlyShaders(false, CompletedRequest->GetMeshMaterialMaps(), CompletedRequest->GetMaterialsToLoad(), CompletedRequest->GetGlobalShaderMap());
 			delete CompletedRequest;
 		}
+		// keep ticking
+		return true;
 	}
-
-	// keep ticking
-	return true;
+	// stop ticking
+	return false;
 }
 
 void FODSCManager::AddThreadedRequest(const TArray<FString>& MaterialsToCompile, EShaderPlatform ShaderPlatform, ODSCRecompileCommand RecompileCommandType)
 {
-	if (IsRunningCookOnTheFly())
+	if (Thread)
 	{
-		check(Thread);
 		Thread->AddRequest(MaterialsToCompile, ShaderPlatform, RecompileCommandType);
 	}
 }
 
 void FODSCManager::AddThreadedShaderPipelineRequest(EShaderPlatform ShaderPlatform, const FString& MaterialName, const FString& VertexFactoryName, const FString& PipelineName, const TArray<FString>& ShaderTypeNames)
 {
-	if (IsRunningCookOnTheFly())
+	if (Thread)
 	{
-		check(Thread);
 		Thread->AddShaderPipelineRequest(ShaderPlatform, MaterialName, VertexFactoryName, PipelineName, ShaderTypeNames);
 	}
 }
diff --git a/Engine/Source/Runtime/Engine/Public/ODSC/ODSCManager.h b/Engine/Source/Runtime/Engine/Public/ODSC/ODSCManager.h
--- a/Engine/Source/Runtime/Engine/Public/ODSC/ODSCManager.h
+++ b/Engine/Source/Runtime/Engine/Public/ODSC/ODSCManager.h
@@ -67,9 +67,12 @@ public:
 	void AddThreadedShaderPipelineRequest(EShaderPlatform ShaderPlatform, const FString& MaterialName, const FString& VertexFactoryName, const FString& PipelineName, const TArray<FString>& ShaderTypeNames);
 private:
 
+	void OnEnginePreExit();
+	void StopThread();
+
 	/** Handles communicating directly with the cook on the fly server. */
 	FODSCThread* Thread = nullptr;
 };
 
 /** The global shader ODSC manager. */
-extern ENGINE_API FODSCManager* GODSCManager;
\ No newline at end of file
+extern ENGINE_API FODSCManager* GODSCManager;
diff --git a/Engine/Source/Runtime/StorageServerClient/Private/StorageServerPlatformFile.cpp b/Engine/Source/Runtime/StorageServerClient/Private/StorageServerPlatformFile.cpp
--- a/Engine/Source/Runtime/StorageServerClient/Private/StorageServerPlatformFile.cpp
+++ b/Engine/Source/Runtime/StorageServerClient/Private/StorageServerPlatformFile.cpp
@@ -777,6 +777,10 @@ int64 FStorageServerPlatformFile::SendReadMessage(uint8* Destination, const FIoC
 bool FStorageServerPlatformFile::SendMessageToServer(const TCHAR* Message, IPlatformFile::IFileServerMessageHandler* Handler)
 {
 #if WITH_COTF
+	if (!CookOnTheFlyServerConnection->IsConnected())
+	{
+		return false;
+	}
 	if (FCString::Stricmp(Message, TEXT("RecompileShaders")) == 0)
 	{
 		UE::Cook::FCookOnTheFlyRequest Request(UE::Cook::ECookOnTheFlyMessage::RecompileShaders);
