COTF & ODSC: Fix assert in FMaterialShaderMap::LoadForRemoteRecompile when a requested material has already been destroyed and can't be found any more.
Fix serialization bug that cause old file based COTF to ignore what shader platform we were requesting (fixed for Zen COTF in 19595445).
Extend the -noshaderddc command line option with -noglobalshaderddc and -nomaterialshaderddc options (to improve COTF debugging iteration times).

#rnx

diff --git a/Engine/Source/Runtime/Engine/Private/Materials/MaterialShader.cpp b/Engine/Source/Runtime/Engine/Private/Materials/MaterialShader.cpp
--- a/Engine/Source/Runtime/Engine/Private/Materials/MaterialShader.cpp
+++ b/Engine/Source/Runtime/Engine/Private/Materials/MaterialShader.cpp
@@ -1446,7 +1446,9 @@ void FMaterialShaderMap::LoadFromDerivedDataCache(const FMaterial* Material, con
 
 			// If NoShaderDDC then don't check for a material the first time we encounter it to simulate
 			// a cold DDC
-			static bool bNoShaderDDC = FParse::Param(FCommandLine::Get(), TEXT("noshaderddc"));
+			static bool bNoShaderDDC =
+				FParse::Param(FCommandLine::Get(), TEXT("noshaderddc")) || 
+				FParse::Param(FCommandLine::Get(), TEXT("nomaterialshaderddc"));
 
 			if (bNoShaderDDC)
 			{
@@ -1627,22 +1629,21 @@ void FMaterialShaderMap::SaveForRemoteRecompile(FArchive& Ar, const TMap<FString
 
 void FMaterialShaderMap::LoadForRemoteRecompile(FArchive& Ar, EShaderPlatform ShaderPlatform, TArray<UMaterialInterface*>& OutLoadedMaterials)
 {
+	FString MaterialName;
+	TArray<TRefCountPtr<FMaterialShaderMap>> LoadedShaderMaps;
+
 	int32 MapSize;
 	Ar << MapSize;
 
 	for (int32 MaterialIndex = 0; MaterialIndex < MapSize; MaterialIndex++)
 	{
-		FString MaterialName;
+		MaterialName.Reset();
 		Ar << MaterialName;
 
-		UMaterialInterface* MatchingMaterial = FindObjectChecked<UMaterialInterface>(NULL, *MaterialName);
-		OutLoadedMaterials.Add(MatchingMaterial);
-
 		int32 NumShaderMaps = 0;
 		Ar << NumShaderMaps;
 
-		TArray<TRefCountPtr<FMaterialShaderMap> > LoadedShaderMaps;
-
+		LoadedShaderMaps.Reset();
 		for (int32 ShaderMapIndex = 0; ShaderMapIndex < NumShaderMaps; ShaderMapIndex++)
 		{
 			uint8 bIsValid = 0;
@@ -1662,6 +1663,13 @@ void FMaterialShaderMap::LoadForRemoteRecompile(FArchive& Ar, EShaderPlatform Sh
 			}
 		}
 
+		UMaterialInterface* MatchingMaterial = FindObject<UMaterialInterface>(nullptr, *MaterialName);
+		if (!MatchingMaterial)
+		{
+			continue;
+		}
+		OutLoadedMaterials.Add(MatchingMaterial);
+
 		// Assign in two passes: first pass for shader maps with unspecified quality levels,
 		// Second pass for shader maps with a specific quality level
 		for (int32 PassIndex = 0; PassIndex < 2; PassIndex++)
diff --git a/Engine/Source/Runtime/Engine/Private/ShaderCompiler/ShaderCompiler.cpp b/Engine/Source/Runtime/Engine/Private/ShaderCompiler/ShaderCompiler.cpp
--- a/Engine/Source/Runtime/Engine/Private/ShaderCompiler/ShaderCompiler.cpp
+++ b/Engine/Source/Runtime/Engine/Private/ShaderCompiler/ShaderCompiler.cpp
@@ -6095,7 +6095,7 @@ void ProcessCookOnTheFlyShaders(bool bReloadGlobalShaders, const TArray<uint8>&
 		if (!bHasFlushed )
 		{
 			// now we need to refresh the RHI resources
-	FlushRenderingCommands();
+			FlushRenderingCommands();
 			bHasFlushed = true;
 		}
 	};
@@ -7005,7 +7005,9 @@ void CompileGlobalShaderMap(EShaderPlatform Platform, const ITargetPlatform* Tar
 
 			// If NoShaderDDC then don't check for a material the first time we encounter it to simulate
 			// a cold DDC
-			static bool bNoShaderDDC = FParse::Param(FCommandLine::Get(), TEXT("noshaderddc"));
+			static bool bNoShaderDDC =
+				FParse::Param(FCommandLine::Get(), TEXT("noshaderddc")) || 
+				FParse::Param(FCommandLine::Get(), TEXT("noglobalshaderddc"));
 			if (UNLIKELY(bNoShaderDDC))
 			{
 				bShaderMapIsBeingCompiled = true;
diff --git a/Engine/Source/Runtime/NetworkFileSystem/Private/NetworkFileServerConnection.cpp b/Engine/Source/Runtime/NetworkFileSystem/Private/NetworkFileServerConnection.cpp
--- a/Engine/Source/Runtime/NetworkFileSystem/Private/NetworkFileServerConnection.cpp
+++ b/Engine/Source/Runtime/NetworkFileSystem/Private/NetworkFileServerConnection.cpp
@@ -1294,13 +1294,14 @@ void FNetworkFileServerClientConnection::ProcessRecompileShaders( FArchive& In,
 	TArray<FString> RecompileModifiedFiles;
 	TArray<uint8> MeshMaterialMaps;
 	TArray<uint8> GlobalShaderMap;
-	FShaderRecompileData RecompileData(ConnectedPlatformName, SP_NumPlatforms, ODSCRecompileCommand::Changed, &RecompileModifiedFiles, &MeshMaterialMaps, &GlobalShaderMap);
+	FShaderRecompileData RecompileData(ConnectedPlatformName, &RecompileModifiedFiles, &MeshMaterialMaps, &GlobalShaderMap);
 
 	// tell other side all the materials to load, by pathname
 	In << RecompileData.MaterialsToLoad;
 
 	int32 iShaderPlatform = static_cast<int32>(RecompileData.ShaderPlatform);
 	In << iShaderPlatform;
+	RecompileData.ShaderPlatform = static_cast<EShaderPlatform>(iShaderPlatform);
 	In << RecompileData.CommandType;
 	In << RecompileData.ShadersToRecompile;
 
