Make it so that async physics tick is called per physics step even if you're not using async physics

diff --git a/Engine/Source/Runtime/Engine/Classes/Components/ActorComponent.h b/Engine/Source/Runtime/Engine/Classes/Components/ActorComponent.h
--- a/Engine/Source/Runtime/Engine/Classes/Components/ActorComponent.h
+++ b/Engine/Source/Runtime/Engine/Classes/Components/ActorComponent.h
@@ -299,7 +299,7 @@ private:
 
 public:
 
-	/** Whether to use use the async physics tick with this component. Requires async physics to be enabled (see project settings) */
+	/** Whether to use use the async physics tick with this component. */
 	uint8 bAsyncPhysicsTickEnabled : 1;
 
 	/** Describes how a component instance will be created */
@@ -721,12 +721,11 @@ public:
 	virtual void TickComponent(float DeltaTime, enum ELevelTick TickType, FActorComponentTickFunction *ThisTickFunction);
 
 	/**
-	 * Override this function to implement custom logic to be executed every async physics step.
+	 * Override this function to implement custom logic to be executed every physics step.
 	 * Only executes if the component is registered, and also bAsyncPhysicsTick must be set to true.
-	 * Requires async physics to be enabled (see physics project settings)
 	 *	
-	 * @param DeltaTime - The async physics time associated with each step (see project settings)
-	 * @param SimTime - This is the total sim time since the sim began, it is an integer multiple of DeltaTime
+	 * @param DeltaTime - The sim time associated with each step
+	 * @param SimTime - This is the total sim time since the sim began
 	 */
 	virtual void AsyncPhysicsTickComponent(float DeltaTime, float SimTime) { ReceiveAsyncPhysicsTick(DeltaTime, SimTime); }
 	
@@ -1046,7 +1045,7 @@ public:
 	UFUNCTION(BlueprintImplementableEvent, meta=(DisplayName = "Tick"))
 	void ReceiveTick(float DeltaSeconds);
 
-	/** Event called every async physics tick if bAsyncPhysicsTickEnabled is true and async physics is enabled (see physics project settings)*/
+	/** Event called every async physics tick if bAsyncPhysicsTickEnabled is true */
 	UFUNCTION(BlueprintImplementableEvent, meta = (DisplayName = "Async Physics Tick"))
 	void ReceiveAsyncPhysicsTick(float DeltaSeconds, float SimSeconds);
 	
diff --git a/Engine/Source/Runtime/Engine/Classes/GameFramework/Actor.h b/Engine/Source/Runtime/Engine/Classes/GameFramework/Actor.h
--- a/Engine/Source/Runtime/Engine/Classes/GameFramework/Actor.h
+++ b/Engine/Source/Runtime/Engine/Classes/GameFramework/Actor.h
@@ -464,7 +464,7 @@ private:
 	/** Set if an Actor tries to be destroyed while it is beginning play so that once BeginPlay ends we can issue the destroy call. */
 	uint8 bActorWantsDestroyDuringBeginPlay : 1;
 	
-	/** Whether to use use the async physics tick with this actor. Requires async physics be enabled (see project settings) */
+	/** Whether to use use the async physics tick with this actor. */
 	UPROPERTY(EditAnywhere, Category=Physics)
 	uint8 bAsyncPhysicsTickEnabled : 1;
 
@@ -1808,7 +1808,7 @@ public:
 	UFUNCTION(BlueprintImplementableEvent, meta=(DisplayName = "Tick"))
 	void ReceiveTick(float DeltaSeconds);
 
-	/** Event called every async physics tick if bAsyncPhysicsTickEnabled is true and async physics is enabled (see physics project settings)*/
+	/** Event called every physics tick if bAsyncPhysicsTickEnabled is true */
 	UFUNCTION(BlueprintImplementableEvent, meta = (DisplayName = "Async Physics Tick"))
 	void ReceiveAsyncPhysicsTick(float DeltaSeconds, float SimSeconds);
 
@@ -2541,12 +2541,11 @@ public:
 	virtual void TickActor( float DeltaTime, enum ELevelTick TickType, FActorTickFunction& ThisTickFunction );
 
 	/**
-	 * Override this function to implement custom logic to be executed every async physics step.
+	 * Override this function to implement custom logic to be executed every physics step.
 	 * bAsyncPhysicsTick must be set to true.
-	 * Requires async physics to be enabled (see physics project settings)
 	 *	
-	 * @param DeltaTime - The physics step delta time (see physics project settings)
-	 * @param SimTime - This is the total sim time since the sim began, it is an integer multiple of DeltaTime
+	 * @param DeltaTime - The physics step delta time
+	 * @param SimTime - This is the total sim time since the sim began.
 	 */
 	virtual void AsyncPhysicsTickActor(float DeltaTime, float SimTime) { ReceiveAsyncPhysicsTick(DeltaTime, SimTime); }
 
diff --git a/Engine/Source/Runtime/Experimental/Chaos/Private/Chaos/Framework/PhysicsSolverBase.cpp b/Engine/Source/Runtime/Experimental/Chaos/Private/Chaos/Framework/PhysicsSolverBase.cpp
--- a/Engine/Source/Runtime/Experimental/Chaos/Private/Chaos/Framework/PhysicsSolverBase.cpp
+++ b/Engine/Source/Runtime/Experimental/Chaos/Private/Chaos/Framework/PhysicsSolverBase.cpp
@@ -432,7 +432,7 @@ namespace Chaos
 				PendingTasks = TGraphTask<FPhysicsSolverProcessPushDataTask>::CreateTask(&Prereqs).ConstructAndDispatchWhenReady(*this, PushData);
 				Prereqs.Add(PendingTasks);
 
-				if(bSolverHasFrozenGameThreadCallbacks && IsUsingFixedDt())
+				if(bSolverHasFrozenGameThreadCallbacks)
 				{
 					PendingTasks = TGraphTask<FPhysicsSolverFrozenGTPreSimCallbacks>::CreateTask(&Prereqs).ConstructAndDispatchWhenReady(*this);
 					Prereqs.Add(PendingTasks);
@@ -461,10 +461,7 @@ namespace Chaos
 	void FAllSolverTasks::AdvanceSolver()
 	{
 		ProcessPushData.ProcessPushData();
-		if (Solver.IsUsingFixedDt())
-		{
-			GTPreSimCallbacks.GTPreSimCallbacks();
-		}
+		GTPreSimCallbacks.GTPreSimCallbacks();
 		AdvanceTask.AdvanceSolver();
 	}
 }
