Fix crash in ImgMediaMipMapInfoManager with -stereoenabled.

diff --git a/Engine/Source/Runtime/Engine/Private/UnrealEngine.cpp b/Engine/Source/Runtime/Engine/Private/UnrealEngine.cpp
--- a/Engine/Source/Runtime/Engine/Private/UnrealEngine.cpp
+++ b/Engine/Source/Runtime/Engine/Private/UnrealEngine.cpp
@@ -3481,15 +3481,20 @@ public:
 
 	virtual EStereoscopicPass GetViewPassForIndex(bool bStereoRequested, int32 ViewIndex) const override
 	{
+		if (!bStereoRequested)
+		{
+			return EStereoscopicPass::eSSP_FULL;
+		}
 #if WITH_MGPU
 		// We have to do all the work for secondary views that are on a different GPU than
 		// the primary view. NB: This assumes that the primary view is assigned to the
 		// first GPU of the AFR group. See FSceneRenderer::ComputeViewGPUMasks.
-		if (bStereoRequested && GNumExplicitGPUsForRendering > 1)
+		else if (GNumExplicitGPUsForRendering > 1)
 		{
 			return EStereoscopicPass::eSSP_PRIMARY;
 		}
 #endif
+
 		return ViewIndex % 2 == 0 ? EStereoscopicPass::eSSP_PRIMARY : EStereoscopicPass::eSSP_SECONDARY;
 	}
 
