Add preload dependencies on sparse class data & its archetype for BP CDOs

Ensures that sparse class data of archetypes is preloaded before attempting to serialize a BPs sparse class data

diff --git a/Engine/Source/Runtime/Engine/Private/BlueprintGeneratedClass.cpp b/Engine/Source/Runtime/Engine/Private/BlueprintGeneratedClass.cpp
--- a/Engine/Source/Runtime/Engine/Private/BlueprintGeneratedClass.cpp
+++ b/Engine/Source/Runtime/Engine/Private/BlueprintGeneratedClass.cpp
@@ -1534,6 +1534,40 @@ void UBlueprintGeneratedClass::GetDefaultObjectPreloadDependencies(TArray<UObjec
 			}
 		}
 	}
+
+	// Add objects related to the sparse class data struct
+	if (SparseClassDataStruct)
+	{
+		// Add the super classes CDO, as it needs to be in place to properly serialize the sparse class data's archetype
+		OutDeps.Add(GetSuperClass()->GetDefaultObject());
+	
+		// Add the sparse class data struct, as it is serialized as part of the CDO
+		OutDeps.Add(SparseClassDataStruct);
+
+		// Also add the sparse class data archetype as it will be copied into any newly created sparse class data 
+		// pre-serialization
+		if(UScriptStruct* SparseClassDataArchetype = GetSparseClassDataArchetypeStruct())
+		{
+			OutDeps.Add(SparseClassDataArchetype);
+		}
+
+		// Add anything that the sparse class data also depends on
+		const void* SparseClassDataToUse = GetSparseClassData(EGetSparseClassDataMethod::ArchetypeIfNull);
+		if (UScriptStruct::ICppStructOps* CppStructOps = SparseClassDataStruct->GetCppStructOps())
+		{
+			CppStructOps->GetPreloadDependencies(const_cast<void*>(SparseClassDataToUse), OutDeps);
+		}
+		// The iterator will recursively loop through all structs in structs/containers too.
+		for (TPropertyValueIterator<FStructProperty> It(SparseClassDataStruct, SparseClassDataToUse); It; ++It)
+		{
+			const UScriptStruct* StructType = It.Key()->Struct;
+			if (UScriptStruct::ICppStructOps* CppStructOps = StructType->GetCppStructOps())
+			{
+				void* StructDataPtr = const_cast<void*>(It.Value());
+				CppStructOps->GetPreloadDependencies(StructDataPtr, OutDeps);
+			}
+		}
+	}
 }
 
 bool UBlueprintGeneratedClass::NeedsLoadForServer() const
