Fix UNetDriver::IsLevelInitializedForActor ignoring child connections
#8670

diff --git a/Engine/Source/Runtime/Engine/Private/NetDriver.cpp b/Engine/Source/Runtime/Engine/Private/NetDriver.cpp
--- a/Engine/Source/Runtime/Engine/Private/NetDriver.cpp
+++ b/Engine/Source/Runtime/Engine/Private/NetDriver.cpp
@@ -2076,8 +2076,30 @@ bool UNetDriver::IsLevelInitializedForActor(const AActor* InActor, const UNetCon
 
 	// we can't create channels while the client is in the wrong world
 	const bool bCorrectWorld = GetWorldPackage() != nullptr && (InConnection->GetClientWorldPackageName() == GetWorldPackage()->GetFName()) && InConnection->ClientHasInitializedLevelFor(InActor);
+	
 	// exception: Special case for PlayerControllers as they are required for the client to travel to the new world correctly			
-	const bool bIsConnectionPC = (InActor == InConnection->PlayerController);
+	bool bIsConnectionPC = false;
+	// Make sure we have Parent Connection
+	const UNetConnection* const ParentConnection = InConnection->IsA<UChildConnection>() ? Cast<UChildConnection>(InConnection)->Parent : InConnection;
+
+	// Check Parent Connection
+	if (InActor == ParentConnection->PlayerController)
+	{
+		bIsConnectionPC = true;
+	}
+	// Check Child Connections
+	else
+	{	
+		for (const UChildConnection* const ChildConnection : ParentConnection->Children)
+		{
+			if (InActor == ChildConnection->PlayerController)
+			{
+				bIsConnectionPC = true;
+				break;
+			}
+		}
+	}
+
 	return bCorrectWorld || bIsConnectionPC;
 }
 
