nop just comments

diff --git a/Engine/Source/Runtime/Engine/Private/Texture2D.cpp b/Engine/Source/Runtime/Engine/Private/Texture2D.cpp
--- a/Engine/Source/Runtime/Engine/Private/Texture2D.cpp
+++ b/Engine/Source/Runtime/Engine/Private/Texture2D.cpp
@@ -508,7 +508,7 @@ float UTexture2D::GetAverageBrightness(bool bIgnoreTrueBlack, bool bUseGrayscale
 	float AvgBrightness = -1.0f;
 #if WITH_EDITOR
 
-	// @@!! this is done in a few places ; factor out to an FImage helper?
+	// @todo Oodle : GetAverageColor is done in a few places ; factor out to an FImage helper?
 
 	TArray64<uint8> RawData;
 	// use the source art if it exists
diff --git a/Engine/Source/Runtime/Engine/Private/TextureLODSettings.cpp b/Engine/Source/Runtime/Engine/Private/TextureLODSettings.cpp
--- a/Engine/Source/Runtime/Engine/Private/TextureLODSettings.cpp
+++ b/Engine/Source/Runtime/Engine/Private/TextureLODSettings.cpp
@@ -270,7 +270,7 @@ void UTextureLODSettings::GetMipGenSettings(const UTexture& Texture, TextureMipG
 	}
 
 	// angular filtering only applies to cubemaps
-	// @@!! currently intentionally NOT allowed on cubearrays
+	// note: currently intentionally NOT allowed on cubearrays
 	if (Setting == TMGS_Angular && !Texture.IsA(UTextureCube::StaticClass()))
 	{
 		Setting = TMGS_NoMipmaps;
