When sampling VT use a fallback color held in a uniform if the root page isn't mapped.

diff --git a/Engine/Source/Runtime/Engine/Private/TextureDerivedData.cpp b/Engine/Source/Runtime/Engine/Private/TextureDerivedData.cpp
--- a/Engine/Source/Runtime/Engine/Private/TextureDerivedData.cpp
+++ b/Engine/Source/Runtime/Engine/Private/TextureDerivedData.cpp
@@ -69,7 +69,7 @@
 // This GUID is mixed into DDC version for virtual textures only, this allows updating DDC version for VT without invalidating DDC for all textures
 // This is useful during development, but once large numbers of VT are present in shipped content, it will have the same problem as TEXTURE_DERIVEDDATA_VER
 // This is put in the DDC1 key but NOT in the DDC2 key
-#define TEXTURE_VT_DERIVEDDATA_VER	TEXT("F36E4B807BC24Y818FE085C655C45176")
+#define TEXTURE_VT_DERIVEDDATA_VER	TEXT("7C16439390E24F1F9468894FB4D4BC54")
 
 #if ENABLE_COOK_STATS
 namespace TextureCookStats
