Throttle texture compiler reschedule when called repeatedly
  - 2m53s to 235ms of reschedule time during level load / PIE and stop

#rnx

diff --git a/Engine/Source/Runtime/Engine/Private/TextureCompiler.cpp b/Engine/Source/Runtime/Engine/Private/TextureCompiler.cpp
--- a/Engine/Source/Runtime/Engine/Private/TextureCompiler.cpp
+++ b/Engine/Source/Runtime/Engine/Private/TextureCompiler.cpp
@@ -578,8 +578,10 @@ void FTextureCompilingManager::ProcessTextures(bool bLimitExecutionTime, int32 M
 			}
 		}
 
-		if (GEngine)
+		if (GEngine && FPlatformTime::Seconds() - LastReschedule > 1.0f)
 		{
+			LastReschedule = FPlatformTime::Seconds();
+
 			TRACE_CPUPROFILER_EVENT_SCOPE(FTextureCompilingManager::Reschedule);
 
 			auto TryRescheduleTexture = 
diff --git a/Engine/Source/Runtime/Engine/Public/TextureCompiler.h b/Engine/Source/Runtime/Engine/Public/TextureCompiler.h
--- a/Engine/Source/Runtime/Engine/Public/TextureCompiler.h
+++ b/Engine/Source/Runtime/Engine/Public/TextureCompiler.h
@@ -97,6 +97,7 @@ private:
 	void PostCompilation(UTexture* Texture);
 	void PostCompilation(TArrayView<UTexture* const> InCompiledTextures);
 
+	double LastReschedule = 0.0f;
 	bool bHasShutdown = false;
 	TArray<TSet<TWeakObjectPtr<UTexture>>> RegisteredTextureBuckets;
 	FAsyncCompilationNotification Notification;
