Strata - also link refraction when dealing with material attributes.

diff --git a/Engine/Source/Runtime/Engine/Private/Materials/Material.cpp b/Engine/Source/Runtime/Engine/Private/Materials/Material.cpp
--- a/Engine/Source/Runtime/Engine/Private/Materials/Material.cpp
+++ b/Engine/Source/Runtime/Engine/Private/Materials/Material.cpp
@@ -2927,6 +2927,7 @@ void UMaterial::ConvertMaterialToStrataMaterial()
 		WorldPositionOffset.Connect(10, BreakMatAtt);
 		AmbientOcclusion.Connect(14, BreakMatAtt);
 		PixelDepthOffset.Connect(24, BreakMatAtt);
+		Refraction.Connect(15, BreakMatAtt);
 
 		if (ShadingModel == MSM_FromMaterialExpression)
 		{
