Strata - SLW non legacy node is now promoted to a operator so that the correct data can be retrieved in the based pass. When using this new node, and the legacy conversion node both follow the same path.

diff --git a/Engine/Source/Runtime/Engine/Private/Materials/HLSLMaterialTranslator.cpp b/Engine/Source/Runtime/Engine/Private/Materials/HLSLMaterialTranslator.cpp
--- a/Engine/Source/Runtime/Engine/Private/Materials/HLSLMaterialTranslator.cpp
+++ b/Engine/Source/Runtime/Engine/Private/Materials/HLSLMaterialTranslator.cpp
@@ -10397,10 +10397,11 @@ int32 FHLSLMaterialTranslator::StrataHairBSDF(int32 BaseColor, int32 Scatter, in
 
 int32 FHLSLMaterialTranslator::StrataSingleLayerWaterBSDF(
 	int32 BaseColor, int32 Metallic, int32 Specular, int32 Roughness, int32 EmissiveColor, int32 TopMaterialOpacity,
-	int32 WaterAlbedo, int32 WaterExtinction, int32 WaterPhaseG, int32 ColorScaleBehindWater, int32 Normal, const FString& SharedLocalBasisIndexMacro)
+	int32 WaterAlbedo, int32 WaterExtinction, int32 WaterPhaseG, int32 ColorScaleBehindWater, int32 Normal, const FString& SharedLocalBasisIndexMacro,
+	FStrataOperator* PromoteToOperator)
 {
 	return AddCodeChunk(
-		MCT_Strata, TEXT("GetStrataSingleLayerWaterBSDF(%s, %s, %s, %s, %s, %s, %s, %s, %s, %s, %s) /* %s */"),
+		MCT_Strata, TEXT("PromoteParameterBlendedBSDFToOperator(GetStrataSingleLayerWaterBSDF(%s, %s, %s, %s, %s, %s, %s, %s, %s, %s, %s), Parameters.StrataTree, %u, %u, %u, %u) /* %s */"),
 		*GetParameterCode(BaseColor),
 		*GetParameterCode(Metallic),
 		*GetParameterCode(Specular),
@@ -10412,6 +10413,10 @@ int32 FHLSLMaterialTranslator::StrataSingleLayerWaterBSDF(
 		*GetParameterCode(WaterPhaseG),
 		*GetParameterCode(ColorScaleBehindWater),
 		*SharedLocalBasisIndexMacro,
+		PromoteToOperator->Index,
+		PromoteToOperator->BSDFIndex,
+		PromoteToOperator->bIsBottom ? 1 : 0,
+		PromoteToOperator->bIsTop ? 1 : 0,
 		*GetParameterCode(Normal)
 	);
 }
diff --git a/Engine/Source/Runtime/Engine/Private/Materials/HLSLMaterialTranslator.h b/Engine/Source/Runtime/Engine/Private/Materials/HLSLMaterialTranslator.h
--- a/Engine/Source/Runtime/Engine/Private/Materials/HLSLMaterialTranslator.h
+++ b/Engine/Source/Runtime/Engine/Private/Materials/HLSLMaterialTranslator.h
@@ -1042,7 +1042,7 @@ protected:
 	virtual int32 StrataSingleLayerWaterBSDF(
 		int32 BaseColor, int32 Metallic, int32 Specular, int32 Roughness, 
 		int32 EmissiveColor, int32 TopMaterialOpacity, int32 WaterAlbedo, int32 WaterExtinction, int32 WaterPhaseG, 
-		int32 ColorScaleBehindWater, int32 Normal, const FString& SharedLocalBasisIndexMacro) override;
+		int32 ColorScaleBehindWater, int32 Normal, const FString& SharedLocalBasisIndexMacro, FStrataOperator* PromoteToOperator) override;
 	virtual int32 StrataHorizontalMixing(int32 Background, int32 Foreground, int32 Mix, int OperatorIndex, uint32 MaxDistanceFromLeaves) override;
 	virtual int32 StrataHorizontalMixingParameterBlending(int32 Background, int32 Foreground, int32 HorizontalMixCodeChunk, int32 NormalMixCodeChunk, const FString& SharedLocalBasisIndexMacro, FStrataOperator* PromoteToOperator) override;
 	virtual int32 StrataVerticalLayering(int32 Top, int32 Base, int OperatorIndex, uint32 MaxDistanceFromLeaves) override;
diff --git a/Engine/Source/Runtime/Engine/Private/Materials/MaterialExpressions.cpp b/Engine/Source/Runtime/Engine/Private/Materials/MaterialExpressions.cpp
--- a/Engine/Source/Runtime/Engine/Private/Materials/MaterialExpressions.cpp
+++ b/Engine/Source/Runtime/Engine/Private/Materials/MaterialExpressions.cpp
@@ -22057,6 +22057,18 @@ int32 UMaterialExpressionStrataSingleLayerWaterBSDF::Compile(class FMaterialComp
 	int32 NormalCodeChunk = CompileWithDefaultNormalWS(Compiler, Normal);
 	const FStrataRegisteredSharedLocalBasis NewRegisteredSharedLocalBasis = StrataCompilationInfoCreateSharedLocalBasis(Compiler, NormalCodeChunk);
 
+	FStrataOperator& StrataOperator = Compiler->StrataCompilationGetOperator(this);
+	StrataOperator.BSDFRegisteredSharedLocalBasis = NewRegisteredSharedLocalBasis;
+
+	if (StrataOperator.bUseParameterBlending)
+	{
+		return Compiler->Errorf(TEXT("Strata SingleLayerWater BSDF node cannot be used with parameter blending."));
+	}
+	else if (StrataOperator.bRootOfParameterBlendingSubTree)
+	{
+		return Compiler->Errorf(TEXT("Strata SingleLayerWater BSDF node cannot be the root of a parameter blending sun tree."));
+	}
+
 	int32 OutputCodeChunk = Compiler->StrataSingleLayerWaterBSDF(
 		CompileWithDefaultFloat3(Compiler, BaseColor, 0.0f, 0.0f, 0.0f),
 		CompileWithDefaultFloat1(Compiler, Metallic, 0.0f),
@@ -22069,7 +22081,8 @@ int32 UMaterialExpressionStrataSingleLayerWaterBSDF::Compile(class FMaterialComp
 		CompileWithDefaultFloat1(Compiler, WaterPhaseG, 0.0f),
 		CompileWithDefaultFloat3(Compiler, ColorScaleBehindWater, 1.0f, 1.0f, 1.0f),
 		NormalCodeChunk,
-		Compiler->GetStrataSharedLocalBasisIndexMacro(NewRegisteredSharedLocalBasis));
+		Compiler->GetStrataSharedLocalBasisIndexMacro(NewRegisteredSharedLocalBasis),
+		&StrataOperator);
 
 	return OutputCodeChunk;
 }
diff --git a/Engine/Source/Runtime/Engine/Public/MaterialCompiler.h b/Engine/Source/Runtime/Engine/Public/MaterialCompiler.h
--- a/Engine/Source/Runtime/Engine/Public/MaterialCompiler.h
+++ b/Engine/Source/Runtime/Engine/Public/MaterialCompiler.h
@@ -563,7 +563,7 @@ public:
 	virtual int32 StrataHairBSDF(int32 BaseColor, int32 Scatter, int32 Specular, int32 Roughness, int32 Backlit, int32 EmissiveColor, int32 Tangent, const FString& SharedLocalBasisIndexMacro) = 0;
 	virtual int32 StrataSingleLayerWaterBSDF(
 		int32 BaseColor, int32 Metallic, int32 Specular, int32 Roughness, int32 EmissiveColor, int32 TopMaterialOpacity, 
-		int32 WaterAlbedo, int32 WaterExtinction, int32 WaterPhaseG, int32 ColorScaleBehindWater, int32 Normal, const FString& SharedLocalBasisIndexMacro) = 0;
+		int32 WaterAlbedo, int32 WaterExtinction, int32 WaterPhaseG, int32 ColorScaleBehindWater, int32 Normal, const FString& SharedLocalBasisIndexMacro, FStrataOperator* PromoteToOperator) = 0;
 	virtual int32 StrataHorizontalMixing(int32 Background, int32 Foreground, int32 Mix, int OperatorIndex, uint32 MaxDistanceFromLeaves) = 0;
 	virtual int32 StrataHorizontalMixingParameterBlending(int32 Background, int32 Foreground, int32 HorizontalMixCodeChunk, int32 NormalMixCodeChunk, const FString& SharedLocalBasisIndexMacro, FStrataOperator* PromoteToOperator) = 0;
 	virtual int32 StrataVerticalLayering(int32 Top, int32 Base, int OperatorIndex, uint32 MaxDistanceFromLeaves) = 0;
@@ -1148,11 +1148,11 @@ public:
 
 	virtual int32 StrataSingleLayerWaterBSDF(
 		int32 BaseColor, int32 Metallic, int32 Specular, int32 Roughness, int32 EmissiveColor, int32 TopMaterialOpacity, 
-		int32 WaterAlbedo, int32 WaterExtinction, int32 WaterPhaseG, int32 ColorScaleBehindWater, int32 Normal, const FString& SharedLocalBasisIndexMacro) override
+		int32 WaterAlbedo, int32 WaterExtinction, int32 WaterPhaseG, int32 ColorScaleBehindWater, int32 Normal, const FString& SharedLocalBasisIndexMacro, FStrataOperator* PromoteToOperator) override
 	{
 		return Compiler->StrataSingleLayerWaterBSDF(
 			BaseColor, Metallic, Specular, Roughness, EmissiveColor, TopMaterialOpacity, 
-			WaterAlbedo, WaterExtinction, WaterPhaseG, ColorScaleBehindWater, Normal, SharedLocalBasisIndexMacro);
+			WaterAlbedo, WaterExtinction, WaterPhaseG, ColorScaleBehindWater, Normal, SharedLocalBasisIndexMacro, PromoteToOperator);
 	}
 
 	virtual int32 StrataHorizontalMixing(int32 Background, int32 Foreground, int32 Mix, int OperatorIndex, uint32 MaxDistanceFromLeaves) override
