Move setting of  bWasDuplicated/bWasDuplicatedForPIE in ULevel::Serialize instead of ULevel::PostDuplicate: PostDuplicate can be called after PostLoad in cases where ULevel was ConditionalPostLoad by outered subobject

#rnx

diff --git a/Engine/Source/Runtime/Engine/Classes/Engine/Level.h b/Engine/Source/Runtime/Engine/Classes/Engine/Level.h
--- a/Engine/Source/Runtime/Engine/Classes/Engine/Level.h
+++ b/Engine/Source/Runtime/Engine/Classes/Engine/Level.h
@@ -889,7 +889,6 @@ public:
 	PRAGMA_ENABLE_DEPRECATION_WARNINGS
 	virtual void PreSave(FObjectPreSaveContext ObjectSaveContext) override;
 	virtual void PreDuplicate(FObjectDuplicationParameters& DupParams) override;
-	virtual void PostDuplicate(bool bDuplicateForPIE) override;
 	virtual bool CanBeClusterRoot() const override;
 	virtual void CreateCluster() override;
 	static void AddReferencedObjects(UObject* InThis, FReferenceCollector& Collector);
diff --git a/Engine/Source/Runtime/Engine/Private/Level.cpp b/Engine/Source/Runtime/Engine/Private/Level.cpp
--- a/Engine/Source/Runtime/Engine/Private/Level.cpp
+++ b/Engine/Source/Runtime/Engine/Private/Level.cpp
@@ -708,6 +708,19 @@ void ULevel::Serialize( FArchive& Ar )
 		}
 	}
 #endif
+
+	if (Ar.IsLoading())
+	{
+		if (Ar.GetPortFlags() & PPF_DuplicateForPIE)
+		{
+			bWasDuplicatedForPIE = true;
+		}
+
+		if (bWasDuplicatedForPIE || (Ar.GetPortFlags() & PPF_Duplicate))
+		{
+			bWasDuplicated = true;
+		}
+	}
 }
 
 void ULevel::CreateReplicatedDestructionInfo(AActor* const Actor)
@@ -1153,14 +1166,6 @@ void ULevel::PreDuplicate(FObjectDuplicationParameters& DupParams)
 #endif
 }
 
-void ULevel::PostDuplicate(bool bDuplicateForPIE)
-{
-	Super::PostDuplicate(bDuplicateForPIE);
-
-	bWasDuplicated = true;
-	bWasDuplicatedForPIE = bDuplicateForPIE;
-}
-
 UWorld* ULevel::GetWorld() const
 {
 	return OwningWorld;
