Add ability to output compressible data to JSON, gated behind console command (allowing for outputting specific data: position, scale, rotation, curve)

diff --git a/Engine/Source/Runtime/Engine/Private/Animation/AnimCompressionTypes.cpp b/Engine/Source/Runtime/Engine/Private/Animation/AnimCompressionTypes.cpp
--- a/Engine/Source/Runtime/Engine/Private/Animation/AnimCompressionTypes.cpp
+++ b/Engine/Source/Runtime/Engine/Private/Animation/AnimCompressionTypes.cpp
@@ -13,12 +13,26 @@
 #include "UObject/FortniteMainBranchObjectVersion.h"
 #include "Animation/AnimSequenceHelpers.h"
 #include "Animation/AnimData/AnimDataModel.h"
+#include "Misc/FileHelper.h"
+#include "Serialization/JsonWriter.h"
 
 CSV_DECLARE_CATEGORY_MODULE_EXTERN(ENGINE_API, Animation);
 
 DECLARE_CYCLE_STAT(TEXT("Build Anim Track Pairs"), STAT_BuildAnimTrackPairs, STATGROUP_Anim);
 DECLARE_CYCLE_STAT(TEXT("Extract Pose From Anim Data"), STAT_ExtractPoseFromAnimData, STATGROUP_Anim);
 
+static FString GCompressionJsonOutput;
+static FAutoConsoleVariableRef CVarCompressionJsonOutput(
+	TEXT("a.Compression.CompressibleDataOutput"),
+	GCompressionJsonOutput,
+	TEXT("Whether to output any JSON file containing the compressible data. (comma delimited)\n")
+	TEXT(" position: output track positional data\n")
+	TEXT(" rotation: output track rotational data\n")
+	TEXT(" scale: output track scale data\n")
+	TEXT(" curve: output rich curve data\n"),
+	ECVF_Cheat
+	);
+
 template<typename ArrayValue>
 void StripFramesEven(TArray<ArrayValue>& Keys, const int32 NumFrames)
 {
@@ -84,6 +98,215 @@ ICompressedAnimData::ICompressedAnimData(const ICompressedAnimData&) = default;
 ICompressedAnimData& ICompressedAnimData::operator=(const ICompressedAnimData&) = default;
 PRAGMA_ENABLE_DEPRECATION_WARNINGS
 
+void FCompressibleAnimData::WriteCompressionDataToJSON(TArrayView<FName> OriginalTrackNames, TArrayView<FRawAnimSequenceTrack> FinalRawAnimationData, TArrayView<FName> FinalTrackNames) const
+{
+	const bool bPositionalData = GCompressionJsonOutput.Contains(TEXT("position"));
+	const bool bRotationalData = GCompressionJsonOutput.Contains(TEXT("rotation"));
+	const bool bScalingData = GCompressionJsonOutput.Contains(TEXT("scale"));
+	const bool bCurveData = GCompressionJsonOutput.Contains(TEXT("curve"));
+
+	if (bPositionalData || bRotationalData || bScalingData || bCurveData)
+	{
+		FString JSONString;
+		TSharedRef<TJsonStringWriter<>> Writer = TJsonStringWriter<>::Create(&JSONString);
+		const UEnum* InterpolationEnum = FindObject<UEnum>(ANY_PACKAGE, TEXT("EAnimInterpolationType"), true);
+
+		Writer->WriteObjectStart();
+		{
+			// Name
+			Writer->WriteValue(TEXT("name"), Name);
+
+			// Interpolation type
+			Writer->WriteValue(TEXT("interpolation"), InterpolationEnum->GetValueAsString(Interpolation));
+
+			// Keys
+			Writer->WriteValue(TEXT("number_of_keys"), NumberOfKeys);
+
+			// Length
+			Writer->WriteValue(TEXT("length_in_seconds"), SequenceLength);
+
+			// Raw Animation
+			if ((bPositionalData || bRotationalData || bScalingData) && FinalRawAnimationData.Num())
+			{
+				Writer->WriteArrayStart(TEXT("animation_tracks"));
+				{
+					for (int32 TrackIndex = 0; TrackIndex < FinalRawAnimationData.Num(); ++TrackIndex)
+					{
+						Writer->WriteObjectStart();
+							
+						// Track name
+						Writer->WriteValue(TEXT("name"), FinalTrackNames[TrackIndex].ToString());
+
+						const FRawAnimSequenceTrack& Track = FinalRawAnimationData[TrackIndex];
+
+						// Position
+						if (bPositionalData)
+						{
+							Writer->WriteArrayStart(TEXT("positional_data"));
+							{
+								for (int32 KeyIndex = 0; KeyIndex < Track.PosKeys.Num(); ++KeyIndex)
+								{
+									Writer->WriteValue(Track.PosKeys[KeyIndex].ToString());
+								}
+							}
+							Writer->WriteArrayEnd();
+						}
+							
+
+						// Rotation
+						if (bRotationalData)
+						{
+							Writer->WriteArrayStart(TEXT("rotational_data"));
+							{
+								for (int32 KeyIndex = 0; KeyIndex < Track.RotKeys.Num(); ++KeyIndex)
+								{
+									Writer->WriteValue(Track.RotKeys[KeyIndex].ToString());
+								}
+							}
+							Writer->WriteArrayEnd();
+						}
+
+						// Scale
+						if (bScalingData)
+						{
+							Writer->WriteArrayStart(TEXT("scaling_data"));
+							{
+								for (int32 KeyIndex = 0; KeyIndex < Track.ScaleKeys.Num(); ++KeyIndex)
+								{
+									Writer->WriteValue(Track.ScaleKeys[KeyIndex].ToString());
+								}
+							}
+							Writer->WriteArrayEnd();
+						}
+
+						Writer->WriteObjectEnd();
+					}
+				}
+				Writer->WriteArrayEnd();
+
+				// Additive Animation
+				if(bIsValidAdditive && AdditiveBaseAnimationData.Num())
+				{
+					Writer->WriteArrayStart(TEXT("additive_base_tracks"));
+					{
+						for (int32 TrackIndex = 0; TrackIndex < AdditiveBaseAnimationData.Num(); ++TrackIndex)
+						{
+							const FRawAnimSequenceTrack& Track = AdditiveBaseAnimationData[TrackIndex];
+							Writer->WriteObjectStart();
+							{
+								// Track name
+								Writer->WriteValue(TEXT("name"), OriginalTrackNames[TrackIndex].ToString());							
+									
+								// Position
+								if (bPositionalData)
+								{
+									Writer->WriteArrayStart(TEXT("positional_data"));
+									{
+										for (int32 KeyIndex = 0; KeyIndex < Track.PosKeys.Num(); ++KeyIndex)
+										{
+											Writer->WriteValue(Track.PosKeys[KeyIndex].ToString());
+										}
+									}
+									Writer->WriteArrayEnd();
+								}
+							
+
+								// Rotation
+								if (bRotationalData)
+								{
+									Writer->WriteArrayStart(TEXT("rotational_data"));
+									{
+										for (int32 KeyIndex = 0; KeyIndex < Track.RotKeys.Num(); ++KeyIndex)
+										{
+											Writer->WriteValue(Track.RotKeys[KeyIndex].ToString());
+										}
+									}
+									Writer->WriteArrayEnd();
+								}
+
+								// Scale
+								if (bScalingData)
+								{
+									Writer->WriteArrayStart(TEXT("scaling_data"));
+									{
+										for (int32 KeyIndex = 0; KeyIndex < Track.ScaleKeys.Num(); ++KeyIndex)
+										{
+											Writer->WriteValue(Track.ScaleKeys[KeyIndex].ToString());
+										}
+									}
+									Writer->WriteArrayEnd();
+								}
+							}
+							Writer->WriteObjectEnd();
+						}
+					}
+					Writer->WriteArrayEnd();
+				}
+			}
+
+			if (bCurveData && RawFloatCurves.Num())
+			{
+				// Num curves
+				Writer->WriteValue(TEXT("number_of_curves"), RawFloatCurves.Num());
+					
+				Writer->WriteArrayStart(TEXT("curve_data"));		
+				for (const FFloatCurve& FloatCurve : RawFloatCurves)
+				{
+					Writer->WriteObjectStart();
+					{
+						Writer->WriteValue(TEXT("curve_name"), FloatCurve.Name.DisplayName.ToString());
+						Writer->WriteValue(TEXT("number_of_keys"), FloatCurve.FloatCurve.GetNumKeys());
+							
+						if(FloatCurve.FloatCurve.GetConstRefOfKeys().Num())
+						{
+							const UEnum* CurveInterpolationEnum = FindObject<UEnum>(ANY_PACKAGE, TEXT("ERichCurveInterpMode"), true);
+							const UEnum* TangentModeEnum = FindObject<UEnum>(ANY_PACKAGE, TEXT("ERichCurveTangentMode"), true);
+							const UEnum* TangentWeightModeEnum = FindObject<UEnum>(ANY_PACKAGE, TEXT("ERichCurveTangentWeightMode"), true);
+								
+							Writer->WriteArrayStart(TEXT("key_data"));
+							for (const FRichCurveKey& Key : FloatCurve.FloatCurve.GetConstRefOfKeys())
+							{
+								Writer->WriteObjectStart();
+								{
+									Writer->WriteValue(TEXT("time"), Key.Time);
+									Writer->WriteValue(TEXT("value"), Key.Value);
+
+									Writer->WriteValue(TEXT("arrive_tangent"), Key.ArriveTangent);
+									Writer->WriteValue(TEXT("arrive_tangent_weight"), Key.ArriveTangentWeight);
+									Writer->WriteValue(TEXT("leave_tangent"), Key.LeaveTangent);
+									Writer->WriteValue(TEXT("leave_tangent_weight"), Key.LeaveTangentWeight);
+										
+									Writer->WriteValue(TEXT("interpolation_mode"), CurveInterpolationEnum->GetNameStringByValue(Key.InterpMode));
+									Writer->WriteValue(TEXT("tangent_mode"), TangentModeEnum->GetNameStringByValue(Key.TangentMode));
+									Writer->WriteValue(TEXT("tangent_weight_mode"), TangentWeightModeEnum->GetNameStringByValue(Key.TangentWeightMode));
+								}
+								Writer->WriteObjectEnd();
+							}
+							Writer->WriteArrayEnd();
+						}
+					}
+					Writer->WriteObjectEnd();
+				}
+				Writer->WriteArrayEnd();
+			}
+		}
+		Writer->WriteObjectEnd();
+		Writer->Close();
+
+		const FString BasePath = FPaths::ProjectSavedDir();
+		const FString FolderPath = BasePath + TEXT("/CompressibleData/");
+		FString NameAsFileName = FullName;
+		NameAsFileName = NameAsFileName.Replace(TEXT("/"), TEXT("_"));
+		int32 LastFullStop = INDEX_NONE;
+		NameAsFileName.FindLastChar('.', LastFullStop);
+		ensure(LastFullStop != INDEX_NONE);
+		NameAsFileName.RemoveAt(LastFullStop, NameAsFileName.Len() - LastFullStop);
+					
+		const FString FilePath = FolderPath + NameAsFileName + TEXT(".json");
+		FFileHelper::SaveStringToFile(JSONString, *FilePath);
+	}
+}
+
 FCompressibleAnimData::FCompressibleAnimData(class UAnimSequence* InSeq, const bool bPerformStripping)
 	: CurveCompressionSettings(InSeq->CurveCompressionSettings)
 	, BoneCompressionSettings(InSeq->BoneCompressionSettings)
@@ -112,30 +335,30 @@ FCompressibleAnimData::FCompressibleAnimData(class UAnimSequence* InSeq, const b
 	/* Always get the resampled data to start off with */
 	const TArray<FBoneAnimationTrack>& ResampledTrackData = InSeq->GetResampledTrackData();
 
-	TArray<FName> TempTrackNames;
+	TArray<FName> OriginalTrackNames;
 	RawAnimationData.Empty(ResampledTrackData.Num());
 	TrackToSkeletonMapTable.Empty(ResampledTrackData.Num());
-	TempTrackNames.Empty(ResampledTrackData.Num());
+	OriginalTrackNames.Empty(ResampledTrackData.Num());
 
 	for (const FBoneAnimationTrack& AnimTrack : ResampledTrackData)
 	{
 		FRawAnimSequenceTrack& Track = RawAnimationData.Add_GetRef(AnimTrack.InternalTrackData);
 		UE::Anim::Compression::SanitizeRawAnimSequenceTrack(Track);
 		TrackToSkeletonMapTable.Add(AnimTrack.BoneTreeIndex);
-		TempTrackNames.Add(AnimTrack.Name);
+		OriginalTrackNames.Add(AnimTrack.Name);
 	}
 
 	const bool bBakeAdditive = InSeq->CanBakeAdditive();
 	if (bBakeAdditive)
 	{
-		InSeq->BakeOutAdditiveIntoRawData(RawAnimationData, TempTrackNames, TrackToSkeletonMapTable, RawFloatCurves, AdditiveBaseAnimationData);
+		InSeq->BakeOutAdditiveIntoRawData(RawAnimationData, OriginalTrackNames, TrackToSkeletonMapTable, RawFloatCurves, AdditiveBaseAnimationData);
 	}
 	else
     {
 		// In case we require baking down transform curves, do so now meaning Virtual Bone baking will incorporate the correct bone transforms
 	    if (InSeq->GetDataModel()->GetNumberOfTransformCurves() > 0)
 	    {
-		    InSeq->BakeTrackCurvesToRawAnimationTracks(RawAnimationData, TempTrackNames, TrackToSkeletonMapTable);
+		    InSeq->BakeTrackCurvesToRawAnimationTracks(RawAnimationData, OriginalTrackNames, TrackToSkeletonMapTable);
 	    }
 
 		RawFloatCurves = InSeq->GetCurveData().FloatCurves;
@@ -143,7 +366,7 @@ FCompressibleAnimData::FCompressibleAnimData(class UAnimSequence* InSeq, const b
 		// If we aren't additive we must bake virtual bones
 	    if (bHasVirtualBones)
 	    {
-		    InSeq->BakeOutVirtualBoneTracks(RawAnimationData, TempTrackNames, TrackToSkeletonMapTable);
+		    InSeq->BakeOutVirtualBoneTracks(RawAnimationData, OriginalTrackNames, TrackToSkeletonMapTable);
 	    }
     } 
 
@@ -169,17 +392,22 @@ FCompressibleAnimData::FCompressibleAnimData(class UAnimSequence* InSeq, const b
 	// Verify bone track names and data, removing any bone that does not exist on the skeleton
     // And for additive animations remove any track deemed not to add any additive animation (identity rotation and zero-vector translation and scale)
 	// Note on (TrackIndex > 0) below : deliberately stop before track 0, compression code doesn't like getting a completely empty animation
+	TArray<FName> FinalTrackNames;
+	FinalTrackNames.Reserve(ResampledTrackData.Num());	
 	for (int32 TrackIndex = RawAnimationData.Num() - 1; TrackIndex > 0; --TrackIndex)
 	{
 		const FRawAnimSequenceTrack& Track = RawAnimationData[TrackIndex];
 		// Try find correct bone index
-		const int32 BoneIndex = RefSkeleton.FindBoneIndex(TempTrackNames[TrackIndex]);
+		const int32 BoneIndex = RefSkeleton.FindBoneIndex(OriginalTrackNames[TrackIndex]);
 		if ((bBakeAdditive && IsRawTrackValidForRemoval(Track)) || BoneIndex == INDEX_NONE)
 		{
 			RawAnimationData.RemoveAtSwap(TrackIndex, 1, false);
-			TempTrackNames.RemoveAtSwap(TrackIndex, 1, false);
 			TrackToSkeletonMapTable.RemoveAtSwap(TrackIndex, 1, false);
 		}
+		else
+        {
+			FinalTrackNames.Add(OriginalTrackNames[TrackIndex]);
+        }
 	}
 
 	// Find or add curve names on skeleton
@@ -226,6 +454,17 @@ FCompressibleAnimData::FCompressibleAnimData(class UAnimSequence* InSeq, const b
 	PRAGMA_DISABLE_DEPRECATION_WARNINGS
 	NumberOfFrames = NumberOfKeys;
     PRAGMA_ENABLE_DEPRECATION_WARNINGS
+
+	if (GCompressionJsonOutput.Len())
+	{
+		// Insert root-bone name
+		if (RawAnimationData.Num())
+		{			
+			FinalTrackNames.Insert(OriginalTrackNames[0], 0);
+		}
+		
+		WriteCompressionDataToJSON(MakeArrayView(OriginalTrackNames), MakeArrayView(RawAnimationData), MakeArrayView(FinalTrackNames));	
+	}
 #endif
 }
 
diff --git a/Engine/Source/Runtime/Engine/Public/Animation/AnimCompressionTypes.h b/Engine/Source/Runtime/Engine/Public/Animation/AnimCompressionTypes.h
--- a/Engine/Source/Runtime/Engine/Public/Animation/AnimCompressionTypes.h
+++ b/Engine/Source/Runtime/Engine/Public/Animation/AnimCompressionTypes.h
@@ -296,7 +296,7 @@ public:
 	}
 
 private:
-
+	void WriteCompressionDataToJSON(TArrayView<FName> OriginalTrackNames, TArrayView<FRawAnimSequenceTrack> FinalRawAnimationData, TArrayView<FName> FinalTrackNames) const;
 };
 
 typedef TSharedPtr<FCompressibleAnimData, ESPMode::ThreadSafe> FCompressibleAnimPtr;
@@ -826,4 +826,4 @@ extern void DecompressPose(	FCompactPose& OutPose,
 							bool bIsBakedAdditive,
 							FName RetargetSource,
 							FName SourceName,
-							const FRootMotionReset& RootMotionReset);
\ No newline at end of file
+							const FRootMotionReset& RootMotionReset);
