Add ability to verify compressed curve evaluated data against source curves (gated behind CVar and disabled by default)

diff --git a/Engine/Source/Runtime/Engine/Private/Animation/AnimCurveCompressionCodec_CompressedRichCurve.cpp b/Engine/Source/Runtime/Engine/Private/Animation/AnimCurveCompressionCodec_CompressedRichCurve.cpp
--- a/Engine/Source/Runtime/Engine/Private/Animation/AnimCurveCompressionCodec_CompressedRichCurve.cpp
+++ b/Engine/Source/Runtime/Engine/Private/Animation/AnimCurveCompressionCodec_CompressedRichCurve.cpp
@@ -3,6 +3,7 @@
 #include "Animation/AnimCurveCompressionCodec_CompressedRichCurve.h"
 #include "Animation/AnimSequence.h"
 #include "Serialization/MemoryWriter.h"
+#include "AnimationCompression.h"
 
 UAnimCurveCompressionCodec_CompressedRichCurve::UAnimCurveCompressionCodec_CompressedRichCurve(const FObjectInitializer& ObjectInitializer)
 	: Super(ObjectInitializer)
@@ -26,6 +27,13 @@ struct FCurveDesc
 };
 
 #if WITH_EDITORONLY_DATA
+
+int32 ValidateCompressedRichCurveEvaluation = 0;
+static FAutoConsoleVariableRef CVarValidateCompressedRichCurveEvaluation(
+	TEXT("a.Compression.ValidateCompressedRichCurveEvaluation"),
+	ValidateCompressedRichCurveEvaluation,
+	TEXT("1 = runs validation, evaluating the compressed rich curve at animation its sampling rate comparing against the MaxCurveError. 0 = validation disabled"));
+
 bool UAnimCurveCompressionCodec_CompressedRichCurve::Compress(const FCompressibleAnimData& AnimSeq, FAnimCurveCompressionResult& OutResult)
 {
 	int32 NumCurves = AnimSeq.RawFloatCurves.Num();
@@ -62,6 +70,26 @@ bool UAnimCurveCompressionCodec_CompressedRichCurve::Compress(const FCompressibl
 
 		KeyDataOffset += CompressedCurve.CompressedKeys.Num();
 		KeyData.Append(CompressedCurve.CompressedKeys);
+
+		// Compressed curve validation
+		if(!!ValidateCompressedRichCurveEvaluation)
+		{
+			const int32 SamplesToTake = AnimSeq.NumberOfKeys;
+			for (int32 KeyIndex = 0; KeyIndex < SamplesToTake; ++KeyIndex)
+			{
+				// Evaluated RawCurve and Compressed Curve
+				const float EvalTime = KeyIndex * Helper.TimePerKey();
+				const float RawValue = RawCurve.Eval(EvalTime);
+				const float CompressedValue = CompressedCurve.Eval(EvalTime);
+
+				const float AbsDelta = FMath::Abs(CompressedValue - RawValue);
+				if (AbsDelta > MaxCurveError)
+				{
+					// Delta larger than tolerated error value
+					UE_LOG(LogAnimationCompression, Warning, TEXT("Curve %s: delta too large %f at %f"), *Curve.Name.DisplayName.ToString(), AbsDelta, EvalTime);
+				}
+			}
+		}
 	}
 
 	TArray<uint8> TempBytes;
