Add scalability to the field clipmap construction

diff --git a/Engine/Plugins/Experimental/ChaosNiagara/Source/ChaosNiagara/Private/NiagaraDataInterfacePhysicsField.cpp b/Engine/Plugins/Experimental/ChaosNiagara/Source/ChaosNiagara/Private/NiagaraDataInterfacePhysicsField.cpp
--- a/Engine/Plugins/Experimental/ChaosNiagara/Source/ChaosNiagara/Private/NiagaraDataInterfacePhysicsField.cpp
+++ b/Engine/Plugins/Experimental/ChaosNiagara/Source/ChaosNiagara/Private/NiagaraDataInterfacePhysicsField.cpp
@@ -164,7 +164,15 @@ public:
 				LocalTargets[Index].Z = FieldResource->FieldInfos.IntegerTargets[Index];
 				LocalTargets[Index].W = 0; // Padding
 			}
-			SetSRVParameter(RHICmdList, ComputeShaderRHI, ClipmapBuffer, FieldResource->ClipmapBuffer.SRV);
+			if (FieldResource->FieldInfos.bBuildClipmap)
+			{
+				SetSRVParameter(RHICmdList, ComputeShaderRHI, ClipmapBuffer, FieldResource->ClipmapBuffer.SRV);
+			}
+			else
+			{
+				SetSRVParameter(RHICmdList, ComputeShaderRHI, ClipmapBuffer, FNiagaraRenderer::GetDummyFloatBuffer());
+			}
+			
 			SetShaderValue(RHICmdList, ComputeShaderRHI, ClipmapCenter, (FVector3f)FieldResource->FieldInfos.ClipmapCenter);
 			SetShaderValue(RHICmdList, ComputeShaderRHI, ClipmapDistance, FieldResource->FieldInfos.ClipmapDistance);
 			SetShaderValue(RHICmdList, ComputeShaderRHI, ClipmapResolution, FieldResource->FieldInfos.ClipmapResolution);
diff --git a/Engine/Source/Runtime/Engine/Private/PhysicsField/PhysicsFieldComponent.cpp b/Engine/Source/Runtime/Engine/Private/PhysicsField/PhysicsFieldComponent.cpp
--- a/Engine/Source/Runtime/Engine/Private/PhysicsField/PhysicsFieldComponent.cpp
+++ b/Engine/Source/Runtime/Engine/Private/PhysicsField/PhysicsFieldComponent.cpp
@@ -34,13 +34,20 @@ DECLARE_GPU_STAT(PhysicsFields_UpdateBuffers);
 */
 DEFINE_LOG_CATEGORY_STATIC(LogGlobalField, Log, All);
 
+/** Boolean to check if we need to build or not the clipmap */
+int32 GPhysicsFieldBuildClipmap = 1;
+FAutoConsoleVariableRef CVarPhysicsFieldBuildClipmap(
+	TEXT("r.PhysicsField.BuildClipmap"),
+	GPhysicsFieldBuildClipmap,
+	TEXT("Build the Physics field clipmap"),
+	ECVF_Scalability | ECVF_RenderThreadSafe);
 
 /** Clipmap enable/disable */
 static TAutoConsoleVariable<int32> CVarPhysicsFieldEnableClipmap(
 	TEXT("r.PhysicsField.EnableField"),
 	1,
 	TEXT("Enable/Disable the Physics field clipmap"),
-	ECVF_RenderThreadSafe);
+	ECVF_Scalability | ECVF_RenderThreadSafe);
 
 /** Clipmap max disatnce */
 float GPhysicsFieldClipmapDistance = 10000;
@@ -1365,7 +1372,7 @@ void UPhysicsFieldComponent::OnRegister()
 {
 	Super::OnRegister();
 
-	TArray<bool> bBuildClipmaps = { true,false };
+	TArray<bool> bBuildClipmaps = { GPhysicsFieldBuildClipmap == 1, false };
 	for (uint32 FieldIndex = 0; FieldIndex < 2; ++FieldIndex)
 	{
 		FPhysicsFieldInstance*& LocalInstance = (FieldIndex == 0) ? FieldInstance : DebugInstance;
diff --git a/Engine/Source/Runtime/Renderer/Private/SceneRendering.cpp b/Engine/Source/Runtime/Renderer/Private/SceneRendering.cpp
--- a/Engine/Source/Runtime/Renderer/Private/SceneRendering.cpp
+++ b/Engine/Source/Runtime/Renderer/Private/SceneRendering.cpp
@@ -1219,7 +1219,14 @@ void SetupPhysicsFieldUniformBufferParameters(const FScene* Scene, FEngineShowFl
 	if (Scene && Scene->PhysicsField && Scene->PhysicsField->FieldResource)
 	{
 		FPhysicsFieldResource* FieldResource = Scene->PhysicsField->FieldResource;
-		ViewUniformShaderParameters.PhysicsFieldClipmapBuffer = FieldResource->ClipmapBuffer.SRV.GetReference();
+		if (FieldResource->FieldInfos.bBuildClipmap)
+		{
+			ViewUniformShaderParameters.PhysicsFieldClipmapBuffer = FieldResource->ClipmapBuffer.SRV.GetReference();
+		}
+		else
+		{
+			ViewUniformShaderParameters.PhysicsFieldClipmapBuffer = GWhiteVertexBufferWithSRV->ShaderResourceViewRHI;
+		}
 		ViewUniformShaderParameters.PhysicsFieldClipmapCenter = (FVector3f)FieldResource->FieldInfos.ClipmapCenter;
 		ViewUniformShaderParameters.PhysicsFieldClipmapDistance = FieldResource->FieldInfos.ClipmapDistance;
 		ViewUniformShaderParameters.PhysicsFieldClipmapResolution = FieldResource->FieldInfos.ClipmapResolution;
