Deprecated public access to time-related member variables on FAnimSequenceDecompressionContext and add gett-ers

diff --git a/Engine/Source/Runtime/Engine/Private/Animation/AnimEncoding_ConstantKeyLerp.cpp b/Engine/Source/Runtime/Engine/Private/Animation/AnimEncoding_ConstantKeyLerp.cpp
--- a/Engine/Source/Runtime/Engine/Private/Animation/AnimEncoding_ConstantKeyLerp.cpp
+++ b/Engine/Source/Runtime/Engine/Private/Animation/AnimEncoding_ConstantKeyLerp.cpp
@@ -278,8 +278,8 @@ inline void AEFConstantKeyLerp<FORMAT>::GetPoseRotations(
 			(ispc::BoneTrackPair*)&DesiredPairs[0],
 			AnimData.CompressedTrackOffsets.GetData(),
 			AnimData.CompressedByteStream.GetData(),
-			DecompContext.SequenceLength,
-			DecompContext.RelativePos,
+			DecompContext.GetPlayableLength(),
+			DecompContext.GetRelativePosition(),
 			(uint8)DecompContext.Interpolation,
 			FORMAT,
 			PairCount);
@@ -330,8 +330,8 @@ inline void AEFConstantKeyLerp<FORMAT>::GetPoseTranslations(
 			(ispc::BoneTrackPair*)&DesiredPairs[0],
 			AnimData.CompressedTrackOffsets.GetData(),
 			AnimData.CompressedByteStream.GetData(),
-			DecompContext.SequenceLength,
-			DecompContext.RelativePos,
+			DecompContext.GetPlayableLength(),
+			DecompContext.GetRelativePosition(),
 			(uint8)DecompContext.Interpolation,
 			FORMAT,
 			PairCount);
@@ -402,8 +402,8 @@ inline void AEFConstantKeyLerp<FORMAT>::GetPoseScales(
 			ScaleOffsets.GetData(),
 			StripSize,
 			AnimData.CompressedByteStream.GetData(),
-			DecompContext.SequenceLength,
-			DecompContext.RelativePos,
+			DecompContext.GetPlayableLength(),
+			DecompContext.GetRelativePosition(),
 			(uint8)DecompContext.Interpolation,
 			FORMAT,
 			PairCount);
diff --git a/Engine/Source/Runtime/Engine/Private/Animation/AnimEncoding_PerTrackCompression.cpp b/Engine/Source/Runtime/Engine/Private/Animation/AnimEncoding_PerTrackCompression.cpp
--- a/Engine/Source/Runtime/Engine/Private/Animation/AnimEncoding_PerTrackCompression.cpp
+++ b/Engine/Source/Runtime/Engine/Private/Animation/AnimEncoding_PerTrackCompression.cpp
@@ -488,12 +488,12 @@ void AEFPerTrackCompressionCodec::GetBoneAtomRotation(
 		{
 			if ((FormatFlags & 0x8) == 0)
 			{
-				Alpha = TimeToIndex(DecompContext.SequenceLength, DecompContext.RelativePos, NumKeys, DecompContext.Interpolation, Index0, Index1);
+				Alpha = TimeToIndex(DecompContext.GetPlayableLength(), DecompContext.GetRelativePosition(), NumKeys, DecompContext.Interpolation, Index0, Index1);
 			}
 			else
 			{
 				const uint8* RESTRICT FrameTable = Align(TrackData + FixedBytes + BytesPerKey * NumKeys, 4);
-				Alpha = TimeToIndex(DecompContext.Interpolation, AnimData.CompressedNumberOfKeys, FrameTable, DecompContext.RelativePos, NumKeys, Index0, Index1);
+				Alpha = TimeToIndex(DecompContext.Interpolation, AnimData.CompressedNumberOfKeys, FrameTable, DecompContext.GetRelativePosition(), NumKeys, Index0, Index1);
 			}
 		}
 
@@ -579,12 +579,12 @@ void AEFPerTrackCompressionCodec::GetBoneAtomTranslation(
 		{
 			if ((FormatFlags & 0x8) == 0)
 			{
-				Alpha = TimeToIndex(DecompContext.SequenceLength, DecompContext.RelativePos, NumKeys, DecompContext.Interpolation, Index0, Index1);
+				Alpha = TimeToIndex(DecompContext.GetPlayableLength(), DecompContext.GetRelativePosition(), NumKeys, DecompContext.Interpolation, Index0, Index1);
 			}
 			else
 			{
 				const uint8* RESTRICT FrameTable = Align(TrackData + FixedBytes + BytesPerKey * NumKeys, 4);
-				Alpha = TimeToIndex(DecompContext.Interpolation, AnimData.CompressedNumberOfKeys, FrameTable, DecompContext.RelativePos, NumKeys, Index0, Index1);
+				Alpha = TimeToIndex(DecompContext.Interpolation, AnimData.CompressedNumberOfKeys, FrameTable, DecompContext.GetRelativePosition(), NumKeys, Index0, Index1);
 			}
 		}
 
@@ -665,12 +665,12 @@ void AEFPerTrackCompressionCodec::GetBoneAtomScale(
 		{
 			if ((FormatFlags & 0x8) == 0)
 			{
-				Alpha = TimeToIndex(DecompContext.SequenceLength, DecompContext.RelativePos, NumKeys, DecompContext.Interpolation, Index0, Index1);
+				Alpha = TimeToIndex(DecompContext.GetPlayableLength(), DecompContext.GetRelativePosition(), NumKeys, DecompContext.Interpolation, Index0, Index1);
 			}
 			else
 			{
 				const uint8* RESTRICT FrameTable = Align(TrackData + FixedBytes + BytesPerKey * NumKeys, 4);
-				Alpha = TimeToIndex(DecompContext.Interpolation, AnimData.CompressedNumberOfKeys, FrameTable, DecompContext.RelativePos, NumKeys, Index0, Index1);
+				Alpha = TimeToIndex(DecompContext.Interpolation, AnimData.CompressedNumberOfKeys, FrameTable, DecompContext.GetRelativePosition(), NumKeys, Index0, Index1);
 			}
 		}
 
@@ -751,8 +751,8 @@ void AEFPerTrackCompressionCodec::GetPoseRotations(
 			AnimData.CompressedTrackOffsets.GetData(),
 			AnimData.CompressedByteStream.GetData(),
 			AnimData.CompressedNumberOfKeys,
-			DecompContext.SequenceLength,
-			DecompContext.RelativePos,
+			DecompContext.GetPlayableLength(),
+			DecompContext.GetRelativePosition(),
 			(uint8)DecompContext.Interpolation,
 			PairCount);
 #endif
@@ -799,8 +799,8 @@ void AEFPerTrackCompressionCodec::GetPoseTranslations(
 			AnimData.CompressedTrackOffsets.GetData(),
 			AnimData.CompressedByteStream.GetData(),
 			AnimData.CompressedNumberOfKeys,
-			DecompContext.SequenceLength,
-			DecompContext.RelativePos,
+			DecompContext.GetPlayableLength(),
+			DecompContext.GetRelativePosition(),
 			(uint8)DecompContext.Interpolation,
 			PairCount);
 #endif
@@ -853,8 +853,8 @@ void AEFPerTrackCompressionCodec::GetPoseScales(
 			StripSize,
 			AnimData.CompressedByteStream.GetData(),
 			AnimData.CompressedNumberOfKeys,
-			DecompContext.SequenceLength,
-			DecompContext.RelativePos,
+			DecompContext.GetPlayableLength(),
+			DecompContext.GetRelativePosition(),
 			(uint8)DecompContext.Interpolation,
 			PairCount);
 #endif
diff --git a/Engine/Source/Runtime/Engine/Private/Animation/AnimEncoding_VariableKeyLerp.cpp b/Engine/Source/Runtime/Engine/Private/Animation/AnimEncoding_VariableKeyLerp.cpp
--- a/Engine/Source/Runtime/Engine/Private/Animation/AnimEncoding_VariableKeyLerp.cpp
+++ b/Engine/Source/Runtime/Engine/Private/Animation/AnimEncoding_VariableKeyLerp.cpp
@@ -244,7 +244,7 @@ void AEFVariableKeyLerp<FORMAT>::GetPoseRotations(
 			AnimData.CompressedTrackOffsets.GetData(),
 			AnimData.CompressedByteStream.GetData(),
 			AnimData.CompressedNumberOfKeys,
-			DecompContext.RelativePos,
+			DecompContext.GetRelativePosition(),
 			(uint8)DecompContext.Interpolation,
 			FORMAT,
 			PairCount);
@@ -296,7 +296,7 @@ void AEFVariableKeyLerp<FORMAT>::GetPoseTranslations(
 			AnimData.CompressedTrackOffsets.GetData(),
 			AnimData.CompressedByteStream.GetData(),
 			AnimData.CompressedNumberOfKeys,
-			DecompContext.RelativePos,
+			DecompContext.GetRelativePosition(),
 			(uint8)DecompContext.Interpolation,
 			FORMAT,
 			PairCount);
@@ -352,7 +352,7 @@ void AEFVariableKeyLerp<FORMAT>::GetPoseScales(
 			StripSize,
 			AnimData.CompressedByteStream.GetData(),
 			AnimData.CompressedNumberOfKeys,
-			DecompContext.RelativePos,
+			DecompContext.GetRelativePosition(),
 			(uint8)DecompContext.Interpolation,
 			FORMAT,
 			PairCount);
diff --git a/Engine/Source/Runtime/Engine/Private/Animation/AnimSequence.cpp b/Engine/Source/Runtime/Engine/Private/Animation/AnimSequence.cpp
--- a/Engine/Source/Runtime/Engine/Private/Animation/AnimSequence.cpp
+++ b/Engine/Source/Runtime/Engine/Private/Animation/AnimSequence.cpp
@@ -1191,10 +1191,10 @@ void UAnimSequence::GetBoneTransform(FTransform& OutAtom, int32 TrackIndex, floa
 	else if (CanEvaluateRawAnimationData())
 	{
 #if WITH_EDITOR
-	ValidateModel();
-	UE::Anim::GetBoneTransformFromModel(DataModel, OutAtom, TrackIndex, Time, Interpolation);
+		ValidateModel();
+		UE::Anim::GetBoneTransformFromModel(DataModel, OutAtom, TrackIndex, Time, Interpolation);
 #endif
-}
+	}
 }
 
 void UAnimSequence::GetBoneTransform(FTransform& OutAtom, int32 TrackIndex, FAnimSequenceDecompressionContext& DecompContext, bool bUseRawData) const
@@ -1212,10 +1212,10 @@ void UAnimSequence::GetBoneTransform(FTransform& OutAtom, int32 TrackIndex, FAni
 	else if (CanEvaluateRawAnimationData())
 	{
 #if WITH_EDITOR
-	ValidateModel();
-	UE::Anim::GetBoneTransformFromModel(DataModel, OutAtom, TrackIndex, DecompContext.Time, Interpolation);
+		ValidateModel();
+		UE::Anim::GetBoneTransformFromModel(DataModel, OutAtom, TrackIndex, DecompContext.GetEvaluationTime(), Interpolation);
 #endif
-}
+	}
 }
 
 void UAnimSequence::ExtractBoneTransform(const TArray<struct FRawAnimSequenceTrack>& InRawAnimationData, FTransform& OutAtom, int32 TrackIndex, float Time) const
diff --git a/Engine/Source/Runtime/Engine/Public/AnimEncoding_ConstantKeyLerp.h b/Engine/Source/Runtime/Engine/Public/AnimEncoding_ConstantKeyLerp.h
--- a/Engine/Source/Runtime/Engine/Public/AnimEncoding_ConstantKeyLerp.h
+++ b/Engine/Source/Runtime/Engine/Public/AnimEncoding_ConstantKeyLerp.h
@@ -206,7 +206,7 @@ FORCEINLINE void AEFConstantKeyLerp<FORMAT>::GetBoneAtomRotation(FTransform& Out
 	{
 		int32 Index0;
 		int32 Index1;
-		float Alpha = TimeToIndex(DecompContext.SequenceLength, DecompContext.RelativePos, NumRotKeys, DecompContext.Interpolation, Index0, Index1);
+		float Alpha = TimeToIndex(DecompContext.GetPlayableLength(), DecompContext.GetRelativePosition(), NumRotKeys, DecompContext.Interpolation, Index0, Index1);
 
 		const int32 RotationStreamOffset = (FORMAT == ACF_IntervalFixed32NoW) ? (sizeof(float)*6) : 0; // offset past Min and Range data
 
@@ -257,7 +257,7 @@ FORCEINLINE void AEFConstantKeyLerp<FORMAT>::GetBoneAtomTranslation(FTransform&
 
 	int32 Index0;
 	int32 Index1;
-	float Alpha = TimeToIndex(DecompContext.SequenceLength, DecompContext.RelativePos, NumTransKeys, DecompContext.Interpolation, Index0, Index1);
+	float Alpha = TimeToIndex(DecompContext.GetPlayableLength(), DecompContext.GetRelativePosition(), NumTransKeys, DecompContext.Interpolation, Index0, Index1);
 
 	const int32 TransStreamOffset = ((FORMAT == ACF_IntervalFixed32NoW) && NumTransKeys > 1) ? (sizeof(float)*6) : 0; // offset past Min and Range data
 
@@ -299,7 +299,7 @@ FORCEINLINE void AEFConstantKeyLerp<FORMAT>::GetBoneAtomScale(FTransform& OutAto
 
 	int32 Index0;
 	int32 Index1;
-	float Alpha = TimeToIndex(DecompContext.SequenceLength, DecompContext.RelativePos, NumScaleKeys, DecompContext.Interpolation, Index0, Index1);
+	float Alpha = TimeToIndex(DecompContext.GetPlayableLength(), DecompContext.GetRelativePosition(), NumScaleKeys, DecompContext.Interpolation, Index0, Index1);
 
 	const int32 ScaleStreamOffset = ((FORMAT == ACF_IntervalFixed32NoW) && NumScaleKeys > 1) ? (sizeof(float)*6) : 0; // offset past Min and Range data
 
diff --git a/Engine/Source/Runtime/Engine/Public/AnimEncoding_VariableKeyLerp.h b/Engine/Source/Runtime/Engine/Public/AnimEncoding_VariableKeyLerp.h
--- a/Engine/Source/Runtime/Engine/Public/AnimEncoding_VariableKeyLerp.h
+++ b/Engine/Source/Runtime/Engine/Public/AnimEncoding_VariableKeyLerp.h
@@ -209,7 +209,7 @@ FORCEINLINE_DEBUGGABLE void AEFVariableKeyLerp<FORMAT>::GetBoneAtomRotation(FTra
 
 		int32 Index0;
 		int32 Index1;
-		float Alpha = TimeToIndex(DecompContext.Interpolation, AnimData.CompressedNumberOfKeys, FrameTable, DecompContext.RelativePos, NumRotKeys, Index0, Index1);
+		float Alpha = TimeToIndex(DecompContext.Interpolation, AnimData.CompressedNumberOfKeys, FrameTable, DecompContext.GetRelativePosition(), NumRotKeys, Index0, Index1);
 
 
 		if (Index0 != Index1)
@@ -263,7 +263,7 @@ FORCEINLINE_DEBUGGABLE void AEFVariableKeyLerp<FORMAT>::GetBoneAtomTranslation(F
 
 	int32 Index0;
 	int32 Index1;
-	float Alpha = TimeToIndex(DecompContext.Interpolation, AnimData.CompressedNumberOfKeys, FrameTable, DecompContext.RelativePos, NumTransKeys, Index0, Index1);
+	float Alpha = TimeToIndex(DecompContext.Interpolation, AnimData.CompressedNumberOfKeys, FrameTable, DecompContext.GetRelativePosition(), NumTransKeys, Index0, Index1);
 	const int32 TransStreamOffset = ((FORMAT == ACF_IntervalFixed32NoW) && NumTransKeys > 1) ? (sizeof(float)*6) : 0; // offset past Min and Range data
 
 	if (Index0 != Index1)
@@ -307,7 +307,7 @@ FORCEINLINE_DEBUGGABLE void AEFVariableKeyLerp<FORMAT>::GetBoneAtomScale(FTransf
 
 	int32 Index0;
 	int32 Index1;
-	float Alpha = TimeToIndex(DecompContext.Interpolation, AnimData.CompressedNumberOfKeys, FrameTable, DecompContext.RelativePos, NumScaleKeys, Index0, Index1);
+	float Alpha = TimeToIndex(DecompContext.Interpolation, AnimData.CompressedNumberOfKeys, FrameTable, DecompContext.GetRelativePosition(), NumScaleKeys, Index0, Index1);
 	const int32 ScaleStreamOffset = ((FORMAT == ACF_IntervalFixed32NoW) && NumScaleKeys > 1) ? (sizeof(float)*6) : 0; // offset past Min and Range data
 
 	if (Index0 != Index1)
diff --git a/Engine/Source/Runtime/Engine/Public/Animation/AnimSequenceDecompressionContext.h b/Engine/Source/Runtime/Engine/Public/Animation/AnimSequenceDecompressionContext.h
--- a/Engine/Source/Runtime/Engine/Public/Animation/AnimSequenceDecompressionContext.h
+++ b/Engine/Source/Runtime/Engine/Public/Animation/AnimSequenceDecompressionContext.h
@@ -11,21 +11,58 @@ struct ICompressedAnimData;
 struct FAnimSequenceDecompressionContext
 {
 	FAnimSequenceDecompressionContext(float SequenceLength_, EAnimInterpolationType Interpolation_, const FName& AnimName_, const ICompressedAnimData& CompressedAnimData_)
-		: SequenceLength(SequenceLength_), Interpolation(Interpolation_), AnimName(AnimName_), CompressedAnimData(CompressedAnimData_), Time(0.f)
-	{}
+			:
+		PRAGMA_DISABLE_DEPRECATION_WARNINGS
+		SequenceLength(SequenceLength_),
+		PRAGMA_ENABLE_DEPRECATION_WARNINGS
+		Interpolation(Interpolation_), AnimName(AnimName_), CompressedAnimData(CompressedAnimData_),
+		PRAGMA_DISABLE_DEPRECATION_WARNINGS
+		Time(0.f), RelativePos(0.f)
+		PRAGMA_ENABLE_DEPRECATION_WARNINGS
+	{
+	}
 
 	// Anim info
+	UE_DEPRECATED(5.1, "Direct access to SequenceLength has been deprecated use GetPlayableLength instead")
 	float SequenceLength;
 	EAnimInterpolationType Interpolation;
 	FName AnimName;
 
 	const ICompressedAnimData& CompressedAnimData;
+
+	UE_DEPRECATED(5.1, "Direct access to Time has been deprecated use GetEvaluationTime or GetInterpolatedEvaluationTime instead")
 	float Time;
+	
+	UE_DEPRECATED(5.1, "Direct access to RelativePos has been deprecated use GetRelativePosition instead")
 	float RelativePos;
 
+	float GetPlayableLength() const
+	{
+		PRAGMA_DISABLE_DEPRECATION_WARNINGS
+		return SequenceLength;
+		PRAGMA_ENABLE_DEPRECATION_WARNINGS
+	}
+
+	float GetEvaluationTime() const
+	{
+		PRAGMA_DISABLE_DEPRECATION_WARNINGS
+		return Time;
+		PRAGMA_ENABLE_DEPRECATION_WARNINGS
+	}
+
+	float GetRelativePosition() const
+	{
+		PRAGMA_DISABLE_DEPRECATION_WARNINGS
+		return RelativePos;
+		PRAGMA_ENABLE_DEPRECATION_WARNINGS
+	}
+
 	void Seek(float SampleAtTime)
 	{
+		PRAGMA_DISABLE_DEPRECATION_WARNINGS
 		Time = SampleAtTime;
-		RelativePos = SampleAtTime / SequenceLength;
+		RelativePos = SequenceLength != 0.f ? SampleAtTime / SequenceLength : 0.f;
+		PRAGMA_ENABLE_DEPRECATION_WARNINGS
 	}
 };
+
