Adding csv stats for subobject deletion

diff --git a/Engine/Source/Runtime/Engine/Private/DataChannel.cpp b/Engine/Source/Runtime/Engine/Private/DataChannel.cpp
--- a/Engine/Source/Runtime/Engine/Private/DataChannel.cpp
+++ b/Engine/Source/Runtime/Engine/Private/DataChannel.cpp
@@ -47,6 +47,8 @@ DECLARE_CYCLE_STAT(TEXT("ActorChan_FindOrCreateRep"), Stat_ActorChanFindOrCreate
 DECLARE_LLM_MEMORY_STAT(TEXT("NetChannel"), STAT_NetChannelLLM, STATGROUP_LLMFULL);
 LLM_DEFINE_TAG(NetChannel, NAME_None, TEXT("Networking"), GET_STATFNAME(STAT_NetChannelLLM), GET_STATFNAME(STAT_NetworkingSummaryLLM));
 
+CSV_DECLARE_CATEGORY_EXTERN(Replication);
+
 extern int32 GDoReplicationContextString;
 extern int32 GNetDormancyValidate;
 extern bool GbNetReuseReplicatorsForDormantObjects;
@@ -3410,6 +3412,10 @@ int64 UActorChannel::ReplicateActor()
 		{
 			QUICK_SCOPE_CYCLE_COUNTER(STAT_NetDeletedSubObjects);
 
+#if !UE_BUILD_SHIPPING
+			CSV_CUSTOM_STAT(Replication, DeletedSubobjectTests, ReplicationMap.Num(), ECsvCustomStatOp::Accumulate);
+#endif
+
 			// Look for deleted subobjects
 			FObjectReplicator* LocalActorReplicator = ActorReplicator.Get();
 			for (auto RepComp = ReplicationMap.CreateIterator(); RepComp; ++RepComp)
@@ -3420,6 +3426,10 @@ int64 UActorChannel::ReplicateActor()
 				{
 					if (LocalReplicator->ObjectNetGUID.IsValid())
 					{
+#if !UE_BUILD_SHIPPING
+						CSV_CUSTOM_STAT(Replication, DeletedSubobjects, 1, ECsvCustomStatOp::Accumulate);
+#endif
+
 						// Write a deletion content header:
 						WriteContentBlockForSubObjectDelete(Bunch, LocalReplicator->ObjectNetGUID);
 
