Extract FNavigationDirtyElement from NavigationTypes.h and remove extra include dependencies

diff --git a/Engine/Source/Runtime/Engine/Classes/AI/Navigation/NavigationDirtyElement.h b/Engine/Source/Runtime/Engine/Classes/AI/Navigation/NavigationDirtyElement.h
new file mode 100644
--- /dev/null
+++ b/Engine/Source/Runtime/Engine/Classes/AI/Navigation/NavigationDirtyElement.h
@@ -0,0 +1,61 @@
+﻿// Copyright Epic Games, Inc. All Rights Reserved.
+
+#pragma once
+
+#include "UObject/WeakObjectPtr.h"
+
+class INavRelevantInterface;
+
+struct FNavigationDirtyElement
+{
+	/** object owning this element */
+	FWeakObjectPtr Owner;
+	
+	/** cached interface pointer */
+	INavRelevantInterface* NavInterface;
+
+	/** override for update flags */
+	int32 FlagsOverride;
+	
+	/** flags of already existing entry for this actor */
+	int32 PrevFlags;
+	
+	/** bounds of already existing entry for this actor */
+	FBox PrevBounds;
+
+	/** prev flags & bounds data are set */
+	uint8 bHasPrevData : 1;
+
+	/** request was invalidated while queued, use prev values to dirty area */
+	uint8 bInvalidRequest : 1;
+
+	FNavigationDirtyElement()
+		: NavInterface(0), FlagsOverride(0), PrevFlags(0), PrevBounds(ForceInit), bHasPrevData(false), bInvalidRequest(false)
+	{
+	}
+
+	FNavigationDirtyElement(UObject* InOwner)
+		: Owner(InOwner), NavInterface(0), FlagsOverride(0), PrevFlags(0), PrevBounds(ForceInit), bHasPrevData(false), bInvalidRequest(false)
+	{
+	}
+
+	FNavigationDirtyElement(UObject* InOwner, INavRelevantInterface* InNavInterface, int32 InFlagsOverride = 0)
+		: Owner(InOwner), NavInterface(InNavInterface),	FlagsOverride(InFlagsOverride), PrevFlags(0), PrevBounds(ForceInit), bHasPrevData(false), bInvalidRequest(false)
+	{
+	}
+
+	bool operator==(const FNavigationDirtyElement& Other) const 
+	{ 
+		return Owner == Other.Owner; 
+	}
+
+	bool operator==(const UObject*& OtherOwner) const 
+	{ 
+		return (Owner == OtherOwner);
+	}
+
+	FORCEINLINE friend uint32 GetTypeHash(const FNavigationDirtyElement& Info)
+	{
+		return GetTypeHash(Info.Owner);
+	}
+};
diff --git a/Engine/Source/Runtime/Engine/Classes/AI/Navigation/NavigationTypes.h b/Engine/Source/Runtime/Engine/Classes/AI/Navigation/NavigationTypes.h
--- a/Engine/Source/Runtime/Engine/Classes/AI/Navigation/NavigationTypes.h
+++ b/Engine/Source/Runtime/Engine/Classes/AI/Navigation/NavigationTypes.h
@@ -2,14 +2,18 @@
 
 #pragma once
 
+#ifdef UE_ENABLE_INCLUDE_ORDER_DEPRECATED_IN_5_1
 #include "CoreMinimal.h"
 #include "Stats/Stats.h"
-#include "UObject/ObjectMacros.h"
 #include "UObject/UObjectGlobals.h"
+#include "UObject/WeakObjectPtr.h"
+#include "AI/Navigation/NavigationDirtyElement.h"
+#endif //UE_ENABLE_INCLUDE_ORDER_DEPRECATED_IN_5_1
+
+#include "UObject/ObjectMacros.h"
 #include "UObject/Object.h"
 #include "UObject/Class.h"
 #include "Templates/SubclassOf.h"
-#include "UObject/WeakObjectPtr.h"
 #include "Misc/CoreStats.h"
 #include "UObject/SoftObjectPath.h"
 #include "GameFramework/Actor.h"
@@ -233,60 +237,6 @@ struct FNavigationBoundsUpdateRequest
 	Type UpdateRequest;
 };
 
-struct FNavigationDirtyElement
-{
-	/** object owning this element */
-	FWeakObjectPtr Owner;
-	
-	/** cached interface pointer */
-	INavRelevantInterface* NavInterface;
-
-	/** override for update flags */
-	int32 FlagsOverride;
-	
-	/** flags of already existing entry for this actor */
-	int32 PrevFlags;
-	
-	/** bounds of already existing entry for this actor */
-	FBox PrevBounds;
-
-	/** prev flags & bounds data are set */
-	uint8 bHasPrevData : 1;
-
-	/** request was invalidated while queued, use prev values to dirty area */
-	uint8 bInvalidRequest : 1;
-
-	FNavigationDirtyElement()
-		: NavInterface(0), FlagsOverride(0), PrevFlags(0), PrevBounds(ForceInit), bHasPrevData(false), bInvalidRequest(false)
-	{
-	}
-
-	FNavigationDirtyElement(UObject* InOwner)
-		: Owner(InOwner), NavInterface(0), FlagsOverride(0), PrevFlags(0), PrevBounds(ForceInit), bHasPrevData(false), bInvalidRequest(false)
-	{
-	}
-
-	FNavigationDirtyElement(UObject* InOwner, INavRelevantInterface* InNavInterface, int32 InFlagsOverride = 0)
-		: Owner(InOwner), NavInterface(InNavInterface),	FlagsOverride(InFlagsOverride), PrevFlags(0), PrevBounds(ForceInit), bHasPrevData(false), bInvalidRequest(false)
-	{
-	}
-
-	bool operator==(const FNavigationDirtyElement& Other) const 
-	{ 
-		return Owner == Other.Owner; 
-	}
-
-	bool operator==(const UObject*& OtherOwner) const 
-	{ 
-		return (Owner == OtherOwner);
-	}
-
-	FORCEINLINE friend uint32 GetTypeHash(const FNavigationDirtyElement& Info)
-	{
-		return GetTypeHash(Info.Owner);
-	}
-};
-
 UENUM()
 enum class ENavDataGatheringMode : uint8
 {
@@ -317,7 +267,7 @@ struct FNavigationPortalEdge
 	FVector Right;
 	NavNodeRef ToRef;
 
-	FNavigationPortalEdge() : Left(0.f), Right(0.f)
+	FNavigationPortalEdge() : Left(0.f), Right(0.f), ToRef(INVALID_NAVNODEREF)
 	{}
 
 	FNavigationPortalEdge(const FVector& InLeft, const FVector& InRight, NavNodeRef InToRef)
diff --git a/Engine/Source/Runtime/NavigationSystem/Private/NavigationSystem.cpp b/Engine/Source/Runtime/NavigationSystem/Private/NavigationSystem.cpp
--- a/Engine/Source/Runtime/NavigationSystem/Private/NavigationSystem.cpp
+++ b/Engine/Source/Runtime/NavigationSystem/Private/NavigationSystem.cpp
@@ -9,6 +9,7 @@
 #include "Engine/World.h"
 #include "GameFramework/Controller.h"
 #include "AI/Navigation/NavRelevantInterface.h"
+#include "AI/Navigation/NavigationDirtyElement.h"
 #include "UObject/UObjectIterator.h"
 #include "EngineUtils.h"
 #include "Logging/MessageLog.h"
diff --git a/Engine/Source/Runtime/NavigationSystem/Public/NavigationDataHandler.h b/Engine/Source/Runtime/NavigationSystem/Public/NavigationDataHandler.h
--- a/Engine/Source/Runtime/NavigationSystem/Public/NavigationDataHandler.h
+++ b/Engine/Source/Runtime/NavigationSystem/Public/NavigationDataHandler.h
@@ -4,6 +4,7 @@
 #include "NavigationOctreeController.h"
 #include "NavigationDirtyAreasController.h"
 
+struct FNavigationDirtyElement;
 
 struct NAVIGATIONSYSTEM_API FNavigationDataHandler
 {
diff --git a/Engine/Source/Runtime/NavigationSystem/Public/NavigationDirtyAreasController.h b/Engine/Source/Runtime/NavigationSystem/Public/NavigationDirtyAreasController.h
--- a/Engine/Source/Runtime/NavigationSystem/Public/NavigationDirtyAreasController.h
+++ b/Engine/Source/Runtime/NavigationSystem/Public/NavigationDirtyAreasController.h
@@ -7,6 +7,7 @@
 NAVIGATIONSYSTEM_API DECLARE_LOG_CATEGORY_EXTERN(LogNavigationDirtyArea, Warning, All);
 
 class ANavigationData;
+struct FNavigationDirtyElement;
 
 struct NAVIGATIONSYSTEM_API FNavigationDirtyAreasController
 {
diff --git a/Engine/Source/Runtime/NavigationSystem/Public/NavigationOctreeController.h b/Engine/Source/Runtime/NavigationSystem/Public/NavigationOctreeController.h
--- a/Engine/Source/Runtime/NavigationSystem/Public/NavigationOctreeController.h
+++ b/Engine/Source/Runtime/NavigationSystem/Public/NavigationOctreeController.h
@@ -2,7 +2,7 @@
 
 #pragma once
 #include "NavigationOctree.h"
-
+#include "AI/Navigation/NavigationDirtyElement.h"
 
 struct NAVIGATIONSYSTEM_API FNavigationOctreeController
 {
diff --git a/Engine/Source/Runtime/NavigationSystem/Public/NavigationSystem.h b/Engine/Source/Runtime/NavigationSystem/Public/NavigationSystem.h
--- a/Engine/Source/Runtime/NavigationSystem/Public/NavigationSystem.h
+++ b/Engine/Source/Runtime/NavigationSystem/Public/NavigationSystem.h
@@ -10,6 +10,7 @@
 #include "Misc/CoreDelegates.h"
 #include "NavFilters/NavigationQueryFilter.h"
 #include "AI/Navigation/NavigationTypes.h"
+#include "AI/Navigation/NavigationDirtyElement.h"
 #include "NavigationSystemTypes.h"
 #include "NavigationData.h"
 #include "AI/NavigationSystemBase.h"
