Fix Network Profiler Showing Wrong Properties For Actors

diff --git a/Engine/Source/Runtime/Engine/Private/DataChannel.cpp b/Engine/Source/Runtime/Engine/Private/DataChannel.cpp
--- a/Engine/Source/Runtime/Engine/Private/DataChannel.cpp
+++ b/Engine/Source/Runtime/Engine/Private/DataChannel.cpp
@@ -3405,6 +3405,7 @@ int64 UActorChannel::ReplicateActor()
 			}
 
 			MemMark.Pop();
+			NETWORK_PROFILER(GNetworkProfiler.TrackReplicateActor(Actor, RepFlags, FPlatformTime::Cycles() - ActorReplicateStartTime, Connection));
 
 			return bWroteSomethingImportant;
 		}
