adding support for tracking world actor deletion

diff --git a/Engine/Source/Runtime/Engine/Classes/Engine/World.h b/Engine/Source/Runtime/Engine/Classes/Engine/World.h
--- a/Engine/Source/Runtime/Engine/Classes/Engine/World.h
+++ b/Engine/Source/Runtime/Engine/Classes/Engine/World.h
@@ -224,6 +224,7 @@ private:
 ENGINE_API DECLARE_LOG_CATEGORY_EXTERN(LogSpawn, Warning, All);
 
 DECLARE_MULTICAST_DELEGATE_OneParam(FOnActorSpawned, AActor*);
+DECLARE_MULTICAST_DELEGATE_OneParam(FOnActorDestroyed, AActor*);
 DECLARE_MULTICAST_DELEGATE_OneParam(FOnFeatureLevelChanged, ERHIFeatureLevel::Type);
 DECLARE_MULTICAST_DELEGATE_OneParam(FOnMovieSceneSequenceTick, float);
 
@@ -1576,6 +1577,9 @@ private:
 	/** a delegate that broadcasts a notification before a newly spawned actor is initialized */
 	mutable FOnActorSpawned OnActorPreSpawnInitialization;
 
+	/** a delegate that broadcasts a notification whenever an actor is destroyed */
+	mutable FOnActorDestroyed OnActorDestroyed;
+
 	/** Reset Async Trace Buffer **/
 	void ResetAsyncTrace();
 
@@ -2813,6 +2817,12 @@ public:
 	/** Remove a listener for OnActorPreSpawnInitialization events */
 	void RemoveOnActorPreSpawnInitialization(FDelegateHandle InHandle) const;
 
+	/** Add a listener for OnActorDestroyed events */
+	FDelegateHandle AddOnActorDestroyedHandler(const FOnActorDestroyed::FDelegate& InHandler) const;
+
+	/** Remove a listener for OnActorDestroyed events */
+	void RemoveOnActorDestroyededHandler(FDelegateHandle InHandle) const;
+
 	/**
 	 * Returns whether the passed in actor is part of any of the loaded levels actors array.
 	 * Warning: Will return true for pending kill actors!
diff --git a/Engine/Source/Runtime/Engine/Private/LevelActor.cpp b/Engine/Source/Runtime/Engine/Private/LevelActor.cpp
--- a/Engine/Source/Runtime/Engine/Private/LevelActor.cpp
+++ b/Engine/Source/Runtime/Engine/Private/LevelActor.cpp
@@ -860,6 +860,8 @@ bool UWorld::DestroyActor( AActor* ThisActor, bool bNetForce, bool bShouldModify
 	// Notify the texture streaming manager about the destruction of this actor.
 	IStreamingManager::Get().NotifyActorDestroyed( ThisActor );
 
+	OnActorDestroyed.Broadcast(ThisActor);
+
 	// Tell this actor it's about to be destroyed.
 	ThisActor->Destroyed();
 
diff --git a/Engine/Source/Runtime/Engine/Private/World.cpp b/Engine/Source/Runtime/Engine/Private/World.cpp
--- a/Engine/Source/Runtime/Engine/Private/World.cpp
+++ b/Engine/Source/Runtime/Engine/Private/World.cpp
@@ -5339,6 +5339,16 @@ void UWorld::RemoveOnActorPreSpawnInitialization(FDelegateHandle InHandle) const
 	OnActorPreSpawnInitialization.Remove(InHandle);
 }
 
+FDelegateHandle UWorld::AddOnActorDestroyedHandler(const FOnActorDestroyed::FDelegate& InHandler) const
+{
+	return OnActorDestroyed.Add(InHandler);
+}
+
+void UWorld::RemoveOnActorDestroyededHandler(FDelegateHandle InHandle) const
+{
+	OnActorDestroyed.Remove(InHandle);
+}
+
 FDelegateHandle UWorld::AddMovieSceneSequenceTickHandler(const FOnMovieSceneSequenceTick::FDelegate& InHandler)
 {
 	return MovieSceneSequenceTick.Add(InHandler);
