Allow FCachedAnimStateArray to work with states from multiple state machines. GetTotalWeight can be > 1 when using states from different state machines

diff --git a/Engine/Source/Runtime/Engine/Private/Animation/CachedAnimData.cpp b/Engine/Source/Runtime/Engine/Private/Animation/CachedAnimData.cpp
--- a/Engine/Source/Runtime/Engine/Private/Animation/CachedAnimData.cpp
+++ b/Engine/Source/Runtime/Engine/Private/Animation/CachedAnimData.cpp
@@ -74,7 +74,7 @@ bool FCachedAnimStateArray::IsValid(UAnimInstance& InAnimInstance) const
 		if (States.Num() > 1)
 		{
 			FName StateMachineName = NAME_None;
-			TArray<FName> UniqueStateNames;
+			TArray<TPair<FName,FName>> UniqueStateNames;
 
 			for (const FCachedAnimStateData& State : States)
 			{
@@ -84,17 +84,18 @@ bool FCachedAnimStateArray::IsValid(UAnimInstance& InAnimInstance) const
 				}
 				else if ((State.StateMachineName != NAME_None) && (State.StateMachineName != StateMachineName))
 				{
-					UE_LOG(LogAnimation, Warning, TEXT("FCachedAnimStateArray::IsValid Mismatched StateMachineName found (%s VS %s) in AnimBP: %s. Renamed or deleted?"), *StateMachineName.ToString(), *State.StateMachineName.ToString(), *GetNameSafe(&InAnimInstance));
-					bCachedIsValid = false;
+					// Array has states from different state machines
+					bHasMultipleStateMachineEntries = true;
 				}
-
-				if (!UniqueStateNames.Contains(State.StateName))
+				
+				TPair<FName, FName> StateName = TPair<FName, FName>(State.StateMachineName, State.StateName);
+				if (!UniqueStateNames.Contains(StateName))
 				{
-					UniqueStateNames.Add(State.StateName);
+					UniqueStateNames.Add(StateName);
 				}
 				else
 				{
-					UE_LOG(LogAnimation, Warning, TEXT("FCachedAnimStateArray::IsValid StateName included multiple times (%s) in AnimBP: %s. Renamed or deleted?"), *State.StateName.ToString(), *GetNameSafe(&InAnimInstance));
+					UE_LOG(LogAnimation, Warning, TEXT("FCachedAnimStateArray::IsValid StateName %s from state machine %s included multiple times in AnimBP: %s."), *State.StateName.ToString(), *State.StateMachineName.ToString(), *GetNameSafe(&InAnimInstance));
 					bCachedIsValid = false;
 				}
 			}
@@ -113,14 +114,45 @@ float FCachedAnimStateArray::GetTotalWeight(UAnimInstance& InAnimInstance) const
 		{
 			TotalWeight += State.GetWeight(InAnimInstance);
 		}
-		return FMath::Min(TotalWeight, 1.f);
+		// Don't clamp to 1 if total is composed from multiple state machines since it can be larger than that
+		if (bHasMultipleStateMachineEntries)
+		{
+			return TotalWeight;
+		}
+		else
+		{
+			return FMath::Min(TotalWeight, 1.f);
+		}
 	}
 	return 0.f;
 }
 
 bool FCachedAnimStateArray::IsFullWeight(UAnimInstance& InAnimInstance) const
 {
-	return FAnimWeight::IsFullWeight(GetTotalWeight(InAnimInstance));
+	if (bHasMultipleStateMachineEntries)
+	{
+		// Array has multiple state machines, return true if one machine has a weight of 1.
+		if (IsValid(InAnimInstance))
+		{
+			TMap<FName, float> TotalWeights;
+			for (const FCachedAnimStateData& State : States)
+			{
+				TotalWeights.FindOrAdd(State.StateMachineName) += State.GetWeight(InAnimInstance);
+			}
+			for (const auto& Iter : TotalWeights)
+			{
+				if (FAnimWeight::IsFullWeight(Iter.Value))
+				{
+					return true;
+				}
+			}
+		}
+		return false;
+	}
+	else
+	{
+		return FAnimWeight::IsFullWeight(GetTotalWeight(InAnimInstance));
+	}
 }
 
 bool FCachedAnimStateArray::IsRelevant(UAnimInstance& InAnimInstance) const
diff --git a/Engine/Source/Runtime/Engine/Public/Animation/CachedAnimData.h b/Engine/Source/Runtime/Engine/Public/Animation/CachedAnimData.h
--- a/Engine/Source/Runtime/Engine/Public/Animation/CachedAnimData.h
+++ b/Engine/Source/Runtime/Engine/Public/Animation/CachedAnimData.h
@@ -66,12 +66,14 @@ struct ENGINE_API FCachedAnimStateArray
 	FCachedAnimStateArray()
 		: bCheckedValidity(false)
 		, bCachedIsValid(true)
+		, bHasMultipleStateMachineEntries(false)
 	{}
 
 	/** Array of states */
 	UPROPERTY(BlueprintReadOnly, EditAnywhere, Category = "State Machine")
 	TArray<FCachedAnimStateData> States;
 
+	/** Returns the total local weight of all states. If the definition contains more than on state machine, the result can be larger than 1 */
 	float GetTotalWeight(UAnimInstance& InAnimInstance) const;
 	bool IsFullWeight(UAnimInstance& InAnimInstance) const;
 	bool IsRelevant(UAnimInstance& InAnimInstance) const;
@@ -80,6 +82,7 @@ private:
 	bool IsValid(UAnimInstance& InAnimInstance) const;
 	mutable bool bCheckedValidity;
 	mutable bool bCachedIsValid;
+	mutable bool bHasMultipleStateMachineEntries;
 };
 
 USTRUCT(BlueprintType)
