Remove MoveTemp in GameplayStatics and add const to allow for UInitialActiveSoundParams to be re-used between PlaySoundAtLocation calls

diff --git a/Engine/Source/Runtime/Engine/Classes/Kismet/GameplayStatics.h b/Engine/Source/Runtime/Engine/Classes/Kismet/GameplayStatics.h
--- a/Engine/Source/Runtime/Engine/Classes/Kismet/GameplayStatics.h
+++ b/Engine/Source/Runtime/Engine/Classes/Kismet/GameplayStatics.h
@@ -664,9 +664,9 @@ public:
 	 *						to do a concurrency limit per owner.
 	 */
 	UFUNCTION(BlueprintCallable, Category="Audio", meta=(WorldContext="WorldContextObject", AdvancedDisplay = "3", UnsafeDuringActorConstruction = "true", Keywords = "play"))
-	static void PlaySoundAtLocation(const UObject* WorldContextObject, USoundBase* Sound, FVector Location, FRotator Rotation, float VolumeMultiplier = 1.f, float PitchMultiplier = 1.f, float StartTime = 0.f, class USoundAttenuation* AttenuationSettings = nullptr, USoundConcurrency* ConcurrencySettings = nullptr, const AActor* OwningActor = nullptr, UInitialActiveSoundParams* InitialParams = nullptr);
+	static void PlaySoundAtLocation(const UObject* WorldContextObject, USoundBase* Sound, FVector Location, FRotator Rotation, float VolumeMultiplier = 1.f, float PitchMultiplier = 1.f, float StartTime = 0.f, class USoundAttenuation* AttenuationSettings = nullptr, USoundConcurrency* ConcurrencySettings = nullptr, const AActor* OwningActor = nullptr, const UInitialActiveSoundParams* InitialParams = nullptr);
 
-	static void PlaySoundAtLocation(const UObject* WorldContextObject, USoundBase* Sound, FVector Location, float VolumeMultiplier = 1.f, float PitchMultiplier = 1.f, float StartTime = 0.f, class USoundAttenuation* AttenuationSettings = nullptr, USoundConcurrency* ConcurrencySettings = nullptr, UInitialActiveSoundParams* InitialParams = nullptr)
+	static void PlaySoundAtLocation(const UObject* WorldContextObject, USoundBase* Sound, FVector Location, float VolumeMultiplier = 1.f, float PitchMultiplier = 1.f, float StartTime = 0.f, class USoundAttenuation* AttenuationSettings = nullptr, USoundConcurrency* ConcurrencySettings = nullptr, const UInitialActiveSoundParams* InitialParams = nullptr)
 	{
 		PlaySoundAtLocation(WorldContextObject, Sound, Location, FRotator::ZeroRotator, VolumeMultiplier, PitchMultiplier, StartTime, AttenuationSettings, ConcurrencySettings, nullptr, InitialParams);
 	}
diff --git a/Engine/Source/Runtime/Engine/Private/GameplayStatics.cpp b/Engine/Source/Runtime/Engine/Private/GameplayStatics.cpp
--- a/Engine/Source/Runtime/Engine/Private/GameplayStatics.cpp
+++ b/Engine/Source/Runtime/Engine/Private/GameplayStatics.cpp
@@ -1597,7 +1597,7 @@ UAudioComponent* UGameplayStatics::SpawnSound2D(const UObject* WorldContextObjec
 	return AudioComponent;
 }
 
-void UGameplayStatics::PlaySoundAtLocation(const UObject* WorldContextObject, class USoundBase* Sound, FVector Location, FRotator Rotation, float VolumeMultiplier, float PitchMultiplier, float StartTime, class USoundAttenuation* AttenuationSettings, class USoundConcurrency* ConcurrencySettings, const AActor* OwningActor, UInitialActiveSoundParams* InitialParams)
+void UGameplayStatics::PlaySoundAtLocation(const UObject* WorldContextObject, class USoundBase* Sound, FVector Location, FRotator Rotation, float VolumeMultiplier, float PitchMultiplier, float StartTime, class USoundAttenuation* AttenuationSettings, class USoundConcurrency* ConcurrencySettings, const AActor* OwningActor, const UInitialActiveSoundParams* InitialParams)
 {
 	if (!Sound || !GEngine || !GEngine->UseSound())
 	{
@@ -1615,7 +1615,7 @@ void UGameplayStatics::PlaySoundAtLocation(const UObject* WorldContextObject, cl
 		TArray<FAudioParameter> Params;
 		if (InitialParams)
 		{
-			Params.Append(MoveTemp(InitialParams->AudioParams));
+			Params.Append(InitialParams->AudioParams);
 		}
 
 		// If OwningActor isn't supplied to this function, derive an owner from the WorldContextObject
