fix Texture.h IsHDR missing TSF_R16F

diff --git a/Engine/Source/Runtime/Engine/Classes/Engine/Texture.h b/Engine/Source/Runtime/Engine/Classes/Engine/Texture.h
--- a/Engine/Source/Runtime/Engine/Classes/Engine/Texture.h
+++ b/Engine/Source/Runtime/Engine/Classes/Engine/Texture.h
@@ -129,10 +129,10 @@ struct FTextureSource
 	/** Default constructor. */
 	ENGINE_API FTextureSource();
 
+	// should match ERawImageFormat::GetBytesPerPixel
 	ENGINE_API static int32 GetBytesPerPixel(ETextureSourceFormat Format);
-	// @@!! TSF_R16F missing
-	//	prefer ERawImageFormat::IsHDR
-	FORCEINLINE static bool IsHDR(ETextureSourceFormat Format) { return (Format == TSF_BGRE8 || Format == TSF_RGBA16F || Format == TSF_RGBA32F); }
+	// should match ERawImageFormat::IsHDR
+	FORCEINLINE static bool IsHDR(ETextureSourceFormat Format) { return (Format == TSF_BGRE8 || Format == TSF_RGBA16F || Format == TSF_R16F || Format == TSF_RGBA32F); }
 	
 	enum class ELockState : uint8
 	{
