small Texture fixes

diff --git a/Engine/Source/Runtime/Engine/Private/Texture.cpp b/Engine/Source/Runtime/Engine/Private/Texture.cpp
--- a/Engine/Source/Runtime/Engine/Private/Texture.cpp
+++ b/Engine/Source/Runtime/Engine/Private/Texture.cpp
@@ -337,14 +337,7 @@ bool UTexture::CanEditChange(const FProperty* InProperty) const
 	if (InProperty)
 	{
 		const FName PropertyName = InProperty->GetFName();
-
-		if (PropertyName == GET_MEMBER_NAME_STRING_CHECKED(UTexture, AdjustVibrance))
-		{
-			// @@!! this doesn't make much sense
-			//  it's one of the few Adjustments that's actually okay on HDR
-			return !HasHDRSource();
-		}
-		
+				
 		// Only enable chromatic adapation method when the white points differ.
 		if (PropertyName == GET_MEMBER_NAME_STRING_CHECKED(FTextureSourceColorSettings, ChromaticAdaptationMethod))
 		{
@@ -381,6 +374,15 @@ void UTexture::ValidateSettingsAfterImportOrEdit(bool * pRequiresNotifyMaterials
 		
 #if WITH_EDITORONLY_DATA
 
+	if ( MipGenSettings == TMGS_LeaveExistingMips && PowerOfTwoMode != ETexturePowerOfTwoSetting::None )
+	{
+		// power of 2 pads not allowed with LeaveExistingMips
+		UE_LOG(LogTexture, Display, TEXT("Power of 2 padding cannot be used with LeaveExistingMips, disabled. (%s)"), *GetName());
+
+		PowerOfTwoMode = ETexturePowerOfTwoSetting::None;
+		// @@!! this is broken in the Texture Editor Details GUI ; the PowerOfTwoMode button does not update
+	}
+
 	// IsPowerOfTwo only checks XY :
 	bool bIsPowerOfTwo = Source.IsPowerOfTwo();
 	if ( ! FMath::IsPowerOfTwo(Source.GetVolumeSizeZ()) )
diff --git a/Engine/Source/Runtime/Engine/Private/TextureDerivedDataTask.cpp b/Engine/Source/Runtime/Engine/Private/TextureDerivedDataTask.cpp
--- a/Engine/Source/Runtime/Engine/Private/TextureDerivedDataTask.cpp
+++ b/Engine/Source/Runtime/Engine/Private/TextureDerivedDataTask.cpp
@@ -541,7 +541,7 @@ void FTextureCacheDerivedDataWorker::BuildTexture(TArray<FTextureBuildSettings>&
 					DerivedData->SizeY = CompressedImage.SizeY;
 					DerivedData->PixelFormat = (EPixelFormat)CompressedImage.PixelFormat;
 
-					// @@!! verify NumSlices
+					// it would be better if CompressedImage just stored NumSlices, rather than recomputing it here
 					if (InBuildSettingsPerLayer[0].bVolume || InBuildSettingsPerLayer[0].bTextureArray)
 					{
 						DerivedData->SetNumSlices(CompressedImage.SizeZ);
