fix Texture upload on consoles; blind copy is needed

diff --git a/Engine/Source/Runtime/Engine/Private/Rendering/Texture2DResource.cpp b/Engine/Source/Runtime/Engine/Private/Rendering/Texture2DResource.cpp
--- a/Engine/Source/Runtime/Engine/Private/Rendering/Texture2DResource.cpp
+++ b/Engine/Source/Runtime/Engine/Private/Rendering/Texture2DResource.cpp
@@ -248,34 +248,45 @@ void FTexture2DResource::GetData( uint32 MipIndex, void* Dest, uint32 DestPitch
 		EffectiveSize,(int)BulkDataSize,
 		SrcPitch,DestPitch);
 
-	#if WITH_EDITORONLY_DATA
-	// in Editor, Mip doesn't come from BulkData, it may be null
-	// MipData[] was set from Editor data
-	// @@!! check MipData[MipIndex] size ! but it's not stored
-	if ( BulkDataSize == 0 )
-	{
-		BulkDataSize = EffectiveSize;
-	}
-	#endif
-
-#if !WITH_EDITORONLY_DATA
-	// only checking when !WITH_EDITORONLY_DATA ? because in Editor BulkDataSize == 0 , so not possible to check
-	checkf((int64)EffectiveSize == BulkDataSize, 
-		TEXT("Texture '%s', mip %d, has a BulkDataSize [%d] that doesn't match calculated size [%d]. Texture size %dx%d, format %d"),
-		*TextureName.ToString(), MipIndex, BulkDataSize, EffectiveSize, GetSizeX(), GetSizeY(), (int32)PixelFormat);
-#endif
-
 	// for platforms that returned 0 pitch from Lock, we need to just use the bulk data directly, never do 
 	// runtime block size checking, conversion, or the like
-	if (DestPitch == 0 || DestPitch == SrcPitch )
+	if (DestPitch == 0)
 	{
-		// checking Dest size before we memcpy would be nice!
+		// check( BulkDataSize >= EffectiveSize );
+
 		FMemory::Memcpy(Dest, MipData[MipIndex], BulkDataSize);
 	}
 	else
 	{
-		// Copy the texture data.
-		CopyTextureData2D(MipData[MipIndex],Dest,MipMap.SizeY,PixelFormat,SrcPitch,DestPitch);
+		#if WITH_EDITORONLY_DATA
+		// in Editor, Mip doesn't come from BulkData, it may be null
+		// MipData[] was set from Editor data
+		// @@!! check MipData[MipIndex] size ! but it's not stored
+		if ( BulkDataSize == 0 )
+		{
+			BulkDataSize = EffectiveSize;
+		}
+		#endif
+
+	#if !WITH_EDITORONLY_DATA
+		// only checking when !WITH_EDITORONLY_DATA ? because in Editor BulkDataSize == 0 , so not possible to check
+		checkf((int64)EffectiveSize == BulkDataSize, 
+			TEXT("Texture '%s', mip %d, has a BulkDataSize [%d] that doesn't match calculated size [%d]. Texture size %dx%d, format %d"),
+			*TextureName.ToString(), MipIndex, BulkDataSize, EffectiveSize, GetSizeX(), GetSizeY(), (int32)PixelFormat);
+	#endif
+
+		// for platforms that returned 0 pitch from Lock, we need to just use the bulk data directly, never do 
+		// runtime block size checking, conversion, or the like
+		if (DestPitch == 0 || DestPitch == SrcPitch )
+		{
+			// checking Dest size before we memcpy would be nice!
+			FMemory::Memcpy(Dest, MipData[MipIndex], BulkDataSize);
+		}
+		else
+		{
+			// Copy the texture data.
+			CopyTextureData2D(MipData[MipIndex],Dest,MipMap.SizeY,PixelFormat,SrcPitch,DestPitch);
+		}
 	}
 	
 	// Free data retrieved via GetCopy inside constructor.
