Chaos: add support for convex in Body instance GetShortestDistanceToPoint and GetClosestPointAndNormal

diff --git a/Engine/Source/Runtime/Engine/Classes/PhysicsEngine/BodySetup.h b/Engine/Source/Runtime/Engine/Classes/PhysicsEngine/BodySetup.h
--- a/Engine/Source/Runtime/Engine/Classes/PhysicsEngine/BodySetup.h
+++ b/Engine/Source/Runtime/Engine/Classes/PhysicsEngine/BodySetup.h
@@ -421,11 +421,12 @@ public:
 	 * Finds the shortest distance between the body setup and a world position. Input and output are given in world space
 	 * @param	WorldPosition	The point we are trying to get close to
 	 * @param	BodyToWorldTM	The transform to convert BodySetup into world space
+	 * @param	bUseConvexShapes When true also check the convex shapes if any (false by default)
 	 * @return					The distance between WorldPosition and the body setup. 0 indicates WorldPosition is inside one of the shapes.
 	 *
-	 * NOTE: This function ignores convex and trimesh data
+	 * NOTE: This function ignores trimesh data
 	 */
-	ENGINE_API float GetShortestDistanceToPoint(const FVector& WorldPosition, const FTransform& BodyToWorldTM) const;
+	ENGINE_API float GetShortestDistanceToPoint(const FVector& WorldPosition, const FTransform& BodyToWorldTM, bool bUseConvexShapes = false) const;
 
 	/** 
 	 * Finds the closest point in the body setup. Input and outputs are given in world space.
@@ -433,11 +434,12 @@ public:
 	 * @param	BodyToWorldTM			The transform to convert BodySetup into world space
 	 * @param	ClosestWorldPosition	The closest point on the body setup to WorldPosition
 	 * @param	FeatureNormal			The normal of the feature associated with ClosestWorldPosition
+	 * @param	bUseConvexShapes When true also check the convex shapes if any (false by default)
 	 * @return							The distance between WorldPosition and the body setup. 0 indicates WorldPosition is inside one of the shapes.
 	 *
-	 * NOTE: This function ignores convex and trimesh data
+	 * NOTE: This function ignores trimesh data
 	 */
-	ENGINE_API float GetClosestPointAndNormal(const FVector& WorldPosition, const FTransform& BodyToWorldTM, FVector& ClosestWorldPosition, FVector& FeatureNormal) const;
+	ENGINE_API float GetClosestPointAndNormal(const FVector& WorldPosition, const FTransform& BodyToWorldTM, FVector& ClosestWorldPosition, FVector& FeatureNormal, bool bUseConvexShapes = false) const;
 
 	/**
 	* Generates the information needed for cooking geometry.
diff --git a/Engine/Source/Runtime/Engine/Classes/PhysicsEngine/ConvexElem.h b/Engine/Source/Runtime/Engine/Classes/PhysicsEngine/ConvexElem.h
--- a/Engine/Source/Runtime/Engine/Classes/PhysicsEngine/ConvexElem.h
+++ b/Engine/Source/Runtime/Engine/Classes/PhysicsEngine/ConvexElem.h
@@ -149,6 +149,24 @@ public:
 
 	ENGINE_API void ScaleElem(FVector DeltaSize, float MinSize);
 
+	/**
+	 * Finds the closest point on the shape given a world position. Input and output are given in world space
+	 * @param	WorldPosition			The point we are trying to get close to
+	 * @param	BodyToWorldTM			The transform to convert BodySetup into world space
+	 * @param	ClosestWorldPosition	The closest point on the shape in world space
+	 * @param	Normal					The normal of the feature associated with ClosestWorldPosition.
+	 * @return							The distance between WorldPosition and the shape. 0 indicates WorldPosition is inside the shape.
+	 */
+	ENGINE_API float GetClosestPointAndNormal(const FVector& WorldPosition, const FTransform& BodyToWorldTM, FVector& ClosestWorldPosition, FVector& Normal) const;
+
+	/**	
+	 * Finds the shortest distance between the element and a world position. Input and output are given in world space
+	 * @param	WorldPosition	The point we are trying to get close to
+	 * @param	BodyToWorldTM	The transform to convert BodySetup into world space
+	 * @return					The distance between WorldPosition and the shape. 0 indicates WorldPosition is inside one of the shapes.
+	 */
+	ENGINE_API float GetShortestDistanceToPoint(const FVector& WorldPosition, const FTransform& BodyToWorldTM) const;
+	
 	ENGINE_API static EAggCollisionShape::Type StaticShapeType;
 
 private:
diff --git a/Engine/Source/Runtime/Engine/Private/PhysicsEngine/BodySetup.cpp b/Engine/Source/Runtime/Engine/Private/PhysicsEngine/BodySetup.cpp
--- a/Engine/Source/Runtime/Engine/Private/PhysicsEngine/BodySetup.cpp
+++ b/Engine/Source/Runtime/Engine/Private/PhysicsEngine/BodySetup.cpp
@@ -1324,17 +1324,18 @@ void UBodySetup::UpdateTriMeshVertices(const TArray<FVector> & NewPositions)
 #endif
 }
 
-template <bool bPositionAndNormal>
+template <bool bPositionAndNormal, bool bUseConvexShapes>
 float GetClosestPointAndNormalImpl(const UBodySetup* BodySetup, const FVector& WorldPosition, const FTransform& LocalToWorld, FVector* ClosestWorldPosition, FVector* FeatureNormal)
 {
 	float ClosestDist = FLT_MAX;
 	FVector TmpPosition, TmpNormal;
+	int32 NumShapeTested = 0;
 
 	//Note that this function is optimized for BodySetup with few elements. This is more common. If we want to optimize the case with many elements we should really return the element during the distance check to avoid pointless iteration
 	for (const FKSphereElem& SphereElem : BodySetup->AggGeom.SphereElems)
 	{
-		
-		if(bPositionAndNormal)
+		NumShapeTested++;
+		if constexpr (bPositionAndNormal)
 		{
 			const float Dist = SphereElem.GetClosestPointAndNormal(WorldPosition, LocalToWorld, TmpPosition, TmpNormal);
 
@@ -1354,7 +1355,8 @@ float GetClosestPointAndNormalImpl(const UBodySetup* BodySetup, const FVector& W
 
 	for (const FKSphylElem& SphylElem : BodySetup->AggGeom.SphylElems)
 	{
-		if (bPositionAndNormal)
+		NumShapeTested++;
+		if constexpr (bPositionAndNormal)
 		{
 			const float Dist = SphylElem.GetClosestPointAndNormal(WorldPosition, LocalToWorld, TmpPosition, TmpNormal);
 
@@ -1374,7 +1376,8 @@ float GetClosestPointAndNormalImpl(const UBodySetup* BodySetup, const FVector& W
 
 	for (const FKBoxElem& BoxElem : BodySetup->AggGeom.BoxElems)
 	{
-		if (bPositionAndNormal)
+		NumShapeTested++;
+		if constexpr (bPositionAndNormal)
 		{
 			const float Dist = BoxElem.GetClosestPointAndNormal(WorldPosition, LocalToWorld, TmpPosition, TmpNormal);
 
@@ -1392,7 +1395,31 @@ float GetClosestPointAndNormalImpl(const UBodySetup* BodySetup, const FVector& W
 		}
 	}
 
-	if (ClosestDist == FLT_MAX)
+	if constexpr (bUseConvexShapes)
+	{
+		NumShapeTested++;
+		for (const FKConvexElem& ConvexElem : BodySetup->AggGeom.ConvexElems)
+		{
+			if constexpr (bPositionAndNormal)
+			{
+				const float Dist = ConvexElem.GetClosestPointAndNormal(WorldPosition, LocalToWorld, TmpPosition, TmpNormal);
+
+				if (Dist < ClosestDist)
+				{
+					*ClosestWorldPosition = TmpPosition;
+					*FeatureNormal = TmpNormal;
+					ClosestDist = Dist;
+				}
+			}
+			else
+			{
+				const float Dist = ConvexElem.GetShortestDistanceToPoint(WorldPosition, LocalToWorld);
+				ClosestDist = Dist < ClosestDist ? Dist : ClosestDist;
+			}
+		}
+	}
+
+	if (NumShapeTested > 0 && ClosestDist == FLT_MAX)
 	{
 		UE_LOG(LogPhysics, Warning, TEXT("GetClosestPointAndNormalImpl ClosestDist for BodySetup %s is coming back as FLT_MAX. WorldPosition = %s, LocalToWorld = %s"), *BodySetup->GetFullName(), *WorldPosition.ToString(), *LocalToWorld.ToHumanReadableString());
 	}
@@ -1400,14 +1427,28 @@ float GetClosestPointAndNormalImpl(const UBodySetup* BodySetup, const FVector& W
 	return ClosestDist;
 }
 
-float UBodySetup::GetShortestDistanceToPoint(const FVector& WorldPosition, const FTransform& LocalToWorld) const
+float UBodySetup::GetShortestDistanceToPoint(const FVector& WorldPosition, const FTransform& LocalToWorld, bool bUseConvexShapes) const
 {
-	return GetClosestPointAndNormalImpl<false>(this, WorldPosition, LocalToWorld, nullptr, nullptr);
+	if (bUseConvexShapes)
+	{
+		return GetClosestPointAndNormalImpl<false, true>(this, WorldPosition, LocalToWorld, nullptr, nullptr);
+	}
+	else
+	{
+		return GetClosestPointAndNormalImpl<false, false>(this, WorldPosition, LocalToWorld, nullptr, nullptr);
+	}
 }
 
-float UBodySetup::GetClosestPointAndNormal(const FVector& WorldPosition, const FTransform& LocalToWorld, FVector& ClosestWorldPosition, FVector& FeatureNormal) const
+float UBodySetup::GetClosestPointAndNormal(const FVector& WorldPosition, const FTransform& LocalToWorld, FVector& ClosestWorldPosition, FVector& FeatureNormal, bool bUseConvexShapes) const
 {
-	return GetClosestPointAndNormalImpl<true>(this, WorldPosition, LocalToWorld, &ClosestWorldPosition, &FeatureNormal);
+	if (bUseConvexShapes)
+	{
+		return GetClosestPointAndNormalImpl<true, true>(this, WorldPosition, LocalToWorld, &ClosestWorldPosition, &FeatureNormal);
+	}
+	else
+	{
+		return GetClosestPointAndNormalImpl<true, false>(this, WorldPosition, LocalToWorld, &ClosestWorldPosition, &FeatureNormal);
+	}
 }
 
 #if WITH_EDITOR
@@ -2396,6 +2437,42 @@ float FKTaperedCapsuleElem::GetScaledHalfLength(const FVector& Scale3D) const
 	return FMath::Max<float>((Length + Radius0 + Radius1) * FMath::Abs(Scale3D.Z) * 0.5f, 0.1f);
 }
 
+float FKConvexElem::GetClosestPointAndNormal(const FVector& WorldPosition, const FTransform& BoneToWorldTM, FVector& ClosestWorldPosition, FVector& Normal) const
+{
+	float MinScale, MinScaleAbs;
+	FVector Scale3DAbs;
+	SetupNonUniformHelper(BoneToWorldTM.GetScale3D() * GetTransform().GetScale3D(), MinScale, MinScaleAbs, Scale3DAbs);
+
+	const FTransform LocalToWorldTM = GetTransform() * BoneToWorldTM;
+	const FVector LocalPosition = LocalToWorldTM.InverseTransformPositionNoScale(WorldPosition);
+	if (ChaosConvex)
+	{
+		Chaos::FVec3 OutNormal;
+		Chaos::FReal Phi = ChaosConvex->PhiWithNormalScaled(LocalPosition, Scale3DAbs, OutNormal);
+		Normal = LocalToWorldTM.TransformVectorNoScale(OutNormal);
+		ClosestWorldPosition = WorldPosition - Normal * Phi;
+		return Phi > UE_SMALL_NUMBER ? Phi : 0.f;
+	}
+	return 0.f;
+}
+
+float FKConvexElem::GetShortestDistanceToPoint(const FVector& WorldPosition, const FTransform& BoneToWorldTM) const
+{
+	float MinScale, MinScaleAbs;
+	FVector Scale3DAbs;
+	SetupNonUniformHelper(BoneToWorldTM.GetScale3D() * GetTransform().GetScale3D(), MinScale, MinScaleAbs, Scale3DAbs);
+
+	const FTransform LocalToWorldTM = GetTransform() * BoneToWorldTM;
+	const FVector LocalPosition = LocalToWorldTM.InverseTransformPositionNoScale(WorldPosition);
+	if (ChaosConvex)
+	{
+		Chaos::FVec3 OutNormal;
+		Chaos::FReal Phi = ChaosConvex->PhiWithNormalScaled(LocalPosition, Scale3DAbs, OutNormal);
+		return Phi > UE_SMALL_NUMBER ? Phi : 0.f;
+	}
+	return 0.f;
+}
+
 class UPhysicalMaterial* UBodySetup::GetPhysMaterial() const
 {
 	UPhysicalMaterial* PhysMat = PhysMaterial;
