Collapsed movement related csvprofiler instrumentation to be more in the spirit of the lightweight profiler

diff --git a/Engine/Source/Runtime/Engine/Private/CharacterMovementComponentAsync.cpp b/Engine/Source/Runtime/Engine/Private/CharacterMovementComponentAsync.cpp
--- a/Engine/Source/Runtime/Engine/Private/CharacterMovementComponentAsync.cpp
+++ b/Engine/Source/Runtime/Engine/Private/CharacterMovementComponentAsync.cpp
@@ -721,7 +721,6 @@ void FCharacterMovementComponentAsyncInput::StartNewPhysics(float deltaTime, int
 void FCharacterMovementComponentAsyncInput::PhysWalking(float deltaTime, int32 Iterations, FCharacterMovementComponentAsyncOutput& Output) const
 {
 	//SCOPE_CYCLE_COUNTER(STAT_CharPhysWalking);
-	CSV_SCOPED_TIMING_STAT_EXCLUSIVE(CharPhysWalking);
 
 	const FCharacterMovementComponentAsyncInput& Input = *this; // TODO Refactor
 
@@ -969,7 +968,6 @@ void FCharacterMovementComponentAsyncInput::PhysWalking(float deltaTime, int32 I
 void FCharacterMovementComponentAsyncInput::PhysFalling(float deltaTime, int32 Iterations, FCharacterMovementComponentAsyncOutput& Output) const
 {
 //	SCOPE_CYCLE_COUNTER(STAT_CharPhysFalling);
-//	CSV_SCOPED_TIMING_STAT_EXCLUSIVE(CharPhysFalling);
 
 	const float MIN_TICK_TIME = UCharacterMovementComponent::MIN_TICK_TIME;
 
diff --git a/Engine/Source/Runtime/Engine/Private/Components/CharacterMovementComponent.cpp b/Engine/Source/Runtime/Engine/Private/Components/CharacterMovementComponent.cpp
--- a/Engine/Source/Runtime/Engine/Private/Components/CharacterMovementComponent.cpp
+++ b/Engine/Source/Runtime/Engine/Private/Components/CharacterMovementComponent.cpp
@@ -4398,7 +4398,6 @@ float UCharacterMovementComponent::BoostAirControl(float DeltaTime, float TickAi
 void UCharacterMovementComponent::PhysFalling(float deltaTime, int32 Iterations)
 {
 	SCOPE_CYCLE_COUNTER(STAT_CharPhysFalling);
-	CSV_SCOPED_TIMING_STAT_EXCLUSIVE(CharPhysFalling);
 
 	if (deltaTime < MIN_TICK_TIME)
 	{
@@ -5053,7 +5052,6 @@ void UCharacterMovementComponent::MaintainHorizontalGroundVelocity()
 void UCharacterMovementComponent::PhysWalking(float deltaTime, int32 Iterations)
 {
 	SCOPE_CYCLE_COUNTER(STAT_CharPhysWalking);
-	CSV_SCOPED_TIMING_STAT_EXCLUSIVE(CharPhysWalking);
 
 	if (deltaTime < MIN_TICK_TIME)
 	{
