C++ Third Person Template project crashes when opened on Android
#fix lazy-intialize attribute types system, as async loading of AnimSequence (AttributeCurve) can happen before EDelayedRegisterRunPhase::ObjectSystemReady is fired

diff --git a/Engine/Source/Runtime/Engine/Private/Animation/AttributeTypes.cpp b/Engine/Source/Runtime/Engine/Private/Animation/AttributeTypes.cpp
--- a/Engine/Source/Runtime/Engine/Private/Animation/AttributeTypes.cpp
+++ b/Engine/Source/Runtime/Engine/Private/Animation/AttributeTypes.cpp
@@ -11,13 +11,19 @@ namespace UE
 		TArray<TWeakObjectPtr<const UScriptStruct>> AttributeTypes::RegisteredTypes;
 		TArray<TUniquePtr<IAttributeBlendOperator>> AttributeTypes::Operators;
 		TArray<TWeakObjectPtr<const UScriptStruct>> AttributeTypes::InterpolatableTypes;
+		std::atomic<bool> AttributeTypes::bInitialized = false;		
+		
+		void AttributeTypes::LazyInitialize()
+		{
+			bool bWasUninitialized = false;
+			if (bInitialized.compare_exchange_strong(bWasUninitialized, true))
+			{
+				Initialize();
+			}
+		}
 
 		void AttributeTypes::Initialize()
 		{
-			static bool bInitialized = false;
-			checkf(bInitialized == false, TEXT("Trying to initialize attribute type system multiple times"));
-			
-			bInitialized = true;
 			RegisterType<FFloatAnimationAttribute>();
 			RegisterType<FIntegerAnimationAttribute>();
 			RegisterType<FStringAnimationAttribute>();
@@ -26,7 +32,7 @@ namespace UE
 		
 		static FDelayedAutoRegisterHelper DelayedAttributeTypesInitializationHelper(EDelayedRegisterRunPhase::ObjectSystemReady, []()
 		{
-			UE::Anim::AttributeTypes::Initialize();
+			UE::Anim::AttributeTypes::LazyInitialize();
 		});
 	}
 }
diff --git a/Engine/Source/Runtime/Engine/Public/Animation/AttributeTypes.h b/Engine/Source/Runtime/Engine/Public/Animation/AttributeTypes.h
--- a/Engine/Source/Runtime/Engine/Public/Animation/AttributeTypes.h
+++ b/Engine/Source/Runtime/Engine/Public/Animation/AttributeTypes.h
@@ -46,9 +46,10 @@ namespace UE
 			static TArray<TWeakObjectPtr<const UScriptStruct>> RegisteredTypes;
 			static TArray<TUniquePtr<IAttributeBlendOperator>> Operators;
 			static TArray<TWeakObjectPtr<const UScriptStruct>> InterpolatableTypes;
-		public:
-			
+			static std::atomic<bool> bInitialized;
 			static void Initialize();
+		public:			
+			static void LazyInitialize();
 
 			/** Used for registering an attribute type for which TAttributeTypeTraits::WithCustomBlendOperator is set to true, use RegisterType() otherwise */
 			template<typename AttributeType, typename OperatorType, typename... OperatorArgs>
@@ -112,6 +113,7 @@ namespace UE
 			/** Returns the blend operator for the provided type, asserts when the type is not registered */
 			static const IAttributeBlendOperator* GetTypeOperator(TWeakObjectPtr<const UScriptStruct> WeakStruct)
 			{
+				LazyInitialize();
 				const int32 Index = AttributeTypes::RegisteredTypes.IndexOfByKey(WeakStruct);
 				ensure(WeakStruct.IsValid());
 				checkf(Index != INDEX_NONE, TEXT("Missing operator for attribute, type %s was not registered previously"), *WeakStruct->GetName());
@@ -121,18 +123,21 @@ namespace UE
 			/** Returns whether or not the provided type can be interpolated, defaults to false when the type is not registered */
 			static bool CanInterpolateType(TWeakObjectPtr<const UScriptStruct> WeakStruct)
 			{
+				LazyInitialize();
 				return AttributeTypes::InterpolatableTypes.Contains(WeakStruct);
 			}
 
 			/** Returns whether or not the type is registered */
 			static bool IsTypeRegistered(const UScriptStruct* ScriptStruct)
 			{
+				LazyInitialize();
 				return AttributeTypes::RegisteredTypes.Contains(ScriptStruct);
 			}
 
 			/** Returns all registered types */
 			static TArray<TWeakObjectPtr<const UScriptStruct>>& GetRegisteredTypes()
 			{
+				LazyInitialize();
 				return AttributeTypes::RegisteredTypes;
 			}
 		};
