5.0 deprecated functionality removal from AnimSequence(Base) and AnimDataModel

diff --git a/Engine/Source/Editor/AnimationBlueprintLibrary/Private/AnimationBlueprintLibrary.cpp b/Engine/Source/Editor/AnimationBlueprintLibrary/Private/AnimationBlueprintLibrary.cpp
--- a/Engine/Source/Editor/AnimationBlueprintLibrary/Private/AnimationBlueprintLibrary.cpp
+++ b/Engine/Source/Editor/AnimationBlueprintLibrary/Private/AnimationBlueprintLibrary.cpp
@@ -2502,16 +2502,6 @@ void UAnimationBlueprintLibrary::RemoveAllBoneAnimation(UAnimSequence* Animation
 	}
 }
 
-void UAnimationBlueprintLibrary::FinalizeBoneAnimation(UAnimSequence* AnimationSequence)
-{
-	if (AnimationSequence)
-	{
-		PRAGMA_DISABLE_DEPRECATION_WARNINGS
-		AnimationSequence->PostProcessSequence();
-		PRAGMA_ENABLE_DEPRECATION_WARNINGS
-	}
-}
-
 static void RecursiveRetrieveAnimationGraphs(UEdGraph* EdGraph, TArray<UAnimationGraph*>& OutAnimationGraphs)
 {
 	if (UAnimationGraph* AnimGraph = Cast<UAnimationGraph>(EdGraph))
diff --git a/Engine/Source/Editor/AnimationBlueprintLibrary/Public/AnimationBlueprintLibrary.h b/Engine/Source/Editor/AnimationBlueprintLibrary/Public/AnimationBlueprintLibrary.h
--- a/Engine/Source/Editor/AnimationBlueprintLibrary/Public/AnimationBlueprintLibrary.h
+++ b/Engine/Source/Editor/AnimationBlueprintLibrary/Public/AnimationBlueprintLibrary.h
@@ -355,7 +355,7 @@ public:
 	/** Apply all the changes made to Bone Tracks to Finalize. This triggers recompression. Note that this is expensive, but will require to get correct compressed data */
 	UE_DEPRECATED(5.0, "FinalizeBoneAnimation has been deprecated, use UAnimDataController instead")
 	UFUNCTION(BlueprintCallable, Category = "AnimationBlueprintLibrary|Curves", meta=(DeprecatedFunction, DeprecationMessage="FinalizeBoneAnimation has been deprecated, use UAnimDataController instead"))
-	static void FinalizeBoneAnimation(UAnimSequence* AnimationSequence);
+	static void FinalizeBoneAnimation(UAnimSequence* AnimationSequence) {}
 
 	// Smart name helper functions
 
diff --git a/Engine/Source/Runtime/Engine/Classes/Animation/AnimData/AnimDataModel.h b/Engine/Source/Runtime/Engine/Classes/Animation/AnimData/AnimDataModel.h
--- a/Engine/Source/Runtime/Engine/Classes/Animation/AnimData/AnimDataModel.h
+++ b/Engine/Source/Runtime/Engine/Classes/Animation/AnimData/AnimDataModel.h
@@ -277,26 +277,6 @@ public:
 	* @return	GUID representing the contained data and state 
 	*/
 	FGuid GenerateGuid() const;
-	
-	/** Backwards compatibility functionality */
-	UE_DEPRECATED(5.0, "Use GetBoneAnimationTracks(), and FBoneAnimationTrack::InternalTrackData instead")
-	const TArray<FRawAnimSequenceTrack>& GetTransientRawAnimationTracks() const;
-	
-	UE_DEPRECATED(5.0, "Use GetBoneAnimationTracks(), and FBoneAnimationTrack::Name instead")
-	const TArray<FName>& GetTransientRawAnimationTrackNames() const;
-
-	UE_DEPRECATED(5.0, "Use GetBoneAnimationTracks(), and FBoneAnimationTrack::BoneTreeIndex instead")
-	const TArray<FTrackToSkeletonMap>& GetTransientRawAnimationTrackSkeletonMappings() const;
-
-	UE_DEPRECATED(5.0, "Non-const access to track data is prohibited, use UAnimDataController API to modify AnimDataModel instead")
-	FRawAnimSequenceTrack& GetNonConstRawAnimationTrackByIndex(int32 TrackIndex);
-
-	UE_DEPRECATED(5.0, "Use GetCurveData() instead")
-	const FRawCurveTracks& GetTransientRawCurveTracks() const;
-
-	UE_DEPRECATED(5.0, "Use Controller for mutating curve data")
-	FAnimationCurveData& GetNonConstCurveData();	
-
 private:
 	/** Helper functionality used by UAnimDataController to retrieve mutable data */ 
 	FBoneAnimationTrack* FindMutableBoneTrackByName(FName Name);
@@ -337,10 +317,6 @@ private:
 			}
 
 			check(BracketCounter >= 0);
-			if (BracketCounter == 0)
-			{
-				GenerateTransientData();
-			}
 		}
 	}
 
@@ -373,16 +349,10 @@ private:
 			}
 
 			check(BracketCounter >= 0);
-			if (BracketCounter == 0)
-			{
-				GenerateTransientData();
-			}
 		}
 	}
 
 private:
-	void GenerateTransientData() const;
-	void CheckTransientData() const;
 
 	UPROPERTY(Transient)
 	int32 BracketCounter = 0;
@@ -420,23 +390,7 @@ private:
 	
 	/** Container with all animated (bone) attribute data */
 	UPROPERTY(VisibleAnywhere, BlueprintReadOnly, Category = "Animation Data Model", meta = (AllowPrivateAccess = "true"))
-	TArray<FAnimatedBoneAttribute> AnimatedBoneAttributes;
-	
-private:
-	/** Transient data, kept around for backward-compatibility */
-	UPROPERTY(Transient, VisibleAnywhere, Category = "Animation Data Model")
-	mutable TArray<FRawAnimSequenceTrack> RawAnimationTracks;
-
-	UPROPERTY(Transient, VisibleAnywhere, Category = "Animation Data Model")
-	mutable TArray<FName> RawAnimationTrackNames;
-
-	UPROPERTY(Transient, VisibleAnywhere, Category = "Animation Data Model")
-	mutable TArray<FTrackToSkeletonMap> RawAnimationTrackSkeletonMappings;
-
-	mutable bool bHasTransientDataBeenGenerated = false;
-
-	UPROPERTY(Transient, VisibleAnywhere, Category = "Animation Data Model")
-	mutable FRawCurveTracks RawCurveTracks;
+	TArray<FAnimatedBoneAttribute> AnimatedBoneAttributes;	
 
 	friend class UAnimDataController;
 	friend class FAnimDataControllerTestBase;
diff --git a/Engine/Source/Runtime/Engine/Classes/Animation/AnimSequence.h b/Engine/Source/Runtime/Engine/Classes/Animation/AnimSequence.h
--- a/Engine/Source/Runtime/Engine/Classes/Animation/AnimSequence.h
+++ b/Engine/Source/Runtime/Engine/Classes/Animation/AnimSequence.h
@@ -291,17 +291,6 @@ public:
 
 	FCompressedAnimSequence CompressedData;
 
-	// Accessors for animation frame count
-	UE_DEPRECATED(5.0, "GetRawNumberOfFrames is deprecated see UAnimDataModel::GetNumberOfKeys for the source keys, or GetNumberOfSampledKeys for resampled keys")
-	int32 GetRawNumberOfFrames() const { return GetNumberOfSampledKeys(); }
-
-	UE_DEPRECATED(5.0, "SetRawNumberOfFrame has been deprecated see UAnimDataController::SetFrameRate")
-	void SetRawNumberOfFrame(int32 InNumFrames) {}
-
-	/** Update the the number of expected keys in the (non-uniform) animation tracks, including T0 */
-	UE_DEPRECATED(5.0, "SetNumberOfSampledKeys has been deprecated see UAnimDataController::SetFrameRate")
-	void SetNumberOfSampledKeys(int32 InNumberOfKeys) {}
-
 	/** Additive animation type. **/
 	UPROPERTY(EditAnywhere, Category=AdditiveSettings, AssetRegistrySearchable)
 	TEnumAsByte<enum EAdditiveAnimationType> AdditiveAnimType;
@@ -445,10 +434,6 @@ public:
 #endif
 	//~ End UAnimSequenceBase Interface
 
-	// Returns the framerate of the animation
-	UE_DEPRECATED(5.0, "GetFrameRate is deprecated see UAnimDataModel::GetFrameRate for the source frame rate, or GetSamplingFrameRate for the target frame rate instead")
-	float GetFrameRate() const { return (float)GetSamplingFrameRate().AsDecimal(); }
-
 	// Extract Root Motion transform from the animation
 	FTransform ExtractRootMotion(float StartTime, float DeltaTime, bool bAllowLooping) const;
 
@@ -481,34 +466,13 @@ public:
 	*/
 	void GetBonePose(struct FAnimationPoseData& OutAnimationPoseData, const FAnimExtractContext& ExtractionContext, bool bForceUseRawData = false) const;
 
-	UE_DEPRECATED(5.0, "GetRawAnimationTrack has been deprecated see UAnimDataModel::GetBoneAnimationTracks")
-	const TArray<FRawAnimSequenceTrack>& GetRawAnimationData() const;
+	const TArray<FTrackToSkeletonMap>& GetCompressedTrackToSkeletonMapTable() const { return CompressedData.CompressedTrackToSkeletonMapTable; }
+	const TArray<struct FSmartName>& GetCompressedCurveNames() const { return CompressedData.CompressedCurveNames; }
 
 #if WITH_EDITORONLY_DATA
-	UE_DEPRECATED(5.0, "SourceRawAnimationData has been deprecated")
-	bool HasSourceRawData() const { return false; }
-	UE_DEPRECATED(5.0, "GetAnimationTrackNames has been deprecated see FBoneAnimationTrack::Name")
-	const TArray<FName>& GetAnimationTrackNames() const;
-	
-	UE_DEPRECATED(5.0, "UpdateCompressedCurveName will be marked protected, updating compressed curve names is now handled by EAnimDataModelNotifyType::CurveRenamed")
+protected:
 	void UpdateCompressedCurveName(SmartName::UID_Type CurveUID, const struct FSmartName& NewCurveName);
-	
-	// Adds a new track (if no track of the supplied name is found) to the raw animation data, optionally setting it to TrackData.
-	UE_DEPRECATED(5.0, "AddNewRawTrack has been deprecated see UAnimDataController::AddBoneTrack and UAnimDataController::SetBoneTrackKeys")
-	int32 AddNewRawTrack(FName TrackName, FRawAnimSequenceTrack* TrackData = nullptr);
 #endif
-	UE_DEPRECATED(5.0, "GetRawTrackToSkeletonMapTable has been deprecated see FBoneAnimationTrack::BoneTreeIndex")
-	const TArray<FTrackToSkeletonMap>& GetRawTrackToSkeletonMapTable() const;
-
-	const TArray<FTrackToSkeletonMap>& GetCompressedTrackToSkeletonMapTable() const { return CompressedData.CompressedTrackToSkeletonMapTable; }
-	const TArray<struct FSmartName>& GetCompressedCurveNames() const { return CompressedData.CompressedCurveNames; }
-
-	UE_DEPRECATED(5.0, "GetRawAnimationTrack, and non-const access to Raw Animation Data has been deprecated see UAnimDataModel::GetBoneTrackByIndex")
-	FRawAnimSequenceTrack& GetRawAnimationTrack(int32 TrackIndex);
-
-	UE_DEPRECATED(5.0, "GetRawAnimationTrack has been deprecated see UAnimDataModel::GetBoneTrackByIndex")
-	const FRawAnimSequenceTrack& GetRawAnimationTrack(int32 TrackIndex) const;
-
 private:
 #if WITH_EDITORONLY_DATA
 	void UpdateRetargetSourceAsset();
@@ -573,49 +537,20 @@ public:
 	 * @param	bUseRawData		If true, use raw animation data instead of compressed data.
 	 */
 	void GetBoneTransform(FTransform& OutAtom, int32 TrackIndex, FAnimSequenceDecompressionContext& DecompContext, bool bUseRawData) const;
-
-	/**
-	 * Extract Bone Transform of the Time given, from InRawAnimationData
-	 *
-	 * @param	InRawAnimationData	RawAnimationData it extracts bone transform from
-	 * @param	OutAtom				[out] Output bone transform.
-	 * @param	TrackIndex			Index of track to interpolate.
-	 * @param	Time				Time on track to interpolate to.
-	 */
-	UE_DEPRECATED(5.0, "ExtractBoneTransform has been deprecated see FAnimationUtils::ExtractTransformFromTrack")
-	void ExtractBoneTransform(const TArray<struct FRawAnimSequenceTrack> & InRawAnimationData, FTransform& OutAtom, int32 TrackIndex, float Time) const;
-
-	/**
-	* Extract Bone Transform of the Time given, from InRawAnimationData
-	*
-	* @param	InRawAnimationTrack	RawAnimationTrack it extracts bone transform from
-	* @param	OutAtom				[out] Output bone transform.
-	* @param	Time				Time on track to interpolate to.
-	*/
-	UE_DEPRECATED(5.0, "ExtractBoneTransform has been deprecated see FAnimationUtils::ExtractTransformFromTrack")
-	void ExtractBoneTransform(const struct FRawAnimSequenceTrack& InRawAnimationTrack, FTransform& OutAtom, float Time) const;
-
-	UE_DEPRECATED(5.0, "ExtractBoneTransform has been deprecated see FAnimSequenceHelpers::ExtractBoneTransform")
-	void ExtractBoneTransform(const struct FRawAnimSequenceTrack& RawTrack, FTransform& OutAtom, int32 KeyIndex) const;
-
 	// End Transform related functions 
 
 	// Begin Memory related functions
 
 	/** @return	estimate uncompressed raw size. This is *not* the real raw size. 
 				Here we estimate what it would be with no trivial compression. */
-#if !WITH_EDITOR
-	UE_DEPRECATED(5.0, "GetUncompressedRawSize will be marked EDITOR_ONLY")
-#endif // !WITH_EDITOR
+#if WITH_EDITOR
 	int32 GetUncompressedRawSize() const;
 
 	/**
 	 * @return		The approximate size of raw animation data.
 	 */
-#if !WITH_EDITOR
-	UE_DEPRECATED(5.0, "GetApproxRawSize will be marked EDITOR_ONLY")
-#endif // !WITH_EDITOR
 	int32 GetApproxRawSize() const;
+#endif // WITH_EDITOR
 
 	/**
 	 * @return		The approximate size of compressed animation data for only bones.
@@ -627,44 +562,14 @@ public:
 	 */
 	int32 GetApproxCompressedSize() const;
 
-	/**
-	 * Removes trivial frames -- frames of tracks when position or orientation is constant
-	 * over the entire animation -- from the raw animation data.  If both position and rotation
-	 * go down to a single frame, the time is stripped out as well.
-	 * @return true if keys were removed.
-	 */
-	UE_DEPRECATED(5.0, "CompressRawAnimData has been deprecated, reduction of Raw Animation data now happens during compression instead")
-	bool CompressRawAnimData(float MaxPosDiff, float MaxAngleDiff) { return false; }
-
-	/**
-	 * Removes trivial frames -- frames of tracks when position or orientation is constant
-	 * over the entire animation -- from the raw animation data.  If both position and rotation
-	 * go down to a single frame, the time is stripped out as well.
-	 * @return true if keys were removed.
-	 */
-	UE_DEPRECATED(5.0, "CompressRawAnimData has been deprecated, reduction of Raw Animation data now happens during compression instead")
-	bool CompressRawAnimData() { return false; }
-
 	// Get compressed data for this UAnimSequence. May be built directly or pulled from DDC
 #if WITH_EDITOR
 	bool ShouldPerformStripping(const bool bPerformFrameStripping, const bool bPerformStrippingOnOddFramedAnims) const;
 	FString GetDDCCacheKeySuffix(const bool bPerformStripping) const;
 	void ApplyCompressedData(const FString& DataCacheKeySuffix, const bool bPerformFrameStripping, const TArray<uint8>& Data);
-#endif
 
-#if !WITH_EDITOR
-	UE_DEPRECATED(5.0, "WaitOnExistingCompression will be marked EDITOR_ONLY")
-#endif // !WITH_EDITOR
 	void WaitOnExistingCompression(const bool bWantResults=true);
-
-#if !WITH_EDITOR
-	UE_DEPRECATED(5.0, "RequestAnimCompression will be marked EDITOR_ONLY")
-#endif // !WITH_EDITOR
 	void RequestAnimCompression(FRequestAnimCompressionParams Params);
-
-#if !WITH_EDITOR
-	UE_DEPRECATED(5.0, "RequestSyncAnimRecompression will be marked EDITOR_ONLY")
-#endif // !WITH_EDITOR
 	void RequestSyncAnimRecompression(bool bOutput = false)
 	{
 		PRAGMA_DISABLE_DEPRECATION_WARNINGS
@@ -672,41 +577,26 @@ public:
 		PRAGMA_ENABLE_DEPRECATION_WARNINGS
 	}
 
-#if !WITH_EDITOR
-	UE_DEPRECATED(5.0, "RequestAsyncAnimRecompression will be marked EDITOR_ONLY")
-#endif // !WITH_EDITOR
 	void RequestAsyncAnimRecompression(bool bOutput = false)
 	{ 
 		PRAGMA_DISABLE_DEPRECATION_WARNINGS
 		RequestAnimCompression(FRequestAnimCompressionParams(true, false, bOutput));
 		PRAGMA_ENABLE_DEPRECATION_WARNINGS
 	}
+#endif
 
 protected:
 	void ApplyCompressedData(const TArray<uint8>& Data);
-
+    void ClearCompressedBoneData();
+    void ClearCompressedCurveData();
+    // Write the compressed data to the supplied FArchive
+    void SerializeCompressedData(FArchive& Ar, bool bDDCData);
 public:
 	bool IsCompressedDataValid() const;
 	bool IsCurveCompressedDataValid() const;
-
-	UE_DEPRECATED(5.0, "ClearCompressedBoneData will be marked protected")
-	void ClearCompressedBoneData();
-	UE_DEPRECATED(5.0, "ClearCompressedCurveData will be marked protected")
-	void ClearCompressedCurveData();
-	// Write the compressed data to the supplied FArchive
-	UE_DEPRECATED(5.0, "SerializeCompressedData will be marked protected")
-	void SerializeCompressedData(FArchive& Ar, bool bDDCData);
-
 	// End Memory related functions
 
 	// Begin Utility functions
-	/**
-	 * Get Skeleton Bone Index from Track Index for raw data
-	 *
-	 * @param	TrackIndex		Track Index
-	 */
-	UE_DEPRECATED(5.0, "DoesContainTransformCurves has been deprecated see FBoneAnimationTrack::BoneTreeIndex")
-	int32 GetSkeletonIndexFromRawDataTrackIndex(const int32 TrackIndex) const;
 
 	/**
 	* Get Skeleton Bone Index from Track Index for compressed data
@@ -717,95 +607,15 @@ public:
 	{
 		return GetCompressedTrackToSkeletonMapTable()[TrackIndex].BoneTreeIndex;
 	}
-
-	/** Clears any data in the AnimSequence */
-	UE_DEPRECATED(5.0, "RecycleAnimSequence has been deprecated use ResetAnimation instead")
-	void RecycleAnimSequence();
-
-#if WITH_EDITOR
-	/** Clears some data in the AnimSequence, so it can be reused when importing a new animation with same name over it. */
-	UE_DEPRECATED(5.0, "CleanAnimSequenceForImport has been deprecated see UAnimDataController:ResetModel")
-	void CleanAnimSequenceForImport();
-#endif
-
-	/** 
-	 * Copy AnimNotifies from one UAnimSequence to another.
-	 */
-	UE_DEPRECATED(5.0, "CopyNotifies has been moved, see UAnimSequenceHelpers::CopyNotifies")
-	static bool CopyNotifies(UAnimSequence* SourceAnimSeq, UAnimSequence* DestAnimSeq, bool bShowDialogs = true);
-
-	/**
-	 * Flip Rotation's W For NonRoot items, and compress it again if SkelMesh exists
-	 */
-	UE_DEPRECATED(5.0, "FlipRotationWForNonRoot has been deprecated, use UAnimDataModel::GetBoneAnimationTracks and UAnimDataController::SetBoneTrackKeys instead")
-	void FlipRotationWForNonRoot(USkeletalMesh * SkelMesh);
-
 	// End Utility functions
+	
 #if WITH_EDITOR
-	/**
-	 * After imported or any other change is made, call this to apply post process
-	 */
-	UE_DEPRECATED(5.0, "PostProcessSequence has been deprecated, any Raw Animation Data modifications should go through UAnimDataController")
-	void PostProcessSequence(bool bForceNewRawDatGuid = true);
-
-	// Kick off compression request when our raw data has changed
-	UE_DEPRECATED(5.0, "OnRawDataChanged has been deprecated, any Raw Animation Data modifications should go through UAnimDataController")
-	void OnRawDataChanged();
-
-	/** 
-	 * Insert extra frame of the first frame at the end of the frame so that it improves the interpolation when it loops
-	 * This increases framecount + time, so that it requires recompression
-	 */
-	UE_DEPRECATED(5.0, "AddLoopingInterpolation has been deprecated see FAnimSequenceHelpers::AnimationData::AddLoopingInterpolation")
-	bool AddLoopingInterpolation();
-
-	/*
-	* Clear all raw animation data that contains bone tracks
-	*/
-	UE_DEPRECATED(5.0, "RemoveAllTracks has been deprecated see UAnimDataController::RemoveAllBoneTracks")
-	void RemoveAllTracks();
-
-	/** 
-	 * Bake Transform Curves.TransformCurves to RawAnimation after making a back up of current RawAnimation
-	 */
-	UE_DEPRECATED(5.0, "BakeTrackCurvesToRawAnimation has been deprecated, transform curves are now baked during compression")
-	void BakeTrackCurvesToRawAnimation() {}
-
 	void BakeTrackCurvesToRawAnimationTracks(TArray<FRawAnimSequenceTrack>& NewRawTracks, TArray<FName>& NewTrackNames, TArray<FTrackToSkeletonMap>& NewTrackToSkeletonMapTable);
 
-	/**
-	 * Sometimes baked data gets invalidated. For example, if you retarget this from another animation
-	 * It won't matter anymore, so in any case, when the data is not valid anymore
-	 * We clear Source Raw Animation Data as well as Transform Curve
-	 */
-	UE_DEPRECATED(5.0, "ClearBakedTransformData has been deprecated, transform curves are now baked during compression")
-	void ClearBakedTransformData();
 	/**
 	 * Add Key to Transform Curves
 	 */
 	void AddKeyToSequence(float Time, const FName& BoneName, const FTransform& AdditiveTransform);
-	/**
-	 * Return true if it needs to re-bake
-	 */
-	UE_DEPRECATED(5.0, "DoesNeedRebake has been deprecated, transform curves are now baked during compression")
-	bool DoesNeedRebake() const { return false; }
-	/**
-	 * Return true if it contains transform curves
-	 */
-	UE_DEPRECATED(5.0, "DoesContainTransformCurves has been deprecated see UAnimDataModel::GetNumberOfTransformCurves")
-	bool DoesContainTransformCurves() const;
-
-	/**
-	 * Returns whether this animation has baked transform curves (i.e. has the raw data been modified)
-	 */
-	UE_DEPRECATED(5.0, "HasBakedTransformCurves has been deprecated, transform curves are now baked during compression")
-	bool HasBakedTransformCurves() const { return false; }
-
-	/**
-	 * Restore the pre baked transform curve raw data
-	 */
-	UE_DEPRECATED(5.0, "RestoreSourceData has been deprecated, transform curves are now baked during compression")
-	void RestoreSourceData() {}
 
 	/**
 	* Return true if compressed data is out of date / missing and so animation needs to use raw data
@@ -825,29 +635,6 @@ public:
 	 */
 	bool CreateAnimation(class UAnimSequence* Sequence);
 
-	/**
-	 * Crops the raw anim data either from Start to CurrentTime or CurrentTime to End depending on
-	 * value of bFromStart.  Can't be called against cooked data.
-	 *
-	 * @param	CurrentTime		marker for cropping (either beginning or end)
-	 * @param	bFromStart		whether marker is begin or end marker
-	 * @return					true if the operation was successful.
-	 */
-	UE_DEPRECATED(5.0, "InsertFramesToRawAnimData has been deprecated see FAnimSequenceHelpers::AnimationData::Trim")
-	bool CropRawAnimData( float CurrentTime, bool bFromStart );
-		
-	/**
-	 * Crops the raw anim data either from Start to CurrentTime or CurrentTime to End depending on
-	 * value of bFromStart.  Can't be called against cooked data.
-	 *
-	 * @param	StartFrame		StartFrame to insert (0-based)
-	 * @param	EndFrame		EndFrame to insert (0-based
-	 * @param	CopyFrame		A frame that we copy from (0-based)
-	 * @return					true if the operation was successful.
-	 */
-	UE_DEPRECATED(5.0, "InsertFramesToRawAnimData has been deprecated see FAnimSequenceHelpers::AnimationData::DuplicateKeys")
-	bool InsertFramesToRawAnimData( int32 StartFrame, int32 EndFrame, int32 CopyFrame);
-
 	/** 
 	 * Add validation check to see if it's being ready to play or not
 	 */
@@ -962,36 +749,8 @@ private:
 	bool ConvertRiggingDataToAnimationData(FAnimSequenceTrackContainer & RiggingAnimationData);
 	int32 GetSpaceBasedAnimationData(TArray< TArray<FTransform> > & AnimationDataInComponentSpace, FAnimSequenceTrackContainer * RiggingAnimationData) const;
 
-	/** Verify Track Map is valid, if not, fix up */
-	UE_DEPRECATED(5.0, "RemoveTrack has been deprecated see UAnimDataController::RemoveBoneTracksMissingFromSkeleton")
-	void VerifyTrackMap(USkeleton* MySkeleton=NULL);
-
-	/** Refresh Track Map from Animation Track Names **/
-	UE_DEPRECATED(5.0, "RemoveTrack has been deprecated see UAnimDataController::RemoveBoneTracksMissingFromSkeleton")
-	void RefreshTrackMapFromAnimTrackNames();
-
-	/**
-	 * Utility function that helps to remove track, you can't just remove RawAnimationData
-	 */
-	UE_DEPRECATED(5.0, "RemoveTrack has been deprecated see UAnimDataController::RemoveBoneTrack")
-	void RemoveTrack(int32 TrackIndex);
-
-	/**
-	 * Utility function that finds the correct spot to insert track to 
-	 */
-	UE_DEPRECATED(5.0, "InsertTrack has been deprecated see UAnimDataController::InsertBoneTrack")
-	int32 InsertTrack(const FName& BoneName);
-
 public:
-	/**
-	 * Utility function to resize the sequence
-	 * It rearranges curve data + notifies
-	 */
-	UE_DEPRECATED(5.0, "ResizeSequence has been deprecated see UAnimDataController::Resize")
-	void ResizeSequence(float NewLength, int32 NewNumFrames, bool bInsert, int32 StartFrame/*inclusive */, int32 EndFrame/*inclusive*/);
-
 #endif
-
 	/** Refresh sync marker data*/
 	void RefreshSyncMarkerDataFromAuthored();
 
@@ -1018,39 +777,29 @@ private:
 public:
 #if WITH_EDITOR
 	UE_DEPRECATED(5.0, "AddBoneCustomAttribute has been deprecated see UAnimDataController::AddAttribute")
-	UFUNCTION(BlueprintCallable, Category=CustomAttributes)
+	UFUNCTION(BlueprintCallable, Category=CustomAttributes, meta=(DeprecatedFunction, DeprecationMessage="AddBoneFloatCustomAttribute has been deprecated, use UAnimDataController::AddAttribute instead"))
 	void AddBoneFloatCustomAttribute(const FName& BoneName, const FName& AttributeName, const TArray<float>& TimeKeys, const TArray<float>& ValueKeys);
 
 	UE_DEPRECATED(5.0, "AddBoneCustomAttribute has been deprecated see UAnimDataController::AddAttribute")
-	UFUNCTION(BlueprintCallable, Category = CustomAttributes)
+	UFUNCTION(BlueprintCallable, Category = CustomAttributes, meta=(DeprecatedFunction, DeprecationMessage="AddBoneIntegerCustomAttribute has been deprecated, use UAnimDataController::AddAttribute instead"))
 	void AddBoneIntegerCustomAttribute(const FName& BoneName, const FName& AttributeName, const TArray<float>& TimeKeys, const TArray<int32>& ValueKeys);
 
 	UE_DEPRECATED(5.0, "AddBoneStringCustomAttribute has been deprecated see UAnimDataController::AddAttribute")
-	UFUNCTION(BlueprintCallable, Category = CustomAttributes)
+	UFUNCTION(BlueprintCallable, Category = CustomAttributes, meta=(DeprecatedFunction, DeprecationMessage="AddBoneStringCustomAttribute has been deprecated, use UAnimDataController::AddAttribute instead"))
 	void AddBoneStringCustomAttribute(const FName& BoneName, const FName& AttributeName, const TArray<float>& TimeKeys, const TArray<FString>& ValueKeys);
 
 	UE_DEPRECATED(5.0, "RemoveCustomAttribute has been deprecated see UAnimDataController::RemoveAttribute")
-	UFUNCTION(BlueprintCallable, Category = CustomAttributes)
+	UFUNCTION(BlueprintCallable, Category = CustomAttributes, meta=(DeprecatedFunction, DeprecationMessage="RemoveCustomAttribute has been deprecated, use UAnimDataController::RemoveAttribute instead"))
 	void RemoveCustomAttribute(const FName& BoneName, const FName& AttributeName);
 
 	UE_DEPRECATED(5.0, "RemoveAllCustomAttributesForBone has been deprecated see UAnimDataController::RemoveAllAttributesForBone")
-	UFUNCTION(BlueprintCallable, Category = CustomAttributes)
+	UFUNCTION(BlueprintCallable, Category = CustomAttributes, meta=(DeprecatedFunction, DeprecationMessage="RemoveAllCustomAttributesForBone has been deprecated, use UAnimDataController::RemoveAllAttributesForBone instead"))
 	void RemoveAllCustomAttributesForBone(const FName& BoneName);
 
 	UE_DEPRECATED(5.0, "RemoveAllCustomAttributes has been deprecated see UAnimDataController::RemoveAllAttributes")
-	UFUNCTION(BlueprintCallable, Category = CustomAttributes)
+	UFUNCTION(BlueprintCallable, Category = CustomAttributes, meta=(DeprecatedFunction, DeprecationMessage="RemoveAllCustomAttributes has been deprecated, use UAnimDataController::RemoveAllAttributes instead"))
 	void RemoveAllCustomAttributes();
-	
-	PRAGMA_DISABLE_DEPRECATION_WARNINGS
-	UE_DEPRECATED(5.0, "GetCustomAttributesForBone has been deprecated see UAnimDataModel::GetAttributesForBone")
-	void GetCustomAttributesForBone(const FName& BoneName, TArray<FCustomAttribute>& OutAttributes) const {}
-	PRAGMA_ENABLE_DEPRECATION_WARNINGS
 #endif // WITH_EDITOR
-	UE_DEPRECATED(5.0, "GetCustomAttributesForBone has been deprecated use EvaluateAttribute instead")
-	void GetCustomAttributes(FAnimationPoseData& OutAnimationPoseData, const FAnimExtractContext& ExtractionContext, bool bUseRawData) const
-	{
-		EvaluateAttributes(OutAnimationPoseData, ExtractionContext, bUseRawData);
-	}
 
 	void EvaluateAttributes(FAnimationPoseData& OutAnimationPoseData, const FAnimExtractContext& ExtractionContext, bool bUseRawData) const;	
 protected:
diff --git a/Engine/Source/Runtime/Engine/Classes/Animation/AnimSequenceBase.h b/Engine/Source/Runtime/Engine/Classes/Animation/AnimSequenceBase.h
--- a/Engine/Source/Runtime/Engine/Classes/Animation/AnimSequenceBase.h
+++ b/Engine/Source/Runtime/Engine/Classes/Animation/AnimSequenceBase.h
@@ -171,11 +171,6 @@ public:
 	// update cache data (notify tracks, sync markers)
 	virtual void RefreshCacheData();
 
-#if WITH_EDITOR
-	UE_DEPRECATED(5.0, "Updating of any relevent Curve data is now handled through the UAnimDataModel notifies")
-	void RefreshCurveData() {}
-#endif // WITH_EDITOR
-
 	//~ Begin UAnimationAsset Interface
 #if WITH_EDITOR
 	virtual void RemapTracksToNewSkeleton(USkeleton* NewSkeleton, bool bConvertSpaces) override;
@@ -253,22 +248,6 @@ public:
 	virtual bool IsValidToPlay() const { return true; }
 	// ideally this would be animsequcnebase, but we might have some issue with that. For now, just allow AnimSequence
 	virtual class UAnimSequence* GetAdditiveBasePose() const { return nullptr; }
-
-	typedef FSimpleMulticastDelegate::FDelegate FOnAnimCurvesChanged;
-	/** Registers a delegate to be called after anim curves have changed*/
-	UE_DEPRECATED(5.0, "Functionality has been deprecated, register to UAnimDataModel::GetModifiedEvent instead")
-	void RegisterOnAnimCurvesChanged(const FOnAnimCurvesChanged& Delegate) {}
-
-	UE_DEPRECATED(5.0, "Functionality has been deprecated, register to UAnimDataModel::GetModifiedEvent instead")
-	void UnregisterOnAnimCurvesChanged(void* Unregister) {}
-
-	typedef FSimpleMulticastDelegate::FDelegate FOnAnimTrackCurvesChanged;
-	/** Registers a delegate to be called after anim track curves have changed*/
-	UE_DEPRECATED(5.0, "Functionality has been deprecated, register to UAnimDataModel::GetModifiedEvent instead")
-	void RegisterOnAnimTrackCurvesChanged(const FOnAnimTrackCurvesChanged& Delegate) {}
-
-	UE_DEPRECATED(5.0, "Functionality has been deprecated, register to UAnimDataModel::GetModifiedEvent instead")
-	void UnregisterOnAnimTrackCurvesChanged(void* Unregister) {}
 #endif
 
 	// return true if anim notify is available 
diff --git a/Engine/Source/Runtime/Engine/Private/Animation/AnimDataModel.cpp b/Engine/Source/Runtime/Engine/Private/Animation/AnimDataModel.cpp
--- a/Engine/Source/Runtime/Engine/Private/Animation/AnimDataModel.cpp
+++ b/Engine/Source/Runtime/Engine/Private/Animation/AnimDataModel.cpp
@@ -34,7 +34,6 @@ void UAnimDataModel::PostDuplicate(bool bDuplicateForPIE)
 {
 	UObject::PostDuplicate(bDuplicateForPIE);
 
-	GenerateTransientData();
 	Notify(EAnimDataModelNotifyType::Populated);
 }
 
@@ -379,72 +378,6 @@ FGuid UAnimDataModel::GenerateGuid() const
 	return Guid;
 }
 
-const TArray<FRawAnimSequenceTrack>& UAnimDataModel::GetTransientRawAnimationTracks() const
-{
-	CheckTransientData();
-	return RawAnimationTracks;
-}
-
-const TArray<FName>& UAnimDataModel::GetTransientRawAnimationTrackNames() const
-{
-	CheckTransientData();
-	return RawAnimationTrackNames;
-}
-
-const TArray<FTrackToSkeletonMap>& UAnimDataModel::GetTransientRawAnimationTrackSkeletonMappings() const
-{
-	CheckTransientData();
-	return RawAnimationTrackSkeletonMappings;
-}
-
-FRawAnimSequenceTrack& UAnimDataModel::GetNonConstRawAnimationTrackByIndex(int32 TrackIndex)
-{
-	CheckTransientData();
-	checkf(BoneAnimationTracks.IsValidIndex(TrackIndex), TEXT("Invalid track index"));
-	return BoneAnimationTracks[TrackIndex].InternalTrackData;
-}
-
-const FRawCurveTracks& UAnimDataModel::GetTransientRawCurveTracks() const
-{
-	CheckTransientData();
-	return RawCurveTracks;
-}
-
-FAnimationCurveData& UAnimDataModel::GetNonConstCurveData()
-{
-	CheckTransientData();	
-	return CurveData;	
-}
-
-void UAnimDataModel::GenerateTransientData() const
-{
-	RawAnimationTracks.Empty(BoneAnimationTracks.Num());
-	RawAnimationTrackNames.Empty(BoneAnimationTracks.Num());
-	RawAnimationTrackSkeletonMappings.Empty(BoneAnimationTracks.Num());
-
-	for (const FBoneAnimationTrack& AnimTrack : BoneAnimationTracks)
-	{
-		RawAnimationTracks.Add(AnimTrack.InternalTrackData);
-		RawAnimationTrackNames.Add(AnimTrack.Name);
-		RawAnimationTrackSkeletonMappings.Add(AnimTrack.BoneTreeIndex);
-	}
-	
-	RawCurveTracks.FloatCurves = CurveData.FloatCurves;
-#if WITH_EDITOR
-	RawCurveTracks.TransformCurves = CurveData.TransformCurves;
-#endif
-	
-	bHasTransientDataBeenGenerated = true;
-}
-
-void UAnimDataModel::CheckTransientData() const
-{
-	if(!bHasTransientDataBeenGenerated)
-	{
-		GenerateTransientData();
-	}
-}
-
 FRichCurve* UAnimDataModel::GetMutableRichCurve(const FAnimationCurveIdentifier& CurveIdentifier)
 {
 	FRichCurve* RichCurve = nullptr;
diff --git a/Engine/Source/Runtime/Engine/Private/Animation/AnimSequence.cpp b/Engine/Source/Runtime/Engine/Private/Animation/AnimSequence.cpp
--- a/Engine/Source/Runtime/Engine/Private/Animation/AnimSequence.cpp
+++ b/Engine/Source/Runtime/Engine/Private/Animation/AnimSequence.cpp
@@ -498,51 +498,47 @@ void UAnimSequence::AddReferencedObjects(UObject* This, FReferenceCollector& Col
 	Collector.AddReferencedObject(AnimSeq->CompressedData.CurveCompressionCodec);
 }
 
+#if WITH_EDITOR
 int32 UAnimSequence::GetUncompressedRawSize() const
 {
-#if WITH_EDITOR
 	if (DataModel)
 	{
 		const int32 BoneRawSize = FRawAnimSequenceTrack::SingleKeySize * DataModel->GetNumBoneTracks() * DataModel->GetNumberOfKeys();
-	int32 CurveRawSize = 0;
+		int32 CurveRawSize = 0;
 
 		for (const FFloatCurve& Curve : DataModel->GetCurveData().FloatCurves)
-	{
-		CurveRawSize += sizeof(FFloatCurve);
-		CurveRawSize += sizeof(FRichCurveKey) * Curve.FloatCurve.Keys.Num();
+		{
+			CurveRawSize += sizeof(FFloatCurve);
+			CurveRawSize += sizeof(FRichCurveKey) * Curve.FloatCurve.Keys.Num();
+		}
+		return BoneRawSize + CurveRawSize;
 	}
-	return BoneRawSize + CurveRawSize;
-}
-
-#endif // WITH_EDITOR
-
 	return 0;
 }
 
 int32 UAnimSequence::GetApproxRawSize() const
 {
-#if WITH_EDITOR
 	if (DataModel)
 	{
 		int32 Total = sizeof(FRawAnimSequenceTrack) * DataModel->GetNumBoneTracks();
 		for (const FBoneAnimationTrack& AnimationTrack : DataModel->GetBoneAnimationTracks())
-	{
-			Total += AnimationTrack.InternalTrackData.PosKeys.Num() * AnimationTrack.InternalTrackData.PosKeys.GetTypeSize();
-			Total += AnimationTrack.InternalTrackData.RotKeys.Num() * AnimationTrack.InternalTrackData.RotKeys.GetTypeSize();
-			Total += AnimationTrack.InternalTrackData.ScaleKeys.Num() * AnimationTrack.InternalTrackData.ScaleKeys.GetTypeSize();
-	}
-
-		for (const FFloatCurve& Curve : DataModel->GetCurveData().FloatCurves)
-	{
-		Total += sizeof(FFloatCurve);
-		Total += sizeof(FRichCurveKey) * Curve.FloatCurve.Keys.Num();
+		{
+				Total += AnimationTrack.InternalTrackData.PosKeys.Num() * AnimationTrack.InternalTrackData.PosKeys.GetTypeSize();
+				Total += AnimationTrack.InternalTrackData.RotKeys.Num() * AnimationTrack.InternalTrackData.RotKeys.GetTypeSize();
+				Total += AnimationTrack.InternalTrackData.ScaleKeys.Num() * AnimationTrack.InternalTrackData.ScaleKeys.GetTypeSize();
+		}
+	
+			for (const FFloatCurve& Curve : DataModel->GetCurveData().FloatCurves)
+		{
+			Total += sizeof(FFloatCurve);
+			Total += sizeof(FRichCurveKey) * Curve.FloatCurve.Keys.Num();
+		}
+		return Total;
 	}
-	return Total;
-}
-#endif // WITH_EDITOR
 
 	return 0;
 }
+#endif // WITH_EDITOR
 
 int32 UAnimSequence::GetApproxBoneCompressedSize() const
 {
@@ -1034,18 +1030,6 @@ void ShowResaveMessage(const UAnimSequence* Sequence)
 	}
 }
 
-void UAnimSequence::VerifyTrackMap(USkeleton* MySkeleton)
-{
-	ValidateModel();
-
-	USkeleton* UseSkeleton = (MySkeleton) ? MySkeleton : GetSkeleton();
-
-	if(Controller->RemoveBoneTracksMissingFromSkeleton(UseSkeleton))
-	{
-		ShowResaveMessage(this);
-	}
-}
-
 #endif // WITH_EDITOR
 void UAnimSequence::BeginDestroy()
 {
@@ -1218,31 +1202,6 @@ void UAnimSequence::GetBoneTransform(FTransform& OutAtom, int32 TrackIndex, FAni
 	}
 }
 
-void UAnimSequence::ExtractBoneTransform(const TArray<struct FRawAnimSequenceTrack>& InRawAnimationData, FTransform& OutAtom, int32 TrackIndex, float Time) const
-{
-	// Bail out if the animation data doesn't exists (e.g. was stripped by the cooker).
-	if (InRawAnimationData.IsValidIndex(TrackIndex))
-	{
-		UE_LOG(LogAnimation, Log, TEXT("UAnimSequence::GetBoneTransform : No valid anim data track!"));
-		OutAtom.SetIdentity();
-		return;
-	}
-
-	PRAGMA_DISABLE_DEPRECATION_WARNINGS
-	return ExtractBoneTransform(InRawAnimationData[TrackIndex], OutAtom, Time);
-	PRAGMA_ENABLE_DEPRECATION_WARNINGS
-}
-
-void UAnimSequence::ExtractBoneTransform(const struct FRawAnimSequenceTrack& RawTrack, FTransform& OutAtom, int32 KeyIndex) const
-{
-	UE::Anim::ExtractBoneTransform(RawTrack, OutAtom, KeyIndex);
-}
-
-void UAnimSequence::ExtractBoneTransform(const struct FRawAnimSequenceTrack& RawTrack, FTransform& OutAtom, float Time) const
-{
-	FAnimationUtils::ExtractTransformFromTrack(Time, GetNumberOfSampledKeys(), GetPlayLength(), RawTrack, Interpolation, OutAtom);
-}
-
 void UAnimSequence::HandleAssetPlayerTickedInternal(FAnimAssetTickContext &Context, const float PreviousTime, const float MoveDelta, const FAnimTickRecord &Instance, struct FAnimNotifyQueue& NotifyQueue) const
 {
 	Super::HandleAssetPlayerTickedInternal(Context, PreviousTime, MoveDelta, Instance, NotifyQueue);
@@ -1612,36 +1571,6 @@ void UAnimSequence::GetBonePose(FAnimationPoseData& OutAnimationPoseData, const
 	EvaluateAttributes(OutAnimationPoseData, ExtractionContext, false);
 }
 
-const TArray<FRawAnimSequenceTrack>& UAnimSequence::GetRawAnimationData() const
-{
-#if WITH_EDITOR
-	if (IsDataModelValid())
-	{
-	PRAGMA_DISABLE_DEPRECATION_WARNINGS
-	return DataModel->GetTransientRawAnimationTracks();
-	PRAGMA_ENABLE_DEPRECATION_WARNINGS
-	}
-#endif
-	
-	static const TArray<FRawAnimSequenceTrack> Tracks;
-	check(false);
-	return Tracks;
-}
-
-#if WITH_EDITORONLY_DATA
-const TArray<FName>& UAnimSequence::GetAnimationTrackNames() const
-{
-	if (IsDataModelValid())
-	{
-	PRAGMA_DISABLE_DEPRECATION_WARNINGS
-	return DataModel->GetTransientRawAnimationTrackNames();
-	PRAGMA_ENABLE_DEPRECATION_WARNINGS
-	}
-	static const TArray<FName> Names;
-	check(false);
-	return Names;
-}
-#endif 
 
 #if WITH_EDITOR
 void UAnimSequence::UpdateCompressedCurveName(SmartName::UID_Type CurveUID, const struct FSmartName& NewCurveName)
@@ -1655,40 +1584,6 @@ void UAnimSequence::UpdateCompressedCurveName(SmartName::UID_Type CurveUID, cons
 		}
 	}
 }
-
-int32 UAnimSequence::AddNewRawTrack(FName TrackName, FRawAnimSequenceTrack* TrackData)
-{
-	ValidateModel();
-
-	int32 NewTrackIndex = INDEX_NONE;
-	Controller->OpenBracket(LOCTEXT("AddNewRawTrack_Bracket", "Adding new Bone Animation Track"));
-	{
-		NewTrackIndex = Controller->AddBoneTrack(TrackName);
-		if (NewTrackIndex != INDEX_NONE && TrackData)
-		{
-			Controller->SetBoneTrackKeys(TrackName, TrackData->PosKeys, TrackData->RotKeys, TrackData->ScaleKeys);
-		}
-	}
-	Controller->CloseBracket();
-
-	if (NewTrackIndex == INDEX_NONE)
-	{
-		FFormatNamedArguments Args;
-		Args.Add(TEXT("TrackName"), FText::FromName(TrackName));
-		FNotificationInfo Info(FText::Format(LOCTEXT("FailedToAddTracks", "Failed to add track with name \"{TrackName}\"."), Args));
-
-		Info.bUseLargeFont = false;
-		Info.ExpireDuration = 5.0f;
-
-		TSharedPtr<SNotificationItem> Notification = FSlateNotificationManager::Get().AddNotification(Info);
-		if (Notification.IsValid())
-		{
-			Notification->SetCompletionState(SNotificationItem::CS_Fail);
-		}
-	}
-
-	return NewTrackIndex;
-}
 #endif
 
 void UAnimSequence::GetBonePose_Additive(FCompactPose& OutPose, FBlendedCurve& OutCurve, const FAnimExtractContext& ExtractionContext) const
@@ -1893,154 +1788,11 @@ FName UAnimSequence::GetRetargetTransformsSourceName() const
 	}
 }
 
-const TArray<FTrackToSkeletonMap>& UAnimSequence::GetRawTrackToSkeletonMapTable() const
-{
-#if WITH_EDITOR
-	if (IsDataModelValid())
-	{
-	PRAGMA_DISABLE_DEPRECATION_WARNINGS
-	return DataModel->GetTransientRawAnimationTrackSkeletonMappings();
-	PRAGMA_ENABLE_DEPRECATION_WARNINGS
-	}
-#endif
-
-	static const TArray<FTrackToSkeletonMap> TempMap;
-	check(false);
-	return TempMap;
-}
-
-const FRawAnimSequenceTrack& UAnimSequence::GetRawAnimationTrack(int32 TrackIndex) const
-{
-#if WITH_EDITOR
-	if (IsDataModelValid())
-	{
-	PRAGMA_DISABLE_DEPRECATION_WARNINGS
-	return DataModel->GetBoneTrackByIndex(TrackIndex).InternalTrackData;
-	PRAGMA_ENABLE_DEPRECATION_WARNINGS 
-	}
-#endif
-	
-	static const FRawAnimSequenceTrack TempTrack;
-	check(false);
-	return TempTrack;
-}
-
-FRawAnimSequenceTrack& UAnimSequence::GetRawAnimationTrack(int32 TrackIndex)
-{
-#if WITH_EDITOR
-	if (IsDataModelValid())
-	{
-	ValidateModel();
-	PRAGMA_DISABLE_DEPRECATION_WARNINGS
-	return DataModel->GetNonConstRawAnimationTrackByIndex(TrackIndex);
-	PRAGMA_ENABLE_DEPRECATION_WARNINGS
-	}
-#endif
-
-		static FRawAnimSequenceTrack TempTrack;
-	check(false);
-	return TempTrack;
-}
-
 void UAnimSequence::RetargetBoneTransform(FTransform& BoneTransform, const int32 SkeletonBoneIndex, const FCompactPoseBoneIndex& BoneIndex, const FBoneContainer& RequiredBones, const bool bIsBakedAdditive) const
 {
 	FAnimationRuntime::RetargetBoneTransform(GetSkeleton(), GetRetargetTransformsSourceName(), GetRetargetTransforms(), BoneTransform, SkeletonBoneIndex, BoneIndex, RequiredBones, bIsBakedAdditive);
 }
 
-#if WITH_EDITOR
-bool UAnimSequence::InsertFramesToRawAnimData( int32 StartFrame, int32 EndFrame, int32 CopyFrame)
-{
-	ValidateModel();
-
-	// make sure the copyframe is valid and start frame is valid
-	const int32 NumFramesToInsert = EndFrame-StartFrame;
-	if ((CopyFrame>=0 && CopyFrame< DataModel->GetNumberOfKeys()) && (StartFrame >= 0 && StartFrame <= DataModel->GetNumberOfKeys()) && NumFramesToInsert > 0)
-	{
-		IAnimationDataController::FScopedBracket ScopedBracket(Controller, LOCTEXT("InsertFramesToRawAnimData", "Inserting Frames into Animation Track Data"));
-		UE::Anim::AnimationData::DuplicateKeys(this, StartFrame, NumFramesToInsert, CopyFrame);
-
-		UE_LOG(LogAnimation, Log, TEXT("\tPlay Length: %f, Number of Keys: %d"), GetPlayLength(), DataModel->GetNumberOfKeys());
-
-		return true;
-	}
-
-	return false;
-}
-
-bool UAnimSequence::CropRawAnimData( float CurrentTime, bool bFromStart )
-{
-	ValidateModel();
-	
-	// Save Total Number of Frames before crop
-	const int32 TotalNumOfKeys = GetDataModel()->GetNumberOfKeys();
-	// if current frame is 1, do not try crop. There is nothing to crop
-	if (TotalNumOfKeys <= 1 )
-	{
-		return false;
-	}
-	
-	// If you're end or beginning, you can't cut all nor nothing. 
-	// Avoiding ambiguous situation what exactly we would like to cut 
-	// Below it clamps range to 1, TotalNumOfFrames-1
-	// causing if you were in below position, it will still crop 1 frame. 
-	// To be clearer, it seems better if we reject those inputs. 
-	// If you're a bit before/after, we assume that you'd like to crop
-	if ( CurrentTime == 0.f || CurrentTime == GetPlayLength() )
-	{
-		return false;
-	}
-
-	const FFrameTime FrameTime = DataModel->GetFrameRate().AsFrameTime(CurrentTime);
-	const FFrameNumber FrameNumber = FrameTime.CeilToFrame();
-	// determine which keys need to be removed.
-	const int32 StartKey = bFromStart ? 0 : FrameNumber.Value;
-	const int32 NumKeysToRemove = bFromStart ? FrameNumber.Value : TotalNumOfKeys - (StartKey + 1);
-
-	// Recalculate number of keys and frames
-	const int32 NewNumKeys = TotalNumOfKeys - NumKeysToRemove;
-	const int32 NewNumFrames = NewNumKeys - 1;
-
-	UE_LOG(LogAnimation, Log, TEXT("UAnimSequence::CropRawAnimData %s - CurrentTime: %f, bFromStart: %d, TotalNumOfKeys: %d, KeyIndex: %d, StartKey: %d, NumKeys: %d"), *GetName(), CurrentTime, bFromStart, TotalNumOfKeys, FrameNumber.Value, StartKey, NumKeysToRemove);
-
-	Controller->OpenBracket(LOCTEXT("CropRawAnimation_Bracket", "Cropping Animation Track Data"));
-	UE::Anim::AnimationData::RemoveKeys(this, StartKey, NumKeysToRemove);
-	Controller->CloseBracket();
-
-	UE_LOG(LogAnimation, Log, TEXT("\tPlay Length: %f, Number of Keys: %d"), GetPlayLength(), GetDataModel()->GetNumberOfKeys());
-
-	MarkPackageDirty();
-
-	return true;
-}
-
-void UAnimSequence::FlipRotationWForNonRoot(USkeletalMesh* SkelMesh)
-{
-	if (!GetSkeleton())
-	{
-		return;
-	}
-
-	ValidateModel();
-	Controller->OpenBracket(LOCTEXT("FlipRotationWForNonRoot_Bracket", "Flipping rotation W component for non-root tracks"));
-	{
-		for (const FBoneAnimationTrack& Track : DataModel->GetBoneAnimationTracks())
-		{
-			// Only apply to non-root bones
-			if (Track.BoneTreeIndex > 0)
-			{
-				TArray<FQuat4f> NewRotationalKeys = Track.InternalTrackData.RotKeys;
-				for (FQuat4f& RotationalKey : NewRotationalKeys)
-				{
-					RotationalKey.W *= -1.f;
-				}
-
-				Controller->SetBoneTrackKeys(Track.Name, Track.InternalTrackData.PosKeys, NewRotationalKeys, Track.InternalTrackData.ScaleKeys);
-			}
-		}
-	}
-	Controller->CloseBracket();
-}
-#endif 
 
 #if WITH_EDITOR
 bool UAnimSequence::ShouldPerformStripping(const bool bPerformFrameStripping, const bool bPerformStrippingOnOddFramedAnims) const
@@ -2062,9 +1814,9 @@ FString UAnimSequence::GetDDCCacheKeySuffix(const bool bPerformStripping) const
 }
 #endif
 
+#if WITH_EDITOR
 void UAnimSequence::WaitOnExistingCompression(const bool bWantResults)
 {
-#if WITH_EDITOR
 	check(IsInGameThread());
 	if (bCompressionInProgress)
 	{
@@ -2073,12 +1825,10 @@ void UAnimSequence::WaitOnExistingCompression(const bool bWantResults)
 		bCompressionInProgress = false;
 		COOK_STAT(Timer.TrackCyclesOnly()); // Need to get hit/miss and size from WaitOnExistingCompression!
 	}
-#endif
 }
 
 void UAnimSequence::RequestAnimCompression(FRequestAnimCompressionParams Params)
 {
-#if WITH_EDITOR
 	check(IsInGameThread());
 	USkeleton* CurrentSkeleton = GetSkeleton();
 	if (CurrentSkeleton == nullptr)
@@ -2194,10 +1944,8 @@ void UAnimSequence::RequestAnimCompression(FRequestAnimCompressionParams Params)
 		}
 
 	}
-#endif
 }
 
-#if WITH_EDITOR
 void UAnimSequence::ApplyCompressedData(const FString& DataCacheKeySuffix, const bool bPerformFrameStripping, const TArray<uint8>& Data)
 {
 	if (GetDDCCacheKeySuffix(bPerformFrameStripping) == DataCacheKeySuffix)
@@ -2240,18 +1988,6 @@ void UAnimSequence::SerializeCompressedData(FArchive& Ar, bool bDDCData)
 		}
 	}
 
-int32 UAnimSequence::GetSkeletonIndexFromRawDataTrackIndex(const int32 TrackIndex) const
-{
-#if WITH_EDITOR
-	if (IsDataModelValid())
-	{
-	return DataModel->GetBoneTrackByIndex(TrackIndex).BoneTreeIndex;
-	}
-#endif // WITH_EDITOR
-	
-	check(false);
-	return INDEX_NONE;
-}
 
 #if WITH_EDITOR
 bool UAnimSequence::CanBakeAdditive() const
@@ -2720,21 +2456,8 @@ void UAnimSequence::UpdateDependentStreamingAnimations() const
 
 #endif
 
-void UAnimSequence::RecycleAnimSequence()
-{
-#if WITH_EDITOR
-	ResetAnimation();
-#endif // WITH_EDITORONLY_DATA
-}
 
 #if WITH_EDITOR
-void UAnimSequence::CleanAnimSequenceForImport()
-{
-	ValidateModel();
-	DeleteBoneAnimationData();
-	RawDataGuid.Invalidate();
-}
-
 void UAnimSequence::ResetAnimation()
 {
 	// clear everything. Making new animation, so need to reset all the things that belong here
@@ -2780,15 +2503,6 @@ void UAnimSequence::DeleteBoneAnimationData()
 	Controller->RemoveAllBoneTracks();
 }
 #endif // WITH_EDITOR
-		
-
-bool UAnimSequence::CopyNotifies(UAnimSequence* SourceAnimSeq, UAnimSequence* DestAnimSeq, bool bShowDialogs /*= true */)
-{
-#if WITH_EDITOR
-	return UE::Anim::CopyNotifies(SourceAnimSeq, DestAnimSeq, bShowDialogs);
-#endif // WITH_EDITOR
-	return false;
-}
 
 bool UAnimSequence::IsValidAdditive() const		
 { 
@@ -3361,17 +3075,6 @@ void UAnimSequence::RemapTracksToNewSkeleton( USkeleton* NewSkeleton, bool bConv
 	Super::RemapTracksToNewSkeleton(NewSkeleton, bConvertSpaces);
 }
 
-void UAnimSequence::PostProcessSequence(bool bForceNewRawDatGuid)
-{
-	PRAGMA_DISABLE_DEPRECATION_WARNINGS
-	// Apply compression
-	MarkRawDataAsModified(bForceNewRawDatGuid);
-	OnRawDataChanged();
-	PRAGMA_ENABLE_DEPRECATION_WARNINGS
-
-	// mark package as dirty
-	MarkPackageDirty();
-}
 
 void UAnimSequence::RemoveNaNTracks()
 {
@@ -3421,34 +3124,6 @@ void UAnimSequence::RemoveNaNTracks()
 	}	
 }
 
-void UAnimSequence::RemoveAllTracks()
-{
-	PRAGMA_DISABLE_DEPRECATION_WARNINGS
-	RawAnimationData.Empty();
-	AnimationTrackNames.Empty();
-	TrackToSkeletonMapTable.Empty();
-	SourceRawAnimationData_DEPRECATED.Empty();
-	PRAGMA_ENABLE_DEPRECATION_WARNINGS
-	
-	ValidateModel();
-	IAnimationDataController::FScopedBracket ScopedBracket(Controller, LOCTEXT("RemoveAllTracks_Bracket", "Removing all Bone Animation and Transform Curve Tracks"));
-	Controller->RemoveAllBoneTracks();
-
-	// clear all transform tracks
-	// not deleting curve names from skeleton 
-	// since we don't know if that name is used by other assets
-	Controller->RemoveAllCurvesOfType(ERawCurveTrackTypes::RCT_Transform);	
-}
-
-void UAnimSequence::RemoveTrack(int32 TrackIndex)
-{
-	ValidateModel();
-	
-	if (DataModel->IsValidBoneTrackIndex(TrackIndex))
-	{
-		Controller->RemoveBoneTrack(DataModel->GetBoneTrackByIndex(TrackIndex).Name);
-	}
-}
 
 int32 FindFirstChildTrack(const USkeleton* MySkeleton, const FReferenceSkeleton& RefSkeleton, const TArray<FName>& AnimationTrackNames, FName BoneName)
 {
@@ -3505,63 +3180,6 @@ int32 FindFirstChildTrack(const USkeleton* MySkeleton, const FReferenceSkeleton&
 	}
 }
 
-int32 UAnimSequence::InsertTrack(const FName& BoneName)
-{
-	ValidateModel();
-
-	// first verify if it doesn't exists, if it does, return
-	const int32 ExistingTrackIndex = DataModel->GetBoneTrackIndexByName(BoneName);
-	if (ExistingTrackIndex != INDEX_NONE)
-	{
-		return ExistingTrackIndex;
-	}
-
-	USkeleton * MySkeleton = GetSkeleton();
-	// should not call this if skeleton was empty
-	if (ensure(MySkeleton) == false)
-	{
-		return INDEX_NONE;
-	}
-
-	const FReferenceSkeleton& RefSkeleton = MySkeleton->GetReferenceSkeleton();
-
-	TArray<FName> TrackNames;
-	DataModel->GetBoneTrackNames(TrackNames);
-
-	const int32 NewTrackIndex = FindFirstChildTrack(MySkeleton, RefSkeleton, TrackNames, BoneName);
-	ensure(NewTrackIndex <= DataModel->GetNumBoneTracks());
-	const int32 BoneIndex = RefSkeleton.FindBoneIndex(BoneName);
-	if (NewTrackIndex != INDEX_NONE)
-	{
-		const TArray<FTransform>& RefPose = RefSkeleton.GetRefBonePose();
-
-		FRawAnimSequenceTrack RawTrack;
-		RawTrack.PosKeys.Add(FVector3f(RefPose[BoneIndex].GetTranslation()));
-		RawTrack.RotKeys.Add(FQuat4f(RefPose[BoneIndex].GetRotation()));
-		RawTrack.ScaleKeys.Add(FVector3f(RefPose[BoneIndex].GetScale3D()));
-
-		// now insert to the track
-		Controller->InsertBoneTrack(BoneName, NewTrackIndex);
-	}
-
-	return NewTrackIndex;
-}
-
-void UAnimSequence::ResizeSequence(float NewLength, int32 NewNumFrames, bool bInsert, int32 StartFrame/*inclusive */, int32 EndFrame/*inclusive*/)
-{
- 	ValidateModel();
-
-	const FFrameRate& FrameRate = DataModel->GetFrameRate();
-	const float T0 = FrameRate.AsSeconds(StartFrame);
-	const float T1 = FrameRate.AsSeconds(EndFrame);
-
-	FFrameRate NewFrameRate = FFrameRate((NewLength / NewNumFrames), 1);
-
-	Controller->Resize(NewLength, T0, T1);
-
-	ensure(DataModel->GetNumberOfKeys() == NewNumFrames);
-}
-
 bool UAnimSequence::GetAllAnimationSequencesReferred(TArray<UAnimationAsset*>& AnimationAssets, bool bRecursive /*= true*/)
 {
 	Super::GetAllAnimationSequencesReferred(AnimationAssets, bRecursive);
@@ -3586,11 +3204,6 @@ void UAnimSequence::ReplaceReferredAnimations(const TMap<UAnimationAsset*, UAnim
 	}
 }
 
-bool UAnimSequence::AddLoopingInterpolation()
-{
-	return UE::Anim::AnimationData::AddLoopingInterpolation(this);
-}
-
 int32 FindParentNodeIndex(URig* Rig, USkeleton* Skeleton, FName ParentNodeName)
 {
 	const int32& ParentNodeIndex = Rig->FindNode(ParentNodeName);
@@ -4047,19 +3660,6 @@ bool UAnimSequence::ConvertRiggingDataToAnimationData(FAnimSequenceTrackContaine
 	return false;
 }
 
-void UAnimSequence::ClearBakedTransformData()
-{
-	UE_LOG(LogAnimation, Warning, TEXT("[%s] Detected previous edited data is invalidated. Clearing transform curve data and Source Data. This can happen if you do retarget another animation to this. If not, please report back to Epic. "), *GetName());
-	//Clear Transform curve data
-	Controller->RemoveAllCurvesOfType(ERawCurveTrackTypes::RCT_Transform);
-}
-
-bool UAnimSequence::DoesContainTransformCurves() const
-{
-	ValidateModel();
-	return DataModel->GetNumberOfTransformCurves() > 0;
-}
-
 void UAnimSequence::BakeTrackCurvesToRawAnimationTracks(TArray<FRawAnimSequenceTrack>& InOutNewRawTracks, TArray<FName>& InOutNewTrackNames, TArray<FTrackToSkeletonMap>& InOutNewTrackToSkeletonMapTable)
 {
 	ValidateModel();
@@ -4219,13 +3819,6 @@ void UAnimSequence::AddKeyToSequence(float Time, const FName& BoneName, const FT
 	Controller->SetTransformCurveKey(TransformCurveId, Time, AdditiveTransform);
 }
 
-void UAnimSequence::RefreshTrackMapFromAnimTrackNames()
-{
-	ValidateModel();
-	const USkeleton * MySkeleton = GetSkeleton();
-	Controller->RemoveBoneTracksMissingFromSkeleton(MySkeleton);
-}
-
 uint8* UAnimSequence::FindSyncMarkerPropertyData(int32 SyncMarkerIndex, FArrayProperty*& ArrayProperty)
 {
 	ArrayProperty = NULL;
@@ -5481,15 +5074,6 @@ void UAnimSequence::EnableRootMotionSettingFromMontage(bool bInEnableRootMotion,
 	}
 }
 
-#if WITH_EDITOR
-void UAnimSequence::OnRawDataChanged()
-{
-	RecompressAnimationData();
-	//MDW - Once we have async anim ddc requests we should do this too
-	//RequestDependentAnimRecompression();
-}
-#endif
-
 bool UAnimSequence::IsCompressedDataValid() const
 {
 	// For bone compressed data, we don't check if we have a codec. It is valid to have no compressed data
diff --git a/Engine/Source/Runtime/Engine/Private/Animation/AnimSequenceHelpers.cpp b/Engine/Source/Runtime/Engine/Private/Animation/AnimSequenceHelpers.cpp
--- a/Engine/Source/Runtime/Engine/Private/Animation/AnimSequenceHelpers.cpp
+++ b/Engine/Source/Runtime/Engine/Private/Animation/AnimSequenceHelpers.cpp
@@ -17,6 +17,7 @@
 #include "Animation/AnimNotifies/AnimNotify.h"
 #include "Animation/AnimNotifies/AnimNotifyState.h"
 #include "Misc/MessageDialog.h"
+#include "Animation/AnimSequenceHelpers.h"
 
 #define LOCTEXT_NAMESPACE "AnimSequenceHelpers"
 
@@ -805,6 +806,16 @@ bool AnimationData::Trim(UAnimSequence* InSequence, float TrimStart, float TrimE
 
 		const int32 StartTrimKeyIndex = StartFrameTrim.Value;
 		const int32 NumTrimmedFrames = EndFrameTrim.Value - StartFrameTrim.Value;
+
+		if (StartTrimKeyIndex == 0 && NumTrimmedFrames + 1 == NumKeys)
+		{
+			return false;
+		}
+
+		if (NumTrimmedFrames == 0)
+		{
+			return false;
+		}
 		
 		IAnimationDataController::FScopedBracket ScopedBracket(Controller, LOCTEXT("TrimRawAnimation_Bracket", "Trimming Animation Track Data"));
 		RemoveKeys(InSequence, StartTrimKeyIndex, NumTrimmedFrames);
