FRichCurve::RemoveRedundantKeys has potential to break curve evaluation outside of user-set threshold
    * Added Sample frame-rate to RemoveRedundantKeys, this allows FRichCurve to evaluate between the to-be-removed keys to make sure removing does not impact _actual_ evaluated values (rather than just at the point of the removed key)
        + Added RemoveRedundantAutoTangentKeys and deprecation to funnel existing use-cases to switch to new implementation (if desired)
    * Added same behaviour to CurveEditorReduceFilter

diff --git a/Engine/Plugins/Runtime/Synthesis/Source/Synthesis/Private/SynthComponentMonoWaveTable.cpp b/Engine/Plugins/Runtime/Synthesis/Source/Synthesis/Private/SynthComponentMonoWaveTable.cpp
--- a/Engine/Plugins/Runtime/Synthesis/Source/Synthesis/Private/SynthComponentMonoWaveTable.cpp
+++ b/Engine/Plugins/Runtime/Synthesis/Source/Synthesis/Private/SynthComponentMonoWaveTable.cpp
@@ -87,7 +87,7 @@ void UMonoWaveTableSynthPreset::EditChangeInternal()
 			bChangedThisCurve = true;
 		}
 
-		Curve.EditorCurveData.RemoveRedundantKeys(SMALL_NUMBER, -0.1, 1.1);
+		Curve.EditorCurveData.RemoveRedundantAutoTangentKeys(SMALL_NUMBER, -0.1, 1.1);
 
 		// Loop through Keys in this Curve...
 		for (auto& Key : Curve.EditorCurveData.Keys)
diff --git a/Engine/Source/Editor/CurveEditor/Private/Filters/CurveEditorReduceFilter.cpp b/Engine/Source/Editor/CurveEditor/Private/Filters/CurveEditorReduceFilter.cpp
--- a/Engine/Source/Editor/CurveEditor/Private/Filters/CurveEditorReduceFilter.cpp
+++ b/Engine/Source/Editor/CurveEditor/Private/Filters/CurveEditorReduceFilter.cpp
@@ -109,14 +109,69 @@ void UCurveEditorReduceFilter::ApplyFilter_Impl(TSharedRef<FCurveEditor> InCurve
 			TArray<FKeyHandle> KeysToRemove;
 			for (int32 TestIndex = 1; TestIndex < KeyHandles.Num() - 1; ++TestIndex)
 			{
+				bool bKeepKey = false;				
+				if (bTryRemoveUserSetTangentKeys && SampleRate.IsValid())
+				{
+					const FFrameTime KeyOneFrameTime = SampleRate.AsFrameTime(SelectedKeyPositions[MostRecentKeepKeyIndex].InputValue);
+					const FFrameTime KeyTwoFrameTime = SampleRate.AsFrameTime(SelectedKeyPositions[TestIndex + 1].InputValue);
+
+					FFrameTime SampleFrameTime = KeyOneFrameTime;
+					const double ToRemoveKeyTime = SelectedKeyPositions[TestIndex].InputValue;
+
+					// Sample curve at intervals, dictated by provided sample-rate, between the keys surrounding the to-be-removed key
+					// this makes sure that any impact on non-lerp keys is taken into account
+					while (SampleFrameTime <= KeyTwoFrameTime)
+					{
+						const double SampleSeconds = SampleRate.AsSeconds(SampleFrameTime);
+						const float ValueWithoutKey = EvalForTwoKeys(SelectedKeyPositions[MostRecentKeepKeyIndex],	SelectedKeyAttributes[MostRecentKeepKeyIndex],
+							SelectedKeyPositions[TestIndex + 1], SelectedKeyAttributes[TestIndex + 1],
+							SampleSeconds);
+
+						const FKeyPosition& KeyOne = SampleSeconds <= ToRemoveKeyTime ? SelectedKeyPositions[MostRecentKeepKeyIndex] : SelectedKeyPositions[TestIndex];
+						const FKeyPosition& KeyTwo = SampleSeconds <= ToRemoveKeyTime ? SelectedKeyPositions[TestIndex] : SelectedKeyPositions[TestIndex + 1];
+						
+						const FKeyAttributes& AttributesOne = SampleSeconds <= ToRemoveKeyTime ? SelectedKeyAttributes[MostRecentKeepKeyIndex] : SelectedKeyAttributes[TestIndex];
+						const FKeyAttributes& AttributesTwo = SampleSeconds <= ToRemoveKeyTime ? SelectedKeyAttributes[TestIndex] : SelectedKeyAttributes[TestIndex + 1];
+						
+						const float ValueWithKey = EvalForTwoKeys(KeyOne, AttributesOne,KeyTwo, AttributesTwo, SampleSeconds);
+				
+						if (FMath::Abs(ValueWithoutKey - ValueWithKey) > Tolerance)
+						{
+							bKeepKey = true;
+							break;
+						}
+
+						// Next frame
+						SampleFrameTime.FrameNumber += 1;	
+					}	
+					
+				}
+				else
+				{
+					auto IsNonUserTangentKey = [](const FKeyAttributes& KeyAttributes)
+					{
+						return !KeyAttributes.HasTangentMode() ||
+							(KeyAttributes.GetTangentMode() != RCTM_User && KeyAttributes.GetTangentMode() != RCTM_Break);
+					};
+					
+					bKeepKey = true;
+
+					// Only try to remove keys which have automatic tangent data
+					if (IsNonUserTangentKey(SelectedKeyAttributes[MostRecentKeepKeyIndex]) &&
+						IsNonUserTangentKey(SelectedKeyAttributes[TestIndex]) &&
+						IsNonUserTangentKey(SelectedKeyAttributes[TestIndex + 1]))
+					{
+						const float KeyValue = SelectedKeyPositions[TestIndex].OutputValue;
+						const float ValueWithoutKey = EvalForTwoKeys(SelectedKeyPositions[MostRecentKeepKeyIndex], SelectedKeyAttributes[MostRecentKeepKeyIndex],
+							SelectedKeyPositions[TestIndex + 1], SelectedKeyAttributes[TestIndex + 1],
+							SelectedKeyPositions[TestIndex].InputValue);
+
+						// Check if there is a great enough change in value to consider this key needed.
+						bKeepKey = FMath::Abs(ValueWithoutKey - KeyValue) > Tolerance;
+					}
+				}
 
-				const float KeyValue = SelectedKeyPositions[TestIndex].OutputValue;
-				const float ValueWithoutKey = EvalForTwoKeys(SelectedKeyPositions[MostRecentKeepKeyIndex], SelectedKeyAttributes[MostRecentKeepKeyIndex],
-					SelectedKeyPositions[TestIndex + 1], SelectedKeyAttributes[TestIndex + 1],
-					SelectedKeyPositions[TestIndex].InputValue);
-
-				// Check if there is a great enough change in value to consider this key needed.
-				if (FMath::Abs(ValueWithoutKey - KeyValue) > Tolerance)
+				if (bKeepKey)
 				{
 					MostRecentKeepKeyIndex = TestIndex;
 				}
diff --git a/Engine/Source/Editor/CurveEditor/Public/Filters/CurveEditorReduceFilter.h b/Engine/Source/Editor/CurveEditor/Public/Filters/CurveEditorReduceFilter.h
--- a/Engine/Source/Editor/CurveEditor/Public/Filters/CurveEditorReduceFilter.h
+++ b/Engine/Source/Editor/CurveEditor/Public/Filters/CurveEditorReduceFilter.h
@@ -13,6 +13,8 @@ public:
 	UCurveEditorReduceFilter() 
 	{
 		Tolerance = 0.1f;
+		bTryRemoveUserSetTangentKeys = false;
+		SampleRate = FFrameRate(120,1);
 	}
 
 protected:
@@ -22,4 +24,12 @@ public:
 	/** Minimum change in values required for a key to be considered distinct enough to keep. */
 	UPROPERTY(EditAnywhere, BlueprintReadWrite, Category = "Settings")
 	float Tolerance;
+	
+	/** Flag whether or not to use SampleRate for sampling between evaluated keys, which allows for removing user-tangent keys. */
+	UPROPERTY(EditAnywhere, BlueprintReadWrite, Category = "Settings")
+	bool bTryRemoveUserSetTangentKeys;
+	
+	/** Rate at which the curve should be sampled to compare against value tolerance. */
+	UPROPERTY(EditAnywhere, BlueprintReadWrite, Category = "Settings", meta = (EditCondition = "bTryRemoveUserSetTangentKeys"))
+	FFrameRate SampleRate;
 };
\ No newline at end of file
diff --git a/Engine/Source/Editor/UnrealEd/Private/SkeletalMeshEdit.cpp b/Engine/Source/Editor/UnrealEd/Private/SkeletalMeshEdit.cpp
--- a/Engine/Source/Editor/UnrealEd/Private/SkeletalMeshEdit.cpp
+++ b/Engine/Source/Editor/UnrealEd/Private/SkeletalMeshEdit.cpp
@@ -1355,7 +1355,7 @@ bool UnFbx::FFbxImporter::ImportCurveToAnimSequence(class UAnimSequence * Target
 		{
 			if (ImportOptions->bRemoveRedundantKeys)
 			{
-				RichCurve.RemoveRedundantKeys(SMALL_NUMBER);
+				RichCurve.RemoveRedundantAutoTangentKeys(SMALL_NUMBER);
 			}
 
 			// Set actual keys on curve within the model
diff --git a/Engine/Source/Runtime/Engine/Classes/Curves/RealCurve.h b/Engine/Source/Runtime/Engine/Classes/Curves/RealCurve.h
--- a/Engine/Source/Runtime/Engine/Classes/Curves/RealCurve.h
+++ b/Engine/Source/Runtime/Engine/Classes/Curves/RealCurve.h
@@ -4,6 +4,7 @@
 
 #include "CoreMinimal.h"
 #include "Curves/IndexedCurve.h"
+#include "Misc/FrameRate.h"
 #include "RealCurve.generated.h"
 
 /** Method of interpolation between this key and the next. */
@@ -131,9 +132,9 @@ public:
 	virtual void BakeCurve(float SampleRate) PURE_VIRTUAL(FRealCurve::BakeCurve, );
 	virtual void BakeCurve(float SampleRate, float FirstKeyTime, float LastKeyTime) PURE_VIRTUAL(FRealCurve::BakeCurve, );
 
-	/** Remove redundant keys, comparing against Tolerance */
-	virtual void RemoveRedundantKeys(float Tolerance) PURE_VIRTUAL(FRealCurve::RemoveRedundantKeys, );
-	virtual void RemoveRedundantKeys(float Tolerance, float FirstKeyTime, float LastKeyTime) PURE_VIRTUAL(FRealCurve::RemoveRedundantKeys, );
+	/** Remove redundant keys, comparing against Tolerance (and optional use of sample-rate for additional testing) */
+	virtual void RemoveRedundantKeys(float Tolerance, FFrameRate SampleRate = FFrameRate(0,0)) PURE_VIRTUAL(FRealCurve::RemoveRedundantKeys, );
+	virtual void RemoveRedundantKeys(float Tolerance, float FirstKeyTime, float LastKeyTime, FFrameRate SampleRate = FFrameRate(0,0)) PURE_VIRTUAL(FRealCurve::RemoveRedundantKeys, );
 
 protected:
 	static void CycleTime(float MinTime, float MaxTime, float& InTime, int& CycleCount);
diff --git a/Engine/Source/Runtime/Engine/Classes/Curves/RichCurve.h b/Engine/Source/Runtime/Engine/Classes/Curves/RichCurve.h
--- a/Engine/Source/Runtime/Engine/Classes/Curves/RichCurve.h
+++ b/Engine/Source/Runtime/Engine/Classes/Curves/RichCurve.h
@@ -313,17 +313,24 @@ public:
 	virtual void BakeCurve(float SampleRate, float FirstKeyTime, float LastKeyTime) final override;
 
 	/** Remove redundant keys, comparing against Tolerance */
-	virtual void RemoveRedundantKeys(float Tolerance) final override;
-	virtual void RemoveRedundantKeys(float Tolerance, float FirstKeyTime, float LastKeyTime) final override;
+	UE_DEPRECATED(5.1, "FRichCurve::RemoveRedundantKeys is deprecated, use signature with additional SampleRate or RemoveRedundantAutoTangentKeys instead")
+	void RemoveRedundantKeys(float Tolerance) { RemoveRedundantAutoTangentKeys(Tolerance); }
+	void RemoveRedundantAutoTangentKeys(float Tolerance);
+	virtual void RemoveRedundantKeys(float Tolerance, FFrameRate SampleRate) final override;
+
+	UE_DEPRECATED(5.1, "FRichCurve::RemoveRedundantKeys is deprecated, use signature with additional SampleRate or RemoveRedundantAutoTangentKeys instead")
+	void RemoveRedundantKeys(float Tolerance, float FirstKeyTime, float LastKeyTime) { RemoveRedundantAutoTangentKeys(Tolerance, FirstKeyTime, LastKeyTime); }
+	void RemoveRedundantAutoTangentKeys(float Tolerance, float FirstKeyTime, float LastKeyTime);
+	virtual void RemoveRedundantKeys(float Tolerance, float FirstKeyTime, float LastKeyTime, FFrameRate SampleRate) final override;
 
 	/** Compresses a rich curve for efficient runtime storage and evaluation */
-	void CompressCurve(struct FCompressedRichCurve& OutCurve, float ErrorThreshold = 0.0001f, float SampleRate = 120.0f) const;
+	void CompressCurve(struct FCompressedRichCurve& OutCurve, float ErrorThreshold = 0.0001f, float SampleRate = 120.0) const;
 
 	/** Allocates a duplicate of the curve */
 	virtual FIndexedCurve* Duplicate() const final { return new FRichCurve(*this); }
 
 private:
-	void RemoveRedundantKeysInternal(float Tolerance, int32 InStartKeepKey, int32 InEndKeepKey);
+	void RemoveRedundantKeysInternal(float Tolerance, int32 InStartKeepKey, int32 InEndKeepKey, FFrameRate SampleRate);
 	virtual int32 GetKeyIndex(float KeyTime, float KeyTimeTolerance) const override final;
 
 public:
diff --git a/Engine/Source/Runtime/Engine/Classes/Curves/SimpleCurve.h b/Engine/Source/Runtime/Engine/Classes/Curves/SimpleCurve.h
--- a/Engine/Source/Runtime/Engine/Classes/Curves/SimpleCurve.h
+++ b/Engine/Source/Runtime/Engine/Classes/Curves/SimpleCurve.h
@@ -183,8 +183,8 @@ public:
 	virtual void BakeCurve(float SampleRate, float FirstKeyTime, float LastKeyTime) final override;
 
 	/** Remove redundant keys, comparing against Tolerance */
-	virtual void RemoveRedundantKeys(float Tolerance) final override;
-	virtual void RemoveRedundantKeys(float Tolerance, float FirstKeyTime, float LastKeyTime) final override;
+	virtual void RemoveRedundantKeys(float Tolerance, FFrameRate SampleRate = FFrameRate(0,0)) final override;
+	virtual void RemoveRedundantKeys(float Tolerance, float FirstKeyTime, float LastKeyTime, FFrameRate SampleRate = FFrameRate(0,0)) final override;
 
 	/** Allocates a duplicate of the curve */
 	virtual FIndexedCurve* Duplicate() const final { return new FSimpleCurve(*this); }
diff --git a/Engine/Source/Runtime/Engine/Private/Animation/AnimCurveCompressionCodec_CompressedRichCurve.cpp b/Engine/Source/Runtime/Engine/Private/Animation/AnimCurveCompressionCodec_CompressedRichCurve.cpp
--- a/Engine/Source/Runtime/Engine/Private/Animation/AnimCurveCompressionCodec_CompressedRichCurve.cpp
+++ b/Engine/Source/Runtime/Engine/Private/Animation/AnimCurveCompressionCodec_CompressedRichCurve.cpp
@@ -55,7 +55,7 @@ bool UAnimCurveCompressionCodec_CompressedRichCurve::Compress(const FCompressibl
 		const FFloatCurve& Curve = AnimSeq.RawFloatCurves[CurveIndex];
 
 		FRichCurve RawCurve = Curve.FloatCurve;	// Copy
-		RawCurve.RemoveRedundantKeys(MaxCurveError);
+		RawCurve.RemoveRedundantAutoTangentKeys(MaxCurveError);
 
 		FCompressedRichCurve CompressedCurve;
 		RawCurve.CompressCurve(CompressedCurve, MaxCurveError, SampleRate);
diff --git a/Engine/Source/Runtime/Engine/Private/Animation/AnimCurveTypes.cpp b/Engine/Source/Runtime/Engine/Private/Animation/AnimCurveTypes.cpp
--- a/Engine/Source/Runtime/Engine/Private/Animation/AnimCurveTypes.cpp
+++ b/Engine/Source/Runtime/Engine/Private/Animation/AnimCurveTypes.cpp
@@ -532,12 +532,12 @@ void FRawCurveTracks::AddFloatCurveKey(const FSmartName& NewCurve, int32 CurveFl
 	FloatCurve->UpdateOrAddKey(Value, Time);
 }
 
-void FRawCurveTracks::RemoveRedundantKeys()
+void FRawCurveTracks::RemoveRedundantKeys(float Tolerance /*= UE_SMALL_NUMBER*/, FFrameRate SampleRate /*= FFrameRate(0,0)*/ )
 {
 	for (auto CurveIter = FloatCurves.CreateIterator(); CurveIter; ++CurveIter)
 	{
 		FFloatCurve& Curve = *CurveIter;
-		Curve.FloatCurve.RemoveRedundantKeys(UE_SMALL_NUMBER);
+		Curve.FloatCurve.RemoveRedundantKeys(Tolerance, SampleRate);
 	}
 }
 #endif
diff --git a/Engine/Source/Runtime/Engine/Private/CurveTable.cpp b/Engine/Source/Runtime/Engine/Private/CurveTable.cpp
--- a/Engine/Source/Runtime/Engine/Private/CurveTable.cpp
+++ b/Engine/Source/Runtime/Engine/Private/CurveTable.cpp
@@ -224,7 +224,7 @@ void UCurveTable::Serialize(FArchive& Ar)
 				FRichCurve* Curve = (FRichCurve*)RowIt.Value();
 				if (Ar.IsCooking() && Ar.IsPersistent() && !Ar.IsObjectReferenceCollector() && !Ar.ShouldSkipBulkData() && CVar_CurveTable_RemoveRedundantKeys > 0)
 				{
-					Curve->RemoveRedundantKeys(0.f);
+					Curve->RemoveRedundantAutoTangentKeys(0.f);
 				}
 				FRichCurve::StaticStruct()->SerializeTaggedProperties(Ar, (uint8*)Curve, FRichCurve::StaticStruct(), nullptr);
 			}
diff --git a/Engine/Source/Runtime/Engine/Private/Curves/RichCurve.cpp b/Engine/Source/Runtime/Engine/Private/Curves/RichCurve.cpp
--- a/Engine/Source/Runtime/Engine/Private/Curves/RichCurve.cpp
+++ b/Engine/Source/Runtime/Engine/Private/Curves/RichCurve.cpp
@@ -197,7 +197,7 @@ bool FRichCurveKey::Serialize(FArchive& Ar)
 
 				if (bHasArriveWeight && bHasLeaveWeight)
 				{
-					TangentWeightMode = RCTWM_WeightedBoth;				
+					TangentWeightMode = RCTWM_WeightedBoth;
 				}
 				else if (bHasArriveWeight)
 				{
@@ -809,7 +809,7 @@ void FRichCurve::ReadjustTimeRange(float NewMinTimeRange, float NewMaxTimeRange,
 		}
 	}
 
-	RemoveRedundantKeys(0.f);
+	RemoveRedundantAutoTangentKeys(0.f);
 
 	// now cull out all out of range 
 	float MinTime, MaxTime;
@@ -892,17 +892,27 @@ void FRichCurve::BakeCurve(float SampleRate, float FirstKeyTime, float LastKeyTi
 	}
 }
 
-void FRichCurve::RemoveRedundantKeys(float Tolerance)
+void FRichCurve::RemoveRedundantAutoTangentKeys(float Tolerance)
+{
+	RemoveRedundantKeys(Tolerance, FFrameRate(0,0));
+}
+
+void FRichCurve::RemoveRedundantKeys(float Tolerance, FFrameRate SampleRate)
 {
 	if (Keys.Num() < 3)
 	{
 		return;
 	}
 
-	RemoveRedundantKeysInternal(Tolerance, 0, Keys.Num() - 1);
+	RemoveRedundantKeysInternal(Tolerance, 0, Keys.Num() - 1, SampleRate);
+}
+
+void FRichCurve::RemoveRedundantAutoTangentKeys(float Tolerance, float FirstKeyTime, float LastKeyTime)
+{
+	RemoveRedundantKeys(Tolerance, FirstKeyTime, LastKeyTime, FFrameRate(0,0));
 }
 
-void FRichCurve::RemoveRedundantKeys(float Tolerance, float FirstKeyTime, float LastKeyTime)
+void FRichCurve::RemoveRedundantKeys(float Tolerance, float FirstKeyTime, float LastKeyTime, FFrameRate SampleRate)
 {
 	if (FirstKeyTime >= LastKeyTime)
 	{
@@ -927,7 +937,7 @@ void FRichCurve::RemoveRedundantKeys(float Tolerance, float FirstKeyTime, float
 
 	if ((StartKey != INDEX_NONE) && (EndKey != INDEX_NONE))
 	{
-		RemoveRedundantKeysInternal(Tolerance, StartKey, EndKey);
+		RemoveRedundantKeysInternal(Tolerance, StartKey, EndKey, SampleRate);
 	}
 }
 
@@ -1192,7 +1202,7 @@ float EvalForTwoKeys(const FRichCurveKey& Key1, const FRichCurveKey& Key2, const
 	}
 }
 
-void FRichCurve::RemoveRedundantKeysInternal(float Tolerance, int32 InStartKeepKey, int32 InEndKeepKey)
+void FRichCurve::RemoveRedundantKeysInternal(float Tolerance, int32 InStartKeepKey, int32 InEndKeepKey, FFrameRate SampleRate)
 {
 	if (Keys.Num() < 3) // Will always keep first and last key
 	{
@@ -1243,11 +1253,59 @@ void FRichCurve::RemoveRedundantKeysInternal(float Tolerance, int32 InStartKeepK
 
 		//Add keys up to the first end keep key if they are not redundant
 		int32 MostRecentKeepKeyIndex = 0;
+
+		const bool bTryRemoveUserTangentKeys = SampleRate.IsValid();
 		for (int32 TestIndex = ActualStartKeepKey+1; TestIndex < ActualEndKeepKey; ++TestIndex) //Loop within the bounds of the first and last key
 		{
-			const float KeyValue = Keys[TestIndex].Value;
-			const float ValueWithoutKey = EvalForTwoKeys(Keys[MostRecentKeepKeyIndex], Keys[TestIndex + 1], Keys[TestIndex].Time);
-			if (FMath::Abs(ValueWithoutKey - KeyValue) > Tolerance) // Is this key needed
+			bool bKeepKey = false;
+			
+			if (bTryRemoveUserTangentKeys)
+			{
+				const FFrameTime KeyOneFrameTime = SampleRate.AsFrameTime(Keys[MostRecentKeepKeyIndex].Time);
+				const FFrameTime KeyTwoFrameTime = SampleRate.AsFrameTime(Keys[TestIndex + 1].Time);
+
+				FFrameTime SampleFrameTime = KeyOneFrameTime;
+				const float ToRemoveKeyTime = Keys[TestIndex].Time;
+
+				// Sample curve at intervals, dictated by provided sample-rate, between the keys surrounding the to-be-removed key
+				// this makes sure that any impact on non-lerp keys is taken into account
+				while (SampleFrameTime <= KeyTwoFrameTime)
+				{
+					const float SampleSeconds = SampleRate.AsSeconds(SampleFrameTime);
+					const float ValueWithoutKey = EvalForTwoKeys(Keys[MostRecentKeepKeyIndex], Keys[TestIndex + 1], SampleSeconds);
+
+					const FRichCurveKey& KeyOne = SampleSeconds <= ToRemoveKeyTime ? Keys[MostRecentKeepKeyIndex] : Keys[TestIndex];
+					const FRichCurveKey& KeyTwo = SampleSeconds <= ToRemoveKeyTime ? Keys[TestIndex] : Keys[TestIndex + 1];
+					const float ValueWithKey = EvalForTwoKeys(KeyOne, KeyTwo, SampleSeconds);
+				
+					if (FMath::Abs(ValueWithoutKey - ValueWithKey) > Tolerance)
+					{
+						bKeepKey = true;
+						break;
+					}
+
+					// Next frame
+					SampleFrameTime.FrameNumber += 1;	
+				}
+			}
+			else
+			{
+				auto IsNonUserTangentKey = [](const FRichCurveKey& Key)
+				{
+					return Key.TangentMode != RCTM_User && Key.TangentMode != RCTM_Break;
+				};
+
+				bKeepKey = true;
+				if (IsNonUserTangentKey(Keys[MostRecentKeepKeyIndex]) && IsNonUserTangentKey(Keys[TestIndex + 1]) && IsNonUserTangentKey(Keys[TestIndex]))
+				{
+					const float KeyValue = Keys[TestIndex].Value;
+					const float ValueWithoutKey = EvalForTwoKeys(Keys[MostRecentKeepKeyIndex], Keys[TestIndex + 1], Keys[TestIndex].Time);
+
+					bKeepKey = FMath::Abs(ValueWithoutKey - KeyValue) > Tolerance;
+				}
+			}			
+			
+			if (bKeepKey) // Is this key needed
 			{
 				MostRecentKeepKeyIndex = TestIndex;
 				NewKeys.Add(Keys[TestIndex]);
@@ -1267,7 +1325,6 @@ void FRichCurve::RemoveRedundantKeysInternal(float Tolerance, int32 InStartKeepK
 	AutoSetTangents();
 
 	// Rebuild KeyHandlesToIndices
-
 	check(Keys.Num() == KeepHandles.Num());
 	KeyHandlesToIndices.Initialize(KeepHandles);
 }
diff --git a/Engine/Source/Runtime/Engine/Private/Curves/SimpleCurve.cpp b/Engine/Source/Runtime/Engine/Private/Curves/SimpleCurve.cpp
--- a/Engine/Source/Runtime/Engine/Private/Curves/SimpleCurve.cpp
+++ b/Engine/Source/Runtime/Engine/Private/Curves/SimpleCurve.cpp
@@ -449,7 +449,7 @@ void FSimpleCurve::BakeCurve(float SampleRate, float FirstKeyTime, float LastKey
 	}
 }
 
-void FSimpleCurve::RemoveRedundantKeys(float Tolerance)
+void FSimpleCurve::RemoveRedundantKeys(float Tolerance, FFrameRate SampleRate /*= FFrameRate(0,0)*/ )
 {
 	if (Keys.Num() < 3)
 	{
@@ -459,7 +459,7 @@ void FSimpleCurve::RemoveRedundantKeys(float Tolerance)
 	RemoveRedundantKeysInternal(Tolerance, 0, Keys.Num() - 1);
 }
 
-void FSimpleCurve::RemoveRedundantKeys(float Tolerance, float FirstKeyTime, float LastKeyTime)
+void FSimpleCurve::RemoveRedundantKeys(float Tolerance, float FirstKeyTime, float LastKeyTime, FFrameRate SampleRate /*= FFrameRate(0,0)*/ )
 {
 	if (FirstKeyTime >= LastKeyTime)
 	{
diff --git a/Engine/Source/Runtime/Engine/Public/Animation/AnimCurveTypes.h b/Engine/Source/Runtime/Engine/Public/Animation/AnimCurveTypes.h
--- a/Engine/Source/Runtime/Engine/Public/Animation/AnimCurveTypes.h
+++ b/Engine/Source/Runtime/Engine/Public/Animation/AnimCurveTypes.h
@@ -847,7 +847,7 @@ struct FRawCurveTracks
 	* Add new float curve from the given UID if not existing and add the key with time/value
 	*/
 	ENGINE_API void AddFloatCurveKey(const FSmartName& NewCurve, int32 CurveFlags, float Time, float Value);
-	ENGINE_API void RemoveRedundantKeys();
+	ENGINE_API void RemoveRedundantKeys(float Tolerance = UE_SMALL_NUMBER, FFrameRate SampleRate = FFrameRate(0,0));
 
 #endif // WITH_EDITOR
 	/**
