Fix crash when changing scalability settings.

If a material instance did not have static permutation resources, the UMaterialInstance::AllMaterialsCacheResourceShadersForRendering() call in ScalabilityCVarsSinkCallback() enqueued an update uniform buffer command for that instance, instead of recreating the buffer (via CacheResourceShadersForRendering -> InitResources). However, the shader map for the parent material changes as a result of the scalability change, so the layout in the new shader map will no longer match the layout of the existing uniform buffer. The solution is to always recreate the buffer in CacheResourceShadersForRendering().

Also fixed the debug name stored in FMaterialInstanceResource. The previous code was trying to use the asset path of the owner material, but that's not set yet at the point where the instance resource is created, so the name was always that of the default material.

#rnx

diff --git a/Engine/Source/Runtime/Engine/Private/Materials/MaterialInstance.cpp b/Engine/Source/Runtime/Engine/Private/Materials/MaterialInstance.cpp
--- a/Engine/Source/Runtime/Engine/Private/Materials/MaterialInstance.cpp
+++ b/Engine/Source/Runtime/Engine/Private/Materials/MaterialInstance.cpp
@@ -116,7 +116,7 @@ FFontParameterValue::ValueType FFontParameterValue::GetValue(const FFontParamete
 }
 
 FMaterialInstanceResource::FMaterialInstanceResource(UMaterialInstance* InOwner)
-	: FMaterialRenderProxy(GetPathNameSafe(InOwner->GetMaterial()))
+	: FMaterialRenderProxy(InOwner->GetName())
 	, Parent(NULL)
 	, Owner(InOwner)
 	, GameThreadParent(NULL)
@@ -2052,10 +2052,10 @@ void UMaterialInstance::CacheResourceShadersForRendering(EMaterialShaderPrecompi
 			ResourcesToCache.Reset();
 			ResourcesToCache.Add(CurrentResource);
 			CacheShadersForResources(ShaderPlatform, ResourcesToCache, PrecompileMode);
-			RecacheUniformExpressions(true);
 		}
 	}
 
+	RecacheUniformExpressions(true);
 	InitResources();
 }
 
