UWorld: Remove NewWorld argument from CleanupWorldInternal, and replace it with bWorldChanged. bWorldChanged was being incorrectly calculated in the recursive CleanupWorldInternal calls on streaming levels, resulting in them incorrectly setting bUnloadFromEditor=true when the toplevel CleanupWorld was called from CleanupLevelThatWasInitializedFromSave during cooking, resulting in indeterminism in the cook due to removing RF_Standalone flags.
#rnx

diff --git a/Engine/Source/Runtime/AIModule/Classes/AISystem.h b/Engine/Source/Runtime/AIModule/Classes/AISystem.h
--- a/Engine/Source/Runtime/AIModule/Classes/AISystem.h
+++ b/Engine/Source/Runtime/AIModule/Classes/AISystem.h
@@ -142,7 +142,9 @@ public:
 	// UAISystemBase begin		
 	virtual void InitializeActorsForPlay(bool bTimeGotReset) override;
 	virtual void WorldOriginLocationChanged(FIntVector OldOriginLocation, FIntVector NewOriginLocation) override;
-	virtual void CleanupWorld(bool bSessionEnded = true, bool bCleanupResources = true, UWorld* NewWorld = NULL) override;
+	virtual void CleanupWorld(bool bSessionEnded = true, bool bCleanupResources = true) override;
+	UE_DEPRECATED(5.1, "NewWorld was unused and not always calculated correctly and we expect it is not needed; let us know on UDN if it is necessary.")
+	virtual void CleanupWorld(bool bSessionEnded, bool bCleanupResources, UWorld* NewWorld) override;
 	virtual void StartPlay() override;
 	// UAISystemBase end
 
diff --git a/Engine/Source/Runtime/AIModule/Private/AISystem.cpp b/Engine/Source/Runtime/AIModule/Private/AISystem.cpp
--- a/Engine/Source/Runtime/AIModule/Private/AISystem.cpp
+++ b/Engine/Source/Runtime/AIModule/Private/AISystem.cpp
@@ -45,7 +45,7 @@ UAISystem::UAISystem(const FObjectInitializer& ObjectInitializer)
 
 void UAISystem::BeginDestroy()
 {
-	CleanupWorld(true, true, NULL);
+	CleanupWorld(true, true);
 	Super::BeginDestroy();
 }
 
@@ -127,7 +127,12 @@ void UAISystem::WorldOriginLocationChanged(FIntVector OldOriginLocation, FIntVec
 
 void UAISystem::CleanupWorld(bool bSessionEnded, bool bCleanupResources, UWorld* NewWorld)
 {
-	Super::CleanupWorld(bSessionEnded, bCleanupResources, NewWorld);
+	CleanupWorld(bSessionEnded, bCleanupResources);
+}
+
+void UAISystem::CleanupWorld(bool bSessionEnded, bool bCleanupResources)
+{
+	Super::CleanupWorld(bSessionEnded, bCleanupResources);
 	
 	if (bCleanupResources)
 	{
diff --git a/Engine/Source/Runtime/Engine/Classes/AI/AISystemBase.h b/Engine/Source/Runtime/Engine/Classes/AI/AISystemBase.h
--- a/Engine/Source/Runtime/Engine/Classes/AI/AISystemBase.h
+++ b/Engine/Source/Runtime/Engine/Classes/AI/AISystemBase.h
@@ -35,10 +35,11 @@ class ENGINE_API UAISystemBase : public UObject
 	 * Called by UWorld::CleanupWorld.
 	 * Should be called by overriding functions.
 	 * @param bSessionEnded whether to notify the viewport that the game session has ended
-	 * @param NewWorld Optional new world that will be loaded after this world is cleaned up. Specify a new world to prevent it and it's sublevels from being GCed during map transitions.
 	 */
-	virtual void CleanupWorld(bool bSessionEnded = true, bool bCleanupResources = true, UWorld* NewWorld = NULL);
-
+	virtual void CleanupWorld(bool bSessionEnded = true, bool bCleanupResources = true);
+	UE_DEPRECATED(5.1, "NewWorld was unused and not always calculated correctly and we expect it is not needed; let us know on UDN if it is necessary.")
+	virtual void CleanupWorld(bool bSessionEnded, bool bCleanupResources, UWorld* NewWorld);
+	
 	/** 
 	 * Called by UWorld::BeginPlay to indicate the gameplay has started.
 	 * Should be called by overriding functions.
diff --git a/Engine/Source/Runtime/Engine/Classes/Engine/World.h b/Engine/Source/Runtime/Engine/Classes/Engine/World.h
--- a/Engine/Source/Runtime/Engine/Classes/Engine/World.h
+++ b/Engine/Source/Runtime/Engine/Classes/Engine/World.h
@@ -3345,7 +3345,7 @@ public:
 
 private:
 	/** Internal version of CleanupWorld. */
-	void CleanupWorldInternal(bool bSessionEnded, bool bCleanupResources, UWorld* NewWorld);
+	void CleanupWorldInternal(bool bSessionEnded, bool bCleanupResources, bool bWorldChanged);
 
 	/** Utility function to handle Exec/Console Commands related to the Trace Tags */
 	bool HandleTraceTagCommand( const TCHAR* Cmd, FOutputDevice& Ar );
diff --git a/Engine/Source/Runtime/Engine/Private/AI/AISystemBase.cpp b/Engine/Source/Runtime/Engine/Private/AI/AISystemBase.cpp
--- a/Engine/Source/Runtime/Engine/Private/AI/AISystemBase.cpp
+++ b/Engine/Source/Runtime/Engine/Private/AI/AISystemBase.cpp
@@ -23,6 +23,11 @@ FSoftClassPath UAISystemBase::GetAISystemClassName()
 }
 
 void UAISystemBase::CleanupWorld(bool bSessionEnded, bool bCleanupResources, UWorld* NewWorld)
+{
+	CleanupWorld(bSessionEnded, bCleanupResources);
+}
+
+void UAISystemBase::CleanupWorld(bool bSessionEnded, bool bCleanupResources)
 {
 	FGameModeEvents::OnGameModeMatchStateSetEvent().Remove(OnMatchStateSetHandle);
 }
diff --git a/Engine/Source/Runtime/Engine/Private/World.cpp b/Engine/Source/Runtime/Engine/Private/World.cpp
--- a/Engine/Source/Runtime/Engine/Private/World.cpp
+++ b/Engine/Source/Runtime/Engine/Private/World.cpp
@@ -4920,17 +4920,17 @@ void UWorld::CleanupWorld(bool bSessionEnded, bool bCleanupResources, UWorld* Ne
 		bool bIsStreamingSubWorld = PersistentLevel && PersistentLevel->OwningWorld != this;
 		UE_CLOG(!bIsStreamingSubWorld, LogWorld, Warning, TEXT("UWorld::CleanupWorld called twice or called without InitWorld called first."));
 	}
-	CleanupWorldInternal(bSessionEnded, bCleanupResources, NewWorld);
+	const bool bWorldChanged = NewWorld != this;
+	CleanupWorldInternal(bSessionEnded, bCleanupResources, bWorldChanged);
 	bInitializedAndNeedsCleanup = false;
 }
 
-void UWorld::CleanupWorldInternal(bool bSessionEnded, bool bCleanupResources, UWorld* NewWorld)
+void UWorld::CleanupWorldInternal(bool bSessionEnded, bool bCleanupResources, bool bWorldChanged)
 {
 	if(CleanupWorldTag == CleanupWorldGlobalTag)
 	{
 		return;
 	}
-	const bool bWorldChanged = NewWorld != this;
 	CleanupWorldTag = CleanupWorldGlobalTag;
 
 	UE_LOG(LogWorld, Log, TEXT("UWorld::CleanupWorld for %s, bSessionEnded=%s, bCleanupResources=%s"), *GetName(), bSessionEnded ? TEXT("true") : TEXT("false"), bCleanupResources ? TEXT("true") : TEXT("false"));
@@ -4952,7 +4952,7 @@ void UWorld::CleanupWorldInternal(bool bSessionEnded, bool bCleanupResources, UW
 
 	if (AISystem != nullptr)
 	{
-		AISystem->CleanupWorld(bSessionEnded, bCleanupResources, NewWorld);
+		AISystem->CleanupWorld(bSessionEnded, bCleanupResources);
 	}
 
 	if (bCleanupResources == true)
@@ -5072,7 +5072,7 @@ void UWorld::CleanupWorldInternal(bool bSessionEnded, bool bCleanupResources, UW
 	for (int32 LevelIndex = 0; LevelIndex < GetNumLevels(); ++LevelIndex)
 	{
 		UWorld* World = CastChecked<UWorld>(GetLevel(LevelIndex)->GetOuter());
-		World->CleanupWorldInternal(bSessionEnded, bCleanupResources, NewWorld);
+		World->CleanupWorldInternal(bSessionEnded, bCleanupResources, bWorldChanged);
 	}
 
 	for (ULevelStreaming* StreamingLevel : GetStreamingLevels())
@@ -5080,7 +5080,7 @@ void UWorld::CleanupWorldInternal(bool bSessionEnded, bool bCleanupResources, UW
 		if (ULevel* Level = StreamingLevel->GetLoadedLevel())
 		{
 			UWorld* World = CastChecked<UWorld>(Level->GetOuter());
-			World->CleanupWorldInternal(bSessionEnded, bCleanupResources, NewWorld);
+			World->CleanupWorldInternal(bSessionEnded, bCleanupResources, bWorldChanged);
 		}
 	}
 
@@ -5093,7 +5093,7 @@ void UWorld::CleanupWorldInternal(bool bSessionEnded, bool bCleanupResources, UW
 			if (Level)
 			{
 				UWorld* const LevelWorld = CastChecked<UWorld>(Level->GetOuter());
-				LevelWorld->CleanupWorldInternal(bSessionEnded, bCleanupResources, NewWorld);
+				LevelWorld->CleanupWorldInternal(bSessionEnded, bCleanupResources, bWorldChanged);
 			}
 		}
 	}
