fix AnimNode_Inertialization value trace data displaying on wrong node in Rewind Debugger AnimGraph

diff --git a/Engine/Source/Runtime/Engine/Private/Animation/AnimNode_Inertialization.cpp b/Engine/Source/Runtime/Engine/Private/Animation/AnimNode_Inertialization.cpp
--- a/Engine/Source/Runtime/Engine/Private/Animation/AnimNode_Inertialization.cpp
+++ b/Engine/Source/Runtime/Engine/Private/Animation/AnimNode_Inertialization.cpp
@@ -423,11 +423,11 @@ void FAnimNode_Inertialization::Evaluate_AnyThread(FPoseContext& Output)
 	DeltaTime = 0.0f;
 	TeleportType = ETeleportType::None;
 
-	TRACE_ANIM_NODE_VALUE(Output, TEXT("State"), *UEnum::GetValueAsString(InertializationState));
-	TRACE_ANIM_NODE_VALUE(Output, TEXT("Elapsed Time"), InertializationElapsedTime);
-	TRACE_ANIM_NODE_VALUE(Output, TEXT("Duration"), InertializationDuration);
-	TRACE_ANIM_NODE_VALUE(Output, TEXT("Max Duration"), InertializationMaxDuration);
-	TRACE_ANIM_NODE_VALUE(Output, TEXT("Normalized Time"), InertializationDuration > UE_KINDA_SMALL_NUMBER ? (InertializationElapsedTime / InertializationDuration) : 0.0f);
+	TRACE_ANIM_NODE_VALUE_WITH_ID(Output, GetNodeIndex(), TEXT("State"), *UEnum::GetValueAsString(InertializationState));
+	TRACE_ANIM_NODE_VALUE_WITH_ID(Output, GetNodeIndex(), TEXT("Elapsed Time"), InertializationElapsedTime);
+	TRACE_ANIM_NODE_VALUE_WITH_ID(Output, GetNodeIndex(), TEXT("Duration"), InertializationDuration);
+	TRACE_ANIM_NODE_VALUE_WITH_ID(Output, GetNodeIndex(), TEXT("Max Duration"), InertializationMaxDuration);
+	TRACE_ANIM_NODE_VALUE_WITH_ID(Output, GetNodeIndex(), TEXT("Normalized Time"), InertializationDuration > UE_KINDA_SMALL_NUMBER ? (InertializationElapsedTime / InertializationDuration) : 0.0f);
 }
 
 
diff --git a/Engine/Source/Runtime/Engine/Private/Animation/AnimTrace.cpp b/Engine/Source/Runtime/Engine/Private/Animation/AnimTrace.cpp
--- a/Engine/Source/Runtime/Engine/Private/Animation/AnimTrace.cpp
+++ b/Engine/Source/Runtime/Engine/Private/Animation/AnimTrace.cpp
@@ -723,7 +723,7 @@ void FAnimTrace::OutputAnimNodeAttribute(const FAnimInstanceProxy& InTargetProxy
 		<< AnimNodeAttribute.NameId(NameId);
 }
 
-void FAnimTrace::OutputAnimNodeValue(const FAnimationBaseContext& InContext, const TCHAR* InKey, bool InValue)
+void FAnimTrace::OutputAnimNodeValue(const FAnimationBaseContext& InContext, uint32 NodeIndex, const TCHAR* InKey, bool InValue)
 {
 	bool bChannelEnabled = UE_TRACE_CHANNELEXPR_IS_ENABLED(AnimationChannel);
 	if (!bChannelEnabled)
@@ -744,13 +744,13 @@ void FAnimTrace::OutputAnimNodeValue(const FAnimationBaseContext& InContext, con
 	UE_TRACE_LOG(Animation, AnimNodeValueBool, AnimationChannel)
 		<< AnimNodeValueBool.Cycle(FPlatformTime::Cycles64())
 		<< AnimNodeValueBool.AnimInstanceId(FObjectTrace::GetObjectId(AnimInstance))
-		<< AnimNodeValueBool.NodeId(InContext.GetCurrentNodeId())
+		<< AnimNodeValueBool.NodeId(NodeIndex)
 		<< AnimNodeValueBool.Value(InValue)
 		<< AnimNodeValueBool.FrameCounter(FObjectTrace::GetObjectWorldTickCounter(AnimInstance))
 		<< AnimNodeValueBool.Key(InKey);
 }
 
-void FAnimTrace::OutputAnimNodeValue(const FAnimationBaseContext& InContext, const TCHAR* InKey, int32 InValue)
+void FAnimTrace::OutputAnimNodeValue(const FAnimationBaseContext& InContext, uint32 NodeIndex, const TCHAR* InKey, int32 InValue)
 {
 	bool bChannelEnabled = UE_TRACE_CHANNELEXPR_IS_ENABLED(AnimationChannel);
 	if (!bChannelEnabled)
@@ -771,13 +771,13 @@ void FAnimTrace::OutputAnimNodeValue(const FAnimationBaseContext& InContext, con
 	UE_TRACE_LOG(Animation, AnimNodeValueInt, AnimationChannel)
 		<< AnimNodeValueInt.Cycle(FPlatformTime::Cycles64())
 		<< AnimNodeValueInt.AnimInstanceId(FObjectTrace::GetObjectId(AnimInstance))
-		<< AnimNodeValueInt.NodeId(InContext.GetCurrentNodeId())
+		<< AnimNodeValueInt.NodeId(NodeIndex)
 		<< AnimNodeValueInt.Value(InValue)
 		<< AnimNodeValueInt.FrameCounter(FObjectTrace::GetObjectWorldTickCounter(AnimInstance))
 		<< AnimNodeValueInt.Key(InKey);
 }
 
-void FAnimTrace::OutputAnimNodeValue(const FAnimationBaseContext& InContext, const TCHAR* InKey, float InValue)
+void FAnimTrace::OutputAnimNodeValue(const FAnimationBaseContext& InContext, uint32 NodeIndex, const TCHAR* InKey, float InValue)
 {
 	bool bChannelEnabled = UE_TRACE_CHANNELEXPR_IS_ENABLED(AnimationChannel);
 	if (!bChannelEnabled)
@@ -798,13 +798,13 @@ void FAnimTrace::OutputAnimNodeValue(const FAnimationBaseContext& InContext, con
 	UE_TRACE_LOG(Animation, AnimNodeValueFloat, AnimationChannel)
 		<< AnimNodeValueFloat.Cycle(FPlatformTime::Cycles64())
 		<< AnimNodeValueFloat.AnimInstanceId(FObjectTrace::GetObjectId(AnimInstance))
-		<< AnimNodeValueFloat.NodeId(InContext.GetCurrentNodeId())
+		<< AnimNodeValueFloat.NodeId(NodeIndex)
 		<< AnimNodeValueFloat.Value(InValue)
 		<< AnimNodeValueFloat.FrameCounter(FObjectTrace::GetObjectWorldTickCounter(AnimInstance))
 		<< AnimNodeValueFloat.Key(InKey);
 }
 
-void FAnimTrace::OutputAnimNodeValue(const FAnimationBaseContext& InContext, const TCHAR* InKey, const FVector2D& InValue)
+void FAnimTrace::OutputAnimNodeValue(const FAnimationBaseContext& InContext, uint32 NodeIndex, const TCHAR* InKey, const FVector2D& InValue)
 {
 	bool bChannelEnabled = UE_TRACE_CHANNELEXPR_IS_ENABLED(AnimationChannel);
 	if (!bChannelEnabled)
@@ -820,7 +820,7 @@ void FAnimTrace::OutputAnimNodeValue(const FAnimationBaseContext& InContext, con
 	UE_TRACE_LOG(Animation, AnimNodeValueVector2D, AnimationChannel)
 		<< AnimNodeValueVector2D.Cycle(FPlatformTime::Cycles64())
 		<< AnimNodeValueVector2D.AnimInstanceId(FObjectTrace::GetObjectId(AnimInstance))
-		<< AnimNodeValueVector2D.NodeId(InContext.GetCurrentNodeId())
+		<< AnimNodeValueVector2D.NodeId(NodeIndex)
 		<< AnimNodeValueVector2D.ValueX(InValue.X)
 		<< AnimNodeValueVector2D.ValueY(InValue.Y)
 		<< AnimNodeValueVector2D.FrameCounter(FObjectTrace::GetObjectWorldTickCounter(AnimInstance))
@@ -828,13 +828,13 @@ void FAnimTrace::OutputAnimNodeValue(const FAnimationBaseContext& InContext, con
 }
 
 
-void FAnimTrace::OutputAnimNodeValue(const FAnimationBaseContext& InContext, const TCHAR* InKey, const FRotator& Value)
+void FAnimTrace::OutputAnimNodeValue(const FAnimationBaseContext& InContext, uint32 NodeIndex, const TCHAR* InKey, const FRotator& Value)
 {
 	const FVector VectorValue(Value.Roll, Value.Pitch, Value.Yaw);
-	OutputAnimNodeValue(InContext, InKey, VectorValue);
+	OutputAnimNodeValue(InContext, NodeIndex, InKey, VectorValue);
 }
 
-void FAnimTrace::OutputAnimNodeValue(const FAnimationBaseContext& InContext, const TCHAR* InKey, const FVector& InValue)
+void FAnimTrace::OutputAnimNodeValue(const FAnimationBaseContext& InContext, uint32 NodeIndex, const TCHAR* InKey, const FVector& InValue)
 {
 	bool bChannelEnabled = UE_TRACE_CHANNELEXPR_IS_ENABLED(AnimationChannel);
 	if (!bChannelEnabled)
@@ -855,7 +855,7 @@ void FAnimTrace::OutputAnimNodeValue(const FAnimationBaseContext& InContext, con
 	UE_TRACE_LOG(Animation, AnimNodeValueVector, AnimationChannel)
 		<< AnimNodeValueVector.Cycle(FPlatformTime::Cycles64())
 		<< AnimNodeValueVector.AnimInstanceId(FObjectTrace::GetObjectId(InContext.AnimInstanceProxy->GetAnimInstanceObject()))
-		<< AnimNodeValueVector.NodeId(InContext.GetCurrentNodeId())
+		<< AnimNodeValueVector.NodeId(NodeIndex)
 		<< AnimNodeValueVector.ValueX(InValue.X)
 		<< AnimNodeValueVector.ValueY(InValue.Y)
 		<< AnimNodeValueVector.ValueZ(InValue.Z)
@@ -863,7 +863,7 @@ void FAnimTrace::OutputAnimNodeValue(const FAnimationBaseContext& InContext, con
 		<< AnimNodeValueVector.Key(InKey);
 }
 
-void FAnimTrace::OutputAnimNodeValue(const FAnimationBaseContext& InContext, const TCHAR* InKey, const FName& InValue)
+void FAnimTrace::OutputAnimNodeValue(const FAnimationBaseContext& InContext, uint32 NodeIndex, const TCHAR* InKey, const FName& InValue)
 {
 	bool bChannelEnabled = UE_TRACE_CHANNELEXPR_IS_ENABLED(AnimationChannel);
 	if (!bChannelEnabled)
@@ -887,13 +887,13 @@ void FAnimTrace::OutputAnimNodeValue(const FAnimationBaseContext& InContext, con
 	UE_TRACE_LOG(Animation, AnimNodeValueString, AnimationChannel)
 		<< AnimNodeValueString.Cycle(FPlatformTime::Cycles64())
 		<< AnimNodeValueString.AnimInstanceId(FObjectTrace::GetObjectId(AnimInstance))
-		<< AnimNodeValueString.NodeId(InContext.GetCurrentNodeId())
+		<< AnimNodeValueString.NodeId(NodeIndex)
 		<< AnimNodeValueString.FrameCounter(FObjectTrace::GetObjectWorldTickCounter(AnimInstance))
 		<< AnimNodeValueString.Key(InKey)
 		<< AnimNodeValueString.Value(Value, ValueLength);
 }
 
-void FAnimTrace::OutputAnimNodeValue(const FAnimationBaseContext& InContext, const TCHAR* InKey, const TCHAR* InValue)
+void FAnimTrace::OutputAnimNodeValue(const FAnimationBaseContext& InContext, uint32 NodeIndex, const TCHAR* InKey, const TCHAR* InValue)
 {
 	bool bChannelEnabled = UE_TRACE_CHANNELEXPR_IS_ENABLED(AnimationChannel);
 	if (!bChannelEnabled)
@@ -914,13 +914,13 @@ void FAnimTrace::OutputAnimNodeValue(const FAnimationBaseContext& InContext, con
 	UE_TRACE_LOG(Animation, AnimNodeValueString, AnimationChannel)
 		<< AnimNodeValueString.Cycle(FPlatformTime::Cycles64())
 		<< AnimNodeValueString.AnimInstanceId(FObjectTrace::GetObjectId(AnimInstance))
-		<< AnimNodeValueString.NodeId(InContext.GetCurrentNodeId())
+		<< AnimNodeValueString.NodeId(NodeIndex)
 		<< AnimNodeValueString.FrameCounter(FObjectTrace::GetObjectWorldTickCounter(AnimInstance))
 		<< AnimNodeValueString.Key(InKey)
 		<< AnimNodeValueString.Value(InValue);
 }
 
-void FAnimTrace::OutputAnimNodeValue(const FAnimationBaseContext& InContext, const TCHAR* InKey, const UObject* InValue)
+void FAnimTrace::OutputAnimNodeValue(const FAnimationBaseContext& InContext, uint32 NodeIndex, const TCHAR* InKey, const UObject* InValue)
 {
 	bool bChannelEnabled = UE_TRACE_CHANNELEXPR_IS_ENABLED(AnimationChannel);
 	if (!bChannelEnabled)
@@ -942,13 +942,13 @@ void FAnimTrace::OutputAnimNodeValue(const FAnimationBaseContext& InContext, con
 	UE_TRACE_LOG(Animation, AnimNodeValueObject, AnimationChannel)
 		<< AnimNodeValueObject.Cycle(FPlatformTime::Cycles64())
 		<< AnimNodeValueObject.AnimInstanceId(FObjectTrace::GetObjectId(AnimInstance))
-		<< AnimNodeValueObject.NodeId(InContext.GetCurrentNodeId())
+		<< AnimNodeValueObject.NodeId(NodeIndex)
 		<< AnimNodeValueObject.Value(FObjectTrace::GetObjectId(InValue))
 		<< AnimNodeValueObject.FrameCounter(FObjectTrace::GetObjectWorldTickCounter(AnimInstance))
 		<< AnimNodeValueObject.Key(InKey);
 }
 
-void FAnimTrace::OutputAnimNodeValue(const FAnimationBaseContext& InContext, const TCHAR* InKey, const UClass* InValue)
+void FAnimTrace::OutputAnimNodeValue(const FAnimationBaseContext& InContext, uint32 NodeIndex, const TCHAR* InKey, const UClass* InValue)
 {
 	bool bChannelEnabled = UE_TRACE_CHANNELEXPR_IS_ENABLED(AnimationChannel);
 	if (!bChannelEnabled)
@@ -970,7 +970,7 @@ void FAnimTrace::OutputAnimNodeValue(const FAnimationBaseContext& InContext, con
 	UE_TRACE_LOG(Animation, AnimNodeValueClass, AnimationChannel)
 		<< AnimNodeValueClass.Cycle(FPlatformTime::Cycles64())
 		<< AnimNodeValueClass.AnimInstanceId(FObjectTrace::GetObjectId(AnimInstance))
-		<< AnimNodeValueClass.NodeId(InContext.GetCurrentNodeId())
+		<< AnimNodeValueClass.NodeId(NodeIndex)
 		<< AnimNodeValueClass.Value(FObjectTrace::GetObjectId(InValue))
 		<< AnimNodeValueClass.FrameCounter(FObjectTrace::GetObjectWorldTickCounter(AnimInstance))
 		<< AnimNodeValueClass.Key(InKey);
diff --git a/Engine/Source/Runtime/Engine/Public/Animation/AnimTrace.h b/Engine/Source/Runtime/Engine/Public/Animation/AnimTrace.h
--- a/Engine/Source/Runtime/Engine/Public/Animation/AnimTrace.h
+++ b/Engine/Source/Runtime/Engine/Public/Animation/AnimTrace.h
@@ -153,16 +153,16 @@ struct FAnimTrace
 	}
 
 	/** Helper function to output a tracked value for an anim node */
-	ENGINE_API static void OutputAnimNodeValue(const FAnimationBaseContext& InContext, const TCHAR* InKey, bool InValue);
-	ENGINE_API static void OutputAnimNodeValue(const FAnimationBaseContext& InContext, const TCHAR* InKey, int32 InValue);
-	ENGINE_API static void OutputAnimNodeValue(const FAnimationBaseContext& InContext, const TCHAR* InKey, float InValue);
-	ENGINE_API static void OutputAnimNodeValue(const FAnimationBaseContext& InContext, const TCHAR* InKey, const FVector2D& InValue);
-	ENGINE_API static void OutputAnimNodeValue(const FAnimationBaseContext& InContext, const TCHAR* InKey, const FVector& InValue);
-	ENGINE_API static void OutputAnimNodeValue(const FAnimationBaseContext& InContext, const TCHAR* InKey, const FRotator& InValue);
-	ENGINE_API static void OutputAnimNodeValue(const FAnimationBaseContext& InContext, const TCHAR* InKey, const FName& InValue);
-	ENGINE_API static void OutputAnimNodeValue(const FAnimationBaseContext& InContext, const TCHAR* InKey, const TCHAR* InValue);
-	ENGINE_API static void OutputAnimNodeValue(const FAnimationBaseContext& InContext, const TCHAR* InKey, const UClass* InValue);
-	ENGINE_API static void OutputAnimNodeValue(const FAnimationBaseContext& InContext, const TCHAR* InKey, const UObject* InValue);
+	ENGINE_API static void OutputAnimNodeValue(const FAnimationBaseContext& InContext, uint32 NodeIndex, const TCHAR* InKey, bool InValue);
+	ENGINE_API static void OutputAnimNodeValue(const FAnimationBaseContext& InContext, uint32 NodeIndex, const TCHAR* InKey, int32 InValue);
+	ENGINE_API static void OutputAnimNodeValue(const FAnimationBaseContext& InContext, uint32 NodeIndex, const TCHAR* InKey, float InValue);
+	ENGINE_API static void OutputAnimNodeValue(const FAnimationBaseContext& InContext, uint32 NodeIndex, const TCHAR* InKey, const FVector2D& InValue);
+	ENGINE_API static void OutputAnimNodeValue(const FAnimationBaseContext& InContext, uint32 NodeIndex, const TCHAR* InKey, const FVector& InValue);
+	ENGINE_API static void OutputAnimNodeValue(const FAnimationBaseContext& InContext, uint32 NodeIndex, const TCHAR* InKey, const FRotator& InValue);
+	ENGINE_API static void OutputAnimNodeValue(const FAnimationBaseContext& InContext, uint32 NodeIndex, const TCHAR* InKey, const FName& InValue);
+	ENGINE_API static void OutputAnimNodeValue(const FAnimationBaseContext& InContext, uint32 NodeIndex, const TCHAR* InKey, const TCHAR* InValue);
+	ENGINE_API static void OutputAnimNodeValue(const FAnimationBaseContext& InContext, uint32 NodeIndex, const TCHAR* InKey, const UClass* InValue);
+	ENGINE_API static void OutputAnimNodeValue(const FAnimationBaseContext& InContext, uint32 NodeIndex, const TCHAR* InKey, const UObject* InValue);
 
 	/** Helper function to output debug info for sequence player nodes */
 	ENGINE_API static void OutputAnimSequencePlayer(const FAnimationBaseContext& InContext, const FAnimNode_SequencePlayerBase& InNode);
@@ -217,7 +217,10 @@ struct FAnimTrace
 	FAnimTrace::FScopedAnimNodeTraceSuspend _ScopedAnimNodeTraceSuspend;
 
 #define TRACE_ANIM_NODE_VALUE(Context, Key, Value) \
-	FAnimTrace::OutputAnimNodeValue(Context, Key, Value);
+	FAnimTrace::OutputAnimNodeValue(Context, Context.GetCurrentNodeId(), Key, Value);
+
+#define TRACE_ANIM_NODE_VALUE_WITH_ID(Context, NodeId, Key, Value) \
+	FAnimTrace::OutputAnimNodeValue(Context, NodeId, Key, Value);
 
 #define TRACE_ANIM_SEQUENCE_PLAYER(Context, Node) \
 	FAnimTrace::OutputAnimSequencePlayer(Context, Node);
@@ -249,6 +252,7 @@ struct FAnimTrace
 #define TRACE_ANIM_NODE_BLENDABLE_ATTRIBUTES(Context, TargetNodeId, SourceNodeId)
 #define TRACE_SCOPED_ANIM_NODE_SUSPEND
 #define TRACE_ANIM_NODE_VALUE(Context, Key, Value)
+#define TRACE_ANIM_NODE_VALUE_WITH_ID(Context, NodeId, Key, Value)
 #define TRACE_ANIM_SEQUENCE_PLAYER(Context, Node)
 #define TRACE_ANIM_STATE_MACHINE_STATE(Context, StateMachineIndex, StateIndex, StateWeight, ElapsedTime)
 #define TRACE_ANIM_NOTIFY(AnimInstance, NotifyEvent, EventType)
