DataLayer API to go through DataLayerSubsystem instead of AWorlDataLayers
Move DataLayers' EditorWorldSetting logic from DataLayerEditorSubsystem & WorldDataLayers to DataLayerSubsystem
Use DataLayerInstance::GetOuterWorldDataLayers instead of retrieving the WorldDataLayers from UWorld::GetWorldDataLayers instead
Rename UDataLayerSubsystem::GetDataLayerFromAsset to GetDataLayerInstanceFromAsset to be inline with the rest of the API
Preparation for multiple AWorldDataLayers existing in the same world

diff --git a/Engine/Source/Editor/DataLayerEditor/Private/DataLayer/DataLayerEditorSubsystem.cpp b/Engine/Source/Editor/DataLayerEditor/Private/DataLayer/DataLayerEditorSubsystem.cpp
--- a/Engine/Source/Editor/DataLayerEditor/Private/DataLayer/DataLayerEditorSubsystem.cpp
+++ b/Engine/Source/Editor/DataLayerEditor/Private/DataLayer/DataLayerEditorSubsystem.cpp
@@ -192,47 +192,43 @@ void UDataLayerEditorSubsystem::OnActorDescContainerInitialized(UActorDescContai
 	{
 		FWorldPartitionActorDesc* ActorDesc = *Iterator;
 		check(ActorDesc->GetContainer() == InActorDescContainer);
-		ActorDesc->DataLayerInstanceNames = FDataLayerUtils::ResolvedDataLayerInstanceNames(ActorDesc, /*WorldDataLayers*/nullptr, WorldDataLayersActorDesc);
+		ActorDesc->DataLayerInstanceNames = FDataLayerUtils::ResolvedDataLayerInstanceNames(ActorDesc, WorldDataLayersActorDesc);
 	}
 }
 
 void UDataLayerEditorSubsystem::OnExecuteActorEditorContextAction(UWorld* InWorld, const EActorEditorContextAction& InType, AActor* InActor)
-{
-	AWorldDataLayers* WorldDataLayer = InWorld->GetWorldDataLayers();
-	if (!WorldDataLayer)
-	{
-		return;
-	}
+{	
+	UDataLayerSubsystem* DataLayerSubsystem = GetWorld()->GetSubsystem<UDataLayerSubsystem>();
 
 	switch (InType)
 	{
 	case EActorEditorContextAction::ApplyContext:
 		check(InActor && InActor->GetWorld() == InWorld);
 		{
-			AddActorToDataLayers(InActor, WorldDataLayer->GetActorEditorContextDataLayers());
+			AddActorToDataLayers(InActor, DataLayerSubsystem->GetActorEditorContextDataLayers());
 		}
 		break;
 	case EActorEditorContextAction::ResetContext:
-		for (UDataLayerInstance* DataLayer : WorldDataLayer->GetActorEditorContextDataLayers())
+		for (UDataLayerInstance* DataLayer : DataLayerSubsystem->GetActorEditorContextDataLayers())
 		{
 			RemoveFromActorEditorContext(DataLayer);
 		}
 		break;
 	case EActorEditorContextAction::PushContext:
-		WorldDataLayer->PushActorEditorContext();
+		DataLayerSubsystem->PushActorEditorContext();
 		BroadcastDataLayerChanged(EDataLayerAction::Reset, NULL, NAME_None);
 		break;
 	case EActorEditorContextAction::PopContext:
-		WorldDataLayer->PopActorEditorContext();
+		DataLayerSubsystem->PopActorEditorContext();
 		BroadcastDataLayerChanged(EDataLayerAction::Reset, NULL, NAME_None);
 		break;
 	}
-}
+}	
 
 bool UDataLayerEditorSubsystem::GetActorEditorContextDisplayInfo(UWorld* InWorld, FActorEditorContextClientDisplayInfo& OutDiplayInfo) const
 {
-	AWorldDataLayers* WorldDataLayer = InWorld->GetWorldDataLayers();
-	if (WorldDataLayer && !WorldDataLayer->GetActorEditorContextDataLayers().IsEmpty())
+	UDataLayerSubsystem* DataLayerSubsystem = GetWorld()->GetSubsystem<UDataLayerSubsystem>();
+	if (!DataLayerSubsystem->GetActorEditorContextDataLayers().IsEmpty())
 	{
 		OutDiplayInfo.Title = TEXT("Current Data Layers");
 		OutDiplayInfo.Brush = FEditorStyle::GetBrush(TEXT("DataLayer.Editor"));
@@ -244,33 +240,33 @@ bool UDataLayerEditorSubsystem::GetActorEditorContextDisplayInfo(UWorld* InWorld
 TSharedRef<SWidget> UDataLayerEditorSubsystem::GetActorEditorContextWidget(UWorld* InWorld) const
 {
 	TSharedRef<SVerticalBox> OutWidget = SNew(SVerticalBox);
-	if (AWorldDataLayers* WorldDataLayer = InWorld->GetWorldDataLayers())
+	
+	UDataLayerSubsystem* DataLayerSubsystem = GetWorld()->GetSubsystem<UDataLayerSubsystem>();
+	TArray<UDataLayerInstance*> DataLayers = DataLayerSubsystem->GetActorEditorContextDataLayers();
+	for (UDataLayerInstance* DataLayer : DataLayers)
 	{
-		TArray<UDataLayerInstance*> DataLayers = WorldDataLayer->GetActorEditorContextDataLayers();
-		for (UDataLayerInstance* DataLayer : DataLayers)
-		{
-			check(IsValid(DataLayer));
-			OutWidget->AddSlot().AutoHeight()
+		check(IsValid(DataLayer));
+		OutWidget->AddSlot().AutoHeight()
 			[
 				SNew(SHorizontalBox)
 				+ SHorizontalBox::Slot()
-				.AutoWidth()
-				.Padding(0.0f, 1.0f, 1.0f, 1.0f)
-				[
-					SNew(SImage)
-					.ColorAndOpacity(DataLayer->GetDebugColor())
-					.Image(FAppStyle::Get().GetBrush("Level.ColorIcon"))
-				]
-				+ SHorizontalBox::Slot()
-				.AutoWidth()
-				.Padding(4.0f, 1.0f, 1.0f, 1.0f)
-				[
-					SNew(STextBlock)
-					.Text(FText::FromString(DataLayer->GetDataLayerShortName()))
-				]
+			.AutoWidth()
+			.Padding(0.0f, 1.0f, 1.0f, 1.0f)
+			[
+				SNew(SImage)
+				.ColorAndOpacity(DataLayer->GetDebugColor())
+			.Image(FAppStyle::Get().GetBrush("Level.ColorIcon"))
+			]
+		+ SHorizontalBox::Slot()
+			.AutoWidth()
+			.Padding(4.0f, 1.0f, 1.0f, 1.0f)
+			[
+				SNew(STextBlock)
+				.Text(FText::FromString(DataLayer->GetDataLayerShortName()))
+			]
 			];
-		}
 	}
+	
 	return OutWidget;
 }
 
@@ -290,17 +286,6 @@ void UDataLayerEditorSubsystem::RemoveFromActorEditorContext(UDataLayerInstance*
 	}
 }
 
-void UDataLayerEditorSubsystem::UpdateDataLayerEditorPerProjectUserSettings()
-{
-	if (AWorldDataLayers* WorldDataLayers = GetWorldDataLayers())
-	{
-		TArray<FName> DataLayersNotLoadedInEditor;
-		TArray<FName> DataLayersLoadedInEditor;
-		WorldDataLayers->GetUserLoadedInEditorStates(DataLayersLoadedInEditor, DataLayersNotLoadedInEditor);
-		GetMutableDefault<UWorldPartitionEditorPerProjectUserSettings>()->SetWorldDataLayersNonDefaultEditorLoadStates(GetWorld(), DataLayersLoadedInEditor, DataLayersNotLoadedInEditor);
-	}
-}
-
 void UDataLayerEditorSubsystem::EditorMapChange()
 {
 	if (UWorld * World = GetWorld())
@@ -689,23 +674,21 @@ bool UDataLayerEditorSubsystem::UpdateActorVisibility(AActor* Actor, bool& bOutS
 	}
 
 	bool bActorBelongsToVisibleDataLayer = false;
-	if (AWorldDataLayers* WorldDataLayers = GetWorldDataLayers())
+	UDataLayerSubsystem* DataLayerSubsystem = UWorld::GetSubsystem<UDataLayerSubsystem>(GetWorld());
+	DataLayerSubsystem->ForEachDataLayer([&Actor, &bOutActorModified, &bActorBelongsToVisibleDataLayer](UDataLayerInstance* DataLayer)
 	{
-		WorldDataLayers->ForEachDataLayer([&Actor, &bOutActorModified, &bActorBelongsToVisibleDataLayer](UDataLayerInstance* DataLayer)
+		if (DataLayer->IsEffectiveVisible() && Actor->ContainsDataLayer(DataLayer))
 		{
-			if (DataLayer->IsEffectiveVisible() && Actor->ContainsDataLayer(DataLayer))
+			if (Actor->SetIsHiddenEdLayer(false))
 			{
-				if (Actor->SetIsHiddenEdLayer(false))
-				{
-					bOutActorModified = true;
-				}
-				// Stop, because we found at least one visible DataLayer the actor belongs to
-				bActorBelongsToVisibleDataLayer = true;
-				return false;
+				bOutActorModified = true;
 			}
-			return true;
-		});
-	}
+			// Stop, because we found at least one visible DataLayer the actor belongs to
+			bActorBelongsToVisibleDataLayer = true;
+			return false;
+		}
+		return true;
+	});
 
 	// If the actor isn't part of a visible DataLayer, hide and de-select it.
 	if (!bActorBelongsToVisibleDataLayer)
@@ -899,19 +882,17 @@ void UDataLayerEditorSubsystem::ToggleDataLayersVisibility(const TArray<UDataLay
 
 void UDataLayerEditorSubsystem::MakeAllDataLayersVisible()
 {
-	if (AWorldDataLayers* WorldDataLayers = GetWorldDataLayers())
+	UDataLayerSubsystem* DataLayerSubsystem = UWorld::GetSubsystem<UDataLayerSubsystem>(GetWorld());
+	DataLayerSubsystem->ForEachDataLayer([this](UDataLayerInstance* DataLayer)
 	{
-		WorldDataLayers->ForEachDataLayer([this](UDataLayerInstance* DataLayer)
+		if (!DataLayer->IsVisible())
 		{
-			if (!DataLayer->IsVisible())
-			{
-				DataLayer->Modify();
-				DataLayer->SetVisible(true);
-				BroadcastDataLayerChanged(EDataLayerAction::Modify, DataLayer, "bIsVisible");
-			}
-			return true;
-		});
-	}
+			DataLayer->Modify();
+			DataLayer->SetVisible(true);
+			BroadcastDataLayerChanged(EDataLayerAction::Modify, DataLayer, "bIsVisible");
+		}
+		return true;
+	});
 	
 	UpdateAllActorsVisibility(true, true);
 }
@@ -986,14 +967,12 @@ bool UDataLayerEditorSubsystem::ToggleDataLayersIsLoadedInEditor(const TArray<UD
 bool UDataLayerEditorSubsystem::ResetUserSettings()
 {
 	bool bRefreshNeeded = false;
-	if (const AWorldDataLayers* WorldDataLayers = GetWorldDataLayers())
+	UDataLayerSubsystem* DataLayerSubsystem = UWorld::GetSubsystem<UDataLayerSubsystem>(GetWorld());
+	DataLayerSubsystem->ForEachDataLayer([this, &bRefreshNeeded](UDataLayerInstance* DataLayer)
 	{
-		WorldDataLayers->ForEachDataLayer([this, &bRefreshNeeded](UDataLayerInstance* DataLayer)
-		{
-			bRefreshNeeded |= SetDataLayerIsLoadedInEditorInternal(DataLayer, DataLayer->IsInitiallyLoadedInEditor(), true);
-			return true;
-		});
-	}
+		bRefreshNeeded |= SetDataLayerIsLoadedInEditorInternal(DataLayer, DataLayer->IsInitiallyLoadedInEditor(), true);
+		return true;
+	});
 	
 	if (bRefreshNeeded)
 	{
@@ -1005,9 +984,9 @@ bool UDataLayerEditorSubsystem::ResetUserSettings()
 
 bool UDataLayerEditorSubsystem::HasDeprecatedDataLayers() const
 {
-	if (const AWorldDataLayers* WorldDataLayers = GetWorldDataLayers())
+	if (AWorldDataLayers* WorldDataLayer = GetWorld()->GetWorldDataLayers())
 	{
-		return WorldDataLayers->HasDeprecatedDataLayers();
+		return WorldDataLayer->HasDeprecatedDataLayers();
 	}
 
 	return false;
@@ -1015,40 +994,32 @@ bool UDataLayerEditorSubsystem::HasDeprecatedDataLayers() const
 
 UDataLayerInstance* UDataLayerEditorSubsystem::GetDataLayerInstance(const FName& DataLayerInstanceName) const
 {
-	const AWorldDataLayers* WorldDataLayers = GetWorldDataLayers();
-	return WorldDataLayers ? const_cast<UDataLayerInstance*>(WorldDataLayers->GetDataLayerInstance(DataLayerInstanceName)) : nullptr;
+	UDataLayerSubsystem* DataLayerSubsystem = UWorld::GetSubsystem<UDataLayerSubsystem>(GetWorld());
+	return DataLayerSubsystem->GetDataLayerInstance(DataLayerInstanceName);
 }
 
 UDataLayerInstance* UDataLayerEditorSubsystem::GetDataLayerInstance(const UDataLayerAsset* DataLayerAsset) const
 {
-	const AWorldDataLayers* WorldDataLayers = GetWorldDataLayers();
-	return WorldDataLayers ? const_cast<UDataLayerInstance*>(WorldDataLayers->GetDataLayerInstance(DataLayerAsset)) : nullptr;
-}
-
-AWorldDataLayers* UDataLayerEditorSubsystem::GetWorldDataLayers() const
-{
-	return GetWorld()->GetWorldDataLayers();
+	UDataLayerSubsystem* DataLayerSubsystem = UWorld::GetSubsystem<UDataLayerSubsystem>(GetWorld());
+	return DataLayerSubsystem ? const_cast<UDataLayerInstance*>(DataLayerSubsystem->GetDataLayerInstance(DataLayerAsset)) : nullptr;
 }
 
 void UDataLayerEditorSubsystem::AddAllDataLayersTo(TArray<TWeakObjectPtr<UDataLayerInstance>>& OutDataLayers) const
 {
-	if (const AWorldDataLayers* WorldDataLayers = GetWorldDataLayers())
+	UDataLayerSubsystem* DataLayerSubsystem = UWorld::GetSubsystem<UDataLayerSubsystem>(GetWorld());
+	DataLayerSubsystem->ForEachDataLayer([&OutDataLayers](UDataLayerInstance* DataLayer)
 	{
-		WorldDataLayers->ForEachDataLayer([&OutDataLayers](UDataLayerInstance* DataLayer)
-			{
-				OutDataLayers.Add(DataLayer);
-				return true;
-			});
-	}
+		OutDataLayers.Add(DataLayer);
+		return true;
+	});
 }
 
 UDataLayerInstance* UDataLayerEditorSubsystem::CreateDataLayerInstance(const FDataLayerCreationParameters& Parameters)
 {
 	UDataLayerInstance* NewDataLayer = nullptr;
 
-
-	AWorldDataLayers* WorldDataLayers = GetWorldDataLayers();
-	if (!WorldDataLayers->HasDeprecatedDataLayers())
+	AWorldDataLayers* WorldDataLayers = Parameters.WorlDataLayers != nullptr ? Parameters.WorlDataLayers.Get() : GetWorld()->GetWorldDataLayers();
+	if (!HasDeprecatedDataLayers())
 	{
 		NewDataLayer = WorldDataLayers->CreateDataLayer<UDataLayerInstanceWithAsset>(Parameters.DataLayerAsset);
 	}
@@ -1068,23 +1039,19 @@ UDataLayerInstance* UDataLayerEditorSubsystem::CreateDataLayerInstance(const FDa
 
 void UDataLayerEditorSubsystem::DeleteDataLayers(const TArray<UDataLayerInstance*>& DataLayersToDelete)
 {
-	if (AWorldDataLayers* WorldDataLayers = GetWorldDataLayers())
+	UDataLayerSubsystem* DataLayerSubsystem = UWorld::GetSubsystem<UDataLayerSubsystem>(GetWorld());
+	if (DataLayerSubsystem->RemoveDataLayers(DataLayersToDelete))
 	{
-		if (WorldDataLayers->RemoveDataLayers(DataLayersToDelete))
-		{
-			BroadcastDataLayerChanged(EDataLayerAction::Delete, NULL, NAME_None);
-		}
+		BroadcastDataLayerChanged(EDataLayerAction::Delete, NULL, NAME_None);
 	}
 }
 
 void UDataLayerEditorSubsystem::DeleteDataLayer(UDataLayerInstance* DataLayerToDelete)
 {
-	if (AWorldDataLayers* WorldDataLayers = GetWorldDataLayers())
+	UDataLayerSubsystem* DataLayerSubsystem = UWorld::GetSubsystem<UDataLayerSubsystem>(GetWorld());
+	if (DataLayerSubsystem->RemoveDataLayer(DataLayerToDelete))
 	{
-		if (WorldDataLayers->RemoveDataLayer(DataLayerToDelete))
-		{
-			BroadcastDataLayerChanged(EDataLayerAction::Delete, NULL, NAME_None);
-		}
+		BroadcastDataLayerChanged(EDataLayerAction::Delete, NULL, NAME_None);
 	}
 }
 
@@ -1104,7 +1071,9 @@ void UDataLayerEditorSubsystem::BroadcastDataLayerChanged(const EDataLayerAction
 void UDataLayerEditorSubsystem::BroadcastDataLayerEditorLoadingStateChanged(bool bIsFromUserChange)
 {
 	DataLayerEditorLoadingStateChanged.Broadcast(bIsFromUserChange);
-	UpdateDataLayerEditorPerProjectUserSettings();
+	
+	UDataLayerSubsystem* DataLayerSubsystem = UWorld::GetSubsystem<UDataLayerSubsystem>(GetWorld());
+	DataLayerSubsystem->UpdateDataLayerEditorPerProjectUserSettings();
 }
 
 void UDataLayerEditorSubsystem::OnSelectionChanged()
@@ -1158,8 +1127,8 @@ bool UDataLayerEditorSubsystem::TryGetDataLayerFromLabel(const FName& DataLayerL
 
 UDataLayerInstance* UDataLayerEditorSubsystem::GetDataLayerFromLabel(const FName& DataLayerLabel) const
 {
-	const AWorldDataLayers* WorldDataLayers = GetWorldDataLayers();
-	return WorldDataLayers ? const_cast<UDataLayerInstance*>(WorldDataLayers->GetDataLayerFromLabel(DataLayerLabel)) : nullptr;
+	UDataLayerSubsystem* DataLayerSubsystem = UWorld::GetSubsystem<UDataLayerSubsystem>(GetWorld());
+	return DataLayerSubsystem->GetDataLayerFromLabel(DataLayerLabel);
 }
 
 PRAGMA_ENABLE_DEPRECATION_WARNINGS
diff --git a/Engine/Source/Editor/DataLayerEditor/Private/DataLayer/DataLayerHierarchy.cpp b/Engine/Source/Editor/DataLayerEditor/Private/DataLayer/DataLayerHierarchy.cpp
--- a/Engine/Source/Editor/DataLayerEditor/Private/DataLayer/DataLayerHierarchy.cpp
+++ b/Engine/Source/Editor/DataLayerEditor/Private/DataLayer/DataLayerHierarchy.cpp
@@ -8,7 +8,7 @@
 #include "DataLayerMode.h"
 #include "DataLayer/DataLayerEditorSubsystem.h"
 #include "WorldPartition/DataLayer/DataLayerAsset.h"
-#include "WorldPartition/DataLayer/WorldDataLayers.h"
+#include "WorldPartition/DataLayer/DataLayerSubsystem.h"
 #include "Engine/Engine.h"
 #include "Engine/World.h"
 #include "EngineUtils.h"
@@ -130,11 +130,7 @@ FSceneOutlinerTreeItemPtr FDataLayerHierarchy::CreateDataLayerTreeItem(UDataLaye
 
 void FDataLayerHierarchy::CreateItems(TArray<FSceneOutlinerTreeItemPtr>& OutItems) const
 {
-	const AWorldDataLayers* WorldDataLayers = RepresentingWorld.Get()->GetWorldDataLayers();
-	if (!WorldDataLayers)
-	{
-		return;
-	}
+	const UDataLayerSubsystem* DataLayerSubsystem = UWorld::GetSubsystem<UDataLayerSubsystem>(RepresentingWorld.Get());
 
 	auto IsDataLayerShown = [this](const UDataLayerInstance* DataLayer)
 	{
@@ -142,7 +138,7 @@ void FDataLayerHierarchy::CreateItems(TArray<FSceneOutlinerTreeItemPtr>& OutItem
 		return ((bIsRuntimeDataLayer && bShowRuntimeDataLayers) || (!bIsRuntimeDataLayer && bShowEditorDataLayers)) && IsDataLayerPartOfSelection(DataLayer);
 	};
 
-	WorldDataLayers->ForEachDataLayer([this, &OutItems, IsDataLayerShown](UDataLayerInstance* DataLayer)
+	DataLayerSubsystem->ForEachDataLayer([this, &OutItems, IsDataLayerShown](UDataLayerInstance* DataLayer)
 	{
 		if (IsDataLayerShown(DataLayer))
 		{
@@ -177,13 +173,13 @@ void FDataLayerHierarchy::CreateItems(TArray<FSceneOutlinerTreeItemPtr>& OutItem
 		{
 			if (UWorldPartition* const WorldPartition = RepresentingWorld.Get()->GetWorldPartition())
 			{
-				FWorldPartitionHelpers::ForEachActorDesc(WorldPartition, [this, WorldPartition, WorldDataLayers, IsDataLayerShown, &OutItems](const FWorldPartitionActorDesc* ActorDesc)
+				FWorldPartitionHelpers::ForEachActorDesc(WorldPartition, [this, WorldPartition, DataLayerSubsystem, IsDataLayerShown, &OutItems](const FWorldPartitionActorDesc* ActorDesc)
 				{
 					if (ActorDesc != nullptr && !ActorDesc->IsLoaded())
 					{
 						for (const FName& DataLayerInstanceName : ActorDesc->GetDataLayerInstanceNames())
 						{
-							if (const UDataLayerInstance* const DataLayerInstance = WorldDataLayers->GetDataLayerInstance(DataLayerInstanceName))
+							if (const UDataLayerInstance* const DataLayerInstance = DataLayerSubsystem->GetDataLayerInstance(DataLayerInstanceName))
 							{
 								if (IsDataLayerShown(DataLayerInstance))
 								{
@@ -299,19 +295,19 @@ void FDataLayerHierarchy::OnLevelActorsRemoved(const TArray<AActor*>& InActors)
 		FSceneOutlinerHierarchyChangedData EventData;
 		EventData.Type = FSceneOutlinerHierarchyChangedData::Removed;
 
-		if (AWorldDataLayers* WorldDataLayers = CurrentWorld->GetWorldDataLayers())
+		for (AActor* Actor : InActors)
 		{
-			for (AActor* Actor : InActors)
+			if (Actor != nullptr && Actor->HasDataLayers())
 			{
-				if (Actor != nullptr && Actor->HasDataLayers())
+				// Make sure the Actor is still associated with a World otherwise, AActor::GetDataLayerInstances retrieve the DataLayerInstances properly below.
+				// It is allowed when exiting a PIE world since the hiearchy will be fully rebuild via a FullRefreshEvent
+				check(Actor->GetWorld() != nullptr || (RepresentingWorld->WorldType == EWorldType::PIE && RepresentingWorld->bIsTearingDown))
+				
+				const TArray<const UDataLayerInstance*> DataLayerInstances = Actor->GetDataLayerInstances();
+				EventData.ItemIDs.Reserve(DataLayerInstances.Num());
+				for (const UDataLayerInstance* DataLayerInstance : DataLayerInstances)
 				{
-					// It is possible here that Actor doesn't have world anymore
-					const TArray<const UDataLayerInstance*> DataLayerInstances = Actor->GetDataLayerInstances(WorldDataLayers);
-					EventData.ItemIDs.Reserve(DataLayerInstances.Num());
-					for (const UDataLayerInstance* DataLayerInstance : DataLayerInstances)
-					{
-						EventData.ItemIDs.Add(FDataLayerActorTreeItem::ComputeTreeItemID(Actor, DataLayerInstance));
-					}
+					EventData.ItemIDs.Add(FDataLayerActorTreeItem::ComputeTreeItemID(Actor, DataLayerInstance));
 				}
 			}
 		}
@@ -428,10 +424,10 @@ void FDataLayerHierarchy::OnActorDescAdded(FWorldPartitionActorDesc* InActorDesc
 	}
 
 	UWorldPartition* const WorldPartition = RepresentingWorld.Get()->GetWorldPartition();
-	const AWorldDataLayers* const WorldDataLayers = RepresentingWorld.Get()->GetWorldDataLayers();
+	const UDataLayerSubsystem* const DataLayerSubsystem = UWorld::GetSubsystem<UDataLayerSubsystem>(RepresentingWorld.Get());
 	const TArray<FName>& DataLayerInstanceNames = InActorDesc->GetDataLayerInstanceNames();
 
-	if (WorldDataLayers != nullptr && RepresentingWorld->GetWorldPartition() == WorldPartition && DataLayerInstanceNames.Num() > 0)
+	if (RepresentingWorld->GetWorldPartition() == WorldPartition && DataLayerInstanceNames.Num() > 0)
 	{
 		FSceneOutlinerHierarchyChangedData EventData;
 		EventData.Type = FSceneOutlinerHierarchyChangedData::Added;
@@ -439,7 +435,7 @@ void FDataLayerHierarchy::OnActorDescAdded(FWorldPartitionActorDesc* InActorDesc
 		EventData.Items.Reserve(DataLayerInstanceNames.Num());
 		for (const FName& DataLayerInstanceName : DataLayerInstanceNames)
 		{
-			const UDataLayerInstance* const DataLayerInstance = WorldDataLayers->GetDataLayerInstance(DataLayerInstanceName);
+			const UDataLayerInstance* const DataLayerInstance = DataLayerSubsystem->GetDataLayerInstance(DataLayerInstanceName);
 			EventData.Items.Add(Mode->CreateItemFor<FDataLayerActorDescTreeItem>(FDataLayerActorDescTreeItemData(InActorDesc->GetGuid(), WorldPartition, const_cast<UDataLayerInstance*>(DataLayerInstance))));
 		}
 		HierarchyChangedEvent.Broadcast(EventData);
@@ -453,10 +449,10 @@ void FDataLayerHierarchy::OnActorDescRemoved(FWorldPartitionActorDesc* InActorDe
 		return;
 	}
 
-	const AWorldDataLayers* const WorldDataLayers = RepresentingWorld.Get()->GetWorldDataLayers();
+	const UDataLayerSubsystem* const DataLayerSubsystem = UWorld::GetSubsystem<UDataLayerSubsystem>(RepresentingWorld.Get());
 	const TArray<FName>& DataLayerInstanceNames = InActorDesc->GetDataLayerInstanceNames();
 	
-	if (WorldDataLayers != nullptr && DataLayerInstanceNames.Num() > 0)
+	if (DataLayerInstanceNames.Num() > 0)
 	{
 		FSceneOutlinerHierarchyChangedData EventData;
 		EventData.Type = FSceneOutlinerHierarchyChangedData::Removed;
@@ -464,7 +460,7 @@ void FDataLayerHierarchy::OnActorDescRemoved(FWorldPartitionActorDesc* InActorDe
 
 		for (const FName& DataLayerInstanceName : DataLayerInstanceNames)
 		{
-			const UDataLayerInstance* const DataLayer = WorldDataLayers->GetDataLayerInstance(DataLayerInstanceName);
+			const UDataLayerInstance* const DataLayer = DataLayerSubsystem->GetDataLayerInstance(DataLayerInstanceName);
 			EventData.ItemIDs.Add(FDataLayerActorDescTreeItem::ComputeTreeItemID(InActorDesc->GetGuid(), DataLayer));
 		}
 		HierarchyChangedEvent.Broadcast(EventData);
diff --git a/Engine/Source/Editor/DataLayerEditor/Private/DataLayer/DataLayerMode.cpp b/Engine/Source/Editor/DataLayerEditor/Private/DataLayer/DataLayerMode.cpp
--- a/Engine/Source/Editor/DataLayerEditor/Private/DataLayer/DataLayerMode.cpp
+++ b/Engine/Source/Editor/DataLayerEditor/Private/DataLayer/DataLayerMode.cpp
@@ -7,7 +7,7 @@
 #include "ContentBrowserModule.h"
 #include "DataLayer/DataLayerEditorSubsystem.h"
 #include "DataLayer/DataLayerOutlinerDeleteButtonColumn.h"
-#include "WorldPartition/DataLayer/WorldDataLayers.h"
+#include "WorldPartition/DataLayer/DataLayerSubsystem.h"
 #include "WorldPartition/DataLayer/DataLayerAsset.h"
 #include "DataLayer/SDataLayerOutliner.h"
 #include "DataLayerHierarchy.h"
@@ -914,15 +914,15 @@ void FDataLayerMode::RegisterContextMenu()
 
 			bool bHasActorEditorContextDataLayers = false;			
 			TArray<const UDataLayerInstance*> AllDataLayers;
-			if (const AWorldDataLayers* WorldDataLayers = RepresentingWorld.IsValid() ? RepresentingWorld->GetWorldDataLayers() : nullptr)
+			if (const UDataLayerSubsystem* DataLayerSubsystem = UWorld::GetSubsystem<UDataLayerSubsystem>(RepresentingWorld.Get()))
 			{
-				WorldDataLayers->ForEachDataLayer([&AllDataLayers](UDataLayerInstance* DataLayerInstance)
+				DataLayerSubsystem->ForEachDataLayer([&AllDataLayers](UDataLayerInstance* DataLayerInstance)
 				{
 					AllDataLayers.Add(DataLayerInstance);
 					return true;
 				});
 
-				bHasActorEditorContextDataLayers = !WorldDataLayers->GetActorEditorContextDataLayers().IsEmpty();
+				bHasActorEditorContextDataLayers = !DataLayerSubsystem->GetActorEditorContextDataLayers().IsEmpty();
 			}
 
 			{
@@ -1293,9 +1293,9 @@ void FDataLayerMode::CreateViewContent(FMenuBuilder& MenuBuilder)
 	);
 
 	TArray<UDataLayerInstance*> AllDataLayers;
-	if (const AWorldDataLayers* WorldDataLayers = RepresentingWorld.IsValid() ? RepresentingWorld->GetWorldDataLayers() : nullptr)
+	if (const UDataLayerSubsystem* DataLayerSubsystem = UWorld::GetSubsystem<UDataLayerSubsystem>(RepresentingWorld.Get()))
 	{
-		WorldDataLayers->ForEachDataLayer([&AllDataLayers](UDataLayerInstance* DataLayer)
+		DataLayerSubsystem->ForEachDataLayer([&AllDataLayers](UDataLayerInstance* DataLayer)
 		{
 			AllDataLayers.Add(DataLayer);
 			return true;
diff --git a/Engine/Source/Editor/DataLayerEditor/Public/DataLayer/DataLayerEditorSubsystem.h b/Engine/Source/Editor/DataLayerEditor/Public/DataLayer/DataLayerEditorSubsystem.h
--- a/Engine/Source/Editor/DataLayerEditor/Public/DataLayer/DataLayerEditorSubsystem.h
+++ b/Engine/Source/Editor/DataLayerEditor/Public/DataLayer/DataLayerEditorSubsystem.h
@@ -635,13 +635,6 @@ private:
 	 */
 	TArray<AActor*> GetSelectedActors() const;
 
-	/**
-	 * Gets the WorldDataLayers for this world
-	 *
-	 * @return	The AWorldDataLayers Object
-	 */
-	class AWorldDataLayers* GetWorldDataLayers() const;
-
 	/**
 	 * Get the current UWorld object.
 	 *
@@ -659,7 +652,6 @@ private:
 	void PostUndoRedo();
 
 	bool SetDataLayerIsLoadedInEditorInternal(UDataLayerInstance* DataLayer, const bool bIsLoadedInEditor, const bool bIsFromUserChange);
-	void UpdateDataLayerEditorPerProjectUserSettings();
 
 	void BroadcastActorDataLayersChanged(const TWeakObjectPtr<AActor>& ChangedActor);
 	void BroadcastDataLayerChanged(const EDataLayerAction Action, const TWeakObjectPtr<const UDataLayerInstance>& ChangedDataLayer, const FName& ChangedProperty);
diff --git a/Engine/Source/Editor/SceneOutliner/Private/SceneOutlinerModule.cpp b/Engine/Source/Editor/SceneOutliner/Private/SceneOutlinerModule.cpp
--- a/Engine/Source/Editor/SceneOutliner/Private/SceneOutlinerModule.cpp
+++ b/Engine/Source/Editor/SceneOutliner/Private/SceneOutlinerModule.cpp
@@ -23,7 +23,7 @@
 #include "FolderTreeItem.h"
 #include "WorldTreeItem.h"
 #include "WorldPartition/DataLayer/DataLayerInstance.h"
-#include "WorldPartition/DataLayer/WorldDataLayers.h"
+#include "WorldPartition/DataLayer/DataLayerSubsystem.h"
 #include "Editor.h"
 
 #define LOCTEXT_NAMESPACE "SceneOutlinerModule"
@@ -254,11 +254,10 @@ void FSceneOutlinerModule::CreateActorInfoColumns(FSceneOutlinerInitializationOp
 			{
 				const UActorDescContainer* ActorDescContainer = ActorDescItem->ActorDescHandle.Container.Get();
 				const UWorld* World = ActorDescContainer ? ActorDescContainer->GetWorld() : nullptr;
-				const AWorldDataLayers* WorldDataLayers = World ? World->GetWorldDataLayers() : nullptr;
-			
-				if (WorldDataLayers)
+				
+				if (const UDataLayerSubsystem* DataLayerSubsystem = UWorld::GetSubsystem<UDataLayerSubsystem>(World))
 				{
-					DataLayerInstances = WorldDataLayers->GetDataLayerInstances(ActorDesc->GetDataLayerInstanceNames());
+					DataLayerInstances = DataLayerSubsystem->GetDataLayerInstances(ActorDesc->GetDataLayerInstanceNames());
 				}
 			}
 		}
diff --git a/Engine/Source/Editor/UnrealEd/Private/Commandlets/WorldPartitionDataLayerToAssetCommandlet.cpp b/Engine/Source/Editor/UnrealEd/Private/Commandlets/WorldPartitionDataLayerToAssetCommandlet.cpp
--- a/Engine/Source/Editor/UnrealEd/Private/Commandlets/WorldPartitionDataLayerToAssetCommandlet.cpp
+++ b/Engine/Source/Editor/UnrealEd/Private/Commandlets/WorldPartitionDataLayerToAssetCommandlet.cpp
@@ -475,11 +475,9 @@ bool UDataLayerToAssetCommandlet::RemapActorDataLayersToAssets(TStrongObjectPtr<
 	UE_SCOPED_TIMER(TEXT("Remapping Actors Data Layers"), LogDataLayerToAssetCommandlet, Display);
 	UE_LOG(LogDataLayerToAssetCommandlet, Log, TEXT("Starting Actor Data Layer Remapping To Data Layer Asset. This can take a while."));
 
-	AWorldDataLayers* WorldDataLayersActor = MainWorld->GetWorldDataLayers();
-
 	uint32 ErrorCount = 0;
 	FWorldPartitionHelpers::ForEachActorWithLoading(MainWorld->GetWorldPartition(), AActor::StaticClass(),
-		[WorldDataLayersActor, &ErrorCount, &CommandletContext, this](const FWorldPartitionActorDesc* ActorDesc)
+		[&ErrorCount, &CommandletContext, this](const FWorldPartitionActorDesc* ActorDesc)
 		{
 			uint32 ActorConversionErrors = 0;
 			if (AActor* Actor = ActorDesc->GetActor())
@@ -640,7 +638,6 @@ bool UDataLayerToAssetCommandlet::RebuildDataLayerHierarchies(TStrongObjectPtr<U
 	UE_SCOPED_TIMER(TEXT("Creating Data Layer Instances Hierarchy"), LogDataLayerToAssetCommandlet, Display);
 
 	uint32 ErrorCount = 0;
-	AWorldDataLayers* WorldDataLayers = MainWorld->GetWorldDataLayers();
 	for (const TWeakObjectPtr<UDataLayerConversionInfo>& ConvertingInfo : CommandletContext->GetConvertingDataLayerConversionInfo())
 	{
 		UDataLayerInstance* Child = ConvertingInfo->DataLayerInstance;
@@ -705,7 +702,6 @@ bool UDataLayerToAssetCommandlet::DeletePreviousConversionsData(TStrongObjectPtr
 	UE_SCOPED_TIMER(TEXT("Delete Conflicting Assets"), LogDataLayerToAssetCommandlet, Display);
 
 	uint32 ErrorCount = 0;
-	AWorldDataLayers* WorldDataLayers = MainWorld->GetWorldDataLayers();
 	for (const TWeakObjectPtr<UDataLayerConversionInfo>& ConvertingInfo : CommandletContext->GetConvertingDataLayerConversionInfo())
 	{
 		for (const TWeakObjectPtr<UDataLayerConversionInfo>& PreviousConversion : ConvertingInfo->GetPreviousConversions())
diff --git a/Engine/Source/Editor/UnrealEd/Private/WorldPartition/WorldPartitionBuilder.cpp b/Engine/Source/Editor/UnrealEd/Private/WorldPartition/WorldPartitionBuilder.cpp
--- a/Engine/Source/Editor/UnrealEd/Private/WorldPartition/WorldPartitionBuilder.cpp
+++ b/Engine/Source/Editor/UnrealEd/Private/WorldPartition/WorldPartitionBuilder.cpp
@@ -13,7 +13,6 @@
 #include "WorldPartition/WorldPartitionHelpers.h"
 #include "WorldPartition/WorldPartitionSubsystem.h"
 #include "WorldPartition/DataLayer/DataLayerSubsystem.h"
-#include "WorldPartition/DataLayer/WorldDataLayers.h"
 #include "WorldPartition/DataLayer/DataLayerInstance.h"
 #include "LevelInstance/LevelInstanceSubsystem.h"
 #include "Math/IntVector.h"
@@ -159,38 +158,37 @@ bool UWorldPartitionBuilder::Run(UWorld* World, FPackageSourceControlHelper& Pac
 	UWorldPartition* WorldPartition = World->GetWorldPartition();
 
 	// Properly Setup DataLayers for Builder
-	if (const AWorldDataLayers* WorldDataLayers = World->GetWorldDataLayers())
+	UDataLayerSubsystem* DataLayerSubsystem = UWorld::GetSubsystem<UDataLayerSubsystem>(World);
+	
+	// Load Data Layers
+	bool bUpdateEditorCells = false;
+	DataLayerSubsystem->ForEachDataLayer([&bUpdateEditorCells, this](UDataLayerInstance* DataLayer)
 	{
-		// Load Data Layers
-		bool bUpdateEditorCells = false;
-		WorldDataLayers->ForEachDataLayer([&bUpdateEditorCells, this](UDataLayerInstance* DataLayer)
-		{
-			const FName DataLayerShortName(DataLayer->GetDataLayerShortName());
+		const FName DataLayerShortName(DataLayer->GetDataLayerShortName());
 
-			// Load all Non Excluded Data Layers + Non DynamicallyLoaded Data Layers + Initially Active Data Layers + Data Layers provided by builder
-			const bool bLoadedInEditor = !ExcludedDataLayerShortNames.Contains(DataLayerShortName) && ((bLoadNonDynamicDataLayers && !DataLayer->IsRuntime()) ||
-										 (bLoadInitiallyActiveDataLayers && DataLayer->GetInitialRuntimeState() == EDataLayerRuntimeState::Activated) ||
-										DataLayerShortNames.Contains(DataLayerShortName));
-			if (DataLayer->IsLoadedInEditor() != bLoadedInEditor)
+		// Load all Non Excluded Data Layers + Non DynamicallyLoaded Data Layers + Initially Active Data Layers + Data Layers provided by builder
+		const bool bLoadedInEditor = !ExcludedDataLayerShortNames.Contains(DataLayerShortName) && ((bLoadNonDynamicDataLayers && !DataLayer->IsRuntime()) ||
+										(bLoadInitiallyActiveDataLayers && DataLayer->GetInitialRuntimeState() == EDataLayerRuntimeState::Activated) ||
+									DataLayerShortNames.Contains(DataLayerShortName));
+		if (DataLayer->IsLoadedInEditor() != bLoadedInEditor)
+		{
+			bUpdateEditorCells = true;
+			DataLayer->SetIsLoadedInEditor(bLoadedInEditor, /*bFromUserChange*/false);
+			if (RequiresCommandletRendering() && bLoadedInEditor)
 			{
-				bUpdateEditorCells = true;
-				DataLayer->SetIsLoadedInEditor(bLoadedInEditor, /*bFromUserChange*/false);
-				if (RequiresCommandletRendering() && bLoadedInEditor)
-				{
-					DataLayer->SetIsInitiallyVisible(true);
-				}
+				DataLayer->SetIsInitiallyVisible(true);
 			}
+		}
 			
-			UE_LOG(LogWorldPartitionBuilder, Display, TEXT("DataLayer '%s' Loaded: %d"), *UDataLayerInstance::GetDataLayerText(DataLayer).ToString(), bLoadedInEditor ? 1 : 0);
+		UE_LOG(LogWorldPartitionBuilder, Display, TEXT("DataLayer '%s' Loaded: %d"), *UDataLayerInstance::GetDataLayerText(DataLayer).ToString(), bLoadedInEditor ? 1 : 0);
 			
-			return true;
-		});
+		return true;
+	});
 	
-		if (bUpdateEditorCells)
-		{
-			UE_LOG(LogWorldPartitionBuilder, Display, TEXT("DataLayer load state changed refreshing editor cells"));
-			FDataLayersEditorBroadcast::StaticOnActorDataLayersEditorLoadingStateChanged(false);
-		}
+	if (bUpdateEditorCells)
+	{
+		UE_LOG(LogWorldPartitionBuilder, Display, TEXT("DataLayer load state changed refreshing editor cells"));
+		FDataLayersEditorBroadcast::StaticOnActorDataLayersEditorLoadingStateChanged(false);
 	}
 
 	const ELoadingMode LoadingMode = GetLoadingMode();
diff --git a/Engine/Source/Editor/UnrealEd/Private/WorldPartition/WorldPartitionMiniMapBuilder.cpp b/Engine/Source/Editor/UnrealEd/Private/WorldPartition/WorldPartitionMiniMapBuilder.cpp
--- a/Engine/Source/Editor/UnrealEd/Private/WorldPartition/WorldPartitionMiniMapBuilder.cpp
+++ b/Engine/Source/Editor/UnrealEd/Private/WorldPartition/WorldPartitionMiniMapBuilder.cpp
@@ -21,7 +21,7 @@
 #include "WorldPartition/ActorDescContainer.h"
 #include "WorldPartition/DataLayer/ActorDataLayer.h"
 #include "WorldPartition/DataLayer/DataLayerInstance.h"
-#include "WorldPartition/DataLayer/WorldDataLayers.h"
+#include "WorldPartition/DataLayer/DataLayerSubsystem.h"
 #include "WorldPartition/HLOD/HLODActor.h"
 #include "WorldPartition/HLOD/HLODActorDesc.h"
 #include "WorldPartition/WorldPartition.h"
@@ -55,17 +55,10 @@ bool UWorldPartitionMiniMapBuilder::PreRun(UWorld* World, FPackageSourceControlH
 	WorldMiniMap->MiniMapTexture = nullptr;
 	MiniMapTiles.Empty();
 
-	AWorldDataLayers* WorldDataLayers = World->GetWorldDataLayers();
-
-	if (WorldDataLayers == nullptr)
-	{
-		UE_LOG(LogWorldPartitionMiniMapBuilder, Error, TEXT("Failed to retrieve WorldDataLayers."));
-		return false;
-	}
-
+	UDataLayerSubsystem* DataLayerSubSystem = UWorld::GetSubsystem<UDataLayerSubsystem>(World);
 	for (const FActorDataLayer& ActorDataLayer : WorldMiniMap->ExcludedDataLayers)
 	{
-		const UDataLayerInstance* DataLayerInstance = WorldDataLayers->GetDataLayerInstance(ActorDataLayer.Name);
+		const UDataLayerInstance* DataLayerInstance = DataLayerSubSystem->GetDataLayerInstance(ActorDataLayer.Name);
 
 		if (DataLayerInstance != nullptr)
 		{
diff --git a/Engine/Source/Editor/UnrealEd/Private/WorldPartition/WorldPartitionNavigationDataBuilder.cpp b/Engine/Source/Editor/UnrealEd/Private/WorldPartition/WorldPartitionNavigationDataBuilder.cpp
--- a/Engine/Source/Editor/UnrealEd/Private/WorldPartition/WorldPartitionNavigationDataBuilder.cpp
+++ b/Engine/Source/Editor/UnrealEd/Private/WorldPartition/WorldPartitionNavigationDataBuilder.cpp
@@ -15,7 +15,7 @@
 #include "WorldPartition/WorldPartition.h"
 #include "WorldPartition/WorldPartitionSubsystem.h"
 #include "WorldPartition/DataLayer/DataLayerAsset.h"
-#include "WorldPartition/DataLayer/WorldDataLayers.h"
+#include "WorldPartition/DataLayer/DataLayerSubsystem.h"
 #include "WorldPartition/NavigationData/NavigationDataChunkActor.h"
 
 DEFINE_LOG_CATEGORY_STATIC(LogWorldPartitionNavigationDataBuilder, Log, All);
@@ -28,21 +28,19 @@ UWorldPartitionNavigationDataBuilder::UWorldPartitionNavigationDataBuilder(const
 bool UWorldPartitionNavigationDataBuilder::PreRun(UWorld* World, FPackageSourceControlHelper& PackageHelper)
 {
 	// Set runtime data layer to be included in the base navmesh generation.
-	if (const AWorldDataLayers* WorldDataLayers = World->GetWorldDataLayers())
+	UDataLayerSubsystem* DataLayerSubsystem = UWorld::GetSubsystem<UDataLayerSubsystem>(GetWorld());
+	for (const TObjectPtr<UDataLayerAsset> DataLayer : World->GetWorldSettings()->BaseNavmeshDataLayers)
 	{
-		for(const TObjectPtr<UDataLayerAsset> DataLayer : World->GetWorldSettings()->BaseNavmeshDataLayers)
+		if (DataLayer != nullptr)
 		{
-			if (DataLayer != nullptr)
+			const UDataLayerInstance* DataLayerInstance = DataLayerSubsystem->GetDataLayerInstance(DataLayer);
+			if (DataLayerInstance == nullptr)
 			{
-				const UDataLayerInstance* DataLayerInstance = WorldDataLayers->GetDataLayerInstance(DataLayer);
-				if(DataLayerInstance == nullptr)
-				{
-					UE_LOG(LogWorldPartitionNavigationDataBuilder, Error, TEXT("Missing UDataLayerInstance for %s."), *DataLayer->GetName());	
-				}
-				else if (DataLayerInstance->IsRuntime())
-				{
-					DataLayerShortNames.Add(FName(DataLayerInstance->GetDataLayerShortName()));
-				}
+				UE_LOG(LogWorldPartitionNavigationDataBuilder, Error, TEXT("Missing UDataLayerInstance for %s."), *DataLayer->GetName());
+			}
+			else if (DataLayerInstance->IsRuntime())
+			{
+				DataLayerShortNames.Add(FName(DataLayerInstance->GetDataLayerShortName()));
 			}
 		}
 	}
diff --git a/Engine/Source/Runtime/Engine/Classes/GameFramework/Actor.h b/Engine/Source/Runtime/Engine/Classes/GameFramework/Actor.h
--- a/Engine/Source/Runtime/Engine/Classes/GameFramework/Actor.h
+++ b/Engine/Source/Runtime/Engine/Classes/GameFramework/Actor.h
@@ -1186,7 +1186,6 @@ public:
 	bool HasValidDataLayers() const;
 	TArray<FName> GetDataLayerInstanceNames() const;
 	TArray<const UDataLayerInstance*> GetDataLayerInstances() const;
-	TArray<const UDataLayerInstance*> GetDataLayerInstances(const AWorldDataLayers* WorldDataLayers) const;
 	bool IsPropertyChangedAffectingDataLayers(FPropertyChangedEvent& PropertyChangedEvent) const;
 	bool IsValidForDataLayer() const;
 	void FixupDataLayers(bool bRevertChangesOnLockedDataLayer = false);
@@ -1219,6 +1218,9 @@ public:
 	UE_DEPRECATED(5.1, "Convert DataLayer using UDataLayerToAssetCommandlet and use GetDataLayerAssets() instead")
 	TArray<FActorDataLayer> const& GetActorDataLayers() const { return DataLayers; }
 
+	UE_DEPRECATED(5.1, "Use GetDataLayerInstances() with no parameters instead")
+	TArray<const UDataLayerInstance*> GetDataLayerInstances(const AWorldDataLayers* WorldDataLayers) const;
+
 	//~ End Deprecated
 #endif
 
diff --git a/Engine/Source/Runtime/Engine/Private/ActorEditor.cpp b/Engine/Source/Runtime/Engine/Private/ActorEditor.cpp
--- a/Engine/Source/Runtime/Engine/Private/ActorEditor.cpp
+++ b/Engine/Source/Runtime/Engine/Private/ActorEditor.cpp
@@ -13,7 +13,8 @@
 #include "AI/NavigationSystemBase.h"
 #include "Engine/BlueprintGeneratedClass.h"
 #include "WorldPartition/WorldPartitionSubsystem.h"
-#include "WorldPartition/DataLayer/WorldDataLayers.h"
+#include "WorldPartition/DataLayer/DataLayer.h"
+#include "WorldPartition/DataLayer/DataLayerSubsystem.h"
 #include "WorldPartition/DataLayer/DataLayerInstance.h"
 #include "EditorSupportDelegates.h"
 #include "Logging/TokenizedMessage.h"
@@ -1382,11 +1383,11 @@ bool AActor::HasDataLayers() const
 
 bool AActor::HasValidDataLayers() const
 {
-	if (const AWorldDataLayers* WorldDataLayers = GetWorld()->GetWorldDataLayers())
+	if (UDataLayerSubsystem* DataLayerSubsystem = UWorld::GetSubsystem<UDataLayerSubsystem>(GetWorld()))
 	{
 		for (const TObjectPtr<const UDataLayerAsset>& DataLayerAsset : DataLayerAssets)
 		{
-			if (const UDataLayerInstance* DataLayerInstance = WorldDataLayers->GetDataLayerInstance(DataLayerAsset))
+			if (const UDataLayerInstance* DataLayerInstance = DataLayerSubsystem->GetDataLayerInstance(DataLayerAsset))
 			{
 				return true;
 			}
@@ -1394,7 +1395,7 @@ bool AActor::HasValidDataLayers() const
 
 		for (const FActorDataLayer& ActorDataLayer : DataLayers)
 		{
-			if (const UDataLayerInstance* DataLayerInstance = WorldDataLayers->GetDataLayerInstance(ActorDataLayer.Name))
+			if (const UDataLayerInstance* DataLayerInstance = DataLayerSubsystem->GetDataLayerInstance(ActorDataLayer.Name))
 			{
 				return true;
 			}
@@ -1407,36 +1408,30 @@ TArray<FName> AActor::GetDataLayerInstanceNames() const
 {
 	TArray<FName> DataLayerInstanceNames;
 	
-	const AWorldDataLayers* WorldDataLayers = GetWorld()->GetWorldDataLayers();
-	if (WorldDataLayers != nullptr)
+	if (UDataLayerSubsystem* DataLayerSubsystem = UWorld::GetSubsystem<UDataLayerSubsystem>(GetWorld()))
 	{
-		DataLayerInstanceNames += WorldDataLayers->GetDataLayerInstanceNames(DataLayerAssets);
+		DataLayerInstanceNames += DataLayerSubsystem->GetDataLayerInstanceNames(DataLayerAssets);
 
 		PRAGMA_DISABLE_DEPRECATION_WARNINGS
-		DataLayerInstanceNames += WorldDataLayers->GetDataLayerInstanceNames(DataLayers);
+		DataLayerInstanceNames += DataLayerSubsystem->GetDataLayerInstanceNames(DataLayers);
 		PRAGMA_ENABLE_DEPRECATION_WARNINGS
 	}
 	return DataLayerInstanceNames;
 }
 
 TArray<const UDataLayerInstance*> AActor::GetDataLayerInstances() const
-{
-	return GetWorld() ? GetDataLayerInstances(GetWorld()->GetWorldDataLayers()) : TArray<const UDataLayerInstance*>();
-}
-
-TArray<const UDataLayerInstance*> AActor::GetDataLayerInstances(const AWorldDataLayers* WorldDataLayers) const
 {
 	TArray<const UDataLayerInstance*> DataLayerInstances;
-	if (WorldDataLayers != nullptr)
+	
+	if (UDataLayerSubsystem* DataLayerSubsystem = UWorld::GetSubsystem<UDataLayerSubsystem>(GetWorld()))
 	{
-		DataLayerInstances += WorldDataLayers->GetDataLayerInstances(DataLayerAssets);
+		DataLayerInstances += DataLayerSubsystem->GetDataLayerInstances(DataLayerAssets);
 
 		PRAGMA_DISABLE_DEPRECATION_WARNINGS
-		DataLayerInstances += WorldDataLayers->GetDataLayerInstances(DataLayers);
+		DataLayerInstances += DataLayerSubsystem->GetDataLayerInstances(DataLayers);
 		PRAGMA_ENABLE_DEPRECATION_WARNINGS
 	}
 
-
 	return DataLayerInstances;
 }
 
@@ -1453,7 +1448,7 @@ void AActor::FixupDataLayers(bool bRevertChangesOnLockedDataLayer /*= false*/)
 
 		if (GetWorld())
 		{
-			if (const AWorldDataLayers* WorldDataLayers = GetWorld()->GetWorldDataLayers())
+			if (UDataLayerSubsystem* DataLayerSubsystem = UWorld::GetSubsystem<UDataLayerSubsystem>(GetWorld()))
 			{
 				if (bRevertChangesOnLockedDataLayer)
 				{
@@ -1461,12 +1456,12 @@ void AActor::FixupDataLayers(bool bRevertChangesOnLockedDataLayer /*= false*/)
 					TSet<const UDataLayerAsset*> PreEdit(PreEditChangeDataLayers);
 					TSet<const UDataLayerAsset*> PostEdit(DataLayerAssets);
 
-					auto DifferenceContainsLockedDataLayers = [WorldDataLayers](const TSet<const UDataLayerAsset*>& A, const TSet<const UDataLayerAsset*>& B)
+					auto DifferenceContainsLockedDataLayers = [DataLayerSubsystem](const TSet<const UDataLayerAsset*>& A, const TSet<const UDataLayerAsset*>& B)
 					{
 						TSet<const UDataLayerAsset*> Diff = A.Difference(B);
 						for (const UDataLayerAsset* DataLayerAsset : Diff)
 						{
-							const UDataLayerInstance* DataLayerInstance = WorldDataLayers->GetDataLayerInstance(DataLayerAsset);
+							const UDataLayerInstance* DataLayerInstance = DataLayerSubsystem->GetDataLayerInstance(DataLayerAsset);
 							if (DataLayerInstance && DataLayerInstance->IsLocked())
 							{
 								return true;
@@ -1483,7 +1478,7 @@ void AActor::FixupDataLayers(bool bRevertChangesOnLockedDataLayer /*= false*/)
 				}
 
 				PRAGMA_DISABLE_DEPRECATION_WARNINGS
-				auto CleanupDataLayers = [this, WorldDataLayers](auto& DataLayerArray)
+				auto CleanupDataLayers = [this, DataLayerSubsystem](auto& DataLayerArray)
 				{
 					using ArrayType = typename TRemoveReference<decltype(DataLayerArray)>::Type;
 					
@@ -1491,7 +1486,7 @@ void AActor::FixupDataLayers(bool bRevertChangesOnLockedDataLayer /*= false*/)
 					for (int32 Index = 0; Index < DataLayerArray.Num();)
 					{
 						auto& DataLayer = DataLayerArray[Index];
-						if (!WorldDataLayers->GetDataLayerInstance(DataLayer) || ExistingDataLayer.Contains(DataLayer))
+						if (!DataLayerSubsystem->GetDataLayerInstance(DataLayer) || ExistingDataLayer.Contains(DataLayer))
 						{
 							DataLayerArray.RemoveAtSwap(Index);
 						}
@@ -1592,6 +1587,19 @@ bool AActor::RemoveDataLayer(const UDEPRECATED_DataLayer* DataLayer)
 	return false;
 }
 
+TArray<const UDataLayerInstance*> AActor::GetDataLayerInstances(const AWorldDataLayers* WorldDataLayers) const
+{
+	TArray<const UDataLayerInstance*> DataLayerInstances;
+
+	if (WorldDataLayers != nullptr)
+	{
+		DataLayerInstances += WorldDataLayers->GetDataLayerInstances(DataLayerAssets);
+		DataLayerInstances += WorldDataLayers->GetDataLayerInstances(DataLayers);
+	}
+
+	return DataLayerInstances;
+}
+
 PRAGMA_ENABLE_DEPRECATION_WARNINGS
 
 bool AActor::ContainsDataLayer(const FActorDataLayer& ActorDataLayer) const
diff --git a/Engine/Source/Runtime/Engine/Private/ActorPartition/ActorPartitionSubsystem.cpp b/Engine/Source/Runtime/Engine/Private/ActorPartition/ActorPartitionSubsystem.cpp
--- a/Engine/Source/Runtime/Engine/Private/ActorPartition/ActorPartitionSubsystem.cpp
+++ b/Engine/Source/Runtime/Engine/Private/ActorPartition/ActorPartitionSubsystem.cpp
@@ -393,9 +393,9 @@ APartitionActor* UActorPartitionSubsystem::GetActor(const TSubclassOf<APartition
 	
 	auto WrappedInActorCreated = [World, InActorCreated](APartitionActor* PartitionActor)
 	{
-		if (const AWorldDataLayers* WorldDataLayers = World->GetWorldDataLayers())
+		if(UDataLayerSubsystem* DataLayerSubsystem = World->GetSubsystem<UDataLayerSubsystem>())
 		{
-			for (UDataLayerInstance* DataLayer : WorldDataLayers->GetActorEditorContextDataLayers())
+			for (UDataLayerInstance* DataLayer : DataLayerSubsystem->GetActorEditorContextDataLayers())
 			{
 				PartitionActor->AddDataLayer(DataLayer);
 			}
diff --git a/Engine/Source/Runtime/Engine/Private/WorldPartition/DataLayer/DataLayerEditorContext.cpp b/Engine/Source/Runtime/Engine/Private/WorldPartition/DataLayer/DataLayerEditorContext.cpp
--- a/Engine/Source/Runtime/Engine/Private/WorldPartition/DataLayer/DataLayerEditorContext.cpp
+++ b/Engine/Source/Runtime/Engine/Private/WorldPartition/DataLayer/DataLayerEditorContext.cpp
@@ -5,7 +5,6 @@
 #if WITH_EDITOR
 
 #include "WorldPartition/DataLayer/DataLayerInstance.h"
-#include "WorldPartition/DataLayer/WorldDataLayers.h"
 #include "WorldPartition/DataLayer/DataLayerSubsystem.h"
 #include "Engine/World.h"
 
@@ -15,15 +14,10 @@
 FDataLayerEditorContext::FDataLayerEditorContext(UWorld* InWorld, const TArray<FName>& InDataLayerInstances)
 	: Hash(FDataLayerEditorContext::EmptyHash)
 {
-	const AWorldDataLayers* WorldDataLayers = InWorld->GetWorldDataLayers();
-	if (!WorldDataLayers)
-	{
-		return;
-	}
-
+	UDataLayerSubsystem* DataLayerSubsystem = UWorld::GetSubsystem<UDataLayerSubsystem>(InWorld);
 	for (const FName& DataLayerInstanceName : InDataLayerInstances)
 	{
-		if (const UDataLayerInstance* DataLayerInstance = WorldDataLayers->GetDataLayerInstance(DataLayerInstanceName))
+		if (const UDataLayerInstance* DataLayerInstance = DataLayerSubsystem->GetDataLayerInstance(DataLayerInstanceName))
 		{
 			DataLayerInstances.AddUnique(DataLayerInstance->GetDataLayerFName());
 		}
diff --git a/Engine/Source/Runtime/Engine/Private/WorldPartition/DataLayer/DataLayerInstanceWithAsset.cpp b/Engine/Source/Runtime/Engine/Private/WorldPartition/DataLayer/DataLayerInstanceWithAsset.cpp
--- a/Engine/Source/Runtime/Engine/Private/WorldPartition/DataLayer/DataLayerInstanceWithAsset.cpp
+++ b/Engine/Source/Runtime/Engine/Private/WorldPartition/DataLayer/DataLayerInstanceWithAsset.cpp
@@ -3,6 +3,7 @@
 #include "WorldPartition/DataLayer/DataLayerInstanceWithAsset.h"
 
 #include "WorldPartition/DataLayer/WorldDataLayers.h"
+#include "WorldPartition/DataLayer/DataLayerSubsystem.h"
 #include "WorldPartition/ErrorHandling/WorldPartitionStreamingGenerationErrorHandler.h"
 
 UDataLayerInstanceWithAsset::UDataLayerInstanceWithAsset(const FObjectInitializer& ObjectInitializer)
@@ -54,7 +55,8 @@ bool UDataLayerInstanceWithAsset::Validate(IStreamingGenerationErrorHandler* Err
 		return false;
 	}
 
-	GetOuterAWorldDataLayers()->ForEachDataLayer([&bIsValid, this, ErrorHandler](UDataLayerInstance* DataLayerInstance)
+	UDataLayerSubsystem* DataLayerSubsystem = UWorld::GetSubsystem<UDataLayerSubsystem>(GetWorld());
+	DataLayerSubsystem->ForEachDataLayer([&bIsValid, this, ErrorHandler](UDataLayerInstance* DataLayerInstance)
 	{
 		if (DataLayerInstance != this)
 		{
diff --git a/Engine/Source/Runtime/Engine/Private/WorldPartition/DataLayer/DataLayerSubsystem.cpp b/Engine/Source/Runtime/Engine/Private/WorldPartition/DataLayer/DataLayerSubsystem.cpp
--- a/Engine/Source/Runtime/Engine/Private/WorldPartition/DataLayer/DataLayerSubsystem.cpp
+++ b/Engine/Source/Runtime/Engine/Private/WorldPartition/DataLayer/DataLayerSubsystem.cpp
@@ -66,6 +66,86 @@ void UDataLayerSubsystem::Initialize(FSubsystemCollectionBase& Collection)
 #endif
 }
 
+#if WITH_EDITOR
+
+bool UDataLayerSubsystem::RemoveDataLayer(const UDataLayerInstance* InDataLayer)
+{
+	if (AWorldDataLayers* WorldDataLayers = GetWorld()->GetWorldDataLayers())
+	{
+		return WorldDataLayers->RemoveDataLayer(InDataLayer);;
+	}
+	else
+	{
+		UE_LOG(LogWorldPartition, Warning, TEXT("Removing a Data Layer without a World Partition"));
+	}
+
+	return false;
+}
+
+bool UDataLayerSubsystem::RemoveDataLayers(const TArray<UDataLayerInstance*>& InDataLayerInstances)
+{
+	if (AWorldDataLayers* WorldDataLayers = GetWorld()->GetWorldDataLayers())
+	{
+		return WorldDataLayers->RemoveDataLayers(InDataLayerInstances);
+	}
+	else
+	{
+		UE_LOG(LogWorldPartition, Warning, TEXT("Removing Data Layers without a World Partition"));
+	}
+
+	return false;
+}
+
+void UDataLayerSubsystem::UpdateDataLayerEditorPerProjectUserSettings() const
+{
+	TArray<FName> DataLayersNotLoadedInEditor;
+	TArray<FName> DataLayersLoadedInEditor;
+
+	GetUserLoadedInEditorStates(DataLayersLoadedInEditor, DataLayersNotLoadedInEditor);
+
+	GetMutableDefault<UWorldPartitionEditorPerProjectUserSettings>()->SetWorldDataLayersNonDefaultEditorLoadStates(GetWorld(), DataLayersLoadedInEditor, DataLayersNotLoadedInEditor);
+}
+
+void UDataLayerSubsystem::GetUserLoadedInEditorStates(TArray<FName>& OutDataLayersLoadedInEditor, TArray<FName>& OutDataLayersNotLoadedInEditor) const
+{
+	OutDataLayersLoadedInEditor.Empty();
+	OutDataLayersNotLoadedInEditor.Empty();
+
+	const TArray<FName>& SettingsDataLayersNotLoadedInEditor = GetMutableDefault<UWorldPartitionEditorPerProjectUserSettings>()->GetWorldDataLayersNotLoadedInEditor(GetWorld());
+	const TArray<FName>& SettingsDataLayersLoadedInEditor = GetMutableDefault<UWorldPartitionEditorPerProjectUserSettings>()->GetWorldDataLayersLoadedInEditor(GetWorld());
+
+	ForEachDataLayer([&OutDataLayersLoadedInEditor, &OutDataLayersNotLoadedInEditor, &SettingsDataLayersNotLoadedInEditor, &SettingsDataLayersLoadedInEditor](UDataLayerInstance* DataLayerInstance)
+	{
+		if (DataLayerInstance->IsLoadedInEditorChangedByUserOperation())
+		{
+			if (!DataLayerInstance->IsLoadedInEditor() && DataLayerInstance->IsInitiallyLoadedInEditor())
+			{
+				OutDataLayersNotLoadedInEditor.Add(DataLayerInstance->GetDataLayerFName());
+			}
+			else if (DataLayerInstance->IsLoadedInEditor() && !DataLayerInstance->IsInitiallyLoadedInEditor())
+			{
+				OutDataLayersLoadedInEditor.Add(DataLayerInstance->GetDataLayerFName());
+			}
+
+			DataLayerInstance->ClearLoadedInEditorChangedByUserOperation();
+		}
+		else
+		{
+			if (SettingsDataLayersNotLoadedInEditor.Contains(DataLayerInstance->GetDataLayerFName()))
+			{
+				OutDataLayersNotLoadedInEditor.Add(DataLayerInstance->GetDataLayerFName());
+			}
+			else if (SettingsDataLayersLoadedInEditor.Contains(DataLayerInstance->GetDataLayerFName()))
+			{
+				OutDataLayersLoadedInEditor.Add(DataLayerInstance->GetDataLayerFName());
+			}
+		}
+
+		return true;
+	});
+}
+#endif
+
 const TSet<FName>& UDataLayerSubsystem::GetEffectiveActiveDataLayerNames() const
 {
 	static TSet<FName> EmptySet;
@@ -80,38 +160,31 @@ const TSet<FName>& UDataLayerSubsystem::GetEffectiveLoadedDataLayerNames() const
 	return WorldDataLayers ? WorldDataLayers->GetEffectiveLoadedDataLayerNames() : EmptySet;
 }
 
-UDataLayerInstance* UDataLayerSubsystem::GetDataLayerFromAsset(const UDataLayerAsset* InDataLayerAsset) const
+UDataLayerInstance* UDataLayerSubsystem::GetDataLayerInstanceFromAsset(const UDataLayerAsset* InDataLayerAsset) const
 {
-	const AWorldDataLayers* WorldDataLayers = GetWorld()->GetWorldDataLayers();
-	return WorldDataLayers ? const_cast<UDataLayerInstance*>(WorldDataLayers->GetDataLayerInstance(InDataLayerAsset)) : nullptr;
+	return GetDataLayerInstance(InDataLayerAsset);
 }
 
 void UDataLayerSubsystem::SetDataLayerInstanceRuntimeState(const UDataLayerAsset* InDataLayerAsset, EDataLayerRuntimeState InState, bool bInIsRecursive)
 {
-	SetDataLayerRuntimeState(GetDataLayerFromAsset(InDataLayerAsset), InState, bInIsRecursive);
+	SetDataLayerRuntimeState(GetDataLayerInstanceFromAsset(InDataLayerAsset), InState, bInIsRecursive);
 }
 
 EDataLayerRuntimeState UDataLayerSubsystem::GetDataLayerInstanceRuntimeState(const UDataLayerAsset* InDataLayerAsset) const
 {
-	return GetDataLayerRuntimeState(GetDataLayerFromAsset(InDataLayerAsset));
+	return GetDataLayerRuntimeState(GetDataLayerInstanceFromAsset(InDataLayerAsset));
 }
 
 EDataLayerRuntimeState UDataLayerSubsystem::GetDataLayerInstanceEffectiveRuntimeState(const UDataLayerAsset* InDataLayerAsset) const
 {
-	return GetDataLayerEffectiveRuntimeState(GetDataLayerFromAsset(InDataLayerAsset));
-}
-
-UDataLayerInstance* UDataLayerSubsystem::GetDataLayerInstance(const FName& InDataLayerInstanceName) const
-{
-	const AWorldDataLayers* WorldDataLayers = GetWorld()->GetWorldDataLayers();
-	return WorldDataLayers ? const_cast<UDataLayerInstance*>(WorldDataLayers->GetDataLayerInstance(InDataLayerInstanceName)) : nullptr;
+	return GetDataLayerEffectiveRuntimeState(GetDataLayerInstanceFromAsset(InDataLayerAsset));
 }
 
 void UDataLayerSubsystem::SetDataLayerRuntimeState(const UDataLayerInstance* InDataLayerInstance, EDataLayerRuntimeState InState, bool bInIsRecursive)
 {
 	if (InDataLayerInstance)
 	{
-		if (AWorldDataLayers* WorldDataLayers = GetWorld()->GetWorldDataLayers())
+		if (AWorldDataLayers* WorldDataLayers = InDataLayerInstance->GetOuterAWorldDataLayers())
 		{
 			WorldDataLayers->SetDataLayerRuntimeState(InDataLayerInstance, InState, bInIsRecursive);
 		}
@@ -122,46 +195,36 @@ void UDataLayerSubsystem::SetDataLayerRuntimeState(const UDataLayerInstance* InD
 	}
 }
 
-EDataLayerRuntimeState UDataLayerSubsystem::GetDataLayerRuntimeState(const UDataLayerInstance* InDataLayer) const
+EDataLayerRuntimeState UDataLayerSubsystem::GetDataLayerRuntimeState(const UDataLayerInstance* InDataLayerInstance) const
 {
-	if (!InDataLayer)
-	{
-		return EDataLayerRuntimeState::Unloaded;
-	}
-
-	return GetDataLayerRuntimeStateByName(InDataLayer->GetDataLayerFName());
+	AWorldDataLayers* WorldDataLayers = InDataLayerInstance->GetOuterAWorldDataLayers();
+	return WorldDataLayers->GetDataLayerRuntimeStateByName(InDataLayerInstance->GetDataLayerFName());
 }
 
 EDataLayerRuntimeState UDataLayerSubsystem::GetDataLayerRuntimeStateByName(const FName& InDataLayerName) const
 {
-	const AWorldDataLayers* WorldDataLayers = GetWorld()->GetWorldDataLayers();
-	if (!WorldDataLayers)
-	{
-		return EDataLayerRuntimeState::Unloaded;
+	if (UDataLayerInstance* DataLayerInstance = GetDataLayerInstance(InDataLayerName))
+	{		
+		return GetDataLayerRuntimeState(DataLayerInstance);
 	}
-
-	return WorldDataLayers->GetDataLayerRuntimeStateByName(InDataLayerName);
+	
+	return EDataLayerRuntimeState::Unloaded;
 }
 
-EDataLayerRuntimeState UDataLayerSubsystem::GetDataLayerEffectiveRuntimeState(const UDataLayerInstance* InDataLayer) const
+EDataLayerRuntimeState UDataLayerSubsystem::GetDataLayerEffectiveRuntimeState(const UDataLayerInstance* InDataLayerInstance) const
 {
-	if (!InDataLayer)
-	{
-		return EDataLayerRuntimeState::Unloaded;
-	}
-
-	return GetDataLayerEffectiveRuntimeStateByName(InDataLayer->GetDataLayerFName());
+	AWorldDataLayers* WorldDataLayers = InDataLayerInstance->GetOuterAWorldDataLayers();
+	return WorldDataLayers->GetDataLayerEffectiveRuntimeStateByName(InDataLayerInstance->GetDataLayerFName());
 }
 
 EDataLayerRuntimeState UDataLayerSubsystem::GetDataLayerEffectiveRuntimeStateByName(const FName& InDataLayerName) const
 {
-	const AWorldDataLayers* WorldDataLayers = GetWorld()->GetWorldDataLayers();
-	if (!WorldDataLayers)
+	if (UDataLayerInstance* DataLayerInstance = GetDataLayerInstance(InDataLayerName))
 	{
-		return EDataLayerRuntimeState::Unloaded;
+		return GetDataLayerEffectiveRuntimeState(DataLayerInstance);
 	}
 
-	return WorldDataLayers->GetDataLayerEffectiveRuntimeStateByName(InDataLayerName);
+	return EDataLayerRuntimeState::Unloaded;
 }
 
 bool UDataLayerSubsystem::IsAnyDataLayerInEffectiveRuntimeState(const TArray<FName>& InDataLayerNames, EDataLayerRuntimeState InState) const
@@ -240,22 +303,19 @@ void UDataLayerSubsystem::DrawDataLayersStatus(UCanvas* Canvas, FVector2D& Offse
 	DrawLayerNames(TEXT("Active Data Layers"), FColor::Green, FColor::White, ActiveDataLayers);
 
 	TSet<FName> UnloadedDataLayers;
-	if (AWorldDataLayers* WorldDataLayers = GetWorld()->GetWorldDataLayers())
-	{		
-		WorldDataLayers->ForEachDataLayer([&LoadedDataLayers, &ActiveDataLayers, &UnloadedDataLayers](UDataLayerInstance* DataLayer)
+	ForEachDataLayer([&LoadedDataLayers, &ActiveDataLayers, &UnloadedDataLayers](UDataLayerInstance* DataLayer)
+	{
+		if (DataLayer->IsRuntime())
 		{
-			if (DataLayer->IsRuntime())
+			const FName DataLayerName = DataLayer->GetDataLayerFName();
+			if (!LoadedDataLayers.Contains(DataLayerName) && !ActiveDataLayers.Contains(DataLayerName))
 			{
-				const FName DataLayerName = DataLayer->GetDataLayerFName();
-				if (!LoadedDataLayers.Contains(DataLayerName) && !ActiveDataLayers.Contains(DataLayerName))
-				{
-					UnloadedDataLayers.Add(DataLayerName);
-				}
+				UnloadedDataLayers.Add(DataLayerName);
 			}
-			return true;
-		});
-		DrawLayerNames(TEXT("Unloaded Data Layers"), FColor::Silver, FColor(192,192,192), UnloadedDataLayers);
-	}
+		}
+		return true;
+	});
+	DrawLayerNames(TEXT("Unloaded Data Layers"), FColor::Silver, FColor(192, 192, 192), UnloadedDataLayers);
 
 	Offset.X += MaxTextWidth + 10;
 
@@ -358,36 +418,34 @@ TArray<UDataLayerInstance*> UDataLayerSubsystem::ConvertArgsToDataLayers(UWorld*
 		Args.Add(Builder.ToString());
 	}
 
+	
+	UDataLayerSubsystem* DataLayerSubsystem = UWorld::GetSubsystem<UDataLayerSubsystem>(World);
 	for (const FString& Arg : Args)
 	{
 		FName DataLayerName = FName(Arg);
-		if (const AWorldDataLayers* WorldDataLayers = World->GetWorldDataLayers())
+		bool bShortNameFound = false;
+		FString SanitizedDataLayerName = DataLayerName.ToString().Replace(TEXT(" "), TEXT(""));
+		DataLayerSubsystem->ForEachDataLayer([&OutDataLayers, &SanitizedDataLayerName, &bShortNameFound](UDataLayerInstance* DataLayerInstance)
 		{
-			bool bShortNameFound = false;
-			FString SanitizedDataLayerName = DataLayerName.ToString().Replace(TEXT(" "), TEXT(""));
-			WorldDataLayers->ForEachDataLayer([&SanitizedDataLayerName, &OutDataLayers, &bShortNameFound](UDataLayerInstance* It)
+			if (DataLayerInstance->GetDataLayerShortName().Find(SanitizedDataLayerName, ESearchCase::IgnoreCase) == 0)
 			{
-
-				if (It->GetDataLayerShortName().Find(SanitizedDataLayerName, ESearchCase::IgnoreCase) == 0)
-				{
-					if (bShortNameFound)
-					{
-						UE_LOG(LogWorldPartition, Error, TEXT("Found 2 data layers with the ShortName %s when converting arguments. Consider using the data layers FullName or renaming one of the two."));
-						return false;
-					}
-
-					OutDataLayers.Add(It);
-					bShortNameFound = true;
-				}
-				else if (It->GetDataLayerFullName().Find(SanitizedDataLayerName, ESearchCase::IgnoreCase) == 0)
+				if (bShortNameFound)
 				{
-					OutDataLayers.Add(It);
+					UE_LOG(LogWorldPartition, Error, TEXT("Found 2 data layers with the ShortName %s when converting arguments. Consider using the data layers FullName or renaming one of the two."));
 					return false;
 				}
-
-				return true;
-			});
-		}
+	
+				OutDataLayers.Add(DataLayerInstance);
+				bShortNameFound = true;
+			}
+			else if (DataLayerInstance->GetDataLayerFullName().Find(SanitizedDataLayerName, ESearchCase::IgnoreCase) == 0)
+			{
+				OutDataLayers.Add(DataLayerInstance);
+				return false;
+			}
+	
+			return true;
+		});
 	}
 
 	return OutDataLayers.Array();
@@ -395,10 +453,49 @@ TArray<UDataLayerInstance*> UDataLayerSubsystem::ConvertArgsToDataLayers(UWorld*
 
 void UDataLayerSubsystem::DumpDataLayers(FOutputDevice& OutputDevice) const
 {
-	if (const AWorldDataLayers* WorldDataLayers = GetWorld()->GetWorldDataLayers())
+	if (AWorldDataLayers* WorldDataLayers = GetWorld()->GetWorldDataLayers())
 	{
 		WorldDataLayers->DumpDataLayers(OutputDevice);
 	}
+	else
+	{
+		UE_LOG(LogWorldPartition, Warning, TEXT("Dumping Data Layers without a World Partition"));
+	}
+}
+
+const UDataLayerInstance* UDataLayerSubsystem::GetDataLayerInstanceFromAssetName(const FName& InDataLayerAssetFullName) const
+{
+	if (AWorldDataLayers* WorldDataLayers = GetWorld()->GetWorldDataLayers())
+	{
+		const UDataLayerInstance* DataLayerInstance = WorldDataLayers->GetDataLayerInstanceFromAssetName(InDataLayerAssetFullName);
+		return const_cast<UDataLayerInstance*>(DataLayerInstance);
+	}
+
+	return nullptr;
+}
+
+void UDataLayerSubsystem::ForEachDataLayer(TFunctionRef<bool(UDataLayerInstance*)> Func)
+{
+	if (AWorldDataLayers* WorldDataLayers = GetWorld()->GetWorldDataLayers())
+	{
+		WorldDataLayers->ForEachDataLayer(Func);
+	}
+	else
+	{
+		UE_LOG(LogWorldPartition, Warning, TEXT("Iterating over Data Layers without a World Partition"));
+	}
+}
+
+void UDataLayerSubsystem::ForEachDataLayer(TFunctionRef<bool(UDataLayerInstance*)> Func) const
+{
+	if (AWorldDataLayers* WorldDataLayers = GetWorld()->GetWorldDataLayers())
+	{
+		WorldDataLayers->ForEachDataLayer(Func);
+	}
+	else
+	{
+		UE_LOG(LogWorldPartition, Warning, TEXT("Iterating over Data Layers without a World Partition"));
+	}
 }
 
 FAutoConsoleCommand UDataLayerSubsystem::ToggleDataLayerActivation(
@@ -466,14 +563,8 @@ FAutoConsoleCommand UDataLayerSubsystem::SetDataLayerRuntimeStateCommand(
 void UDataLayerSubsystem::GetDataLayerDebugColors(TMap<FName, FColor>& OutMapping) const
 {
 	OutMapping.Reset();
-
-	const AWorldDataLayers* WorldDataLayers = GetWorld()->GetWorldDataLayers();
-	if (!WorldDataLayers)
-	{
-		return;
-	}
 	
-	WorldDataLayers->ForEachDataLayer([&OutMapping](UDataLayerInstance* DataLayer)
+	ForEachDataLayer([&OutMapping](UDataLayerInstance* DataLayer)
 	{
 		OutMapping.Add(DataLayer->GetDataLayerFName(), DataLayer->GetDebugColor());
 		return true;
@@ -481,19 +572,43 @@ void UDataLayerSubsystem::GetDataLayerDebugColors(TMap<FName, FColor>& OutMappin
 }
 
 #if WITH_EDITOR
+
+void UDataLayerSubsystem::PushActorEditorContext() const
+{
+	if (AWorldDataLayers* WorldDataLayers = GetWorld()->GetWorldDataLayers())
+	{
+		WorldDataLayers->PushActorEditorContext();
+	}
+}
+
+void UDataLayerSubsystem::PopActorEditorContext() const
+{
+	if (AWorldDataLayers* WorldDataLayers = GetWorld()->GetWorldDataLayers())
+	{
+		WorldDataLayers->PopActorEditorContext();
+	}
+}
+
+TArray<UDataLayerInstance*> UDataLayerSubsystem::GetActorEditorContextDataLayers() const
+{
+	if(AWorldDataLayers* WorldDataLayers = GetWorld()->GetWorldDataLayers())
+	{
+		return WorldDataLayers->GetActorEditorContextDataLayers();
+	}
+
+	return TArray<UDataLayerInstance*>();
+}
+
 uint32 UDataLayerSubsystem::GetDataLayerEditorContextHash() const
 {
-	if (const AWorldDataLayers* WorldDataLayers = GetWorld()->GetWorldDataLayers())
+	TArray<FName> DataLayerInstanceNames;
+	for (UDataLayerInstance* DataLayerInstance : GetActorEditorContextDataLayers())
 	{
-		TArray<FName> DataLayerInstanceNames;
-		for (UDataLayerInstance* DataLayerInstance : WorldDataLayers->GetActorEditorContextDataLayers())
-		{
-			DataLayerInstanceNames.Add(DataLayerInstance->GetDataLayerFName());
-		}
-		return FDataLayerEditorContext(GetWorld(), DataLayerInstanceNames).GetHash();
+		DataLayerInstanceNames.Add(DataLayerInstance->GetDataLayerFName());
 	}
-	return FDataLayerEditorContext::EmptyHash;
+	return FDataLayerEditorContext(GetWorld(), DataLayerInstanceNames).GetHash();
 }
+
 #endif
 
 //~ Begin Deprecated
@@ -507,8 +622,13 @@ UDataLayerInstance* UDataLayerSubsystem::GetDataLayer(const FActorDataLayer& InD
 
 UDataLayerInstance* UDataLayerSubsystem::GetDataLayerFromLabel(FName InDataLayerLabel) const
 {
-	const AWorldDataLayers* WorldDataLayers = GetWorld()->GetWorldDataLayers();
-	return WorldDataLayers ? const_cast<UDataLayerInstance*>(WorldDataLayers->GetDataLayerFromLabel(InDataLayerLabel)) : nullptr;
+	if (AWorldDataLayers* WorldDataLayers = GetWorld()->GetWorldDataLayers())
+	{
+		const UDataLayerInstance* DataLayerInstance = GetWorld()->GetWorldDataLayers()->GetDataLayerFromLabel(InDataLayerLabel);
+		return const_cast<UDataLayerInstance*>(DataLayerInstance);
+	}
+
+	return nullptr;
 }
 
 EDataLayerRuntimeState UDataLayerSubsystem::GetDataLayerRuntimeStateByLabel(const FName& InDataLayerLabel) const
@@ -549,8 +669,7 @@ PRAGMA_ENABLE_DEPRECATION_WARNINGS
 
 UDataLayerInstance* UDataLayerSubsystem::GetDataLayerFromName(FName InDataLayerName) const
 {
-	const AWorldDataLayers* WorldDataLayers = GetWorld()->GetWorldDataLayers();
-	return WorldDataLayers ? const_cast<UDataLayerInstance*>(WorldDataLayers->GetDataLayerInstance(InDataLayerName)) : nullptr;
+	return GetDataLayerInstance(InDataLayerName);
 }
 
 EDataLayerRuntimeState UDataLayerSubsystem::GetDataLayerRuntimeState(const FActorDataLayer& InDataLayer) const
diff --git a/Engine/Source/Runtime/Engine/Private/WorldPartition/DataLayer/DataLayerUtils.cpp b/Engine/Source/Runtime/Engine/Private/WorldPartition/DataLayer/DataLayerUtils.cpp
--- a/Engine/Source/Runtime/Engine/Private/WorldPartition/DataLayer/DataLayerUtils.cpp
+++ b/Engine/Source/Runtime/Engine/Private/WorldPartition/DataLayer/DataLayerUtils.cpp
@@ -4,12 +4,12 @@
 
 #if WITH_EDITOR
 #include "WorldPartition/DataLayer/WorldDataLayersActorDesc.h"
-#include "WorldPartition/DataLayer/WorldDataLayers.h"
 #include "WorldPartition/DataLayer/DataLayerInstanceWithAsset.h"
 #include "WorldPartition/DataLayer/DataLayerAsset.h"
 #include "WorldPartition/DataLayer/DataLayerType.h"
 #include "WorldPartition/WorldPartitionActorDesc.h"
 #include "WorldPartition/ActorDescContainer.h"
+#include "WorldPartition/DataLayer/DataLayerSubsystem.h"
 
 FWorldDataLayersActorDesc* FDataLayerUtils::GetWorldDataLayersActorDesc(const UActorDescContainer* InContainer, bool bInCheckValid)
 {
@@ -29,29 +29,25 @@ FWorldDataLayersActorDesc* FDataLayerUtils::GetWorldDataLayersActorDesc(const UA
 	return nullptr;
 }
 
-TArray<FName> FDataLayerUtils::ResolvedDataLayerInstanceNames(const FWorldPartitionActorDesc* InActorDesc, const AWorldDataLayers* InWorldDataLayers, const FWorldDataLayersActorDesc* InWorldDataLayersActorDesc, bool* bOutIsResultValid)
+TArray<FName> FDataLayerUtils::ResolvedDataLayerInstanceNames(const FWorldPartitionActorDesc* InActorDesc, const FWorldDataLayersActorDesc* InWorldDataLayersActorDesc, bool* bOutIsResultValid)
 {
 	bool bLocalIsSuccess = true;
 	bool& bIsSuccess = bOutIsResultValid ? *bOutIsResultValid : bLocalIsSuccess;
 	bIsSuccess = true;
 
 	// Prioritize in-memory AWorldDataLayers
-	const AWorldDataLayers* WorldDataLayers = InWorldDataLayers;
-	if (!WorldDataLayers)
-	{
-		UWorld* World = InActorDesc->GetContainer() ? InActorDesc->GetContainer()->GetWorld() : nullptr;
-		WorldDataLayers = World ? World->GetWorldDataLayers() : nullptr;
-	}
+	UWorld* World = InActorDesc->GetContainer() ? InActorDesc->GetContainer()->GetWorld() : nullptr;
+	const UDataLayerSubsystem* DataLayerSubsystem = World ? World->GetSubsystem<UDataLayerSubsystem>() : nullptr;
 
 	// DataLayers not using DataLayer Assets represent DataLayerInstanceNames
 	if (!InActorDesc->IsUsingDataLayerAsset())
 	{
-		if (WorldDataLayers)
+		if (DataLayerSubsystem)
 		{
 			TArray<FName> Result;
 			for (const FName& DataLayerInstanceName : InActorDesc->DataLayers)
 			{
-				if (const UDataLayerInstance* DataLayerInstance = WorldDataLayers->GetDataLayerInstance(DataLayerInstanceName))
+				if (const UDataLayerInstance* DataLayerInstance = DataLayerSubsystem->GetDataLayerInstance(DataLayerInstanceName))
 				{
 					Result.Add(DataLayerInstanceName);
 				}
@@ -75,12 +71,12 @@ TArray<FName> FDataLayerUtils::ResolvedDataLayerInstanceNames(const FWorldPartit
 	// ActorDesc DataLayers represents DataLayer Asset Paths
 	else
 	{
-		if (WorldDataLayers)
+		if (DataLayerSubsystem)
 		{
 			TArray<FName> Result;
 			for (const FName& DataLayerAssetPath : InActorDesc->DataLayers)
 			{
-				WorldDataLayers->ForEachDataLayer([DataLayerAssetPath, &Result](UDataLayerInstance* DataLayerInstance)
+				DataLayerSubsystem->ForEachDataLayer([DataLayerAssetPath, &Result](UDataLayerInstance* DataLayerInstance)
 				{
 					const UDataLayerInstanceWithAsset* DataLayerInstanceWithAsset = Cast<UDataLayerInstanceWithAsset>(DataLayerInstance);
 					const UDataLayerAsset* DataLayerAsset = DataLayerInstanceWithAsset ? DataLayerInstanceWithAsset->GetAsset() : nullptr;
@@ -131,12 +127,12 @@ TArray<FName> FDataLayerUtils::ResolveRuntimeDataLayerInstanceNames(const FWorld
 	if (Container)
 	{
 		UWorld* World = Container->GetWorld();
-		if (AWorldDataLayers* WorldDataLayers = World ? World->GetWorldDataLayers() : nullptr)
+		if (const UDataLayerSubsystem* DataLayerSubsystem = World ? World->GetSubsystem<UDataLayerSubsystem>() : nullptr)
 		{
 			bIsSuccess = true;
 			for (FName DataLayerInstanceName : InActorDesc->GetDataLayerInstanceNames())
 			{
-				const UDataLayerInstance* DataLayerInstance = WorldDataLayers->GetDataLayerInstance(DataLayerInstanceName);
+				const UDataLayerInstance* DataLayerInstance = DataLayerSubsystem->GetDataLayerInstance(DataLayerInstanceName);
 				if (DataLayerInstance && DataLayerInstance->IsRuntime())
 				{
 					Result.Add(DataLayerInstanceName);
diff --git a/Engine/Source/Runtime/Engine/Private/WorldPartition/DataLayer/WorldDataLayers.cpp b/Engine/Source/Runtime/Engine/Private/WorldPartition/DataLayer/WorldDataLayers.cpp
--- a/Engine/Source/Runtime/Engine/Private/WorldPartition/DataLayer/WorldDataLayers.cpp
+++ b/Engine/Source/Runtime/Engine/Private/WorldPartition/DataLayer/WorldDataLayers.cpp
@@ -414,44 +414,6 @@ TUniquePtr<class FWorldPartitionActorDesc> AWorldDataLayers::CreateClassActorDes
 	return TUniquePtr<FWorldPartitionActorDesc>(new FWorldDataLayersActorDesc());
 }
 
-void AWorldDataLayers::GetUserLoadedInEditorStates(TArray<FName>& OutDataLayersLoadedInEditor, TArray<FName>& OutDataLayersNotLoadedInEditor) const
-{
-	const TArray<FName>& SettingsDataLayersNotLoadedInEditor = GetMutableDefault<UWorldPartitionEditorPerProjectUserSettings>()->GetWorldDataLayersNotLoadedInEditor(GetWorld());
-	const TArray<FName>& SettingsDataLayersLoadedInEditor = GetMutableDefault<UWorldPartitionEditorPerProjectUserSettings>()->GetWorldDataLayersLoadedInEditor(GetWorld());
-
-	OutDataLayersNotLoadedInEditor.Empty();
-	OutDataLayersLoadedInEditor.Empty();
-	ForEachDataLayer([&OutDataLayersNotLoadedInEditor, &OutDataLayersLoadedInEditor, &SettingsDataLayersNotLoadedInEditor, &SettingsDataLayersLoadedInEditor](UDataLayerInstance* DataLayer)
-	{
-		if (DataLayer->IsLoadedInEditorChangedByUserOperation())
-		{
-			if (!DataLayer->IsLoadedInEditor() && DataLayer->IsInitiallyLoadedInEditor())
-			{
-				OutDataLayersNotLoadedInEditor.Add(DataLayer->GetDataLayerFName());
-			}
-			else if (DataLayer->IsLoadedInEditor() && !DataLayer->IsInitiallyLoadedInEditor())
-			{
-				OutDataLayersLoadedInEditor.Add(DataLayer->GetDataLayerFName());
-			}
-			
-			DataLayer->ClearLoadedInEditorChangedByUserOperation();
-		}
-		else
-		{
-			if (SettingsDataLayersNotLoadedInEditor.Contains(DataLayer->GetDataLayerFName()))
-			{
-				OutDataLayersNotLoadedInEditor.Add(DataLayer->GetDataLayerFName());
-			}
-			else if (SettingsDataLayersLoadedInEditor.Contains(DataLayer->GetDataLayerFName()))
-			{
-				OutDataLayersLoadedInEditor.Add(DataLayer->GetDataLayerFName());
-			}
-		}
-		return true;
-	});
-}
-
-
 AWorldDataLayers* AWorldDataLayers::Create(UWorld* World, FName InWorldDataLayerName)
 {
 	check(World);
diff --git a/Engine/Source/Runtime/Engine/Private/WorldPartition/WorldPartition.cpp b/Engine/Source/Runtime/Engine/Private/WorldPartition/WorldPartition.cpp
--- a/Engine/Source/Runtime/Engine/Private/WorldPartition/WorldPartition.cpp
+++ b/Engine/Source/Runtime/Engine/Private/WorldPartition/WorldPartition.cpp
@@ -43,6 +43,7 @@
 #include "WorldPartition/WorldPartitionMiniMapHelper.h"
 #include "WorldPartition/DataLayer/DataLayerSubsystem.h"
 #include "WorldPartition/DataLayer/DataLayerInstance.h"
+#include "WorldPartition/DataLayer/DataLayerSubsystem.h"
 #include "WorldPartition/DataLayer/WorldDataLayers.h"
 #include "WorldPartition/WorldPartitionActorDescViewProxy.h"
 #include "WorldPartition/HLOD/HLODLayer.h"
@@ -64,14 +65,13 @@ static FAutoConsoleVariableRef CVarGDefaultLoadingRangeHLOD0(
 TMap<FName, FString> GetDataLayersDumpString(const UWorldPartition* WorldPartition)
 {
 	TMap<FName, FString> DataLayersDumpString;
-	if (const AWorldDataLayers* WorldDataLayers = WorldPartition->GetWorld()->GetWorldDataLayers())
+	const UDataLayerSubsystem* DataLayerSubsystem = UWorld::GetSubsystem<UDataLayerSubsystem>(WorldPartition->GetWorld());
+	DataLayerSubsystem->ForEachDataLayer([&DataLayersDumpString](const UDataLayerInstance* DataLayer)
 	{
-		WorldDataLayers->ForEachDataLayer([&DataLayersDumpString](const UDataLayerInstance* DataLayer)
-		{
-			DataLayersDumpString.FindOrAdd(DataLayer->GetDataLayerFName()) = FString::Format(TEXT("{0}({1})"), { DataLayer->GetDataLayerShortName(), DataLayer->GetDataLayerFName().ToString() });
-			return true;
-		});
-	}
+		DataLayersDumpString.FindOrAdd(DataLayer->GetDataLayerFName()) = FString::Format(TEXT("{0}({1})"), { DataLayer->GetDataLayerShortName(), DataLayer->GetDataLayerFName().ToString() });
+		return true;
+	});
+	
 	return DataLayersDumpString;
 }
 
@@ -799,7 +799,7 @@ UWorldPartition* UWorldPartition::CreateOrRepairWorldPartition(AWorldSettings* W
 		WorldSettings->MarkPackageDirty();
 
 		WorldPartition->DefaultHLODLayer = UHLODLayer::GetEngineDefaultHLODLayersSetup();
-
+		
 		AWorldDataLayers* WorldDataLayers = World->GetWorldDataLayers();
 		if (!WorldDataLayers)
 		{
@@ -1025,41 +1025,40 @@ bool UWorldPartition::ShouldActorBeLoadedByEditorCells(const FWorldPartitionActo
 		return false;
 	}
 
-	if (const AWorldDataLayers* WorldDataLayers = GetWorld()->GetWorldDataLayers())
-	{
-		// Use DataLayers of loaded/dirty Actor if available to handle dirtied actors
-		FWorldPartitionActorViewProxy ActorDescProxy(ActorDesc);
+	UDataLayerSubsystem* DataLayerSubsystem = UWorld::GetSubsystem<UDataLayerSubsystem>(GetWorld());
+	
+	// Use DataLayers of loaded/dirty Actor if available to handle dirtied actors
+	FWorldPartitionActorViewProxy ActorDescProxy(ActorDesc);
 
-		if (IsRunningCookCommandlet())
+	if (IsRunningCookCommandlet())
+	{
+		// When running cook commandlet, dont allow loading of actors with dynamically loaded data layers
+		for (const FName& DataLayerInstanceName : ActorDescProxy.GetDataLayers())
 		{
-			// When running cook commandlet, dont allow loading of actors with dynamically loaded data layers
-			for (const FName& DataLayerInstanceName : ActorDescProxy.GetDataLayers())
+			const UDataLayerInstance* DataLayerInstance = DataLayerSubsystem->GetDataLayerInstance(DataLayerInstanceName);
+			if (DataLayerInstance && DataLayerInstance->IsRuntime())
 			{
-				const UDataLayerInstance* DataLayerInstance = WorldDataLayers->GetDataLayerInstance(DataLayerInstanceName);
-				if (DataLayerInstance && DataLayerInstance->IsRuntime())
-				{
-					return false;
-				}
+				return false;
 			}
 		}
-		else
+	}
+	else
+	{
+		uint32 NumValidLayers = 0;
+		for (const FName& DataLayerInstanceName : ActorDescProxy.GetDataLayers())
 		{
-			uint32 NumValidLayers = 0;
-			for (const FName& DataLayerInstanceName : ActorDescProxy.GetDataLayers())
+			if (const UDataLayerInstance* DataLayerInstance = DataLayerSubsystem->GetDataLayerInstance(DataLayerInstanceName))
 			{
-				if (const UDataLayerInstance* DataLayerInstance = WorldDataLayers->GetDataLayerInstance(DataLayerInstanceName))
+				if (DataLayerInstance->IsEffectiveLoadedInEditor())
 				{
-					if (DataLayerInstance->IsEffectiveLoadedInEditor())
-					{
-						return true;
-					}
-					NumValidLayers++;
+					return true;
 				}
+				NumValidLayers++;
 			}
-			return !NumValidLayers;
 		}
+		return !NumValidLayers;
 	}
-
+	
 	return true;
 };
 
diff --git a/Engine/Source/Runtime/Engine/Private/WorldPartition/WorldPartitionActorCluster.cpp b/Engine/Source/Runtime/Engine/Private/WorldPartition/WorldPartitionActorCluster.cpp
--- a/Engine/Source/Runtime/Engine/Private/WorldPartition/WorldPartitionActorCluster.cpp
+++ b/Engine/Source/Runtime/Engine/Private/WorldPartition/WorldPartitionActorCluster.cpp
@@ -9,18 +9,18 @@
 #include "WorldPartition/WorldPartition.h"
 #include "WorldPartition/ActorDescContainer.h"
 #include "WorldPartition/WorldPartitionActorDescView.h"
-#include "WorldPartition/DataLayer/WorldDataLayers.h"
+#include "WorldPartition/DataLayer/DataLayerSubsystem.h"
 #include "WorldPartition/DataLayer/DataLayerInstance.h"
 
 template<class LayerNameContainer>
 TSet<const UDataLayerInstance*> GetDataLayerInstances(UWorld* InWorld, const LayerNameContainer& DataLayerInstanceNames)
 {
 	TSet<const UDataLayerInstance*> DataLayerInstances;
-	if (const AWorldDataLayers* WorldDataLayers = InWorld->GetWorldDataLayers())
+	if (const UDataLayerSubsystem* DataLayerSubsystem = UWorld::GetSubsystem<UDataLayerSubsystem>(InWorld))
 	{
 		for (const FName& DataLayerInstanceName : DataLayerInstanceNames)
 		{
-			if (const UDataLayerInstance* DataLayerInstance = WorldDataLayers->GetDataLayerInstance(DataLayerInstanceName))
+			if (const UDataLayerInstance* DataLayerInstance = DataLayerSubsystem->GetDataLayerInstance(DataLayerInstanceName))
 			{
 				if (ensure(DataLayerInstance->IsRuntime()))
 				{
diff --git a/Engine/Source/Runtime/Engine/Private/WorldPartition/WorldPartitionActorDesc.cpp b/Engine/Source/Runtime/Engine/Private/WorldPartition/WorldPartitionActorDesc.cpp
--- a/Engine/Source/Runtime/Engine/Private/WorldPartition/WorldPartitionActorDesc.cpp
+++ b/Engine/Source/Runtime/Engine/Private/WorldPartition/WorldPartitionActorDesc.cpp
@@ -16,7 +16,7 @@
 #include "WorldPartition/WorldPartitionLog.h"
 #include "WorldPartition/ActorDescContainer.h"
 #include "WorldPartition/HLOD/HLODLayer.h"
-#include "WorldPartition/DataLayer/WorldDataLayers.h"
+#include "WorldPartition/DataLayer/DataLayerSubsystem.h"
 #include "WorldPartition/DataLayer/DataLayerAsset.h"
 #include "WorldPartition/DataLayer/DataLayerUtils.h"
 #include "Engine/Public/ActorReferencesUtils.h"
@@ -59,22 +59,21 @@ void FWorldPartitionActorDesc::Init(const AActor* InActor)
 	{
 		TArray<FName> LocalDataLayerAssetPaths;
 		TArray<FName> LocalDataLayerInstanceNames;
-		const AWorldDataLayers* WorldDataLayers = InActor->GetWorld()->GetWorldDataLayers();
-		if (WorldDataLayers)
+		UDataLayerSubsystem* DataLayerSubsystem = UWorld::GetSubsystem<UDataLayerSubsystem>(InActor->GetWorld());
+		
+		LocalDataLayerAssetPaths.Reserve(InActor->GetDataLayerAssets().Num());
+		for (const TObjectPtr<const UDataLayerAsset>& DataLayerAsset : InActor->GetDataLayerAssets())
 		{
-			LocalDataLayerAssetPaths.Reserve(InActor->GetDataLayerAssets().Num());
-			for (const TObjectPtr<const UDataLayerAsset>& DataLayerAsset : InActor->GetDataLayerAssets())
+			if (DataLayerAsset && DataLayerSubsystem->GetDataLayerInstance(DataLayerAsset))
 			{
-				if (DataLayerAsset && WorldDataLayers->GetDataLayerInstance(DataLayerAsset))
-				{
-					LocalDataLayerAssetPaths.Add(*DataLayerAsset->GetPathName());
-				}
+				LocalDataLayerAssetPaths.Add(*DataLayerAsset->GetPathName());
 			}
-
-			PRAGMA_DISABLE_DEPRECATION_WARNINGS
-			LocalDataLayerInstanceNames = WorldDataLayers->GetDataLayerInstanceNames(InActor->GetActorDataLayers());
-			PRAGMA_ENABLE_DEPRECATION_WARNINGS
 		}
+
+		PRAGMA_DISABLE_DEPRECATION_WARNINGS
+		LocalDataLayerInstanceNames = DataLayerSubsystem->GetDataLayerInstanceNames(InActor->GetActorDataLayers());
+		PRAGMA_ENABLE_DEPRECATION_WARNINGS
+		
 		// Validation
 		const bool bHasDataLayerAssets = LocalDataLayerAssetPaths.Num() > 0;
 		const bool bHasDeprecatedDataLayers = LocalDataLayerInstanceNames.Num() > 0;
@@ -85,7 +84,7 @@ void FWorldPartitionActorDesc::Init(const AActor* InActor)
 		DataLayers = bIsUsingDataLayerAsset ? MoveTemp(LocalDataLayerAssetPaths) : MoveTemp(LocalDataLayerInstanceNames);
 
 		// Init DataLayers transient info
-		DataLayerInstanceNames = FDataLayerUtils::ResolvedDataLayerInstanceNames(this, WorldDataLayers);
+		DataLayerInstanceNames = FDataLayerUtils::ResolvedDataLayerInstanceNames(this);
 	}
 
 	ActorPackage = InActor->GetPackage()->GetFName();
diff --git a/Engine/Source/Runtime/Engine/Private/WorldPartition/WorldPartitionRuntimeCell.cpp b/Engine/Source/Runtime/Engine/Private/WorldPartition/WorldPartitionRuntimeCell.cpp
--- a/Engine/Source/Runtime/Engine/Private/WorldPartition/WorldPartitionRuntimeCell.cpp
+++ b/Engine/Source/Runtime/Engine/Private/WorldPartition/WorldPartitionRuntimeCell.cpp
@@ -5,7 +5,7 @@
 #include "WorldPartition/WorldPartitionDebugHelper.h"
 #include "WorldPartition/WorldPartitionRuntimeSpatialHash.h"
 #include "WorldPartition/WorldPartitionLevelStreamingPolicy.h"
-#include "WorldPartition/DataLayer/WorldDataLayers.h"
+#include "WorldPartition/DataLayer/DataLayerSubsystem.h"
 #include "WorldPartition/DataLayer/DataLayersID.h"
 #include "Engine/World.h"
 
@@ -51,14 +51,14 @@ void UWorldPartitionRuntimeCell::UpdateDebugName()
 	Builder += FString::Printf(TEXT("L%d_X%d_Y%d"), Coords.Z, Coords.X, Coords.Y);
 	int32 DataLayerCount = DataLayers.Num();
 
-	const AWorldDataLayers* WorldDataLayers = GetOuterUWorldPartition()->GetWorld()->GetWorldDataLayers();
+	const UDataLayerSubsystem* DataLayerSubsystem = UWorld::GetSubsystem<UDataLayerSubsystem>(GetOuterUWorldPartition()->GetWorld());
 	TArray<const UDataLayerInstance*> DataLayerObjects;
-	if (WorldDataLayers && (DataLayerCount > 0))
+	if (DataLayerCount > 0)
 	{
 		Builder += TEXT(" DL[");
 		for (int i = 0; i < DataLayerCount; ++i)
 		{
-			const UDataLayerInstance* DataLayer = WorldDataLayers->GetDataLayerInstance(DataLayers[i]);
+			const UDataLayerInstance* DataLayer = DataLayerSubsystem->GetDataLayerInstance(DataLayers[i]);
 			DataLayerObjects.Add(DataLayer);
 			Builder += DataLayer->GetDataLayerShortName();
 			Builder += TEXT(",");
diff --git a/Engine/Source/Runtime/Engine/Private/WorldPartition/WorldPartitionStreamingGeneration.cpp b/Engine/Source/Runtime/Engine/Private/WorldPartition/WorldPartitionStreamingGeneration.cpp
--- a/Engine/Source/Runtime/Engine/Private/WorldPartition/WorldPartitionStreamingGeneration.cpp
+++ b/Engine/Source/Runtime/Engine/Private/WorldPartition/WorldPartitionStreamingGeneration.cpp
@@ -9,7 +9,7 @@
 #include "WorldPartition/WorldPartitionRuntimeHash.h"
 #include "WorldPartition/WorldPartitionStreamingPolicy.h"
 #include "WorldPartition/WorldPartitionActorCluster.h"
-#include "WorldPartition/DataLayer/WorldDataLayers.h"
+#include "WorldPartition/DataLayer/DataLayerSubsystem.h"
 #include "WorldPartition/ErrorHandling/WorldPartitionStreamingGenerationNullErrorHandler.h"
 #include "WorldPartition/ErrorHandling/WorldPartitionStreamingGenerationLogErrorHandler.h"
 #include "WorldPartition/ErrorHandling/WorldPartitionStreamingGenerationMapCheckErrorHandler.h"
@@ -431,14 +431,12 @@ class FWorldPartitionStreamingGenerator
 			FContainerDescriptor& ContainerDescriptor = ContainerIt.Value();
 			if (ContainerID.IsMainContainer() && ContainerDescriptor.Container->GetWorld())
 			{
-				if (AWorldDataLayers* WorldDataLayers = ContainerDescriptor.Container->GetWorld()->GetWorldDataLayers())
+				UDataLayerSubsystem* DataLayerSubsystem = UWorld::GetSubsystem<UDataLayerSubsystem>(ContainerDescriptor.Container->GetWorld());
+				DataLayerSubsystem->ForEachDataLayer([this](const UDataLayerInstance* DataLayerInstance)
 				{
-					WorldDataLayers->ForEachDataLayer([this](const UDataLayerInstance* DataLayerInstance) 
-					{
-						DataLayerInstance->Validate(ErrorHandler);
-						return true;
-					});
-				}
+					DataLayerInstance->Validate(ErrorHandler);
+					return true;
+				});
 
 				break;
 			}
diff --git a/Engine/Source/Runtime/Engine/Private/WorldSettings.cpp b/Engine/Source/Runtime/Engine/Private/WorldSettings.cpp
--- a/Engine/Source/Runtime/Engine/Private/WorldSettings.cpp
+++ b/Engine/Source/Runtime/Engine/Private/WorldSettings.cpp
@@ -39,7 +39,7 @@
 #include "Settings/EditorExperimentalSettings.h"
 #include "Landscape.h"
 #include "Rendering/StaticLightingSystemInterface.h"
-#include "WorldPartition/DataLayer/WorldDataLayers.h"
+#include "WorldPartition/DataLayer/DataLayerSubsystem.h"
 #endif 
 
 #define LOCTEXT_NAMESPACE "ErrorChecking"
@@ -512,10 +512,9 @@ void AWorldSettings::SaveDefaultWorldPartitionSettings()
 	if (WorldPartition)
 	{
 		DefaultWorldPartitionSettings.LoadedEditorGridCells = WorldPartition->GetUserLoadedEditorGridCells();		
-		if (AWorldDataLayers* WorldDataLayers = GetWorld()->GetWorldDataLayers())
-		{
-			WorldDataLayers->GetUserLoadedInEditorStates(DefaultWorldPartitionSettings.LoadedDataLayers, DefaultWorldPartitionSettings.NotLoadedDataLayers);
-		}
+
+		const UDataLayerSubsystem* DataLayerSubsystem =  UWorld::GetSubsystem<UDataLayerSubsystem>(GetWorld());
+		DataLayerSubsystem->GetUserLoadedInEditorStates(DefaultWorldPartitionSettings.LoadedDataLayers, DefaultWorldPartitionSettings.NotLoadedDataLayers);
 	}
 }
 
diff --git a/Engine/Source/Runtime/Engine/Public/WorldPartition/DataLayer/DataLayerSubsystem.h b/Engine/Source/Runtime/Engine/Public/WorldPartition/DataLayer/DataLayerSubsystem.h
--- a/Engine/Source/Runtime/Engine/Public/WorldPartition/DataLayer/DataLayerSubsystem.h
+++ b/Engine/Source/Runtime/Engine/Public/WorldPartition/DataLayer/DataLayerSubsystem.h
@@ -6,6 +6,7 @@
 #include "WorldPartition/DataLayer/ActorDataLayer.h"
 #include "WorldPartition/DataLayer/DataLayerInstance.h"
 #include "WorldPartition/DataLayer/DataLayerEditorContext.h"
+#include "WorldPartition/DataLayer/WorldDataLayers.h"
 #include "DataLayerSubsystem.generated.h"
 
 /**
@@ -51,7 +52,7 @@ public:
 
 	/** Find a Data Layer by its asset. */
 	UFUNCTION(BlueprintCallable, Category = DataLayers)
-	UDataLayerInstance* GetDataLayerFromAsset(const UDataLayerAsset* InDataLayerAsset) const;
+	UDataLayerInstance* GetDataLayerInstanceFromAsset(const UDataLayerAsset* InDataLayerAsset) const;
 
 	UFUNCTION(BlueprintCallable, Category = DataLayers)
 	EDataLayerRuntimeState GetDataLayerInstanceRuntimeState(const UDataLayerAsset* InDataLayerAsset) const;
@@ -69,7 +70,33 @@ public:
 
 	//~ End Blueprint callable functions
 
-	UDataLayerInstance* GetDataLayerInstance(const FName& InDataLayerInstanceName) const;
+#if WITH_EDITOR
+	bool RemoveDataLayer(const UDataLayerInstance* InDataLayer);
+	bool RemoveDataLayers(const TArray<UDataLayerInstance*>& InDataLayerInstances);
+
+	void UpdateDataLayerEditorPerProjectUserSettings() const;
+	void GetUserLoadedInEditorStates(TArray<FName>& OutDataLayersLoadedInEditor, TArray<FName>& OutDataLayersNotLoadedInEditor) const;
+
+	void PushActorEditorContext() const;
+	void PopActorEditorContext() const;
+	TArray<UDataLayerInstance*> GetActorEditorContextDataLayers() const;
+	uint32 GetDataLayerEditorContextHash() const;
+#endif
+
+	template<class T>
+	UDataLayerInstance* GetDataLayerInstance(const T& InDataLayerIdentifier) const;
+
+	template<class T>
+	TArray<const UDataLayerInstance*> GetDataLayerInstances(const TArray<T>& InDataLayerIdentifiers) const;
+
+	template<class T>
+	TArray<FName> GetDataLayerInstanceNames(const TArray<T>& InDataLayerIdentifiers) const;
+
+	const UDataLayerInstance* GetDataLayerInstanceFromAssetName(const FName& InDataLayerAssetFullName) const;
+
+	void ForEachDataLayer(TFunctionRef<bool(UDataLayerInstance*)> Func);
+	void ForEachDataLayer(TFunctionRef<bool(UDataLayerInstance*)> Func) const;
+
 	void SetDataLayerRuntimeState(const UDataLayerInstance* InDataLayerInstance, EDataLayerRuntimeState InState, bool bInIsRecursive = false);
 	EDataLayerRuntimeState GetDataLayerRuntimeState(const UDataLayerInstance* InDataLayer) const;
 	EDataLayerRuntimeState GetDataLayerRuntimeStateByName(const FName& InDataLayerName) const;
@@ -162,13 +189,6 @@ public:
 
 	//~ End Deprecated
 
-#if WITH_EDITOR
-public:
-	uint32 GetDataLayerEditorContextHash() const;
-
-private:
-#endif
-
 private:
 
 	/** Console command used to toggle activation of a DataLayer */
@@ -179,4 +199,37 @@ private:
 
 	/** Console command used to set Runtime DataLayer state*/
 	static class FAutoConsoleCommand SetDataLayerRuntimeStateCommand;
-};
\ No newline at end of file
+};
+
+template<class T>
+UDataLayerInstance* UDataLayerSubsystem::GetDataLayerInstance(const T& InDataLayerIdentifier) const
+{
+	PRAGMA_DISABLE_DEPRECATION_WARNINGS
+	const UDataLayerInstance * DataLayerInstance = GetWorld()->GetWorldDataLayers()->GetDataLayerInstance(InDataLayerIdentifier);
+	PRAGMA_ENABLE_DEPRECATION_WARNINGS
+	return const_cast<UDataLayerInstance*>(DataLayerInstance);
+}
+
+template<class T>
+TArray<FName> UDataLayerSubsystem::GetDataLayerInstanceNames(const TArray<T>& InDataLayerIdentifiers) const
+{
+	TArray<FName> DataLayerInstanceNames;
+	DataLayerInstanceNames.Reserve(InDataLayerIdentifiers.Num());
+
+	PRAGMA_DISABLE_DEPRECATION_WARNINGS
+	DataLayerInstanceNames = GetWorld()->GetWorldDataLayers()->GetDataLayerInstanceNames(InDataLayerIdentifiers);
+	PRAGMA_ENABLE_DEPRECATION_WARNINGS
+	return DataLayerInstanceNames;
+}
+
+template<class T>
+TArray<const UDataLayerInstance*> UDataLayerSubsystem::GetDataLayerInstances(const TArray<T>& InDataLayerIdentifiers) const
+{
+	TArray<const UDataLayerInstance*> DataLayerInstances;
+	DataLayerInstances.Reserve(InDataLayerIdentifiers.Num());
+
+	PRAGMA_DISABLE_DEPRECATION_WARNINGS
+	DataLayerInstances = GetWorld()->GetWorldDataLayers()->GetDataLayerInstances(InDataLayerIdentifiers);
+	PRAGMA_ENABLE_DEPRECATION_WARNINGS
+	return DataLayerInstances;
+}
\ No newline at end of file
diff --git a/Engine/Source/Runtime/Engine/Public/WorldPartition/DataLayer/DataLayerUtils.h b/Engine/Source/Runtime/Engine/Public/WorldPartition/DataLayer/DataLayerUtils.h
--- a/Engine/Source/Runtime/Engine/Public/WorldPartition/DataLayer/DataLayerUtils.h
+++ b/Engine/Source/Runtime/Engine/Public/WorldPartition/DataLayer/DataLayerUtils.h
@@ -23,7 +23,7 @@ public:
 		return IconNameByType[static_cast<uint32>(DataLayerType)];
 	}
 	static FWorldDataLayersActorDesc* GetWorldDataLayersActorDesc(const UActorDescContainer* InContainer, bool bInCheckValid = true);
-	static TArray<FName> ResolvedDataLayerInstanceNames(const FWorldPartitionActorDesc* InActorDesc, const AWorldDataLayers* InWorldDataLayers = nullptr, const FWorldDataLayersActorDesc* InWorldDataLayersActorDesc = nullptr, bool* bOutIsResultValid = nullptr);
+	static TArray<FName> ResolvedDataLayerInstanceNames(const FWorldPartitionActorDesc* InActorDesc, const FWorldDataLayersActorDesc* InWorldDataLayersActorDesc = nullptr, bool* bOutIsResultValid = nullptr);
 	static TArray<FName> ResolveRuntimeDataLayerInstanceNames(const FWorldPartitionActorDesc* InActorDesc, const UActorDescContainer* InContainer, bool* bOutIsResultValid = nullptr);
 #endif
 
diff --git a/Engine/Source/Runtime/Engine/Public/WorldPartition/DataLayer/WorldDataLayers.h b/Engine/Source/Runtime/Engine/Public/WorldPartition/DataLayer/WorldDataLayers.h
--- a/Engine/Source/Runtime/Engine/Public/WorldPartition/DataLayer/WorldDataLayers.h
+++ b/Engine/Source/Runtime/Engine/Public/WorldPartition/DataLayer/WorldDataLayers.h
@@ -81,9 +81,6 @@ public:
 	// Allows overriding of DataLayers with PlayFromHere
 	template<class T>
 	void OverwriteDataLayerRuntimeStates(const TArray<T>* InActiveDataLayers, const TArray<T>* InLoadedDataLayers );
-
-	// Returns the DataLayer user loaded editor states
-	void GetUserLoadedInEditorStates(TArray<FName>& OutDataLayersLoadedInEditor, TArray<FName>& OutDataLayersNotLoadedInEditor) const;
 #endif
 	
 	void DumpDataLayers(FOutputDevice& OutputDevice) const;
@@ -147,6 +144,9 @@ public:
 	UE_DEPRECATED(5.1, "Convert UDataLayers to UDataLayerAsset and UDataLayerInstance using DataLayerToAssetCommandLet")
 	const UDataLayerInstance* GetDataLayerFromLabel(const FName& InDataLayerLabel) const;
 
+	UE_DEPRECATED(5.1, "Use UDataLayerSubsystem::GetUserLoadedInEditorStates insteaed")
+	void GetUserLoadedInEditorStates(TArray<FName>& OutDataLayersLoadedInEditor, TArray<FName>& OutDataLayersNotLoadedInEditor) const {}
+
 	//~ End Deprecated
 
 protected:
diff --git a/Engine/Source/Runtime/MovieSceneTracks/Private/Systems/MovieSceneDataLayerSystem.cpp b/Engine/Source/Runtime/MovieSceneTracks/Private/Systems/MovieSceneDataLayerSystem.cpp
--- a/Engine/Source/Runtime/MovieSceneTracks/Private/Systems/MovieSceneDataLayerSystem.cpp
+++ b/Engine/Source/Runtime/MovieSceneTracks/Private/Systems/MovieSceneDataLayerSystem.cpp
@@ -611,7 +611,7 @@ void UMovieSceneDataLayerSystem::UpdateDesiredStates()
 
 			for (const UDataLayerAsset* DataLayerAsset : Section->GetDataLayerAssets())
 			{
-				const UDataLayerInstance* DataLayerInstance = DataLayerSubsystem->GetDataLayerFromAsset(DataLayerAsset);
+				const UDataLayerInstance* DataLayerInstance = DataLayerSubsystem->GetDataLayerInstanceFromAsset(DataLayerAsset);
 				if (DataLayerInstance)
 				{
 					this->DesiredLayerStates->ApplyNewState(DataLayerInstance->GetDataLayerFName(), OptHBiases ? OptHBiases[Index] : 0, DesiredState, bRequiresStreamingFlush);
@@ -667,7 +667,7 @@ void UMovieSceneDataLayerSystem::BeginTrackingEntities()
 
 				for (const UDataLayerAsset* DataLayerAsset : Section->GetDataLayerAssets())
 				{
-					UDataLayerInstance* DataLayerInstance = DataLayerSubsystem->GetDataLayerFromAsset(DataLayerAsset);
+					UDataLayerInstance* DataLayerInstance = DataLayerSubsystem->GetDataLayerInstanceFromAsset(DataLayerAsset);
 					if (DataLayerInstance)
 					{
 						FPreAnimatedStateEntry Entry = PreAnimatedStorage->MakeEntry(DataLayerInstance);
