Fix bug with out params to RPCs

diff --git a/Engine/Source/Runtime/Engine/Private/NetDriver.cpp b/Engine/Source/Runtime/Engine/Private/NetDriver.cpp
--- a/Engine/Source/Runtime/Engine/Private/NetDriver.cpp
+++ b/Engine/Source/Runtime/Engine/Private/NetDriver.cpp
@@ -257,6 +257,7 @@ namespace NetDriverInternal
 
 				check(((uint8*)Dest - (uint8*)LocalParms) + CopySize <= Function->ParmsSize);
 
+				It->InitializeValue(Dest);
 				It->CopyCompleteValue(Dest, Out->PropAddr);
 
 				CopiedProperties.Emplace(*It, LocalParms);
