Fix child connections ViewTarget being nullptr when the linked PlayerController has no view target
#8669

diff --git a/Engine/Source/Runtime/Engine/Private/NetDriver.cpp b/Engine/Source/Runtime/Engine/Private/NetDriver.cpp
--- a/Engine/Source/Runtime/Engine/Private/NetDriver.cpp
+++ b/Engine/Source/Runtime/Engine/Private/NetDriver.cpp
@@ -4271,14 +4271,19 @@ int32 UNetDriver::ServerReplicateActors_PrepConnections( const float DeltaSecond
 			{
 				UNetConnection *Child = Connection->Children[ChildIdx];
 				APlayerController* ChildPlayerController = Child->PlayerController;
-				if ( ChildPlayerController != NULL )
-				{
-					Child->ViewTarget = ChildPlayerController->GetViewTarget();
-				}
-				else
+				AActor* DesiredChildViewTarget = Child->OwningActor;
+
+				if (ChildPlayerController)
 				{
-					Child->ViewTarget = NULL;
+					AActor* ChildViewTarget = ChildPlayerController->GetViewTarget();
+
+					if (ChildViewTarget && ChildViewTarget->GetWorld())
+					{
+						DesiredChildViewTarget = ChildViewTarget;
+					}
 				}
+
+				Child->ViewTarget = DesiredChildViewTarget;
 			}
 		}
 		else
