PR #8726: Missing Clean() of Character SavedRootMotion in MovementComp. This could cause refs from stale root motion sources to prevent GC'ing for longer than necessary, possibly indefinitely under the right circumstances.

diff --git a/Engine/Source/Runtime/Engine/Private/Components/CharacterMovementComponent.cpp b/Engine/Source/Runtime/Engine/Private/Components/CharacterMovementComponent.cpp
--- a/Engine/Source/Runtime/Engine/Private/Components/CharacterMovementComponent.cpp
+++ b/Engine/Source/Runtime/Engine/Private/Components/CharacterMovementComponent.cpp
@@ -8068,6 +8068,7 @@ bool UCharacterMovementComponent::ClientUpdatePositionAfterServerUpdate()
 			CurrentRootMotion.UpdateStateFrom(CharacterOwner->SavedRootMotion);
 			CharacterOwner->bClientResimulateRootMotionSources = false;
 		}
+		CharacterOwner->SavedRootMotion.Clear();
 
 		return false;
 	}
