Fixes issue where players with valid pawns would not be spawned via RestartPlayerAtTransform
PR #9063

diff --git a/Engine/Source/Runtime/Engine/Private/GameModeBase.cpp b/Engine/Source/Runtime/Engine/Private/GameModeBase.cpp
--- a/Engine/Source/Runtime/Engine/Private/GameModeBase.cpp
+++ b/Engine/Source/Runtime/Engine/Private/GameModeBase.cpp
@@ -1298,7 +1298,7 @@ void AGameModeBase::RestartPlayerAtTransform(AController* NewPlayer, const FTran
 		}
 	}
 
-	if (IsValid(NewPlayer->GetPawn()))
+	if (!IsValid(NewPlayer->GetPawn()))
 	{
 		FailedToRestartPlayer(NewPlayer);
 	}
