Minor change to API to avoid static analysis warning

diff --git a/Engine/Source/Runtime/Engine/Private/Materials/MaterialHLSLEmitter.cpp b/Engine/Source/Runtime/Engine/Private/Materials/MaterialHLSLEmitter.cpp
--- a/Engine/Source/Runtime/Engine/Private/Materials/MaterialHLSLEmitter.cpp
+++ b/Engine/Source/Runtime/Engine/Private/Materials/MaterialHLSLEmitter.cpp
@@ -501,7 +501,7 @@ static void GetMaterialEnvironment(EShaderPlatform InPlatform,
 	// Setup defines to map each VT stack to either 1 or 2 page table textures, depending on how many layers it uses
 	for (int i = 0; i < EmitMaterialData.VTStacks.Num(); ++i)
 	{
-		const FMaterialVirtualTextureStack& Stack = MaterialCompilationOutput.UniformExpressionSet.GetVTStacks()[i];
+		const FMaterialVirtualTextureStack& Stack = MaterialCompilationOutput.UniformExpressionSet.GetVTStack(i);
 
 		TStringBuilder<256> PageTableName;
 		PageTableName.Appendf(TEXT("VIRTUALTEXTURE_PAGETABLE_%d"), i);
diff --git a/Engine/Source/Runtime/Engine/Private/Materials/MaterialHLSLTree.cpp b/Engine/Source/Runtime/Engine/Private/Materials/MaterialHLSLTree.cpp
--- a/Engine/Source/Runtime/Engine/Private/Materials/MaterialHLSLTree.cpp
+++ b/Engine/Source/Runtime/Engine/Private/Materials/MaterialHLSLTree.cpp
@@ -703,7 +703,7 @@ uint32 AcquireVTStackIndex(
 	for (int32 Index = EmitMaterialData.VTStackHash.First(Hash.Hash); EmitMaterialData.VTStackHash.IsValid(Index); Index = EmitMaterialData.VTStackHash.Next(Index))
 	{
 		const FVTStackEntry& Entry = EmitMaterialData.VTStacks[Index];
-		if (!UniformExpressionSet.GetVTStacks()[Index].AreLayersFull() &&
+		if (!UniformExpressionSet.GetVTStack(Index).AreLayersFull() &&
 			Entry.EmitTexCoordValue == EmitTexCoordValue &&
 			Entry.EmitTexCoordValueDdx == EmitTexCoordValueDdx &&
 			Entry.EmitTexCoordValueDdy == EmitTexCoordValueDdy &&
diff --git a/Engine/Source/Runtime/Engine/Public/MaterialShared.h b/Engine/Source/Runtime/Engine/Public/MaterialShared.h
--- a/Engine/Source/Runtime/Engine/Public/MaterialShared.h
+++ b/Engine/Source/Runtime/Engine/Public/MaterialShared.h
@@ -584,6 +584,7 @@ public:
 	uint32 AddDefaultParameterValue(const UE::Shader::FValue& Value);
 
 	TConstArrayView<FMaterialVirtualTextureStack> GetVTStacks() const { return VTStacks; }
+	const FMaterialVirtualTextureStack& GetVTStack(int32 Index) const { return VTStacks[Index]; }
 	int32 AddVTStack(int32 InPreallocatedStackTextureIndex);
 	int32 AddVTLayer(int32 StackIndex, int32 TextureIndex);
 
