Add check() to make it clear to static analysis that value may only be assigned once

diff --git a/Engine/Source/Runtime/Engine/Private/Materials/MaterialInstanceConstant.cpp b/Engine/Source/Runtime/Engine/Private/Materials/MaterialInstanceConstant.cpp
--- a/Engine/Source/Runtime/Engine/Private/Materials/MaterialInstanceConstant.cpp
+++ b/Engine/Source/Runtime/Engine/Private/Materials/MaterialInstanceConstant.cpp
@@ -230,6 +230,7 @@ void UMaterialInstanceConstant::UpdateCachedData()
 
 		if (bUsingNewHLSLGenerator && bHasStaticPermutationResource)
 		{
+			check(!LocalCachedExpressionData);
 			LocalCachedExpressionData = new FMaterialCachedExpressionData();
 			LocalCachedExpressionData->Reset();
 			LocalCachedExpressionData->UpdateForCachedHLSLTree(GetCachedHLSLTree(), &LocalStaticParameters);
