Add check() to tell static analysis UTexture* isn't allowed to be nullptr

diff --git a/Engine/Source/Runtime/Engine/Private/Materials/MaterialHLSLTree.cpp b/Engine/Source/Runtime/Engine/Private/Materials/MaterialHLSLTree.cpp
--- a/Engine/Source/Runtime/Engine/Private/Materials/MaterialHLSLTree.cpp
+++ b/Engine/Source/Runtime/Engine/Private/Materials/MaterialHLSLTree.cpp
@@ -823,15 +823,17 @@ uint32 AcquireVTStackIndex(
 void FExpressionTextureSample::EmitValueShader(FEmitContext& Context, FEmitScope& Scope, const FRequestedType& RequestedType, FEmitValueShaderResult& OutResult) const
 {
 	FMaterialTextureValue TextureValue;
-	TextureExpression->GetValueObject(Context, Scope, TextureValue);
+	verify(TextureExpression->GetValueObject(Context, Scope, TextureValue));
 	UTexture* Texture = TextureValue.Texture;
+	check(Texture);
+
 	const EMaterialValueType TextureType = Texture->GetMaterialType();
 	const Shader::EValueType TexCoordType = Private::GetTexCoordType(TextureType);
 
 	TextureAddress StaticAddressX = TA_Wrap;
 	TextureAddress StaticAddressY = TA_Wrap;
 	TextureAddress StaticAddressZ = TA_Wrap;
-	if (Texture && Texture->Source.GetNumBlocks() > 1)
+	if (Texture->Source.GetNumBlocks() > 1)
 	{
 		// UDIM (multi-block) texture are forced to use wrap address mode
 		// This is important for supporting VT stacks made from UDIMs with differing number of blocks, as this requires wrapping vAddress for certain layers
@@ -844,12 +846,9 @@ void FExpressionTextureSample::EmitValueShader(FEmitContext& Context, FEmitScope
 		switch (SamplerSource)
 		{
 		case SSM_FromTextureAsset:
-			if (Texture)
-			{
-				StaticAddressX = Texture->GetTextureAddressX();
-				StaticAddressY = Texture->GetTextureAddressY();
-				StaticAddressZ = Texture->GetTextureAddressZ();
-			}
+			StaticAddressX = Texture->GetTextureAddressX();
+			StaticAddressY = Texture->GetTextureAddressY();
+			StaticAddressZ = Texture->GetTextureAddressZ();
 			break;
 		case SSM_Wrap_WorldGroupSettings:
 			StaticAddressX = TA_Wrap;
