Force outstanding static mesh compilation when saving a world to prevent construction scripts to be deferred

diff --git a/Engine/Source/Runtime/Engine/Private/World.cpp b/Engine/Source/Runtime/Engine/Private/World.cpp
--- a/Engine/Source/Runtime/Engine/Private/World.cpp
+++ b/Engine/Source/Runtime/Engine/Private/World.cpp
@@ -106,6 +106,7 @@
 	#include "Engine/LODActor.h"
 	#include "PIEPreviewDeviceProfileSelectorModule.h"
 	#include "AssetCompilingManager.h"
+	#include "StaticMeshCompiler.h"
 	#include "WorldPartition/DataLayer/WorldDataLayers.h"
 	#include "PieFixupSerializer.h"
 	#include "ActorFolder.h"
@@ -1282,6 +1283,9 @@ void UWorld::PreSaveRoot(FObjectPreSaveRootContext ObjectSaveContext)
 	{
 		PersistentLevel->DetachAttachAllActorsPackages(/*bReattach*/false);
 	}
+
+	// Flush outstanding static mesh compilation to ensure that construction scripts are properly ran and not deferred prior to saving
+	FStaticMeshCompilingManager::Get().FinishAllCompilation();
 #endif
 
 	// Update components and keep track off whether we need to clean them up afterwards.
