Remove dead code path that's not executed with !WITH_EDITOR for calculating bWorldIsPaused

diff --git a/Engine/Source/Runtime/Engine/Private/SceneView.cpp b/Engine/Source/Runtime/Engine/Private/SceneView.cpp
--- a/Engine/Source/Runtime/Engine/Private/SceneView.cpp
+++ b/Engine/Source/Runtime/Engine/Private/SceneView.cpp
@@ -2769,12 +2769,8 @@ FSceneViewFamily::FSceneViewFamily(const ConstructionValues& CVS)
 				EngineShowFlags.LOD = 1;
 			}
 
-#if WITH_EDITOR
 			// If a single frame step or toggling between Play-in-Editor and Simulate-in-Editor happened this frame, then unpause for this frame.
 			bWorldIsPaused = !(World->IsCameraMoveable() || World->bDebugFrameStepExecutedThisFrame || World->bToggledBetweenPIEandSIEThisFrame);
-#else
-			bWorldIsPaused = !World->IsCameraMoveable();
-#endif
 		}
 	}
 
