Fix deadlock in DDC fill due to the streamed audio building calling accessors that normally wait until the async compilation is finished to return

#rnx

diff --git a/Engine/Source/Runtime/Engine/Private/AudioDerivedData.cpp b/Engine/Source/Runtime/Engine/Private/AudioDerivedData.cpp
--- a/Engine/Source/Runtime/Engine/Private/AudioDerivedData.cpp
+++ b/Engine/Source/Runtime/Engine/Private/AudioDerivedData.cpp
@@ -50,6 +50,67 @@ namespace AudioCookStats
 
 #if WITH_EDITORONLY_DATA
 
+// Any thread implicated in the streamed audio platform data build must have a valid scope to be granted access
+// to properties being modified by the build itself without triggering a FinishCache.
+// Any other thread that is a consumer of the FStreamedAudioPlatformData will trigger a FinishCache when accessing
+// incomplete properties which will wait until the builder thread has finished before returning property that is ready to be read.
+class FStreamedAudioBuildScope
+{
+public:
+	FStreamedAudioBuildScope(const FStreamedAudioPlatformData* PlatformData)
+	{
+		PreviousScope = PlatformDataBeingAsyncCompiled;
+		PlatformDataBeingAsyncCompiled = PlatformData;
+	}
+
+	~FStreamedAudioBuildScope()
+	{
+		check(PlatformDataBeingAsyncCompiled);
+		PlatformDataBeingAsyncCompiled = PreviousScope;
+	}
+
+	static bool ShouldWaitOnIncompleteProperties(const FStreamedAudioPlatformData* PlatformData)
+	{
+		return PlatformDataBeingAsyncCompiled != PlatformData;
+	}
+
+private:
+	const FStreamedAudioPlatformData* PreviousScope = nullptr;
+	// Only the thread(s) compiling this platform data will have full access to incomplete properties without causing any stalls.
+	static thread_local const FStreamedAudioPlatformData* PlatformDataBeingAsyncCompiled;
+};
+
+thread_local const FStreamedAudioPlatformData* FStreamedAudioBuildScope::PlatformDataBeingAsyncCompiled = nullptr;
+
+#endif  // #ifdef WITH_EDITORONLY_DATA
+
+TIndirectArray<struct FStreamedAudioChunk>& FStreamedAudioPlatformData::GetChunks() const
+{
+#if WITH_EDITORONLY_DATA
+	if (FStreamedAudioBuildScope::ShouldWaitOnIncompleteProperties(this))
+	{
+		// For the chunks to be available, any async task need to complete first.
+		const_cast<FStreamedAudioPlatformData*>(this)->FinishCache();
+	}
+#endif
+	return const_cast<FStreamedAudioPlatformData*>(this)->Chunks;
+}
+
+int32 FStreamedAudioPlatformData::GetNumChunks() const
+{
+#if WITH_EDITORONLY_DATA
+	if (FStreamedAudioBuildScope::ShouldWaitOnIncompleteProperties(this))
+	{
+		// NumChunks is written by the caching process, any async task need to complete before we can read it.
+		const_cast<FStreamedAudioPlatformData*>(this)->FinishCache();
+	}
+#endif
+
+	return NumChunks;
+}
+
+#if WITH_EDITORONLY_DATA
+
 /*------------------------------------------------------------------------------
 Derived data key generation.
 ------------------------------------------------------------------------------*/
@@ -559,6 +620,10 @@ public:
 	{
 		TRACE_CPUPROFILER_EVENT_SCOPE(FStreamedAudioCacheDerivedDataWorker::DoWork);
 
+		// This scope will let us access any incomplete properties since we are the producer of those properties and
+		// we can't wait on ourself without causing a deadlock.
+		FStreamedAudioBuildScope StreamedAudioBuildScope(DerivedData);
+
 		TArray<uint8> RawDerivedData;
 		bool bForceRebuild = (CacheFlags & EStreamedAudioCacheFlags::ForceRebuild) != 0;
 		bool bInlineChunks = (CacheFlags & EStreamedAudioCacheFlags::InlineChunks) != 0;
diff --git a/Engine/Source/Runtime/Engine/Private/SoundWave.cpp b/Engine/Source/Runtime/Engine/Private/SoundWave.cpp
--- a/Engine/Source/Runtime/Engine/Private/SoundWave.cpp
+++ b/Engine/Source/Runtime/Engine/Private/SoundWave.cpp
@@ -200,25 +200,6 @@ void FSoundWaveData::DiscardZerothChunkData()
 	ZerothChunkData.Empty();
 }
 
-TIndirectArray<struct FStreamedAudioChunk>& FStreamedAudioPlatformData::GetChunks() const
-{
-#if WITH_EDITORONLY_DATA
-	// For the chunks to be available, any async task need to complete first.
-	const_cast<FStreamedAudioPlatformData*>(this)->FinishCache();
-#endif
-	return const_cast<FStreamedAudioPlatformData*>(this)->Chunks;
-}
-
-int32 FStreamedAudioPlatformData::GetNumChunks() const
-{
-#if WITH_EDITORONLY_DATA
-	// NumChunks is written by the caching process, any async task need to complete before we can read it.
-	const_cast<FStreamedAudioPlatformData*>(this)->FinishCache();
-#endif
-
-	return NumChunks;
-}
-
 FSoundWaveData::MaxChunkSizeResults FSoundWaveData::GetMaxChunkSizeResults() const
 {
 	MaxChunkSizeResults Results;
