Fix SetLoadedPath not setting the proper path on unsaved worlds breaking (PackageName != LoadedPath)

This was causing a bad instancing context mapping on 2nd client when starting a Multiplayer PIE session on a World Partition world

diff --git a/Engine/Source/Runtime/Engine/Private/World.cpp b/Engine/Source/Runtime/Engine/Private/World.cpp
--- a/Engine/Source/Runtime/Engine/Private/World.cpp
+++ b/Engine/Source/Runtime/Engine/Private/World.cpp
@@ -3565,7 +3565,7 @@ UWorld* UWorld::DuplicateWorldForPIE(const FString& PackageName, UWorld* OwningW
 	UPackage* PIELevelPackage = CreatePackage(*PrefixedLevelName);
 	PIELevelPackage->SetPackageFlags(PKG_PlayInEditor);
 	PIELevelPackage->SetPIEInstanceID(PIEInstanceID);
-	PIELevelPackage->SetLoadedPath(FPackagePath::FromPackageNameChecked(PackageFName));
+	PIELevelPackage->SetLoadedPath(EditorLevelPackage->GetLoadedPath());
 	PRAGMA_DISABLE_DEPRECATION_WARNINGS
 	PIELevelPackage->SetGuid( EditorLevelPackage->GetGuid() );
 	PRAGMA_ENABLE_DEPRECATION_WARNINGS
