Fixed crash when parsing datalayer console commands

diff --git a/Engine/Source/Runtime/Engine/Private/WorldPartition/DataLayer/DataLayerSubsystem.cpp b/Engine/Source/Runtime/Engine/Private/WorldPartition/DataLayer/DataLayerSubsystem.cpp
--- a/Engine/Source/Runtime/Engine/Private/WorldPartition/DataLayer/DataLayerSubsystem.cpp
+++ b/Engine/Source/Runtime/Engine/Private/WorldPartition/DataLayer/DataLayerSubsystem.cpp
@@ -445,7 +445,7 @@ TArray<UDataLayerInstance*> UDataLayerSubsystem::ConvertArgsToDataLayers(UWorld*
 			{
 				if (bShortNameFound)
 				{
-					UE_LOG(LogWorldPartition, Error, TEXT("Found 2 data layers with the ShortName %s when converting arguments. Consider using the data layers FullName or renaming one of the two."));
+					UE_LOG(LogWorldPartition, Error, TEXT("Found 2 data layers with the ShortName %s when converting arguments. Consider using the data layers FullName or renaming one of the two."), *SanitizedDataLayerName);
 					return false;
 				}
 	
