Fixed instancing of partitioned worlds (used by level instance)

diff --git a/Engine/Source/Runtime/Engine/Private/WorldPartition/WorldPartition.cpp b/Engine/Source/Runtime/Engine/Private/WorldPartition/WorldPartition.cpp
--- a/Engine/Source/Runtime/Engine/Private/WorldPartition/WorldPartition.cpp
+++ b/Engine/Source/Runtime/Engine/Private/WorldPartition/WorldPartition.cpp
@@ -463,8 +463,7 @@ void UWorldPartition::Initialize(UWorld* InWorld, const FTransform& InTransform)
 				if (bIsInstanced)
 				{
 					const FString LongActorPackageName = ActorDescIterator->GetActorPackage().ToString();
-					const FString ActorPackageName = FPaths::GetBaseFilename(LongActorPackageName);
-					const FString InstancedName = FString::Printf(TEXT("%s_InstanceOf_%s"), *LevelPackage->GetName(), *ActorPackageName);
+					const FString InstancedName = ULevel::GetExternalActorPackageInstanceName(LevelPackage->GetName(), LongActorPackageName);
 
 					InstancingContext.AddMapping(*LongActorPackageName, *InstancedName);
 
@@ -483,6 +482,10 @@ void UWorldPartition::Initialize(UWorld* InWorld, const FTransform& InTransform)
 
 	if (bIsEditor)
 	{
+		// Here we need to flush any async loading before starting any synchronous load (mixing synchronous and asynchronous load on the same package is not handled properly).
+		// Also, this will ensure that FindObject will work on an external actor that was loading asynchronously.
+		FlushAsyncLoading();
+
 		// Make sure to preload only AWorldDataLayers actor first (ShouldActorBeLoadedByEditorCells requires it)
 		for (UActorDescContainer::TIterator<> ActorDescIterator(this); ActorDescIterator; ++ActorDescIterator)
 		{
diff --git a/Engine/Source/Runtime/Engine/Private/WorldPartition/WorldPartitionActorDesc.cpp b/Engine/Source/Runtime/Engine/Private/WorldPartition/WorldPartitionActorDesc.cpp
--- a/Engine/Source/Runtime/Engine/Private/WorldPartition/WorldPartitionActorDesc.cpp
+++ b/Engine/Source/Runtime/Engine/Private/WorldPartition/WorldPartitionActorDesc.cpp
@@ -376,26 +376,17 @@ AActor* FWorldPartitionActorDesc::Load() const
 			Container->GetInstancingContext(InstancingContext, SoftObjectPathFixupArchive);
 		}
 
-		const FName PackageNameToLoad = *ActorPackage.ToString();
-		FName PackageNameToCreate = PackageNameToLoad;
+		UPackage* Package = nullptr;
+
 		if (InstancingContext)
 		{
-			PackageNameToCreate = InstancingContext->Remap(ActorPackage);
-			check(PackageNameToCreate != ActorPath);
-		}
+			FName RemappedPackageName = InstancingContext->Remap(ActorPackage);
+			check(RemappedPackageName != ActorPath);
 
-		// Here we are not using LoadPackage because we could be inside an async load and it's currently not properly handled.
-		// Instead, use the async version and flush.
-		UPackage* Package = nullptr;
-		LoadPackageAsync(FPackagePath::FromPackageNameChecked(PackageNameToLoad), PackageNameToCreate, FLoadPackageAsyncDelegate::CreateLambda([&Package](const FName&, UPackage* LoadedPackage, EAsyncLoadingResult::Type Result)
-		{
-			if (Result == EAsyncLoadingResult::Succeeded)
-			{
-				Package = LoadedPackage;
-			}
-		}), PKG_None, INDEX_NONE, 0, InstancingContext);
+			Package = CreatePackage(*RemappedPackageName.ToString());
+		}
 
-		FlushAsyncLoading();
+		Package = LoadPackage(Package, *ActorPackage.ToString(), LOAD_None, nullptr, InstancingContext);
 
 		if (Package)
 		{
