Remove ensure message for morph target count

diff --git a/Engine/Source/Runtime/Engine/Private/SkeletalRenderGPUSkin.cpp b/Engine/Source/Runtime/Engine/Private/SkeletalRenderGPUSkin.cpp
--- a/Engine/Source/Runtime/Engine/Private/SkeletalRenderGPUSkin.cpp
+++ b/Engine/Source/Runtime/Engine/Private/SkeletalRenderGPUSkin.cpp
@@ -420,12 +420,6 @@ void FSkeletalMeshObjectGPUSkin::Update(int32 LODIndex,USkinnedMeshComponent* In
 		InitMorphResources(InMeshComponent->bPerBoneMotionBlur, MorphTargetWeights);		
 	}
 
-	const FSkeletalMeshLODRenderData& LODData = SkeletalMeshRenderData->LODRenderData[LODIndex];
-	ensureAlwaysMsgf(MorphTargetWeights.Num() == LODData.MorphTargetVertexInfoBuffers.GetNumMorphs(),
-				TEXT("Skeletal mesh '%s' morph data mismatching: MorphTargetWeights=%d, [LOD%d]MorphTargetVertexInfoBuffers=%d, MorphTargets=%d"),
-				InMeshComponent->SkeletalMesh ? *InMeshComponent->SkeletalMesh->GetName() : TEXT(""), MorphTargetWeights.Num(), LODIndex, LODData.MorphTargetVertexInfoBuffers.GetNumMorphs(),
-				InMeshComponent->SkeletalMesh ? InMeshComponent->SkeletalMesh->GetMorphTargets().Num() : 0);
-
 	// create the new dynamic data for use by the rendering thread
 	// this data is only deleted when another update is sent
 	FDynamicSkelMeshObjectDataGPUSkin* NewDynamicData = FDynamicSkelMeshObjectDataGPUSkin::AllocDynamicSkelMeshObjectDataGPUSkin();		
