- Embedded Level Instance Actors are now loaded through Instancing Context instead of Duplication
- Duplicate Folder objects instead of referencing Persistent level Folder objects to support Level Instance folders properly with instancing
- Remove no longer needed ActorFolders from the RuntimeCells

#test AncientGame,CitySample (Cook, PIE, -game) + specific tests in QAGame (LevelInstanceStreaming_WP and LevelInstanceSoftObjectPath_WP)
#rnx

diff --git a/Engine/Source/Runtime/Engine/Classes/Engine/Level.h b/Engine/Source/Runtime/Engine/Classes/Engine/Level.h
--- a/Engine/Source/Runtime/Engine/Classes/Engine/Level.h
+++ b/Engine/Source/Runtime/Engine/Classes/Engine/Level.h
@@ -73,7 +73,7 @@ private:
 	friend class UWorld;
 	friend class UActorFolder;
 	friend class UWorldPartitionConvertCommandlet;
-	friend class UWorldPartitionLevelStreamingDynamic;
+	friend class FWorldPartitionLevelHelper;
 };
 
 class ENGINE_API FLevelPartitionOperationScope
@@ -780,6 +780,8 @@ public:
 	ENGINE_API static bool GetPartitionedLevelCanBeUsedByLevelInstanceFromAsset(const FAssetData& Asset);
 	ENGINE_API static bool GetPartitionedLevelCanBeUsedByLevelInstanceFromPackage(FName LevelPackage);
 	ENGINE_API static const FName LoadAllExternalObjectsTag;
+	ENGINE_API static const FName DontLoadExternalObjectsTag;
+
 	ENGINE_API bool GetPromptWhenAddingToLevelOutsideBounds() const;
 	ENGINE_API bool GetPromptWhenAddingToLevelBeforeCheckout() const;
 #endif
@@ -1341,6 +1343,7 @@ private:
 	void SetUseActorFoldersInternal(bool bInEnabled);
 
 	friend struct FLevelActorFoldersHelper;
+	friend class FWorldPartitionLevelHelper;
 
 	/** Replace the existing LSA (if set) by spawning a new one based on this level's script blueprint */
 	void RegenerateLevelScriptActor();
diff --git a/Engine/Source/Runtime/Engine/Private/ActorFolder.cpp b/Engine/Source/Runtime/Engine/Private/ActorFolder.cpp
--- a/Engine/Source/Runtime/Engine/Private/ActorFolder.cpp
+++ b/Engine/Source/Runtime/Engine/Private/ActorFolder.cpp
@@ -163,17 +163,9 @@ void UActorFolder::MarkAsDeleted()
 	bIsDeleted = true;
 }
 
-UActorFolder* UActorFolder::GetParent() const
+UActorFolder* UActorFolder::GetParent(bool bSkipDeleted) const
 {
-	UActorFolder* Parent = ParentFolderGuid.IsValid() ? GetOuterULevel()->GetActorFolder(ParentFolderGuid) : nullptr;
-	if (Parent)
-	{
-		if (!Parent->IsValid())
-		{
-			return Parent->GetParent();
-		}
-	}
-	return Parent;
+	return ParentFolderGuid.IsValid() ? GetOuterULevel()->GetActorFolder(ParentFolderGuid, bSkipDeleted) : nullptr;
 }
 
 FName UActorFolder::GetPath() const
diff --git a/Engine/Source/Runtime/Engine/Private/Level.cpp b/Engine/Source/Runtime/Engine/Private/Level.cpp
--- a/Engine/Source/Runtime/Engine/Private/Level.cpp
+++ b/Engine/Source/Runtime/Engine/Private/Level.cpp
@@ -366,6 +366,7 @@ TMap<FName, TWeakObjectPtr<UWorld> > ULevel::StreamedLevelsOwningWorld;
 
 #if WITH_EDITOR
 const FName ULevel::LoadAllExternalObjectsTag(TEXT("LoadAllExternalObjectsTag"));
+const FName ULevel::DontLoadExternalObjectsTag(TEXT("DontLoadExternalObjectsTag"));
 #endif
 
 ULevel::ULevel( const FObjectInitializer& ObjectInitializer )
@@ -944,7 +945,9 @@ void ULevel::PostLoad()
 	// if we use external actors, load dynamic actors here
 	if (IsUsingExternalActors())
 	{
-		if (!bWasDuplicated && (!bIsPartitioned || Linker->GetInstancingContext().HasTag(ULevel::LoadAllExternalObjectsTag)))
+		if (!bWasDuplicated &&
+			(!bIsPartitioned || Linker->GetInstancingContext().HasTag(ULevel::LoadAllExternalObjectsTag)) &&
+			!Linker->GetInstancingContext().HasTag(ULevel::DontLoadExternalObjectsTag))
 		{
 			UPackage* LevelPackage = GetPackage();
 			bool bPackageForPIE = LevelPackage->HasAnyPackageFlags(PKG_PlayInEditor);
diff --git a/Engine/Source/Runtime/Engine/Private/WorldPartition/WorldPartitionLevelHelper.cpp b/Engine/Source/Runtime/Engine/Private/WorldPartition/WorldPartitionLevelHelper.cpp
--- a/Engine/Source/Runtime/Engine/Private/WorldPartition/WorldPartitionLevelHelper.cpp
+++ b/Engine/Source/Runtime/Engine/Private/WorldPartition/WorldPartitionLevelHelper.cpp
@@ -5,7 +5,7 @@
  */
 
 #include "WorldPartition/WorldPartitionLevelHelper.h"
-#include "WorldPartition/WorldPartitionPackageCache.h"
+#include "WorldPartition/WorldPartitionPackageHelper.h"
 
 #if WITH_EDITOR
 
@@ -20,6 +20,36 @@
 #include "StaticMeshCompiler.h"
 #include "LevelUtils.h"
 #include "Templates/SharedPointer.h"
+#include "ActorFolder.h"
+
+TMap<FName, FWorldPartitionLevelHelper::FPackageReference> FWorldPartitionLevelHelper::PackageReferences;
+
+void FWorldPartitionLevelHelper::FPackageReferencer::AddReference(UPackage* InPackage)
+{
+	check(InPackage);
+	FPackageReference& RefInfo = FWorldPartitionLevelHelper::PackageReferences.FindOrAdd(InPackage->GetFName());
+	check(RefInfo.Package == nullptr || RefInfo.Package == InPackage);
+	RefInfo.Package = InPackage;
+	RefInfo.Referencers.Add(this);
+}
+
+void FWorldPartitionLevelHelper::FPackageReferencer::RemoveReferences()
+{
+	for (auto It = FWorldPartitionLevelHelper::PackageReferences.CreateIterator(); It; ++It)
+	{
+		FPackageReference& RefInfo = It->Value;
+		RefInfo.Referencers.Remove(this);
+		if (RefInfo.Referencers.Num() == 0)
+		{
+			if (UPackage* Package = RefInfo.Package.Get())
+			{
+				FWorldPartitionPackageHelper::UnloadPackage(Package);
+			}
+			It.RemoveCurrent();
+		}
+	}
+}
+
 
  /**
   * Defaults World's initialization values for World Partition StreamingLevels
@@ -66,7 +96,7 @@ void FWorldPartitionLevelHelper::MoveExternalActorsToLevel(const TArray<FWorldPa
 		AActor* Actor = FindObject<AActor>(nullptr, *PackageObjectMapping.LoadedPath.ToString());
 		if (Actor)
 		{
-			UPackage* ActorExternalPackage = Actor->IsPackageExternal() ? Actor->GetPackage() : nullptr;
+			UPackage* ActorExternalPackage = Actor->GetPackage();
 
 			const bool bSameOuter = (InLevel == Actor->GetOuter());
 			Actor->SetPackageExternal(false, false);
@@ -85,8 +115,7 @@ void FWorldPartitionLevelHelper::MoveExternalActorsToLevel(const TArray<FWorldPa
 			}
 
 			// Include objects found in the source actor package in the destination level package
-			// @todo_ow: Generalize this case and handle properly the case of LevelInstance (LI actor's package is the duplicated LI package, i.e. not an external actor package)
-			if (ActorExternalPackage)
+			if (ensure(ActorExternalPackage))
 			{
 				TArray<UObject*> Objects;
 				const bool bIncludeNestedSubobjects = false;
@@ -224,7 +253,39 @@ ULevel* FWorldPartitionLevelHelper::CreateEmptyLevelForRuntimeCell(const UWorld*
 	return NewLevel;
 }
 
-bool FWorldPartitionLevelHelper::LoadActors(ULevel* InDestLevel, TArrayView<FWorldPartitionRuntimeCellObjectMapping> InActorPackages, FWorldPartitionPackageCache& InPackageCache, TFunction<void(bool)> InCompletionCallback, bool bInLoadAsync, FLinkerInstancingContext* InOutInstancingContext /*=nullptr*/)
+void FWorldPartitionLevelHelper::DuplicateActorFolderToRuntimeCell(ULevel* CellLevel, ULevel* SrcLevel, const FGuid& ActorFolderGuid)
+{
+	check(SrcLevel->IsUsingActorFolders() && CellLevel->IsUsingActorFolders());
+		
+	FGuid CurrentGuid = ActorFolderGuid;
+	while (CurrentGuid.IsValid() && !CellLevel->GetActorFolder(CurrentGuid))
+	{
+		const bool bSkipDeleted = false;
+		UActorFolder* ActorFolder = SrcLevel->GetActorFolder(CurrentGuid, bSkipDeleted);
+		CurrentGuid.Invalidate();
+		if (ActorFolder)
+		{
+			CurrentGuid.Invalidate();
+			FObjectDuplicationParameters DupParams = InitStaticDuplicateObjectParams(ActorFolder, CellLevel);
+			DupParams.ApplyFlags |= RF_Transient;
+			DupParams.FlagMask &= ~RF_Transactional;
+			DupParams.bAssignExternalPackages = false;
+			UActorFolder* DuplicatedFolder = Cast<UActorFolder>(StaticDuplicateObjectEx(DupParams));
+
+			const bool bShouldDirtyLevel = false;
+			const bool bShouldBroadcast = false;
+			FLevelActorFoldersHelper::AddActorFolder(CellLevel, DuplicatedFolder, bShouldDirtyLevel, bShouldBroadcast);
+
+			// Continue to parent
+			if (UActorFolder* ParentFolder = ActorFolder->GetParent(bSkipDeleted))
+			{
+				CurrentGuid = ParentFolder->GetGuid();
+			}
+		}
+	}
+}
+
+bool FWorldPartitionLevelHelper::LoadActors(UWorld* InOwningWorld, ULevel* InDestLevel, TArrayView<FWorldPartitionRuntimeCellObjectMapping> InActorPackages, FWorldPartitionLevelHelper::FPackageReferencer& InPackageReferencer, TFunction<void(bool)> InCompletionCallback, bool bInLoadAsync, FLinkerInstancingContext InstancingContext)
 {
 	UPackage* DestPackage = InDestLevel ? InDestLevel->GetPackage() : nullptr;
 	FString ShortLevelPackageName = DestPackage? FPackageName::GetShortName(DestPackage->GetFName()) : FString();
@@ -237,47 +298,45 @@ bool FWorldPartitionLevelHelper::LoadActors(ULevel* InDestLevel, TArrayView<FWor
 	TSharedPtr<FLoadProgress> LoadProgress = MakeShared<FLoadProgress>();
 
 	// Actors to load
-	TArray<FString> ActorPackageInstanceNames;
-	ActorPackageInstanceNames.Reserve(InActorPackages.Num());
 	TArray<FWorldPartitionRuntimeCellObjectMapping*> ActorPackages;
 	ActorPackages.Reserve(InActorPackages.Num());
 
-	// Levels to Load with actors to duplicate
-	TMap<FActorContainerID, TArray<FWorldPartitionRuntimeCellObjectMapping*>> PackagesToDuplicate;
-	check(!InDestLevel || InOutInstancingContext);
-
+	TMap<FActorContainerID, FLinkerInstancingContext> LinkerInstancingContexts;
+	// Add Main container context
+	LinkerInstancingContexts.Add(FActorContainerID::GetMainContainerID(), MoveTemp(InstancingContext));
+			
 	for (FWorldPartitionRuntimeCellObjectMapping& PackageObjectMapping : InActorPackages)
 	{
-		if (PackageObjectMapping.ContainerID.IsMainContainer())
+		FLinkerInstancingContext* Context = LinkerInstancingContexts.Find(PackageObjectMapping.ContainerID);
+		if (!Context)
 		{
-			if (InOutInstancingContext)
-			{
-				FString ObjectPath = PackageObjectMapping.Package.ToString();
-				FString ActorPackageName = FPackageName::ObjectPathToPackageName(ObjectPath);
-				FString ActorPackageInstanceName = FString::Printf(TEXT("%s_%s"), *ActorPackageName, *ShortLevelPackageName);
+			check(!PackageObjectMapping.ContainerID.IsMainContainer());
+			const FName ContainerPackageInstanceName(*FString::Printf(TEXT("/Temp%s_%s"), *PackageObjectMapping.ContainerPackage.ToString(), *PackageObjectMapping.ContainerID.ToString()));
 
-				ActorPackageInstanceNames.Add(ActorPackageInstanceName);
-				InOutInstancingContext->AddMapping(FName(*ActorPackageName), FName(*ActorPackageInstanceName));
-			}
-			ActorPackages.Add(&PackageObjectMapping);
+			FLinkerInstancingContext& NewContext = LinkerInstancingContexts.Add(PackageObjectMapping.ContainerID);
+			NewContext.AddTag(ULevel::DontLoadExternalObjectsTag);
+			NewContext.AddMapping(PackageObjectMapping.ContainerPackage, ContainerPackageInstanceName);
+			Context = &NewContext;
 		}
-		else
+		
+		const FName ContainerPackageInstanceName = Context->Remap(PackageObjectMapping.ContainerPackage);
+		if (PackageObjectMapping.ContainerPackage != ContainerPackageInstanceName)
 		{
-			// When loading LevelInstances we duplicate the package 
-			PackagesToDuplicate.FindOrAdd(PackageObjectMapping.ContainerID).Add(&PackageObjectMapping);
+			const FString ActorPackageName = FPackageName::ObjectPathToPackageName(PackageObjectMapping.Package.ToString());
+			const FString ActorPackageInstanceName = ULevel::GetExternalActorPackageInstanceName(ContainerPackageInstanceName.ToString(), ActorPackageName);
+
+			Context->AddMapping(FName(*ActorPackageName), FName(*ActorPackageInstanceName));
 		}
+		ActorPackages.Add(&PackageObjectMapping);
 	}
 
-	LoadProgress->NumPendingLoadRequests = ActorPackages.Num() + PackagesToDuplicate.Num();
+	LoadProgress->NumPendingLoadRequests = ActorPackages.Num();
 
-	for (int32 ChildIndex = 0; ChildIndex < ActorPackages.Num(); ChildIndex++)
+	for (FWorldPartitionRuntimeCellObjectMapping* PackageObjectMapping : ActorPackages)
 	{
-		const FWorldPartitionRuntimeCellObjectMapping& PackageObjectMapping = *ActorPackages[ChildIndex];
-		const FString ActorPackageName = FPackageName::ObjectPathToPackageName(PackageObjectMapping.Package.ToString());
-		FName ActorName = *FPaths::GetExtension(PackageObjectMapping.Path.ToString());
-
-		FLoadPackageAsyncDelegate CompletionCallback = FLoadPackageAsyncDelegate::CreateLambda([LoadProgress, ActorName, InDestLevel, InCompletionCallback](const FName& LoadedPackageName, UPackage* LoadedPackage, EAsyncLoadingResult::Type Result)
+		FLoadPackageAsyncDelegate CompletionCallback = FLoadPackageAsyncDelegate::CreateLambda([LoadProgress, PackageObjectMapping, &InPackageReferencer, InOwningWorld, InDestLevel, InCompletionCallback](const FName& LoadedPackageName, UPackage* LoadedPackage, EAsyncLoadingResult::Type Result)
 		{
+			const FName ActorName = *FPaths::GetExtension(PackageObjectMapping->Path.ToString());
 			check(LoadProgress->NumPendingLoadRequests);
 			LoadProgress->NumPendingLoadRequests--;
 
@@ -285,13 +344,66 @@ bool FWorldPartitionLevelHelper::LoadActors(ULevel* InDestLevel, TArrayView<FWor
 
 			if (Actor)
 			{
+				const UWorld* ContainerWorld = PackageObjectMapping->ContainerID.IsMainContainer() ? InOwningWorld : Actor->GetTypedOuter<UWorld>();
+				
+				// Duplicate Folder if needed (this will recreate a transient folder structure in InDestLevel if it doesn't exist (only used in PIE)
+				if (FGuid ActorFolderGuid = Actor->GetFolderGuid(); InDestLevel && ActorFolderGuid.IsValid())
+				{
+					// Make sure Dest level is properly setup
+					if (!InDestLevel->IsUsingActorFolders())
+					{
+						FLevelActorFoldersHelper::SetUseActorFolders(InDestLevel, true);
+					}
+					InDestLevel->bFixupActorFoldersAtLoad = false;
+
+					// Make sure Source level actor folder fixup was called
+					if (!ContainerWorld->PersistentLevel->LoadedExternalActorFolders.IsEmpty())
+					{
+						ContainerWorld->PersistentLevel->bFixupActorFoldersAtLoad = false;
+						ContainerWorld->PersistentLevel->FixupActorFolders();
+					}
+					
+					FWorldPartitionLevelHelper::DuplicateActorFolderToRuntimeCell(InDestLevel, ContainerWorld->PersistentLevel, ActorFolderGuid);
+				}
+
+				if (!PackageObjectMapping->ContainerID.IsMainContainer())
+				{					
+					// Add Cache handle on world so it gets unloaded properly
+					InPackageReferencer.AddReference(ContainerWorld->GetPackage());
+										
+					FString SourceWorldPath, RemappedWorldPath;
+					ContainerWorld->GetSoftObjectPathMapping(SourceWorldPath, RemappedWorldPath);
+
+					// Rename through UObject to avoid changing Actor's external packaging and folder properties
+					Actor->UObject::Rename(*FString::Printf(TEXT("%s_%s"), *Actor->GetName(), *PackageObjectMapping->ContainerID.ToString()), InDestLevel, REN_NonTransactional | REN_ForceNoResetLoaders | REN_DoNotDirty | REN_DontCreateRedirectors);
+					
+					FLevelUtils::FApplyLevelTransformParams TransformParams(nullptr, PackageObjectMapping->ContainerTransform);
+					TransformParams.Actor = Actor;
+					TransformParams.bDoPostEditMove = false;
+					FLevelUtils::ApplyLevelTransform(TransformParams);
+						
+					// Path to use when searching for this actor in MoveExternalActorsToLevel
+					PackageObjectMapping->LoadedPath = *Actor->GetPathName();
+
+					// Fixup any FSoftObjectPath from this Actor (and its SubObjects) in this container to another object in the same container with a ContainerID suffix that can be remapped to
+					// to a Cell in the StreamingPolicy (this relies on the fact that the _DUP package doesn't get fixed up)
+					FSoftObjectPathFixupArchive FixupArchive([&](FSoftObjectPath& Value)
+					{
+						if (!Value.IsNull() && Value.GetAssetPathString().Equals(SourceWorldPath, ESearchCase::IgnoreCase))
+						{
+							Value.SetSubPathString(AddActorContainerIDToSubPathString(PackageObjectMapping->ContainerID, Value.GetSubPathString()));
+						}
+					});
+					FixupArchive.Fixup(Actor);
+				}
+
 				if (InDestLevel)
 				{
 					check(Actor->IsPackageExternal());
 					InDestLevel->Actors.Add(Actor);
 					checkf(Actor->GetLevel() == InDestLevel, TEXT("Levels mismatch, got : %s, expected: %s\nActor: %s\nActorFullName: %s\nActorPackage: %s"), *InDestLevel->GetFullName(), *Actor->GetLevel()->GetFullName(), *Actor->GetActorNameOrLabel(), *Actor->GetFullName(), *Actor->GetPackage()->GetFullName());
 				}
-								
+
 				UE_LOG(LogEngine, Verbose, TEXT(" ==> Loaded %s (remaining: %d)"), *Actor->GetFullName(), LoadProgress->NumPendingLoadRequests);
 			}
 			else
@@ -307,119 +419,32 @@ bool FWorldPartitionLevelHelper::LoadActors(ULevel* InDestLevel, TArrayView<FWor
 			}
 		});
 
+		FName PackageToLoad(*FPackageName::ObjectPathToPackageName(PackageObjectMapping->Package.ToString()));
+		const FLinkerInstancingContext& ContainerInstancingContext = LinkerInstancingContexts.FindChecked(PackageObjectMapping->ContainerID);
+		FName PackageName = ContainerInstancingContext.Remap(PackageToLoad);
+
 		if (bInLoadAsync)
 		{
-			FPackagePath PackagePath = FPackagePath::FromPackageNameChecked(*ActorPackageName);
+			FPackagePath PackagePath = FPackagePath::FromPackageNameChecked(PackageToLoad);
+
 			check(DestPackage);
 			const EPackageFlags PackageFlags = DestPackage->HasAnyPackageFlags(PKG_PlayInEditor) ? PKG_PlayInEditor : PKG_None;
-			::LoadPackageAsync(PackagePath, FName(*ActorPackageInstanceNames[ChildIndex]), CompletionCallback, PackageFlags, DestPackage->GetPIEInstanceID(), 0, InOutInstancingContext);
+			::LoadPackageAsync(PackagePath, PackageName, CompletionCallback, PackageFlags, DestPackage->GetPIEInstanceID(), 0, &ContainerInstancingContext);
 		}
 		else
 		{
 			UPackage* InstancingPackage = nullptr;
-			if(InOutInstancingContext)
-			{ 
-				InstancingPackage = CreatePackage(*ActorPackageInstanceNames[ChildIndex]);
-			}
-
-			UPackage* Package = LoadPackage(InstancingPackage, *ActorPackageName, LOAD_None, nullptr, InOutInstancingContext);
-			CompletionCallback.Execute(*ActorPackageName, Package, Package ? EAsyncLoadingResult::Succeeded : EAsyncLoadingResult::Failed);
-		}
-	}
-
-	for (auto& Pair : PackagesToDuplicate)
-	{
-		FString PackageToLoadFrom = Pair.Value[0]->ContainerPackage.ToString();
-		FLoadPackageAsyncDelegate CompletionCallback = FLoadPackageAsyncDelegate::CreateLambda([LoadProgress, ContainerInstanceID = Pair.Key, Mappings = MoveTemp(Pair.Value), InDestLevel, &InPackageCache, InCompletionCallback](const FName& LoadedPackageName, UPackage* LoadedPackage, EAsyncLoadingResult::Type Result)
-		{
-			check(LoadProgress->NumPendingLoadRequests);
-			LoadProgress->NumPendingLoadRequests--;
-
-			if (LoadedPackage)
-			{
-				FName DuplicatePackageName(*FString::Printf(TEXT("%s_%s"), *LoadedPackage->GetName(), *ContainerInstanceID.ToString()));
-				check(!FindPackage(nullptr, *DuplicatePackageName.ToString()));
-				UPackage* DuplicatedPackage = FWorldPartitionPackageHelper::DuplicateWorldPackage(LoadedPackage, DuplicatePackageName);
-				check(DuplicatedPackage);
-				UWorld* DuplicatedWorld = UWorld::FindWorldInPackage(DuplicatedPackage);
-				check(DuplicatedWorld);
-				
-				// Get Source Path for this world (ReplaceFrom)
-				UWorld* LoadedWorld = UWorld::FindWorldInPackage(LoadedPackage);
-				FString SourceWorldPath, RemappedWorldPath;
-				LoadedWorld->GetSoftObjectPathMapping(SourceWorldPath, RemappedWorldPath);
-				
-				for (auto Mapping : Mappings)
-				{
-					FString ActorName = FPaths::GetExtension(Mapping->Path.ToString());
-					AActor* DuplicatedActor = FindObject<AActor>(DuplicatedWorld->PersistentLevel, *ActorName);
-					// Possible actor isn't found if Actor package failed to load because of missing dependencies (ex: deleted Blueprint)
-					if (DuplicatedActor)
-					{
-						DuplicatedActor->Rename(*FString::Printf(TEXT("%s_%s"), *DuplicatedActor->GetName(), *Mapping->ContainerID.ToString()), InDestLevel, REN_NonTransactional | REN_ForceNoResetLoaders | REN_DoNotDirty | REN_DontCreateRedirectors);
-						USceneComponent* RootComponent = DuplicatedActor->GetRootComponent();
-
-						FLevelUtils::FApplyLevelTransformParams TransformParams(nullptr, Mapping->ContainerTransform);
-						TransformParams.Actor = DuplicatedActor;
-						TransformParams.bDoPostEditMove = false;
-						FLevelUtils::ApplyLevelTransform(TransformParams);
-
-						// Path to use when searching for this actor in MoveExternalActorsToLevel
-						Mapping->LoadedPath = *DuplicatedActor->GetPathName();
-						
-						// Fixup any FSoftObjectPath from this Actor (and its SubObjects) in this container to another object in the same container with a ContainerID suffix that can be remapped to
-						// to a Cell in the StreamingPolicy (this relies on the fact that the _DUP package doesn't get fixed up)
-						FSoftObjectPathFixupArchive FixupArchive([&](FSoftObjectPath& Value)
-						{
-							if (!Value.IsNull() && Value.GetAssetPathString().Equals(SourceWorldPath, ESearchCase::IgnoreCase))
-							{
-								Value.SetSubPathString(AddActorContainerIDToSubPathString(Mapping->ContainerID, Value.GetSubPathString()));
-							}
-						});
-						FixupArchive.Fixup(DuplicatedActor);
-
-						UE_LOG(LogEngine, Verbose, TEXT(" ==> Duplicated %s (remaining: %d)"), *DuplicatedActor->GetFullName(), LoadProgress->NumPendingLoadRequests);
-					}
-					else
-					{
-						Mapping->LoadedPath = NAME_None;
-						LoadProgress->NumFailedLoadedRequests++;
-						UE_LOG(LogEngine, Warning, TEXT("Failed to find actor %s in package %s"), *ActorName, *LoadedPackageName.ToString());
-					}
-				}
-
-				if (InDestLevel)
-				{
-					// In the case where we have a destination level, we don't need the duplicated package anymore
-					FWorldPartitionPackageHelper::UnloadPackage(DuplicatedPackage);
-				}
-				else
-				{
-					// Push the package in the cache so that package unloading will be handled by the cache destruction
-					// The only existing case is during cook (Actors are loaded with a null destination level, then MoveExternalActorsToLevel is called).
-					InPackageCache.CacheLoadedPackage(DuplicatedPackage);
-				}
-			}
-			else
+			if (PackageName != PackageToLoad)
 			{
-				UE_LOG(LogEngine, Warning, TEXT("Failed to load %s"), *LoadedPackageName.ToString());
-				LoadProgress->NumFailedLoadedRequests++;
+				InstancingPackage = CreatePackage(*PackageName.ToString());
 			}
 
-			if (!LoadProgress->NumPendingLoadRequests)
-			{
-				InCompletionCallback(!LoadProgress->NumFailedLoadedRequests);
-			}
-		});
-
-		// Prefix with /Temp so package doesn't appear in Content, PostFix the Package so that we don't load over the original package
-		FName PackageDuplicateSeedName(*FString::Printf(TEXT("/Temp%s_DUP"), *PackageToLoadFrom));
-		
-		// Use package cache here because level instances can load the same levels across cells we want to avoid Loading the same package multiple times...
-		InPackageCache.LoadPackage(PackageDuplicateSeedName, *PackageToLoadFrom, CompletionCallback, /*bLoadAsync*/bInLoadAsync, /*bInWorldPackage*/true);
+			UPackage* Package = LoadPackage(InstancingPackage, *PackageToLoad.ToString(), LOAD_None, nullptr, &ContainerInstancingContext);
+			CompletionCallback.Execute(PackageToLoad, Package, Package ? EAsyncLoadingResult::Succeeded : EAsyncLoadingResult::Failed);
+		}
 	}
 
 	return (LoadProgress->NumPendingLoadRequests == 0);
 }
 
-#endif
\ No newline at end of file
+#endif
diff --git a/Engine/Source/Runtime/Engine/Private/WorldPartition/WorldPartitionLevelStreamingDynamic.cpp b/Engine/Source/Runtime/Engine/Private/WorldPartition/WorldPartitionLevelStreamingDynamic.cpp
--- a/Engine/Source/Runtime/Engine/Private/WorldPartition/WorldPartitionLevelStreamingDynamic.cpp
+++ b/Engine/Source/Runtime/Engine/Private/WorldPartition/WorldPartitionLevelStreamingDynamic.cpp
@@ -99,7 +99,6 @@ void UWorldPartitionLevelStreamingDynamic::Initialize(const UWorldPartitionRunti
 	bShouldBeAlwaysLoaded = InCell.IsAlwaysLoaded();
 	StreamingPriority = 0;
 	UnsavedActorsContainer = InCell.UnsavedActorsContainer;
-	ActorFolders = InCell.GetActorFolders().Array();
 
 	UWorld* OuterWorld = InCell.GetOuterUWorldPartition()->GetTypedOuter<UWorld>();
 	Initialize(OuterWorld, InCell.GetPackages());
@@ -134,19 +133,6 @@ void UWorldPartitionLevelStreamingDynamic::CreateRuntimeLevel()
 	UPackage* RuntimeLevelPackage = RuntimeLevel->GetPackage();
 	check(RuntimeLevelPackage);
 
-	// Propagate ActorFolder flag to the runtime level and prepare its ActorFolders list
-	if (World->PersistentLevel->IsUsingActorFolders() && ActorFolders.Num())
-	{
-		FLevelActorFoldersHelper::SetUseActorFolders(RuntimeLevel, true);
-		for (const FGuid& ActorFolderGuid : ActorFolders)
-		{
-			if (UActorFolder* ActorFolder = World->PersistentLevel->GetActorFolder(ActorFolderGuid))
-			{
-				FLevelActorFoldersHelper::AddActorFolder(RuntimeLevel, ActorFolder, /*bInShouldDirtyLevel*/ false, /*bInShouldBroadcast*/ false);
-			}
-		}
-	}
-
 	// Set flag here as this level isn't async loaded
 	RuntimeLevel->bClientOnlyVisible = bClientOnlyVisible;
 
@@ -358,7 +344,7 @@ bool UWorldPartitionLevelStreamingDynamic::IssueLoadRequests()
 	// Load saved actors
 	if (ChildPackagesToLoad.Num())
 	{
-		FWorldPartitionLevelHelper::LoadActors(RuntimeLevel, ChildPackagesToLoad, PackageCache, FinalizeLoading, GetWorld()->IsGameWorld(), &InstancingContext);
+		FWorldPartitionLevelHelper::LoadActors(World, RuntimeLevel, ChildPackagesToLoad, PackageReferencer, FinalizeLoading, World->IsGameWorld(), MoveTemp(InstancingContext));
 	}
 	else
 	{
@@ -425,7 +411,7 @@ void UWorldPartitionLevelStreamingDynamic::FinalizeRuntimeLevel()
 	// Mark this package as fully loaded with regards to external objects
 	RuntimeLevel->GetPackage()->SetDynamicPIEPackagePending(false);
 
-	PackageCache.UnloadPackages();
+	PackageReferencer.RemoveReferences();
 }
 
 /**
@@ -435,7 +421,7 @@ void UWorldPartitionLevelStreamingDynamic::OnCleanupLevel()
 {
 	if (RuntimeLevel)
 	{
-		PackageCache.UnloadPackages();
+		PackageReferencer.RemoveReferences();
 
 		RuntimeLevel->OnCleanupLevel.Remove(OnCleanupLevelDelegateHandle);
 
diff --git a/Engine/Source/Runtime/Engine/Private/WorldPartition/WorldPartitionPackageCache.cpp b/Engine/Source/Runtime/Engine/Private/WorldPartition/WorldPartitionPackageCache.cpp
deleted file mode 100644
diff --git a/Engine/Source/Runtime/Engine/Private/WorldPartition/WorldPartitionPackageHelper.cpp b/Engine/Source/Runtime/Engine/Private/WorldPartition/WorldPartitionPackageHelper.cpp
--- a/Engine/Source/Runtime/Engine/Private/WorldPartition/WorldPartitionPackageHelper.cpp
+++ b/Engine/Source/Runtime/Engine/Private/WorldPartition/WorldPartitionPackageHelper.cpp
@@ -46,21 +46,4 @@ void FWorldPartitionPackageHelper::UnloadPackage(UPackage* InPackage)
 	}
 }
 
-UPackage* FWorldPartitionPackageHelper::DuplicateWorldPackage(UPackage* InPackage, FName InDuplicatePackageName)
-{
-	UPackage* DuplicatedPackage = nullptr;
-	UWorld* PackageWorld = UWorld::FindWorldInPackage(InPackage);
-	check(PackageWorld);
-
-	DuplicatedPackage = CreatePackage(*InDuplicatePackageName.ToString());
-	FObjectDuplicationParameters DuplicationParameters(PackageWorld, DuplicatedPackage);
-	DuplicationParameters.bAssignExternalPackages = false;
-	DuplicationParameters.DuplicateMode = EDuplicateMode::World;
-
-	UWorld* DuplicatedWorld = Cast<UWorld>(StaticDuplicateObjectEx(DuplicationParameters));
-	check(DuplicatedWorld);
-
-	return DuplicatedPackage;
-}
-
 #endif
\ No newline at end of file
diff --git a/Engine/Source/Runtime/Engine/Private/WorldPartition/WorldPartitionRuntimeLevelStreamingCell.cpp b/Engine/Source/Runtime/Engine/Private/WorldPartition/WorldPartitionRuntimeLevelStreamingCell.cpp
--- a/Engine/Source/Runtime/Engine/Private/WorldPartition/WorldPartitionRuntimeLevelStreamingCell.cpp
+++ b/Engine/Source/Runtime/Engine/Private/WorldPartition/WorldPartitionRuntimeLevelStreamingCell.cpp
@@ -12,7 +12,6 @@
 #if WITH_EDITOR
 #include "WorldPartition/WorldPartitionLevelStreamingPolicy.h"
 #include "WorldPartition/WorldPartitionLevelHelper.h"
-#include "WorldPartition/WorldPartitionPackageCache.h"
 #include "Engine/LevelStreamingAlwaysLoaded.h"
 #endif
 
@@ -102,10 +101,6 @@ void UWorldPartitionRuntimeLevelStreamingCell::AddActorToCell(const FWorldPartit
 {
 	check(!ActorDescView.GetActorIsEditorOnly());
 	Packages.Emplace(ActorDescView.GetActorPackage(), ActorDescView.GetActorPath(), InContainerID, InContainerTransform, InContainer->GetContainerPackage());
-	if (ActorDescView.GetFolderGuid().IsValid())
-	{
-		ActorFolders.Add(ActorDescView.GetFolderGuid());
-	}
 }
 
 UWorldPartitionLevelStreamingDynamic* UWorldPartitionRuntimeLevelStreamingCell::CreateLevelStreaming(const FString& InPackageName) const
@@ -147,10 +142,12 @@ void UWorldPartitionRuntimeLevelStreamingCell::MoveAlwaysLoadedContentToPersiste
 	check(IsAlwaysLoaded());
 	if (GetActorCount() > 0)
 	{
-		FWorldPartitionPackageCache PackageCache;
-		verify(FWorldPartitionLevelHelper::LoadActors(nullptr, Packages, PackageCache, [](bool) {}, /*bLoadAsync*/false));
+		FWorldPartitionLevelHelper::FPackageReferencer PackageReferencer;
+		const bool bLoadAsync = false;
 
 		UWorld* OuterWorld = GetOuterUWorldPartition()->GetTypedOuter<UWorld>();
+		verify(FWorldPartitionLevelHelper::LoadActors(OuterWorld, nullptr, Packages, PackageReferencer, [](bool) {}, bLoadAsync, FLinkerInstancingContext()));
+				
 		FWorldPartitionLevelHelper::MoveExternalActorsToLevel(Packages, OuterWorld->PersistentLevel);
 
 		// Empty cell's package list (this ensures that no one can rely on cell's content).
@@ -192,12 +189,14 @@ bool UWorldPartitionRuntimeLevelStreamingCell::PopulateGeneratedPackageForCook(U
 		}
 
 		// Load cell Actors
-		FWorldPartitionPackageCache PackageCache;
-		verify(FWorldPartitionLevelHelper::LoadActors(nullptr, Packages, PackageCache, [](bool) {}, /*bLoadAsync*/false));
-
-		// Create a level and move these actors in it
 		UWorldPartition* WorldPartition = GetOuterUWorldPartition();
 		UWorld* OuterWorld = WorldPartition->GetTypedOuter<UWorld>();
+
+		FWorldPartitionLevelHelper::FPackageReferencer PackageReferencer;
+		const bool bLoadAsync = false;
+		verify(FWorldPartitionLevelHelper::LoadActors(OuterWorld, nullptr, Packages, PackageReferencer, [](bool) {}, bLoadAsync, FLinkerInstancingContext()));
+
+		// Create a level and move these actors in it
 		ULevel* NewLevel = FWorldPartitionLevelHelper::CreateEmptyLevelForRuntimeCell(OuterWorld, LevelStreaming->GetWorldAsset().ToString(), InPackage);
 		check(NewLevel->GetPackage() == InPackage);
 		FWorldPartitionLevelHelper::MoveExternalActorsToLevel(Packages, NewLevel);
diff --git a/Engine/Source/Runtime/Engine/Public/ActorFolder.h b/Engine/Source/Runtime/Engine/Public/ActorFolder.h
--- a/Engine/Source/Runtime/Engine/Public/ActorFolder.h
+++ b/Engine/Source/Runtime/Engine/Public/ActorFolder.h
@@ -33,7 +33,7 @@ public:
 	//~ End UObject
 
 	void SetParent(UActorFolder* InParent);
-	UActorFolder* GetParent() const;
+	UActorFolder* GetParent(bool bSkipDeleted = true) const;
 
 	void SetLabel(const FString& InFolderLabel);
 	FString GetLabel() const { return FolderLabel; }
diff --git a/Engine/Source/Runtime/Engine/Public/WorldPartition/WorldPartitionActorCluster.h b/Engine/Source/Runtime/Engine/Public/WorldPartition/WorldPartitionActorCluster.h
--- a/Engine/Source/Runtime/Engine/Public/WorldPartition/WorldPartitionActorCluster.h
+++ b/Engine/Source/Runtime/Engine/Public/WorldPartition/WorldPartitionActorCluster.h
@@ -60,6 +60,11 @@ struct FActorContainerID
 		return GetTypeHash(InContainerID.ID);
 	}
 
+	static FActorContainerID GetMainContainerID()
+	{
+		return FActorContainerID();
+	}
+
 	UPROPERTY()
 	uint64 ID;
 };
diff --git a/Engine/Source/Runtime/Engine/Public/WorldPartition/WorldPartitionLevelHelper.h b/Engine/Source/Runtime/Engine/Public/WorldPartition/WorldPartitionLevelHelper.h
--- a/Engine/Source/Runtime/Engine/Public/WorldPartition/WorldPartitionLevelHelper.h
+++ b/Engine/Source/Runtime/Engine/Public/WorldPartition/WorldPartitionLevelHelper.h
@@ -14,23 +14,39 @@
 #include "Engine/World.h"
 #include "WorldPartition/WorldPartitionRuntimeCell.h"
 
-class FWorldPartitionPackageCache;
+class FWorldPartitionPackageHelper;
 class UWorldPartition;
 struct FActorContainerID;
 
 class FWorldPartitionLevelHelper
 {
 public:
+	struct FPackageReferencer
+	{
+		~FPackageReferencer() { RemoveReferences(); }
+
+		void AddReference(UPackage* InPackage);
+		void RemoveReferences();
+	};
+
 	static ULevel* CreateEmptyLevelForRuntimeCell(const UWorld* InWorld, const FString& InWorldAssetName, UPackage* DestPackage = nullptr);
+	static void DuplicateActorFolderToRuntimeCell(ULevel* CellLevel, ULevel* SrcLevel, const FGuid& ActorFolderGuid);
 	static void MoveExternalActorsToLevel(const TArray<FWorldPartitionRuntimeCellObjectMapping>& InChildPackages, ULevel* InLevel);
 	static void RemapLevelSoftObjectPaths(ULevel* InLevel, UWorldPartition* InWorldPartition);
 	
-	static bool LoadActors(ULevel* InDestLevel, TArrayView<FWorldPartitionRuntimeCellObjectMapping> InActorPackages, FWorldPartitionPackageCache& InPackageCache, TFunction<void(bool)> InCompletionCallback, bool bInLoadAsync, FLinkerInstancingContext* InOutInstancingContext = nullptr);
+	static bool LoadActors(UWorld* InOwningWorld, ULevel* InDestLevel, TArrayView<FWorldPartitionRuntimeCellObjectMapping> InActorPackages, FPackageReferencer& InPackageReferencer, TFunction<void(bool)> InCompletionCallback, bool bInLoadAsync, FLinkerInstancingContext InOutInstancingContext);
 	
 	static FString AddActorContainerIDToActorPath(const FActorContainerID& InContainerID, const FString& InActorPath);
 	static FString AddActorContainerIDToSubPathString(const FActorContainerID& InContainerID, const FString& InSubPathString);
 private:
 	static UWorld::InitializationValues GetWorldInitializationValues();
+
+	struct FPackageReference
+	{
+		TSet<FPackageReferencer*> Referencers;
+		TWeakObjectPtr<UPackage> Package;
+	};
+	static TMap<FName, FPackageReference> PackageReferences;
 };
 
 #endif
\ No newline at end of file
diff --git a/Engine/Source/Runtime/Engine/Public/WorldPartition/WorldPartitionLevelStreamingDynamic.h b/Engine/Source/Runtime/Engine/Public/WorldPartition/WorldPartitionLevelStreamingDynamic.h
--- a/Engine/Source/Runtime/Engine/Public/WorldPartition/WorldPartitionLevelStreamingDynamic.h
+++ b/Engine/Source/Runtime/Engine/Public/WorldPartition/WorldPartitionLevelStreamingDynamic.h
@@ -16,7 +16,7 @@
 #include "WorldPartition/WorldPartition.h"
 #include "WorldPartition/WorldPartitionRuntimeLevelStreamingCell.h"
 #if WITH_EDITOR
-#include "WorldPartition/WorldPartitionPackageCache.h"
+#include "WorldPartition/WorldPartitionLevelHelper.h"
 #endif
 #include "WorldPartitionLevelStreamingDynamic.generated.h"
 
@@ -60,7 +60,6 @@ private:
 	FName OriginalLevelPackageName;
 	TArray<FWorldPartitionRuntimeCellObjectMapping> ChildPackages;
 	TArray<FWorldPartitionRuntimeCellObjectMapping> ChildPackagesToLoad;
-	TArray<FGuid> ActorFolders;
 
 #if WITH_EDITORONLY_DATA
 	UPROPERTY()
@@ -70,7 +69,7 @@ private:
 	FDelegateHandle OnCleanupLevelDelegateHandle;
 	bool bLoadRequestInProgress;
 	bool bLoadSucceeded;
-	FWorldPartitionPackageCache PackageCache;
+	FWorldPartitionLevelHelper::FPackageReferencer PackageReferencer;
 #endif
 
 	UPROPERTY()
diff --git a/Engine/Source/Runtime/Engine/Public/WorldPartition/WorldPartitionPackageCache.h b/Engine/Source/Runtime/Engine/Public/WorldPartition/WorldPartitionPackageCache.h
deleted file mode 100644
diff --git a/Engine/Source/Runtime/Engine/Public/WorldPartition/WorldPartitionPackageHelper.h b/Engine/Source/Runtime/Engine/Public/WorldPartition/WorldPartitionPackageHelper.h
--- a/Engine/Source/Runtime/Engine/Public/WorldPartition/WorldPartitionPackageHelper.h
+++ b/Engine/Source/Runtime/Engine/Public/WorldPartition/WorldPartitionPackageHelper.h
@@ -12,7 +12,6 @@ class FWorldPartitionPackageHelper
 {
 public:
 	static void UnloadPackage(UPackage* InPackage);
-	static UPackage* DuplicateWorldPackage(UPackage* InPackage, FName InDuplicatePackageName);
 };
 
 #endif
\ No newline at end of file
diff --git a/Engine/Source/Runtime/Engine/Public/WorldPartition/WorldPartitionRuntimeLevelStreamingCell.h b/Engine/Source/Runtime/Engine/Public/WorldPartition/WorldPartitionRuntimeLevelStreamingCell.h
--- a/Engine/Source/Runtime/Engine/Public/WorldPartition/WorldPartitionRuntimeLevelStreamingCell.h
+++ b/Engine/Source/Runtime/Engine/Public/WorldPartition/WorldPartitionRuntimeLevelStreamingCell.h
@@ -41,7 +41,6 @@ class UWorldPartitionRuntimeLevelStreamingCell : public UWorldPartitionRuntimeSp
 	//~End UWorldPartitionRuntimeCell Interface
 
 	const TArray<FWorldPartitionRuntimeCellObjectMapping>& GetPackages() const { return Packages; }
-	const TSet<FGuid>& GetActorFolders() const { return ActorFolders; }
 #endif
 
 private:
@@ -61,9 +60,6 @@ private:
 #if WITH_EDITORONLY_DATA
 	UPROPERTY()
 	TArray<FWorldPartitionRuntimeCellObjectMapping> Packages;
-
-	UPROPERTY()
-	TSet<FGuid> ActorFolders;
 #endif
 
 	UPROPERTY()
