IsStereo redundancy elimination
-IsStereoscopic3D already calls IsStereoRenderingAllowed, so we can delete the redundant call.

diff --git a/Engine/Source/Runtime/Engine/Private/UnrealClient.cpp b/Engine/Source/Runtime/Engine/Private/UnrealClient.cpp
--- a/Engine/Source/Runtime/Engine/Private/UnrealClient.cpp
+++ b/Engine/Source/Runtime/Engine/Private/UnrealClient.cpp
@@ -1674,7 +1674,7 @@ const TArray<FColor>& FViewport::GetRawHitProxyData(FIntRect InRect)
 {
 	FScopedConditionalWorldSwitcher WorldSwitcher(ViewportClient);
 
-	const bool bIsRenderingStereo = GEngine->IsStereoscopic3D( this ) && this->IsStereoRenderingAllowed();
+	const bool bIsRenderingStereo = GEngine->IsStereoscopic3D( this );
 
 	bool bFetchHitProxyBytes = !bIsRenderingStereo && ( !bHitProxiesCached || (SizeY*SizeX) != CachedHitProxyData.Num() );
 
