Fixed DataLayer console command invalid reporting of duplicate datalayers.

diff --git a/Engine/Source/Runtime/Engine/Private/WorldPartition/DataLayer/DataLayerSubsystem.cpp b/Engine/Source/Runtime/Engine/Private/WorldPartition/DataLayer/DataLayerSubsystem.cpp
--- a/Engine/Source/Runtime/Engine/Private/WorldPartition/DataLayer/DataLayerSubsystem.cpp
+++ b/Engine/Source/Runtime/Engine/Private/WorldPartition/DataLayer/DataLayerSubsystem.cpp
@@ -441,7 +441,7 @@ TArray<UDataLayerInstance*> UDataLayerSubsystem::ConvertArgsToDataLayers(UWorld*
 		FString SanitizedDataLayerName = DataLayerName.ToString().Replace(TEXT(" "), TEXT(""));
 		DataLayerSubsystem->ForEachDataLayer([&OutDataLayers, &SanitizedDataLayerName, &bShortNameFound](UDataLayerInstance* DataLayerInstance)
 		{
-			if (DataLayerInstance->GetDataLayerShortName().Find(SanitizedDataLayerName, ESearchCase::IgnoreCase) == 0)
+			if (DataLayerInstance->GetDataLayerShortName().Compare(SanitizedDataLayerName, ESearchCase::IgnoreCase) == 0)
 			{
 				if (bShortNameFound)
 				{
