Fixed repair of WorldDataLayers a world partition initialization.

diff --git a/Engine/Source/Runtime/Engine/Private/WorldPartition/WorldPartition.cpp b/Engine/Source/Runtime/Engine/Private/WorldPartition/WorldPartition.cpp
--- a/Engine/Source/Runtime/Engine/Private/WorldPartition/WorldPartition.cpp
+++ b/Engine/Source/Runtime/Engine/Private/WorldPartition/WorldPartition.cpp
@@ -500,15 +500,9 @@ void UWorldPartition::Initialize(UWorld* InWorld, const FTransform& InTransform)
 		AWorldDataLayers* WorldDataLayers = World->GetWorldDataLayers();
 		if (!WorldDataLayers)
 		{
-			if (WorldDataLayersActor.IsValid())
-			{
-				WorldDataLayers = CastChecked<AWorldDataLayers>((*WorldDataLayersActor)->GetActor());
-			}
-			else
-			{
-				WorldDataLayers = AWorldDataLayers::Create(World);
-			}
-
+			// WorldDataLayersActor is invalid here orelse its PostLoad would have set itself as the world's WorldDataLayers
+			check(!WorldDataLayersActor.IsValid());
+			WorldDataLayers = AWorldDataLayers::Create(World);
 			World->SetWorldDataLayers(WorldDataLayers);
 		}
 	}
