Protect audio render state callback delegates to prevent issues with multiple threads attempting to add and remove delegates.

NOTE: When deadend lifted you MUST also un-deadend CL 19669399

diff --git a/Engine/Source/Runtime/Engine/Private/AudioDevice.cpp b/Engine/Source/Runtime/Engine/Private/AudioDevice.cpp
--- a/Engine/Source/Runtime/Engine/Private/AudioDevice.cpp
+++ b/Engine/Source/Runtime/Engine/Private/AudioDevice.cpp
@@ -4706,13 +4706,39 @@ void FAudioDevice::UpdateAudioEngineSubsystems()
 	}
 }
 
-void FAudioDevice::NotifyAudioDevicePreRender(const FAudioDeviceRenderInfo& InInfo) const
+FDelegateHandle FAudioDevice::AddPreRenderDelegate(const FOnAudioDevicePreRender::FDelegate& InDelegate)
 {
+	FScopeLock LockCallbacks(&RenderStateCallbackListCritSec);
+	return OnAudioDevicePreRender.Add(InDelegate);
+}
+
+bool FAudioDevice::RemovePreRenderDelegate(const FDelegateHandle& InHandle)
+{
+	FScopeLock LockCallbacks(&RenderStateCallbackListCritSec);
+	return OnAudioDevicePreRender.Remove(InHandle);
+}
+
+void FAudioDevice::NotifyAudioDevicePreRender(const FAudioDeviceRenderInfo& InInfo)
+{
+	FScopeLock LockCallbacks(&RenderStateCallbackListCritSec);
 	OnAudioDevicePreRender.Broadcast(InInfo);
 }
 
-void FAudioDevice::NotifyAudioDevicePostRender(const FAudioDeviceRenderInfo& InInfo) const
+FDelegateHandle FAudioDevice::AddPostRenderDelegate(const FOnAudioDevicePostRender::FDelegate& InDelegate)
+{
+	FScopeLock LockCallbacks(&RenderStateCallbackListCritSec);
+	return OnAudioDevicePostRender.Add(InDelegate);
+}
+
+bool FAudioDevice::RemovePostRenderDelegate(const FDelegateHandle& InHandle)
+{
+	FScopeLock LockCallbacks(&RenderStateCallbackListCritSec);
+	return OnAudioDevicePostRender.Remove(InHandle);
+}
+
+void FAudioDevice::NotifyAudioDevicePostRender(const FAudioDeviceRenderInfo& InInfo)
 {
+	FScopeLock LockCallbacks(&RenderStateCallbackListCritSec);
 	OnAudioDevicePostRender.Broadcast(InInfo);
 }
 
diff --git a/Engine/Source/Runtime/Engine/Public/AudioDevice.h b/Engine/Source/Runtime/Engine/Public/AudioDevice.h
--- a/Engine/Source/Runtime/Engine/Public/AudioDevice.h
+++ b/Engine/Source/Runtime/Engine/Public/AudioDevice.h
@@ -900,6 +900,14 @@ protected:
 	// Handle for our device destroyed delegate
 	FDelegateHandle DeviceDestroyedHandle;
 
+	FCriticalSection RenderStateCallbackListCritSec;
+
+	// Callback as audio device is about to render a buffer
+	FOnAudioDevicePreRender OnAudioDevicePreRender;
+
+	// Callback as audio device has just finished rendering a buffer
+	FOnAudioDevicePostRender OnAudioDevicePostRender;
+
 public:
 	/**
 	 * Registers a sound class with the audio device
@@ -953,9 +961,6 @@ public:
 		return nullptr;
 	}
 
-	FOnAudioDevicePreRender OnAudioDevicePreRender;
-	FOnAudioDevicePostRender OnAudioDevicePostRender;
-
 	virtual void InitSoundEffectPresets() {}
 
 	/**
@@ -1512,10 +1517,10 @@ protected:
 	 */
 	virtual void OnListenerUpdated(const TArray<FListener>& InListeners) {};
 
-	void NotifyAudioDevicePreRender(const FAudioDeviceRenderInfo& InInfo) const;
+	void NotifyAudioDevicePreRender(const FAudioDeviceRenderInfo& InInfo);
+
+	void NotifyAudioDevicePostRender(const FAudioDeviceRenderInfo& InInfo);
 
-	void NotifyAudioDevicePostRender(const FAudioDeviceRenderInfo& InInfo) const;
-	
 private:
 
 	/**
@@ -1724,6 +1729,14 @@ public:
 		DefaultBaseSoundMix = InDefaultBaseSoundMix;
 	}
 
+	FDelegateHandle AddPreRenderDelegate(const FOnAudioDevicePreRender::FDelegate& InDelegate);
+
+	bool RemovePreRenderDelegate(const FDelegateHandle& InHandle);
+
+	FDelegateHandle AddPostRenderDelegate(const FOnAudioDevicePostRender::FDelegate& InDelegate);
+
+	bool RemovePostRenderDelegate(const FDelegateHandle& InHandle);
+
 private:
 	/**
 	 * Internal helper function used by ParseSoundClasses to traverse the tree.
