Generate compile error if material tries to append more than 4 componets worth of data

diff --git a/Engine/Source/Runtime/Engine/Private/Materials/HLSLMaterialTranslator.cpp b/Engine/Source/Runtime/Engine/Private/Materials/HLSLMaterialTranslator.cpp
--- a/Engine/Source/Runtime/Engine/Private/Materials/HLSLMaterialTranslator.cpp
+++ b/Engine/Source/Runtime/Engine/Private/Materials/HLSLMaterialTranslator.cpp
@@ -8142,6 +8142,11 @@ int32 FHLSLMaterialTranslator::AppendVector(int32 A,int32 B)
 	const int32 NumComponentsA = GetNumComponents(TypeA);
 	const int32 NumComponentsB = GetNumComponents(TypeB);
 	const int32 NumResultComponents = NumComponentsA + NumComponentsB;
+	if (NumResultComponents > 4)
+	{
+		return Errorf(TEXT("Can't append %s to %s"), DescribeType(TypeA), DescribeType(TypeB));
+	}
+
 	const EMaterialValueType ResultType = bIsLWC ? GetLWCVectorType(NumResultComponents) : GetVectorType(NumResultComponents);
 
 	if(GetParameterUniformExpression(A) && GetParameterUniformExpression(B))
