Deprecate an unused member in UBlueprint and remove all usage in engine code where it is no longer required.

diff --git a/Engine/Source/Editor/KismetCompiler/Private/KismetCompiler.cpp b/Engine/Source/Editor/KismetCompiler/Private/KismetCompiler.cpp
--- a/Engine/Source/Editor/KismetCompiler/Private/KismetCompiler.cpp
+++ b/Engine/Source/Editor/KismetCompiler/Private/KismetCompiler.cpp
@@ -4836,10 +4836,6 @@ void FKismetCompilerContext::CompileFunctions(EInternalCompilerFlags InternalFla
 			const bool bGenerateStubsOnly = !bIsFullCompile || (0 != MessageLog.NumErrors);
 			BP_SCOPED_COMPILER_EVENT_STAT(EKismetCompilerStats_CodeGenerationTime);
 			Backend_VM.GenerateCodeFromClass(NewClass, FunctionList, bGenerateStubsOnly);
-			if (!bGenerateStubsOnly)
-			{
-				Blueprint->bHasAnyNonReducibleFunction = Backend_VM.bAnyNonReducibleFunctionGenerated ? UBlueprint::EIsBPNonReducible::Yes : UBlueprint::EIsBPNonReducible::No;
-			}
 		}
 
 		// Fill ScriptAndPropertyObjectReferences arrays in functions
diff --git a/Engine/Source/Editor/KismetCompiler/Private/KismetCompilerBackend.h b/Engine/Source/Editor/KismetCompiler/Private/KismetCompilerBackend.h
--- a/Engine/Source/Editor/KismetCompiler/Private/KismetCompilerBackend.h
+++ b/Engine/Source/Editor/KismetCompiler/Private/KismetCompilerBackend.h
@@ -29,14 +29,11 @@ protected:
 
 	TStatementToSkipSizeMap UbergraphStatementLabelMap;
 public:
-	bool bAnyNonReducibleFunctionGenerated;
-
 	FKismetCompilerVMBackend(UBlueprint* InBlueprint, UEdGraphSchema_K2* InSchema, FKismetCompilerContext& InContext)
 		: Blueprint(InBlueprint)
 		, Schema(InSchema)
 		, MessageLog(InContext.MessageLog)
 		, CompilerContext(InContext)
-		, bAnyNonReducibleFunctionGenerated(false)
 	{
 	}
 
diff --git a/Engine/Source/Editor/KismetCompiler/Private/KismetCompilerVMBackend.cpp b/Engine/Source/Editor/KismetCompiler/Private/KismetCompilerVMBackend.cpp
--- a/Engine/Source/Editor/KismetCompiler/Private/KismetCompilerVMBackend.cpp
+++ b/Engine/Source/Editor/KismetCompiler/Private/KismetCompilerVMBackend.cpp
@@ -2362,11 +2362,6 @@ void FKismetCompilerVMBackend::ConstructFunction(FKismetFunctionContext& Functio
 					FBlueprintCompiledStatement* Statement = (*StatementList)[StatementIndex];
 
 					ScriptWriter.GenerateCodeForStatement(CompilerContext, FunctionContext, *Statement, StatementNode);
-					
-					const bool bUberGraphFunctionCall = Statement->FunctionToCall && (Statement->FunctionToCall == Class->UberGraphFunction)
-						&& (EKismetCompiledStatementType::KCST_CallFunction == Statement->Type);
-					const bool bIsReducible = FKismetCompilerUtilities::IsStatementReducible(Statement->Type) || bUberGraphFunctionCall;
-					bAnyNonReducibleFunctionGenerated |= !bIsReducible;
 				}
 			}
 		}
diff --git a/Engine/Source/Runtime/Engine/Classes/Engine/Blueprint.h b/Engine/Source/Runtime/Engine/Classes/Engine/Blueprint.h
--- a/Engine/Source/Runtime/Engine/Classes/Engine/Blueprint.h
+++ b/Engine/Source/Runtime/Engine/Classes/Engine/Blueprint.h
@@ -809,15 +809,20 @@ public:
 	UPROPERTY(transient, duplicatetransient)
 	TSet<TWeakObjectPtr<UStruct>> CachedUDSDependencies;
 
-	enum class EIsBPNonReducible : uint8
+	// @todo: BP2CPP_remove
+	enum class UE_DEPRECATED(5.1, "Blueprint Nativization has been removed as a supported feature. This type will eventually be removed.") EIsBPNonReducible : uint8
 	{
 		Unkown,
 		Yes,
 		No,
 	};
 
+	// @todo: BP2CPP_remove
 	// Cached information if the BP contains any non-reducible functions (that can benefit from nativization).
+	UE_DEPRECATED(5.1, "Blueprint Nativization has been removed as a supported feature. This member will eventually be removed.")
+	PRAGMA_DISABLE_DEPRECATION_WARNINGS
 	EIsBPNonReducible bHasAnyNonReducibleFunction;
+	PRAGMA_ENABLE_DEPRECATION_WARNINGS
 
 	// If this BP is just a duplicate created for a specific compilation, the reference to original GeneratedClass is needed
 	UPROPERTY(transient, duplicatetransient)
diff --git a/Engine/Source/Runtime/Engine/Private/Blueprint.cpp b/Engine/Source/Runtime/Engine/Private/Blueprint.cpp
--- a/Engine/Source/Runtime/Engine/Private/Blueprint.cpp
+++ b/Engine/Source/Runtime/Engine/Private/Blueprint.cpp
@@ -354,7 +354,10 @@ UBlueprint::UBlueprint(const FObjectInitializer& ObjectInitializer)
 #if WITH_EDITORONLY_DATA
 	, bDuplicatingReadOnly(false)
 	, bCachedDependenciesUpToDate(false)
+	// @todo: BP2CPP_remove
+	PRAGMA_DISABLE_DEPRECATION_WARNINGS
 	, bHasAnyNonReducibleFunction(EIsBPNonReducible::Unkown)
+	PRAGMA_ENABLE_DEPRECATION_WARNINGS
 #endif
 {
 }
