- Adding more layers to the "get cvars for other platform" system: Scalability at a default level, and ConsoleVariables.ini

diff --git a/Engine/Source/Runtime/Core/Private/HAL/ConsoleManager.cpp b/Engine/Source/Runtime/Core/Private/HAL/ConsoleManager.cpp
--- a/Engine/Source/Runtime/Core/Private/HAL/ConsoleManager.cpp
+++ b/Engine/Source/Runtime/Core/Private/HAL/ConsoleManager.cpp
@@ -10,6 +10,7 @@ ConsoleManager.cpp: console command handling
 #include "Misc/Paths.h"
 #include "Stats/Stats.h"
 #include "Misc/ConfigCacheIni.h"
+#include "Misc/CoreDelegates.h"
 #include "Modules/ModuleManager.h"
 #include "HAL/PlatformProcess.h"
 #include "Misc/RemoteConfigIni.h"
@@ -352,8 +353,12 @@ static void ExpandScalabilityCVar(FConfigCacheIni* ConfigSystem, const FString&
 	}
 }
 
-bool IConsoleManager::VisitPlatformCVarsForEmulation(FName PlatformName, bool bVisitPlatformDeviceProfile, TFunctionRef<void(const FString& CVarName, const FString& CVarValue, EConsoleVariableFlags SetBy)> Visit)
+bool IConsoleManager::VisitPlatformCVarsForEmulation(FName PlatformName, const FString& DeviceProfileName, TFunctionRef<void(const FString& CVarName, const FString& CVarValue, EConsoleVariableFlags SetBy)> Visit)
 {
+	// we can't get to Scalability code in here (it's in Engine) but we still want to apply the default level... 
+	// there is a static_assert in UDeviceProfile::ExpandDeviceProfileCVars to make sure the default doesn't change from this value
+	const int DefaultScalabilityLevel = 3;
+
 	// get the config system for the platform the DP uses
 	FConfigCacheIni* ConfigSystem = FConfigCacheIni::ForPlatform(*PlatformName.ToString());
 	if (ConfigSystem == nullptr)
@@ -365,56 +370,135 @@ bool IConsoleManager::VisitPlatformCVarsForEmulation(FName PlatformName, bool bV
 	//   doesn't come from ini
 
 	// ECVF_SetByScalability:
-	//	 skipped, i believe this is not really loaded as a normal layer per-se, it's up to the other sections to set with this one
+	//	 initializes sg. cvars to the Default level, likely DPs etc will replace most/all of these settings
+
+	// ECVF_SetByProjectSetting:
+	// ECVF_SetBySystemSettingsIni:
+	//   read from ini files
+	
+	// ECVF_SetByDeviceProfile: 
+	//   read from the DP parenting chain
 
 	// ECVF_SetByGameSetting:
 	//   skipped, since we don't have a user
 
+	// ECVF_SetByConsoleVariablesIni:
+	//   maybe skip this? it's a weird one, but maybe?
 
-	// ECVF_SetByProjectSetting:
-	// ECVF_SetBySystemSettingsIni:
-	// ECVF_SetByDeviceProfile: 
-	//   ONLY using the base DP named for the platform, as the DPManager needs to be used for proper DeviceProfiles
+	// ECVF_SetByCommandline:
+	//   skip as this would not be expected to apply to emulation
+
+	// ECVF_SetByCode:
+	//   skip because it cannot be set by code
+
+	// ECVF_SetByConsole
+	//   we could have this if we made a per-platform CVar, not just the shared default value
 
 	const TCHAR* DeviceProfileTag = TEXT("_NamedDeviceProfile");
+	const TCHAR* ScalabilityTag = TEXT("_Scalability");
 	struct FSectionPair
 	{
 		const TCHAR* Name; EConsoleVariableFlags SetBy;
 	} Sections[] =
 	{
+		// this order is not based on priority order, but by the order as seen in LaunchEngineLoop:
+		//		GSystemSettings.Initialize()
+		{ TEXT("SystemSettings"), ECVF_SetBySystemSettingsIni },
+
+		//		ApplyCVarSettingsFromIni(...)
 		{ TEXT("/Script/Engine.RendererSettings"), ECVF_SetByProjectSetting },
 		{ TEXT("/Script/Engine.RendererOverrideSettings"), ECVF_SetByProjectSetting },
 		{ TEXT("/Script/Engine.StreamingSettings"), ECVF_SetByProjectSetting },
 		{ TEXT("/Script/Engine.GarbageCollectionSettings"), ECVF_SetByProjectSetting },
 		{ TEXT("/Script/Engine.NetworkSettings"), ECVF_SetByProjectSetting },
 
-		{ TEXT("SystemSettings"), ECVF_SetBySystemSettingsIni },
-		{ TEXT("ConsoleVariables"), ECVF_SetBySystemSettingsIni },
+		//		Scalability::InitScalabilitySystem()
+		{ ScalabilityTag, ECVF_SetByDeviceProfile },
 
+		//		UDeviceProfileManager::InitializeCVarsForActiveDeviceProfile()
 		{ DeviceProfileTag, ECVF_SetByDeviceProfile },
+
+		//		FConfigCacheIni::LoadConsoleVariablesFromINI()
+		{ TEXT("Startup"), ECVF_SetByConsoleVariablesIni },
+		{ TEXT("ConsoleVariables"), ECVF_SetBySystemSettingsIni },
+
+	};
+
+	TMap<FString, int32> CVarSetByMap;
+	auto VisitIfAllowed = [&CVarSetByMap, &Visit](const FString& Name, const FString& Value, EConsoleVariableFlags SetBy)
+	{
+		int32 SetByInt = (int32)(SetBy & ECVF_SetByMask);
+		int32 WasSetBy = CVarSetByMap.FindOrAdd(Name);
+		if (SetByInt < WasSetBy)
+		{
+			UE_LOG(LogConsoleManager, Log, TEXT("Skipping CVar %s=%s while visiting another platform, because it was already visited with a higher priority"), *Name, *Value);
+			return;
+		}
+
+		CVarSetByMap[Name] = SetByInt;
+
+		// now look up the cvar, if it exists (it's okay if it doesn't, it may not exist on host platform, but then it's not previewable!)
+		IConsoleVariable* CVar = IConsoleManager::Get().FindConsoleVariable(*Name, false /* bTrackFrequentCalls */);
+		EConsoleVariableFlags PreviewFlag = (CVar != nullptr) ? (EConsoleVariableFlags)(CVar->GetFlags() & ECVF_Preview) : ECVF_Default;
+		Visit(Name, Value, (EConsoleVariableFlags)(SetBy | PreviewFlag));
 	};
 
 	// now walk up the stack getting current values
 	for (const FSectionPair& SectionPair : Sections)
 	{
+		bool bDeleteSection = false;
 		FConfigSection* Section;
 		bool bIsDeviceProfile = FCString::Strcmp(SectionPair.Name, DeviceProfileTag) == 0;
+		bool bIsScalabilityLevel = FCString::Strcmp(SectionPair.Name, ScalabilityTag) == 0;
 
 		if (bIsDeviceProfile)
 		{
-			// skip this if we don't want it
-			if (bVisitPlatformDeviceProfile == false)
+			// skip this if we didn't specify one
+			if (DeviceProfileName.Len() == 0)
+			{
+				continue;
+			}
+			if (!FCoreDelegates::GatherDeviceProfileCVars.IsBound())
 			{
+				UE_LOG(LogConsoleManager, Warning, TEXT("Attempted to get CVars for another platform before FCoreDelegates::GatherDeviceProfileCVars was bound to a callback. CVar values are likely incorrect."));
 				continue;
 			}
-			FString DPSectionName = PlatformName.ToString() + TEXT(" DeviceProfile");
-			Section = ConfigSystem->GetSectionPrivate(*DPSectionName, false, true, GDeviceProfilesIni);
+
+			// make a fake section of the dp cvars - this will let the expansion happen as normal below
+			Section = new FConfigSection();
+			bDeleteSection = true;
+
+			// run the delegate (this code can't get into DP code directly, so we use a delegate), and walk over the results
+			for (TPair<FName, FString>& Pair : FCoreDelegates::GatherDeviceProfileCVars.Execute(DeviceProfileName))
+			{
+				Section->Add(Pair);
+			}
+		}
+		else if (bIsScalabilityLevel)
+		{
+			// walk over all of the scalability groups, and assign them to the default level
+			FString DefaultLevel = FString::Printf(TEXT("%d"), DefaultScalabilityLevel);
+
+			// make a fake section of the sg. vars - this will let the expansion happen as normal below
+			Section = new FConfigSection();
+			bDeleteSection = true;
+			IConsoleManager::Get().ForEachConsoleObjectThatStartsWith(FConsoleObjectVisitor::CreateLambda([&DefaultLevel, Section](const TCHAR* Name, IConsoleObject* Obj)
+				{
+					if (Obj->TestFlags(ECVF_ScalabilityGroup))
+					{
+						Section->Add(Name, FConfigValue(DefaultLevel));
+					}
+				}
+			));
 		}
 		else
 		{
-			Section = ConfigSystem->GetSectionPrivate(SectionPair.Name, false, true, GEngineIni);
+			static const FString ConsoleVariablesIni = FPaths::EngineDir() + TEXT("Config/ConsoleVariables.ini");
+			const FString& IniFile = (SectionPair.SetBy == ECVF_SetByConsoleVariablesIni) ? ConsoleVariablesIni : GEngineIni;
+			Section = ConfigSystem->GetSectionPrivate(SectionPair.Name, false, true, IniFile);
 		}
 
+
 		if (Section != nullptr)
 		{
 			// add the cvars from the section
@@ -423,25 +507,6 @@ bool IConsoleManager::VisitPlatformCVarsForEmulation(FName PlatformName, bool bV
 				FString Key = Pair.Key.ToString();
 				FString Value = Pair.Value.GetValue();
 
-				// DPs have +CVars= prefix for the key, and the real cvar KVP is in the Value
-				// and other keys are not CVars so are ignored
-				if (bIsDeviceProfile)
-				{
-					if (Key.StartsWith(TEXT("CVars")))
-					{
-						FString ValueCopy = Value;
-						// expect a second = in the value, skip any that don't
-						if (!ValueCopy.Split(TEXT("="), &Key, &Value))
-						{
-							continue;
-						}
-					}
-					else
-					{
-						continue;
-					}
-				}
-
 				// don't bother tracking when looking up other platform cvars
 				IConsoleVariable* CVar = IConsoleManager::Get().FindConsoleVariable(*Key, false /* bTrackFrequentCalls */);
 				EConsoleVariableFlags PreviewFlag = (CVar != nullptr) ? (EConsoleVariableFlags)(CVar->GetFlags() & ECVF_Preview) : ECVF_Default;
@@ -457,28 +522,23 @@ bool IConsoleManager::VisitPlatformCVarsForEmulation(FName PlatformName, bool bV
 
 					for (const auto& ScalabilityPair : ScalabilityCVars)
 					{
-						Visit(ScalabilityPair.Key, ScalabilityPair.Value, (EConsoleVariableFlags)(ECVF_SetByScalability | PreviewFlag));
+						VisitIfAllowed(ScalabilityPair.Key, ScalabilityPair.Value, (EConsoleVariableFlags)(ECVF_SetByScalability | PreviewFlag));
 					}
 				}
 				else
 				{
-					Visit(Key, Value, (EConsoleVariableFlags)(SectionPair.SetBy | PreviewFlag));
+					VisitIfAllowed(Key, Value, (EConsoleVariableFlags)(SectionPair.SetBy | PreviewFlag));
 				}
 			}
-		}
-	}
-
-	// ECVF_SetByConsoleVariablesIni:
-	//   maybe skip this? it's a weird one, but maybe?
-
-	// ECVF_SetByCommandline:
-	//   skip as this would not be expected to apply to emulation
 
-	// ECVF_SetByCode:
-	//   skip because it cannot be set by code
+			// clean up the temp section we made
+			if (bDeleteSection)
+			{
+				delete Section;
+			}
 
-	// ECVF_SetByConsole
-	//   we could have this if we made a per-platform CVar, not just the shared default value
+		}
+	}
 
 
 	return true;
@@ -487,7 +547,6 @@ bool IConsoleManager::VisitPlatformCVarsForEmulation(FName PlatformName, bool bV
 
 static bool GetConfigValueFromRuntimeSources(FName PlatformName, IConsoleVariable* CVar, FString& OutValue, EConsoleVariableFlags& LastSetBy)
 {
-
 	FString VariableName = IConsoleManager::Get().FindConsoleObjectName(CVar);
 	if (VariableName.Len() == 0)
 	{
@@ -497,7 +556,8 @@ static bool GetConfigValueFromRuntimeSources(FName PlatformName, IConsoleVariabl
 	OutValue = CVar->GetDefaultValueVariable()->GetString();
 	LastSetBy = ECVF_SetByConstructor;
 
-	IConsoleManager::VisitPlatformCVarsForEmulation(PlatformName, true, 
+	// use the platform's base DeviceProfile for emulation
+	IConsoleManager::VisitPlatformCVarsForEmulation(PlatformName, PlatformName.ToString(),
 		[&VariableName, &OutValue, &LastSetBy](const FString& CVarName, const FString& CVarValue, EConsoleVariableFlags SetByAndPreview)
 		{
 			// if this key is the variable, set it at the current SetBy level
diff --git a/Engine/Source/Runtime/Core/Private/Misc/CoreDelegates.cpp b/Engine/Source/Runtime/Core/Private/Misc/CoreDelegates.cpp
--- a/Engine/Source/Runtime/Core/Private/Misc/CoreDelegates.cpp
+++ b/Engine/Source/Runtime/Core/Private/Misc/CoreDelegates.cpp
@@ -88,6 +88,9 @@ FCoreDelegates::FPakSigningKeysDelegate& FCoreDelegates::GetPakSigningKeysDelega
     FSimpleMulticastDelegate FCoreDelegates::PreSlateModal;
     FSimpleMulticastDelegate FCoreDelegates::PostSlateModal;
 #endif	//WITH_EDITOR
+#if ALLOW_OTHER_PLATFORM_CONFIG
+	FCoreDelegates::FGatherDeviceProfileCVars FCoreDelegates::GatherDeviceProfileCVars;
+#endif
 FSimpleMulticastDelegate FCoreDelegates::OnShutdownAfterError;
 FSimpleMulticastDelegate FCoreDelegates::OnInit;
 FSimpleMulticastDelegate FCoreDelegates::OnPostEngineInit;
diff --git a/Engine/Source/Runtime/Core/Public/HAL/IConsoleManager.h b/Engine/Source/Runtime/Core/Public/HAL/IConsoleManager.h
--- a/Engine/Source/Runtime/Core/Public/HAL/IConsoleManager.h
+++ b/Engine/Source/Runtime/Core/Public/HAL/IConsoleManager.h
@@ -973,10 +973,10 @@ struct CORE_API IConsoleManager
 	 * It also won't include any UserSettings or ConsoleVariables.ini settings.
 	 * 
 	 * @param PlatformName The platform name (the ini name, so Windows, not Win64)
-	 * @param bVisitPlatformDeviceProfile Set to true if you want Visit called with CVars found in the base DeviceProfile named PlatformName 
+	 * @param DeviceProfileName If this is non-empty, the given deviceprofile will be loaded from the platform's inis and inserted into the CVars
 	 * @param Visit the callback to run for each CVar found
 	 */
-	static bool VisitPlatformCVarsForEmulation(FName PlatformName, bool bVisitPlatformDeviceProfile, TFunctionRef<void(const FString& CVarName, const FString& CVarValue, EConsoleVariableFlags SetBy)> Visit);
+	static bool VisitPlatformCVarsForEmulation(FName PlatformName, const FString& DeviceProfileName, TFunctionRef<void(const FString& CVarName, const FString& CVarValue, EConsoleVariableFlags SetBy)> Visit);
 #endif
 
 	virtual FConsoleVariableMulticastDelegate& OnCVarUnregistered() = 0;
diff --git a/Engine/Source/Runtime/Core/Public/Misc/CoreDelegates.h b/Engine/Source/Runtime/Core/Public/Misc/CoreDelegates.h
--- a/Engine/Source/Runtime/Core/Public/Misc/CoreDelegates.h
+++ b/Engine/Source/Runtime/Core/Public/Misc/CoreDelegates.h
@@ -239,7 +239,14 @@ public:
     static FSimpleMulticastDelegate PostSlateModal;
     
 #endif	//WITH_EDITOR
-	
+
+#if ALLOW_OTHER_PLATFORM_CONFIG
+	// Called when the CVar (ConsoleManager) needs to retrieve CVars for a deviceprofile for another platform - this dramatically simplifies module dependencies
+	typedef TMap<FName, FString> FCVarKeyValueMap;
+	DECLARE_DELEGATE_RetVal_OneParam(FCVarKeyValueMap, FGatherDeviceProfileCVars, const FString&  /*DeviceProfileName*/);
+	static FGatherDeviceProfileCVars GatherDeviceProfileCVars;
+#endif
+
 	// Called when an error occurred.
 	static FSimpleMulticastDelegate OnShutdownAfterError;
 
diff --git a/Engine/Source/Runtime/Engine/Classes/DeviceProfiles/DeviceProfile.h b/Engine/Source/Runtime/Engine/Classes/DeviceProfiles/DeviceProfile.h
--- a/Engine/Source/Runtime/Engine/Classes/DeviceProfiles/DeviceProfile.h
+++ b/Engine/Source/Runtime/Engine/Classes/DeviceProfiles/DeviceProfile.h
@@ -179,14 +179,16 @@ private:
 public:
 	const TMap<FString, FString>& GetAllExpandedCVars();
 	const TMap<FString, FString>& GetAllPreviewCVars();
-	void AddExpandedCVars(const TMap<FString, FString>& CVarsToMerge);
-	void AddPreviewCVars(const TMap<FString, FString>& CVarsToMerge);
 	void ClearAllExpandedCVars();
 	/** Set the memory size bucket to be used when previewing this DP, changing this will reset the expanded cvars. */
 	void SetPreviewMemorySizeBucket(EPlatformMemorySizeBucket PreviewMemorySizeBucketIn);
 	EPlatformMemorySizeBucket GetPreviewMemorySizeBucket() const;
 
 private:
+	
+	// calculate the cvars for another platform's deviceprofile
+	void ExpandDeviceProfileCVars();
+		
 	/** Resolved CVars, including expanded scalability cvars used to properly emulate one platform on another */
 	TMap<FString, FString> AllExpandedCVars;
 
diff --git a/Engine/Source/Runtime/Engine/Classes/DeviceProfiles/DeviceProfileManager.h b/Engine/Source/Runtime/Engine/Classes/DeviceProfiles/DeviceProfileManager.h
--- a/Engine/Source/Runtime/Engine/Classes/DeviceProfiles/DeviceProfileManager.h
+++ b/Engine/Source/Runtime/Engine/Classes/DeviceProfiles/DeviceProfileManager.h
@@ -202,11 +202,6 @@ public:
 	static bool GetScalabilityCVar(const FString& CvarName, int32& OutValue);
 	static bool GetScalabilityCVar(const FString& CvarName, float& OutValue);
 
-#if ALLOW_OTHER_PLATFORM_CONFIG
-	/** Retrieve another platform's DeviceProfile, and make sure it's CVars are expanded, including scalability */
-	static void ExpandDeviceProfileCVars(UDeviceProfile* DeviceProfile);
-#endif
-
 	/**
 	* Enable/Disable a tagged fragment of the active device profile.
 	* This unsets the entire cvar DP state, then re-sets the new DP+fragment state.
@@ -220,6 +215,18 @@ public:
 	const FSelectedFragmentProperties* GetActiveDeviceProfileFragmentByTag(FName& FragmentTag) const;
 
 
+	enum class EDeviceProfileMode : uint8
+	{
+		DPM_SetCVars,
+		DPM_CacheValues,
+	};
+
+	/**
+	 * Walk the device profile/fragment chain to get the final set ot CVars in a unified way
+	 */
+	static TMap<FName, FString> GatherDeviceProfileCVars(const FString& DeviceProfileName, EDeviceProfileMode GatherMode);
+
+
 private:
 	/**
 	 * Set the active device profile - set via the device profile blueprint.
@@ -239,17 +246,11 @@ private:
 	/** Sees if the texture settings are the same between two profiles */
 	bool AreTextureGroupsTheSame(UDeviceProfile* Profile1, UDeviceProfile* Profile2) const;
 
-	enum class EDeviceProfileMode : uint8
-	{
-		DPM_SetCVars,
-		DPM_CacheValues,
-	};
-
 	/**
 	 * Perform the processing of ini sections, going up to parents, etc. Depending on runtime vs editor processing of another platform,
 	 * the Mode will control how the settings are handled
 	 */
-	static void ProcessDeviceProfileIniSettings(const FString& DeviceProfileName, EDeviceProfileMode Mode);
+	static void SetDeviceProfileCVars(const FString& DeviceProfileName);
 
 
 	/** Read and process all of the fragment matching rules. Returns an array containing the names of fragments selected. */
diff --git a/Engine/Source/Runtime/Engine/Private/DeviceProfiles/DeviceProfile.cpp b/Engine/Source/Runtime/Engine/Private/DeviceProfiles/DeviceProfile.cpp
--- a/Engine/Source/Runtime/Engine/Private/DeviceProfiles/DeviceProfile.cpp
+++ b/Engine/Source/Runtime/Engine/Private/DeviceProfiles/DeviceProfile.cpp
@@ -387,12 +387,31 @@ const TMap<FString, FString>& UDeviceProfile::GetConsolidatedCVars() const
 
 
 #if ALLOW_OTHER_PLATFORM_CONFIG
+
+void UDeviceProfile::ExpandDeviceProfileCVars()
+{
+	// VisitPlatformCVarsForEmulation can't access Scalability.h, so make sure it doesn't change away from 3, or if it does, to fix up the hardcoded number
+	static_assert(Scalability::DefaultQualityLevel == 3, "If this trips, update this and IConsoleManager::VisitPlatformCVarsForEmulation with the new value!");
+
+	TMap<FString, FString> CVarsToAdd;
+	TMap<FString, FString> PreviewCVars;
+	IConsoleManager::VisitPlatformCVarsForEmulation(*DeviceType, GetName(),
+		[this](const FString& CVarName, const FString& CVarValue, EConsoleVariableFlags SetByAndPreview)
+		{
+			AllExpandedCVars.Add(CVarName, CVarValue);
+			if (SetByAndPreview & EConsoleVariableFlags::ECVF_Preview)
+			{
+				AllPreviewCVars.Add(CVarName, CVarValue);
+			}
+		});
+}
+
 const TMap<FString, FString>& UDeviceProfile::GetAllExpandedCVars()
 {
 	// expand on first use
 	if (AllExpandedCVars.Num() == 0)
 	{
-		UDeviceProfileManager::Get().ExpandDeviceProfileCVars(this);
+		ExpandDeviceProfileCVars();
 	}
 
 	return AllExpandedCVars;
@@ -403,29 +422,12 @@ const TMap<FString, FString>& UDeviceProfile::GetAllPreviewCVars()
 	// expand on first use
 	if (AllPreviewCVars.Num() == 0)
 	{
-		UDeviceProfileManager::Get().ExpandDeviceProfileCVars(this);
+		ExpandDeviceProfileCVars();
 	}
 
 	return AllPreviewCVars;
 }
 
-void UDeviceProfile::AddExpandedCVars(const TMap<FString, FString>& CVarsToMerge)
-{
-	// merge the cvars, overwriting any existing ones
-	for (const auto& Pair : CVarsToMerge)
-	{
-		AllExpandedCVars.Add(Pair.Key, Pair.Value);
-	}
-}
-
-void UDeviceProfile::AddPreviewCVars(const TMap<FString, FString>& CVarsToMerge)
-{
-	for (const auto& Pair : CVarsToMerge)
-	{
-		AllPreviewCVars.Add(Pair.Key, Pair.Value);
-	}
-}
-
 void UDeviceProfile::ClearAllExpandedCVars()
 {
 	AllExpandedCVars.Empty();
diff --git a/Engine/Source/Runtime/Engine/Private/DeviceProfiles/DeviceProfileManager.cpp b/Engine/Source/Runtime/Engine/Private/DeviceProfiles/DeviceProfileManager.cpp
--- a/Engine/Source/Runtime/Engine/Private/DeviceProfiles/DeviceProfileManager.cpp
+++ b/Engine/Source/Runtime/Engine/Private/DeviceProfiles/DeviceProfileManager.cpp
@@ -16,11 +16,14 @@
 #include "Interfaces/ITargetPlatformManagerModule.h"
 #include "PlatformInfo.h"
 #include "PIEPreviewDeviceProfileSelectorModule.h"
+#include "Framework/Notifications/NotificationManager.h"
+#include "Widgets/Notifications/SNotificationList.h"
 #endif
 #include "ProfilingDebugging/CsvProfiler.h"
 #include "DeviceProfiles/DeviceProfileFragment.h"
 #include "GenericPlatform/GenericPlatformCrashContext.h"
 
+
 DEFINE_LOG_CATEGORY_STATIC(LogDeviceProfileManager, Log, All);
 
 static TAutoConsoleVariable<FString> CVarDeviceProfileOverride(
@@ -59,10 +62,9 @@ UDeviceProfileManager& UDeviceProfileManager::Get(bool bFromPostCDOContruct)
 		DeviceProfileManagerSingleton = NewObject<UDeviceProfileManager>();
 
 		DeviceProfileManagerSingleton->AddToRoot();
-		if (!FPlatformProperties::RequiresCookedData())
-		{
-			DeviceProfileManagerSingleton->LoadProfiles();
-		}
+#if ALLOW_OTHER_PLATFORM_CONFIG
+		DeviceProfileManagerSingleton->LoadProfiles();
+#endif
 
 		// always start with an active profile, even if we create it on the spot
 		UDeviceProfile* ActiveProfile = DeviceProfileManagerSingleton->FindProfile(GetPlatformDeviceProfileName());
@@ -85,6 +87,8 @@ UDeviceProfileManager& UDeviceProfileManager::Get(bool bFromPostCDOContruct)
 		);
 
 		InitializeSharedSamplerStates();
+
+		FCoreDelegates::GatherDeviceProfileCVars.BindLambda([](const FString& DeviceProfileName) { return UDeviceProfileManager::Get().GatherDeviceProfileCVars(DeviceProfileName, EDeviceProfileMode::DPM_CacheValues); });
 	}
 	return *DeviceProfileManagerSingleton;
 }
@@ -142,136 +146,55 @@ static void ExpandScalabilityCVar(FConfigCacheIni* ConfigSystem, const FString&
 	}
 }
 
-void UDeviceProfileManager::ProcessDeviceProfileIniSettings(const FString& DeviceProfileName, EDeviceProfileMode Mode)
+TMap<FName, FString> UDeviceProfileManager::GatherDeviceProfileCVars(const FString& DeviceProfileName, EDeviceProfileMode GatherMode)
 {
 	FConfigCacheIni* ConfigSystem = GConfig;
-	EPlatformMemorySizeBucket PreviewMemorySizeBucket = EPlatformMemorySizeBucket::Default;
+	// build up cvars into a map
+	TMap<FName, FString> DeviceProfileCVars;
+	TArray<FSelectedFragmentProperties> FragmentsSelected;
+	TSet<FString> FragmentCVarKeys;
+	TArray<FString> SelectedFragmentCVars;
+
 
-	if (Mode == EDeviceProfileMode::DPM_CacheValues)
+	EPlatformMemorySizeBucket MemBucket = FPlatformMemory::GetMemorySizeBucket();
+	// if caching (for another platfomr), then we use the DP's PreviewMemoryBucket, instead of querying for it
+	if (GatherMode == EDeviceProfileMode::DPM_CacheValues)
 	{
 #if ALLOW_OTHER_PLATFORM_CONFIG
 		// caching is not done super early, so we can assume DPs have been found now
-		UDeviceProfile* Profile = UDeviceProfileManager::Get().FindProfile(DeviceProfileName, false);
-		check(Profile);
+		UDeviceProfile* Profile = UDeviceProfileManager::Get().FindProfile(DeviceProfileName, true);
+		if (Profile == nullptr)
+		{
+			UE_LOG(LogDeviceProfileManager, Log, TEXT("Unable to find DeviceProfile '%s' for gathering cvars for caching. Aborting...."), *DeviceProfileName);
+			check(Profile);
+			return DeviceProfileCVars;
+		}
 		// use the DP's platform's configs, NOT the running platform
 		ConfigSystem = FConfigCacheIni::ForPlatform(*Profile->DeviceType);
-		PreviewMemorySizeBucket = Profile->GetPreviewMemorySizeBucket();
+		MemBucket = Profile->GetPreviewMemorySizeBucket();
+
+		FragmentsSelected = FindMatchingFragments(DeviceProfileName, ConfigSystem);
 #else
 		checkNoEntry();
 #endif
 	}
-
-	if (Mode == EDeviceProfileMode::DPM_SetCVars)
-	{
-		UE_LOG(LogDeviceProfileManager, Log, TEXT("Applying CVar settings loaded from the selected device profile: [%s]"), *DeviceProfileName);
-	}
-
-	check(ConfigSystem);
-
-	TArray< FString > AvailableProfiles;
-	ConfigSystem->GetSectionNames( GDeviceProfilesIni, AvailableProfiles );
-
-	// Look up the ini for this tree as we are far too early to use the UObject system
-	AvailableProfiles.Remove( TEXT( "DeviceProfiles" ) );
-
-	// Next we need to create a hierarchy of CVars from the Selected Device Profile, to it's eldest parent
-	// if we are just caching, this also contains the set of all CVars (including expanding scalability groups)
-	TMap<FString, FString> CVarsAlreadySetList;
-
-	// previewable subset of CVarsAlreadySetList
-	TMap<FString, FString> PreviewableCVars;
-
-	// reset some global state for "active DP" mode
-	if (Mode == EDeviceProfileMode::DPM_SetCVars)
+	else
 	{
-		DeviceProfileScalabilityCVars.Empty();
-
-		// we should have always pushed away old values by the time we get here
-		check(PushedSettings.Num() == 0);
-
-#if !UE_BUILD_SHIPPING
-#if PLATFORM_ANDROID
-		// allow ConfigRules to override cvars first
-		const TMap<FString, FString>& ConfigRules = FAndroidMisc::GetConfigRulesTMap();
-		for (const TPair<FString, FString>& Pair : ConfigRules)
+		// the very first time, this is called to set cvars, we cache the current device's fragments
+		if (PlatformFragmentsSelected.Num() == 0)
 		{
-			FString Key = Pair.Key;
-			if (Key.StartsWith("cvar_"))
-			{
-				FString CVarKey = Key.RightChop(5);
-				FString CVarValue = Pair.Value;
-
-				UE_LOG(LogDeviceProfileManager, Log, TEXT("Setting ConfigRules Device Profile CVar: [[%s:%s]]"), *CVarKey, *CVarValue);
-
-				// set it and remember it
-				OnSetCVarFromIniEntry(*GDeviceProfilesIni, *CVarKey, *CVarValue, ECVF_SetByDeviceProfile);
-				CVarsAlreadySetList.Add(CVarKey, CVarValue);
-			}
-		}
-#endif
-#endif
-
-#if !UE_BUILD_SHIPPING
-		// pre-apply any -dpcvars= items, so that they override anything in the DPs
-		FString DPCVarString;
-		if (FParse::Value(FCommandLine::Get(), TEXT("DPCVars="), DPCVarString, false) || FParse::Value(FCommandLine::Get(), TEXT("DPCVar="), DPCVarString, false))
-		{
-			// look over a list of cvars
-			TArray<FString> DPCVars;
-			DPCVarString.ParseIntoArray(DPCVars, TEXT(","), true);
-			for (const FString& DPCVar : DPCVars)
-			{
-				// split up each Key=Value pair
-				FString CVarKey, CVarValue;
-				if (DPCVar.Split(TEXT("="), &CVarKey, &CVarValue))
-				{
-					UE_LOG(LogDeviceProfileManager, Log, TEXT("Setting CommandLine Device Profile CVar: [[%s:%s]]"), *CVarKey, *CVarValue);
-
-					// set it and remember it (no thanks, Ron Popeil)
-					OnSetCVarFromIniEntry(*GDeviceProfilesIni, *CVarKey, *CVarValue, ECVF_SetByDeviceProfile);
-					CVarsAlreadySetList.Add(CVarKey, CVarValue);
-				}
-			}
-		}
-#endif
-
-		// Preload a cvar we rely on
-		if (FConfigSection* Section = ConfigSystem->GetSectionPrivate(TEXT("ConsoleVariables"), false, true, *GEngineIni))
-		{
-			static FName AllowScalabilityAtRuntimeName = TEXT("dp.AllowScalabilityGroupsToChangeAtRuntime");
-			if (const FConfigValue* Value = Section->Find(AllowScalabilityAtRuntimeName))
-			{
-				const FString& KeyString = AllowScalabilityAtRuntimeName.ToString();
-				const FString& ValueString = Value->GetValue();
-				OnSetCVarFromIniEntry(*GEngineIni, *KeyString, *ValueString, ECVF_SetBySystemSettingsIni);
-			}
+			PlatformFragmentsSelected = FindMatchingFragments(DeviceProfileName, GConfig);
 		}
+		FragmentsSelected = PlatformFragmentsSelected;
 	}
 
-	FString SectionSuffix = *FString::Printf(TEXT(" %s"), *UDeviceProfile::StaticClass()->GetName());
 
-	TSet<FString> FragmentCVarKeys;
-	TArray<FString> SelectedFragmentCVars;
-	TArray<FSelectedFragmentProperties> FragmentsSelected;
+	// directly look in the .ini instead of using FindDeviceProfile - not all of the parent DPs may be in existence, but we still need to go up the chain
+	// this will cache the DP names we know about
+	TArray<FString> AvailableProfiles;
+	ConfigSystem->GetSectionNames(GDeviceProfilesIni, AvailableProfiles);
+	AvailableProfiles.Remove(TEXT("DeviceProfiles"));
 
-	// Process the fragment matching rules.
-	// Only perform the matching process for the base DP for simplicity.
-	// Perform the matching process if no fragments were selected or we are previewing. re-use the matched array if already present.
-	if (Mode == EDeviceProfileMode::DPM_CacheValues || PlatformFragmentsSelected.Num() == 0)
-	{
-		// run the matching rules.
-		FragmentsSelected = FindMatchingFragments(DeviceProfileName, ConfigSystem);
-		if (Mode != EDeviceProfileMode::DPM_CacheValues)
-		{
-			// Store the newly selected fragments.
-			PlatformFragmentsSelected = FragmentsSelected;
-		}
-	}
-	else if (Mode != EDeviceProfileMode::DPM_CacheValues)
-	{
-		// use the existing selected fragment state.
-		FragmentsSelected = PlatformFragmentsSelected;
-	}
 
 	// Here we gather the cvars from selected fragments in reverse order
 	for (int i = FragmentsSelected.Num() - 1; i >= 0; i--)
@@ -296,17 +219,18 @@ void UDeviceProfileManager::ProcessDeviceProfileIniSettings(const FString& Devic
 		}
 	}
 
+	static FString SectionSuffix = *FString::Printf(TEXT(" %s"), *UDeviceProfile::StaticClass()->GetName());
+
 	// For each device profile, starting with the selected and working our way up the BaseProfileName tree,
-	// Find all CVars and set them 
 	FString BaseDeviceProfileName = DeviceProfileName;
 	bool bReachedEndOfTree = BaseDeviceProfileName.IsEmpty();
-	while( bReachedEndOfTree == false ) 
+	while (bReachedEndOfTree == false)
 	{
 		FString CurrentSectionName = BaseDeviceProfileName + SectionSuffix;
-		
-		// Check the profile was available.
-		bool bProfileExists = AvailableProfiles.Contains( CurrentSectionName );
-		if( bProfileExists )
+
+		// check if there is a section named for the DeviceProfile
+		FConfigSection* CurrentSection = ConfigSystem->GetSectionPrivate(*CurrentSectionName, false, true, GDeviceProfilesIni);
+		if (CurrentSection != nullptr)
 		{
 			// put this up in some shared code somewhere in FGenericPlatformMemory
 			const TCHAR* BucketNames[] = {
@@ -315,7 +239,7 @@ void UDeviceProfileManager::ProcessDeviceProfileIniSettings(const FString& Devic
 				TEXT("_Default"),
 				TEXT("_Smaller"),
 				TEXT("_Smallest"),
-                TEXT("_Tiniest"),
+				TEXT("_Tiniest"),
 			};
 
 			for (int Pass = 0; Pass < 2; Pass++)
@@ -324,20 +248,12 @@ void UDeviceProfileManager::ProcessDeviceProfileIniSettings(const FString& Devic
 				FString ArrayName = TEXT("CVars");
 				if (Pass == 0)
 				{
-					// PreviewMemorySizeBucket is used when previewing, since we don't have a current device to emulate.
-					if (Mode == EDeviceProfileMode::DPM_CacheValues)
-					{
-						ArrayName += BucketNames[(int32)PreviewMemorySizeBucket];
-					}
-					else
-					{
-						ArrayName += BucketNames[(int32)FPlatformMemory::GetMemorySizeBucket()];
-					}
+					ArrayName += BucketNames[(int32)MemBucket];
 				}
 
 				TArray< FString > CurrentProfilesCVars, FragmentCVars;
 				GetCVarsFromDPFragmentIncludes(*CurrentSectionName, *ArrayName, FragmentCVars, ConfigSystem);
-				ConfigSystem->GetArray(*CurrentSectionName, *ArrayName, CurrentProfilesCVars, GDeviceProfilesIni);
+				CurrentSection->MultiFind(*ArrayName, CurrentProfilesCVars, true);
 
 				if (FragmentCVars.Num())
 				{
@@ -378,68 +294,10 @@ void UDeviceProfileManager::ProcessDeviceProfileIniSettings(const FString& Devic
 					FString CVarKey, CVarValue;
 					if ((*CVarIt).Split(TEXT("="), &CVarKey, &CVarValue))
 					{
-						if (!CVarsAlreadySetList.Find(CVarKey))
+						FName CVarKeyName(*CVarKey);
+						if (!DeviceProfileCVars.Find(CVarKeyName))
 						{
-							IConsoleVariable* CVar = IConsoleManager::Get().FindConsoleVariable(*CVarKey);
-							if (Mode == EDeviceProfileMode::DPM_SetCVars)
-							{
-								if (CVar)
-								{
-									// remember the previous value and priority
-									FPushedCVarSetting OldSetting = FPushedCVarSetting(CVar->GetString(), CVar->GetFlags());
-									PushedSettings.Add(CVarKey, OldSetting);
-
-									// indicate we are pushing, not setting
-									UE_LOG(LogDeviceProfileManager, Log, TEXT("Pushing Device Profile CVar: [[%s:%s -> %s]]"), *CVarKey, *OldSetting.Value, *CVarValue);
-								}
-								else
-								{
-									UE_LOG(LogDeviceProfileManager, Warning, TEXT("Creating unregistered Device Profile CVar: [[%s:%s]]"), *CVarKey, *CVarValue);
-								}
-							}
-
-							// General scalability bucket cvars are set as a suggested default but can be overridden by game settings.
-							bool bIsScalabilityBucket = CVarKey.StartsWith(TEXT("sg."));
-
-							if (Mode == EDeviceProfileMode::DPM_CacheValues)
-							{
-								if (bIsScalabilityBucket)
-								{
-									// don't overwrite any existing cvars when expanding
-									ExpandScalabilityCVar(ConfigSystem, CVarKey, CVarValue, CVarsAlreadySetList, false);
-								}
-
-								// cache key with value
-								CVarsAlreadySetList.Add(CVarKey, CVarValue);
-
-#if WITH_EDITOR
-								// if this is previewable, add it to the list - and also make any expanded scalability cvars previewable
-								if (CVar != nullptr && CVar->TestFlags(EConsoleVariableFlags::ECVF_Preview))
-								{
-									if (bIsScalabilityBucket)
-									{
-										// don't overwrite any existing cvars when expanding
-										ExpandScalabilityCVar(ConfigSystem, CVarKey, CVarValue, PreviewableCVars, false);
-									}
-									PreviewableCVars.Add(CVarKey, CVarValue);
-								}
-#endif
-							}
-							// actually set the cvar if not just caching
-							else
-							{
-								
-								// Cache any scalability related cvars so we can conveniently reapply them later as a way to reset the device defaults
-								if (bIsScalabilityBucket && CVarAllowScalabilityGroupsToChangeAtRuntime.GetValueOnGameThread() > 0)
-								{
-									DeviceProfileScalabilityCVars.Add(*CVarKey, *CVarValue);
-								}
-
-								//If this is a dp preview then we set cvars with their existing priority so that we don't cause future issues when setting by scalability levels etc.
-								uint32 CVarPriority = bIsScalabilityBucket ? ECVF_SetByScalability : ECVF_SetByDeviceProfile;
-								OnSetCVarFromIniEntry(*GDeviceProfilesIni, *CVarKey, *CVarValue, CVarPriority);
-								CVarsAlreadySetList.Add(CVarKey, CVarValue);
-							}
+							DeviceProfileCVars.Add(CVarKeyName, CVarValue);
 						}
 					}
 				}
@@ -457,22 +315,111 @@ void UDeviceProfileManager::ProcessDeviceProfileIniSettings(const FString& Devic
 				BaseDeviceProfileName.Empty();
 			}
 		}
-		
+
 		// Check if we have inevitably reached the end of the device profile tree.
-		bReachedEndOfTree = !bProfileExists || BaseDeviceProfileName.IsEmpty();
+		bReachedEndOfTree = CurrentSection == nullptr || BaseDeviceProfileName.IsEmpty();
 	}
 
-#if ALLOW_OTHER_PLATFORM_CONFIG
-	// copy the running cache into the DP
-	if (Mode == EDeviceProfileMode::DPM_CacheValues)
+	return DeviceProfileCVars;
+}
+
+void UDeviceProfileManager::SetDeviceProfileCVars(const FString& DeviceProfileName)
+{
+	// walk over the parent chain, gathering the cvars this DP inherits and contains
+	TMap<FName, FString> DeviceProfileCVars = GatherDeviceProfileCVars(DeviceProfileName, EDeviceProfileMode::DPM_SetCVars);
+
+	// reset some global state
+	DeviceProfileScalabilityCVars.Empty();
+
+	// we should have always popped away old values by the time we get here
+	check(PushedSettings.Num() == 0);
+
+	// Preload a cvar we rely on in the loop below
+	if (FConfigSection* Section = GConfig->GetSectionPrivate(TEXT("ConsoleVariables"), false, true, *GEngineIni))
+	{
+		static FName AllowScalabilityAtRuntimeName = TEXT("dp.AllowScalabilityGroupsToChangeAtRuntime");
+		if (const FConfigValue* Value = Section->Find(AllowScalabilityAtRuntimeName))
+		{
+			const FString& KeyString = AllowScalabilityAtRuntimeName.ToString();
+			const FString& ValueString = Value->GetValue();
+			OnSetCVarFromIniEntry(*GEngineIni, *KeyString, *ValueString, ECVF_SetBySystemSettingsIni);
+		}
+	}
+
+	// for each DP cvar - push the old value to restore, cache ScalabiliutyGroups, and then set the values
+	for (TPair<FName, FString>& Pair : DeviceProfileCVars)
+	{
+		FString CVarKey = Pair.Key.ToString();
+		const FString& CVarValue = Pair.Value;
+
+		// get the actual cvar object
+		IConsoleVariable* CVar = IConsoleManager::Get().FindConsoleVariable(*CVarKey);
+		if (CVar)
+		{
+			// remember the previous value and priority
+			FPushedCVarSetting OldSetting = FPushedCVarSetting(CVar->GetString(), CVar->GetFlags());
+			PushedSettings.Add(CVarKey, OldSetting);
+
+			// indicate we are pushing, not setting
+			UE_LOG(LogDeviceProfileManager, Log, TEXT("Pushing Device Profile CVar: [[%s:%s -> %s]]"), *CVarKey, *OldSetting.Value, *CVarValue);
+		}
+		else
+		{
+			UE_LOG(LogDeviceProfileManager, Warning, TEXT("Creating unregistered Device Profile CVar: [[%s:%s]]"), *CVarKey, *Pair.Value);
+		}
+
+		// Cache any scalability related cvars so we can conveniently reapply them later as a way to reset the device defaults
+		bool bIsScalabilityBucket = CVarKey.StartsWith(TEXT("sg."));
+		if (bIsScalabilityBucket && CVarAllowScalabilityGroupsToChangeAtRuntime.GetValueOnGameThread() > 0)
+		{
+			DeviceProfileScalabilityCVars.Add(*CVarKey, *CVarValue);
+		}
+
+		// Set by scalability or DP, depending
+		uint32 CVarPriority = bIsScalabilityBucket ? ECVF_SetByScalability : ECVF_SetByDeviceProfile;
+		OnSetCVarFromIniEntry(*GDeviceProfilesIni, *CVarKey, *CVarValue, CVarPriority);
+	}
+
+
+	// now allow for ovverrides in !SHIPPING
+#if !UE_BUILD_SHIPPING
+#if PLATFORM_ANDROID
+	// allow ConfigRules to override cvars first
+	const TMap<FString, FString>& ConfigRules = FAndroidMisc::GetConfigRulesTMap();
+	for (const TPair<FString, FString>& Pair : ConfigRules)
 	{
-		UDeviceProfile* Profile = UDeviceProfileManager::Get().FindProfile(DeviceProfileName, false);
-		Profile->AddExpandedCVars(CVarsAlreadySetList);
+		FString Key = Pair.Key;
+		if (Key.StartsWith("cvar_"))
+		{
+			FString CVarKey = Key.RightChop(5);
+			FString CVarValue = Pair.Value;
 
-		// now remember all previewable cvars
-		Profile->AddPreviewCVars(PreviewableCVars);
+			UE_LOG(LogDeviceProfileManager, Log, TEXT("Setting ConfigRules Device Profile CVar: [[%s:%s]]"), *CVarKey, *CVarValue);
 
-		Profile->SelectedFragments = FragmentsSelected;
+			// set it and remember it
+			OnSetCVarFromIniEntry(*GDeviceProfilesIni, *CVarKey, *CVarValue, ECVF_SetByDeviceProfile);
+		}
+	}
+#endif
+	// pre-apply any -dpcvars= items, so that they override anything in the DPs
+	FString DPCVarString;
+	if (FParse::Value(FCommandLine::Get(), TEXT("DPCVars="), DPCVarString, false) || FParse::Value(FCommandLine::Get(), TEXT("DPCVar="), DPCVarString, false))
+	{
+		// look over a list of cvars
+		TArray<FString> DPCVars;
+		DPCVarString.ParseIntoArray(DPCVars, TEXT(","), true);
+		for (const FString& DPCVar : DPCVars)
+		{
+			// split up each Key=Value pair
+			FString CVarKey, CVarValue;
+			if (DPCVar.Split(TEXT("="), &CVarKey, &CVarValue))
+			{
+				UE_LOG(LogDeviceProfileManager, Log, TEXT("Setting CommandLine Device Profile CVar: [[%s:%s]]"), *CVarKey, *CVarValue);
+
+				// set it and remember it (no thanks, Ron Popeil)
+				OnSetCVarFromIniEntry(*GDeviceProfilesIni, *CVarKey, *CVarValue, ECVF_SetByDeviceProfile);
+			}
+		}
 	}
 #endif
 }
@@ -535,7 +482,7 @@ void UDeviceProfileManager::ChangeTaggedFragmentState(FName FragmentTag, bool bN
 				RestorePushedState(PushedSettings);
 				// set the new state and reapply all fragments.
 				Fragment.bEnabled = bNewState;
-				ProcessDeviceProfileIniSettings(GetActiveDeviceProfileName(), EDeviceProfileMode::DPM_SetCVars);
+				SetDeviceProfileCVars(GetActiveDeviceProfileName());
 			}
 			break;
 		}
@@ -555,36 +502,9 @@ void UDeviceProfileManager::InitializeCVarsForActiveDeviceProfile()
 		ActiveProfileName = GetPlatformDeviceProfileName();
 	}
 
-	ProcessDeviceProfileIniSettings(ActiveProfileName, EDeviceProfileMode::DPM_SetCVars);
+	SetDeviceProfileCVars(ActiveProfileName);
 }
 
-#if ALLOW_OTHER_PLATFORM_CONFIG
-void UDeviceProfileManager::ExpandDeviceProfileCVars(UDeviceProfile* DeviceProfile)
-{
-	TMap<FString, FString> CVarsToAdd;
-	TMap<FString, FString> PreviewCVars;
-	IConsoleManager::VisitPlatformCVarsForEmulation(*DeviceProfile->DeviceType, false,
-		[&CVarsToAdd, &PreviewCVars](const FString& CVarName, const FString& CVarValue, EConsoleVariableFlags SetByAndPreview)
-		{
-			CVarsToAdd.Add(CVarName, CVarValue);
-			if (SetByAndPreview & EConsoleVariableFlags::ECVF_Preview)
-			{
-				PreviewCVars.Add(CVarName, CVarValue);
-			}
-		});
-
-	// we now have all the cvars for this platform without _any_ device profiles, so now add those into the map
-	DeviceProfile->AddExpandedCVars(CVarsToAdd);
-	DeviceProfile->AddPreviewCVars(PreviewCVars);
-
-	// and finalize them into a cache
-	ProcessDeviceProfileIniSettings(DeviceProfile->GetName(), EDeviceProfileMode::DPM_CacheValues);
-}
-#endif
-
-
-
-
 
 bool UDeviceProfileManager::DoActiveProfilesReference(const TSet<FString>& DeviceProfilesToQuery)
 {
