- Fixed compile error for monolithics

diff --git a/Engine/Source/Runtime/Engine/Private/DeviceProfiles/DeviceProfileManager.cpp b/Engine/Source/Runtime/Engine/Private/DeviceProfiles/DeviceProfileManager.cpp
--- a/Engine/Source/Runtime/Engine/Private/DeviceProfiles/DeviceProfileManager.cpp
+++ b/Engine/Source/Runtime/Engine/Private/DeviceProfiles/DeviceProfileManager.cpp
@@ -88,7 +88,9 @@ UDeviceProfileManager& UDeviceProfileManager::Get(bool bFromPostCDOContruct)
 
 		InitializeSharedSamplerStates();
 
+#if ALLOW_OTHER_PLATFORM_CONFIG
 		FCoreDelegates::GatherDeviceProfileCVars.BindLambda([](const FString& DeviceProfileName) { return UDeviceProfileManager::Get().GatherDeviceProfileCVars(DeviceProfileName, EDeviceProfileMode::DPM_CacheValues); });
+#endif
 	}
 	return *DeviceProfileManagerSingleton;
 }
