Upgrade SoundWave to use FEditorBulkData data

diff --git a/Engine/Plugins/Runtime/AudioCapture/Source/AudioCaptureEditor/Private/AudioRecordingManager.cpp b/Engine/Plugins/Runtime/AudioCapture/Source/AudioCaptureEditor/Private/AudioRecordingManager.cpp
--- a/Engine/Plugins/Runtime/AudioCapture/Source/AudioCaptureEditor/Private/AudioRecordingManager.cpp
+++ b/Engine/Plugins/Runtime/AudioCapture/Source/AudioCaptureEditor/Private/AudioRecordingManager.cpp
@@ -407,10 +407,8 @@ void FAudioRecordingManager::StopRecording(TArray<USoundWave*>& OutSoundWaves)
 					NewSoundWave->InvalidateCompressedData(true, false);
 
 					// Copy the raw wave data file to the sound wave for storage. Will allow the recording to be exported.
-					NewSoundWave->RawData.Lock(LOCK_READ_WRITE);
-					void* LockedData = NewSoundWave->RawData.Realloc(RawWaveData.Num());
-					FMemory::Memcpy(LockedData, RawWaveData.GetData(), RawWaveData.Num());
-					NewSoundWave->RawData.Unlock();
+					FSharedBuffer UpdatedBuffed = FSharedBuffer::Clone(RawWaveData.GetData(), RawWaveData.Num());
+					NewSoundWave->RawData.UpdatePayload(UpdatedBuffed);
 
 					bCreatedPackage = true;
 				}
diff --git a/Engine/Source/Editor/AudioEditor/Private/Factories/SoundFactory.cpp b/Engine/Source/Editor/AudioEditor/Private/Factories/SoundFactory.cpp
--- a/Engine/Source/Editor/AudioEditor/Private/Factories/SoundFactory.cpp
+++ b/Engine/Source/Editor/AudioEditor/Private/Factories/SoundFactory.cpp
@@ -479,9 +479,9 @@ UObject* USoundFactory::CreateObject
 			// copy the data into the bulk byte data
 
 			// Get the raw data bulk byte pointer and copy over the .wav files we generated
-			Sound->RawData.Lock(LOCK_READ_WRITE);
 
-			uint8* LockedData = (uint8*)Sound->RawData.Realloc(TotalSize);
+			FUniqueBuffer EditableBuffer = FUniqueBuffer::Alloc(TotalSize);
+			uint8* LockedData = (uint8*)EditableBuffer.GetData();
 			int16* LockedDataInt16 = reinterpret_cast<int16*>(LockedData);
 
 			int32 RawDataOffset = 0;
@@ -524,15 +524,15 @@ UObject* USoundFactory::CreateObject
 					RawDataOffset += ChannelSize;
 				}
 			}
-			Sound->RawData.Unlock();
+
+			Sound->RawData.UpdatePayload(EditableBuffer.MoveToShared());
 		}
 		else
 		{
 			// For mono and stereo assets, just copy the data into the buffer
-			Sound->RawData.Lock(LOCK_READ_WRITE);
-			void* LockedData = Sound->RawData.Realloc(RawWaveDataBufferSize);
-			FMemory::Memcpy(LockedData, Buffer, RawWaveDataBufferSize);
-			Sound->RawData.Unlock();
+			FSharedBuffer UpdatedBuffer = FSharedBuffer::Clone(RawWaveData.GetData(), RawWaveData.Num());
+			Sound->RawData.UpdatePayload(UpdatedBuffer);
+
 		}
 
 		Sound->Duration = (float)NumFrames / *WaveInfo.pSamplesPerSec;
diff --git a/Engine/Source/Editor/UnrealEd/Private/EditorExporters.cpp b/Engine/Source/Editor/UnrealEd/Private/EditorExporters.cpp
--- a/Engine/Source/Editor/UnrealEd/Private/EditorExporters.cpp
+++ b/Engine/Source/Editor/UnrealEd/Private/EditorExporters.cpp
@@ -151,9 +151,15 @@ bool USoundExporterWAV::SupportsObject(UObject* Object) const
 bool USoundExporterWAV::ExportBinary( UObject* Object, const TCHAR* Type, FArchive& Ar, FFeedbackContext* Warn, int32 FileIndex, uint32 PortFlags )
 {
 	USoundWave* Sound = CastChecked<USoundWave>( Object );
-	void* RawWaveData = Sound->RawData.Lock( LOCK_READ_ONLY );
-	Ar.Serialize( RawWaveData, Sound->RawData.GetBulkDataSize() );
-	Sound->RawData.Unlock();
+
+	// We a non-const pointer for the sake of calling serialize, but this is on writing FArchive so it's never modified
+	// Assert that's the case
+	if(ensure(Ar.IsSaving()))
+	{	
+		TFuture<FSharedBuffer> FutureBuffer = Sound->RawData.GetPayload(); // This will block, but for now that's ok.
+		void* RawWaveData = (void*)FutureBuffer.Get().GetData();
+		Ar.Serialize(RawWaveData, FutureBuffer.Get().GetSize());
+	}
 	return true;
 }
 
@@ -254,15 +260,18 @@ bool USoundSurroundExporterWAV::ExportBinary( UObject* Object, const TCHAR* Type
 
 	USoundWave* Sound = CastChecked<USoundWave>( Object );
 	if ( Sound->ChannelSizes.Num() > 0 )
-	{
-		uint8* RawWaveData = ( uint8* )Sound->RawData.Lock( LOCK_READ_ONLY );
-
-		if( Sound->ChannelSizes[ FileIndex ] )
+	{		
+		if (ensure(Ar.IsSaving()))
 		{
-			Ar.Serialize( RawWaveData + Sound->ChannelOffsets[ FileIndex ], Sound->ChannelSizes[ FileIndex ] );
-		}
+			// Non-const pointer for the sake of calling serialize, so assert that we are writing an archive.
+			TFuture<FSharedBuffer> FutureBuffer = Sound->RawData.GetPayload();
+			uint8* RawWaveData = (uint8*) FutureBuffer.Get().GetData();
 
-		Sound->RawData.Unlock();
+			if( Sound->ChannelSizes[ FileIndex ] )
+			{
+				Ar.Serialize( RawWaveData + Sound->ChannelOffsets[ FileIndex ], Sound->ChannelSizes[ FileIndex ] );
+			}
+		}
 
 		bResult = Sound->ChannelSizes[ FileIndex ] != 0;
 	}
diff --git a/Engine/Source/Editor/UnrealEd/Private/ThumbnailRendering/SoundWaveThumbnailRenderer.cpp b/Engine/Source/Editor/UnrealEd/Private/ThumbnailRendering/SoundWaveThumbnailRenderer.cpp
--- a/Engine/Source/Editor/UnrealEd/Private/ThumbnailRendering/SoundWaveThumbnailRenderer.cpp
+++ b/Engine/Source/Editor/UnrealEd/Private/ThumbnailRendering/SoundWaveThumbnailRenderer.cpp
@@ -31,8 +31,7 @@ void USoundWaveThumbnailRenderer::Draw(UObject* Object, int32 X, int32 Y, uint32
 	}
 
 	SoundWave->SetRedrawThumbnail(false);
-
-	if (SoundWave->NumChannels == 0 || SoundWave->RawData.GetBulkDataSize() == 0)
+	if (SoundWave->NumChannels == 0 || !SoundWave->RawData.HasPayloadData())
 	{
 		return;
 	}
@@ -41,9 +40,10 @@ void USoundWaveThumbnailRenderer::Draw(UObject* Object, int32 X, int32 Y, uint32
 	FCanvasLineItem LineItem;
 	LineItem.SetColor(FLinearColor::White);
 
-	uint8* RawWaveData = (uint8*)SoundWave->RawData.Lock(LOCK_READ_ONLY);
-	int32 RawDataSize = SoundWave->RawData.GetBulkDataSize();
- 
+	TFuture<FSharedBuffer> BufferFuture = SoundWave->RawData.GetPayload();
+	const uint8* RawWaveData = (const uint8 *)BufferFuture.Get().GetData();  // Bulk
+	int32 RawDataSize = BufferFuture.Get().GetSize();
+	 
 	// Compute the scaled y-value used to render the channel data
 	const float SampleYScale = Height / (2.f * 32767 * SoundWave->NumChannels);
 
@@ -168,7 +168,6 @@ void USoundWaveThumbnailRenderer::Draw(UObject* Object, int32 X, int32 Y, uint32
 		CurrentRawWaveByteIndex += RawChannelFileByteSize;
 	}
 
-	SoundWave->RawData.Unlock();
 }
 
 bool USoundWaveThumbnailRenderer::AllowsRealtimeThumbnails(UObject* Object) const
diff --git a/Engine/Source/Runtime/Core/Public/UObject/UE5MainStreamObjectVersion.h b/Engine/Source/Runtime/Core/Public/UObject/UE5MainStreamObjectVersion.h
--- a/Engine/Source/Runtime/Core/Public/UObject/UE5MainStreamObjectVersion.h
+++ b/Engine/Source/Runtime/Core/Public/UObject/UE5MainStreamObjectVersion.h
@@ -209,6 +209,9 @@ struct CORE_API FUE5MainStreamObjectVersion
 		// Enforcing new automatic tangent behaviour, enforcing auto-tangents for Key0 and KeyN to be flat, for Animation Assets.
 		ForceUpdateAnimationAssetCurveTangents,
 
+		// SoundWave Update to use EditorBuildData for it's RawData
+		SoundWaveVirtualizationUpdate,
+
 		// -----<new versions can be added above this line>-------------------------------------------------
 		VersionPlusOne,
 		LatestVersion = VersionPlusOne - 1
diff --git a/Engine/Source/Runtime/Engine/Classes/Sound/SoundWave.h b/Engine/Source/Runtime/Engine/Classes/Sound/SoundWave.h
--- a/Engine/Source/Runtime/Engine/Classes/Sound/SoundWave.h
+++ b/Engine/Source/Runtime/Engine/Classes/Sound/SoundWave.h
@@ -16,6 +16,7 @@
 #include "Interfaces/Interface_AsyncCompilation.h"
 #include "Serialization/BulkData.h"
 #include "Serialization/BulkDataBuffer.h"
+#include "Serialization/EditorBulkData.h"
 #include "Sound/SoundGroups.h"
 #include "Sound/SoundWaveLoadingBehavior.h"
 #include "UObject/ObjectKey.h"
@@ -786,7 +787,8 @@ public:
 	const uint8* GetResourceData() const;
 
 	/** Uncompressed wav data 16 bit in mono or stereo - stereo not allowed for multichannel data */
-	FByteBulkData RawData;
+
+	UE::Serialization::FEditorBulkData RawData;
 
 	/** GUID used to uniquely identify this node so it can be found in the DDC */
 	FGuid CompressedDataGuid;
diff --git a/Engine/Source/Runtime/Engine/Private/AudioDerivedData.cpp b/Engine/Source/Runtime/Engine/Private/AudioDerivedData.cpp
--- a/Engine/Source/Runtime/Engine/Private/AudioDerivedData.cpp
+++ b/Engine/Source/Runtime/Engine/Private/AudioDerivedData.cpp
@@ -7,6 +7,7 @@
 #include "Serialization/BulkData.h"
 #include "Serialization/MemoryReader.h"
 #include "Serialization/MemoryWriter.h"
+#include "Serialization/EditorBulkData.h"
 #include "Misc/ScopedSlowTask.h"
 #include "Audio.h"
 #include "Interfaces/ITargetPlatform.h"
@@ -1238,7 +1239,7 @@ struct FAudioCookInputs
 	// copy-on-write mechanism that doesn't require us to make
 	// a copy in memory to get immutability.
 	FCriticalSection& BulkDataCriticalSection;
-	FByteBulkData& BulkData;
+	UE::Serialization::FEditorBulkData& BulkData;
 #endif
 
 	FAudioCookInputs(USoundWave* InSoundWave, FName InBaseFormat, FName InHashFormat, const FPlatformAudioCookOverrides* InCookOverrides)
@@ -1318,19 +1319,15 @@ static void CookSimpleWave(const FAudioCookInputs& Inputs, TArray<uint8>& Output
 
 	FScopeLock ScopeLock(&Inputs.BulkDataCriticalSection);
 
-	if (!Inputs.BulkData.IsBulkDataLoaded())
-	{
-		Inputs.BulkData.ForceBulkDataResident();
-	}
+	TFuture<FSharedBuffer> FutureBuffer = Inputs.BulkData.GetPayload();
 	
-	UE_CLOG(!Inputs.BulkData.IsBulkDataLoaded(), LogAudioDerivedData, Display, TEXT("Calling ForceBulkDataResident for LPCM data for USoundWave %s failed."), *Inputs.SoundFullName);
 	
-	// Lock raw wave data.
-	const uint8* RawWaveData = (const uint8*)Inputs.BulkData.LockReadOnly();
-	int32 RawDataSize = Inputs.BulkData.GetBulkDataSize();
-	ON_SCOPE_EXIT { Inputs.BulkData.Unlock(); };
 
-	if (!RawWaveData || RawDataSize == 0)
+	
+	const uint8* RawWaveData = (const uint8*)FutureBuffer.Get().GetData(); // Will block 
+	int32 RawDataSize = FutureBuffer.Get().GetSize();
+
+	if (!Inputs.BulkData.HasPayloadData())
 	{
 		UE_LOG(LogAudioDerivedData, Warning, TEXT("LPCM data failed to load for sound %s"), *Inputs.SoundFullName);
 	}
@@ -1421,18 +1418,10 @@ static void CookSurroundWave(const FAudioCookInputs& Inputs,  TArray<uint8>& Out
 	}
 
 	FScopeLock ScopeLock(&Inputs.BulkDataCriticalSection);
-
-	if (!Inputs.BulkData.IsBulkDataLoaded())
-	{
-		Inputs.BulkData.ForceBulkDataResident();
-	}
-	
-	UE_CLOG(!Inputs.BulkData.IsBulkDataLoaded(), LogAudioDerivedData, Display, TEXT("Calling ForceBulkDataResident for LPCM data for USoundWave %s failed."), *Inputs.SoundFullName);
-
 	// Lock raw wave data.
-	const uint8* RawWaveData = (const uint8*)Inputs.BulkData.LockReadOnly();
-	int32 RawDataSize = Inputs.BulkData.GetBulkDataSize();
-	ON_SCOPE_EXIT { Inputs.BulkData.Unlock(); };
+	TFuture<FSharedBuffer> FutureBuffer = Inputs.BulkData.GetPayload();
+	const uint8* RawWaveData = (const uint8*)FutureBuffer.Get().GetData(); // Will block 
+	int32 RawDataSize = FutureBuffer.Get().GetSize();
 
 	// Front left channel is the master
 	static_assert(SPEAKER_FrontLeft == 0, "Front-left speaker must be first.");
@@ -1886,9 +1875,6 @@ void USoundWave::WillNeverCacheCookedPlatformDataAgain()
 	// so need to keep the cached platform data around
 	Super::WillNeverCacheCookedPlatformDataAgain();
 
-	// TODO: We can clear these arrays if we never need to cook again.
-	RawData.RemoveBulkData();
-
 	check(SoundWaveDataPtr);
 	SoundWaveDataPtr->CompressedFormatData.FlushData();
 }
diff --git a/Engine/Source/Runtime/Engine/Private/SampleBufferIO.cpp b/Engine/Source/Runtime/Engine/Private/SampleBufferIO.cpp
--- a/Engine/Source/Runtime/Engine/Private/SampleBufferIO.cpp
+++ b/Engine/Source/Runtime/Engine/Private/SampleBufferIO.cpp
@@ -483,10 +483,8 @@ namespace Audio
 		// TODO: Check to see if we need to call USoundWave::FreeResources here.
 
 		// Emplace wav data in the RawData component of the sound wave.
-		CurrentSoundWave->RawData.Lock(LOCK_READ_WRITE);
-		void* LockedData = CurrentSoundWave->RawData.Realloc(SerializedWavData.Num());
-		FMemory::Memcpy(LockedData, SerializedWavData.GetData(), SerializedWavData.Num());
-		CurrentSoundWave->RawData.Unlock();
+		FSharedBuffer Buffer = FSharedBuffer::Clone(SerializedWavData.GetData(), SerializedWavData.Num());
+		CurrentSoundWave->RawData.UpdatePayload(Buffer);
 
 		USoundWave* SavedSoundWave = CurrentSoundWave;
 
diff --git a/Engine/Source/Runtime/Engine/Private/SoundWave.cpp b/Engine/Source/Runtime/Engine/Private/SoundWave.cpp
--- a/Engine/Source/Runtime/Engine/Private/SoundWave.cpp
+++ b/Engine/Source/Runtime/Engine/Private/SoundWave.cpp
@@ -36,6 +36,7 @@
 #include "Misc/OutputDeviceArchiveWrapper.h"
 #include "Async/Async.h"
 #include "Misc/CommandLine.h"
+#include "UObject/UE5MainStreamObjectVersion.h"
 
 #if WITH_EDITOR
 #include "Framework/Notifications/NotificationManager.h"
@@ -662,6 +663,15 @@ USoundWave::USoundWave(const FObjectInitializer& ObjectInitializer)
 	EnvelopeFollowerFrameSize = 1024;
 	EnvelopeFollowerAttackTime = 10;
 	EnvelopeFollowerReleaseTime = 100;
+
+#if UE_ENABLE_VIRTUALIZATION_TOGGLE
+	// Calling ::SetVirtualizationOptOut allows us to prevent the bulkdata from virtualizing if a project
+	// has that option enabled. This is not intended to be a shipping feature.
+	PRAGMA_DISABLE_DEPRECATION_WARNINGS
+		RawData.SetVirtualizationOptOut(true);
+	PRAGMA_ENABLE_DEPRECATION_WARNINGS
+#endif //UE_ENABLE_VIRTUALIZATION_TOGGLE
+
 #endif
 
 	bCachedSampleRateFromPlatformSettings = false;
@@ -828,6 +838,7 @@ void USoundWave::Serialize( FArchive& Ar )
 	}
 
 	Ar.UsingCustomVersion(FFrameworkObjectVersion::GUID);
+	Ar.UsingCustomVersion(FUE5MainStreamObjectVersion::GUID);
 
 	if (Ar.IsLoading() && (Ar.UEVer() >= VER_UE4_SOUND_COMPRESSION_TYPE_ADDED) && (Ar.CustomVer(FFrameworkObjectVersion::GUID) < FFrameworkObjectVersion::RemoveSoundWaveCompressionName))
 	{
@@ -838,6 +849,9 @@ void USoundWave::Serialize( FArchive& Ar )
 	bool bShouldStreamSound = false;
 
 #if WITH_EDITORONLY_DATA
+	FByteBulkData TempOldBulkData;
+	bool bBulkDataConverted = false;
+
 	bLoadedFromCookedData = Ar.IsLoading() && bCooked;
 	if (bVirtualizeWhenSilent_DEPRECATED)
 	{
@@ -962,12 +976,30 @@ void USoundWave::Serialize( FArchive& Ar )
 		// only save the raw data for non-cooked packages
 #if WITH_EDITORONLY_DATA
 		FScopeLock ScopeLock(&RawDataCriticalSection);
+
+		if (Ar.IsLoading() && Ar.CustomVer(FUE5MainStreamObjectVersion::GUID) < FUE5MainStreamObjectVersion::SoundWaveVirtualizationUpdate)
+		{
+			TempOldBulkData.Serialize(Ar, this, INDEX_NONE, false);	
+			bBulkDataConverted = true;
+		}
+		else
+		{
+			RawData.Serialize(Ar, this);
+		}
 #endif
-		RawData.Serialize(Ar, this, INDEX_NONE, false);
 	}
 
 	Ar << CompressedDataGuid;
 
+#if WITH_EDITORONLY_DATA	
+	// If we're converting to FEditorBulkData, we need to first serialize the CompressedDataGuid, as we use it as our key.
+	if (bBulkDataConverted)
+	{
+		FScopeLock ScopeLock(&RawDataCriticalSection);
+		RawData.CreateFromBulkData(TempOldBulkData, CompressedDataGuid, this); 
+	}
+#endif //WITH_EDITORONLY_DATA	
+	
 	bool bBuiltStreamedAudio = false;
 
 	if (bShouldStreamSound)
@@ -1192,11 +1224,11 @@ bool USoundWave::GetImportedSoundWaveData(TArray<uint8>& OutRawPCMData, uint32&
 #endif
 
 	// Can only get sound wave data if there is bulk data 
-	if (RawData.GetBulkDataSize() > 0)
+	if (RawData.HasPayloadData())
 	{
-		FBulkDataReadScopeLock LockedBulkData(RawData);
-		const uint8* Data = LockedBulkData.GetData<uint8>();
-		int32 DataSize = RawData.GetBulkDataSize();
+		TFuture<FSharedBuffer> BufferFuture = RawData.GetPayload();
+		const uint8* Data = (const uint8*) BufferFuture.Get().GetData(); // Will block.
+		int32 DataSize = BufferFuture.Get().GetSize();
 
 		if (NumChannels > 2)
 		{
@@ -1666,14 +1698,6 @@ void USoundWave::PostLoad()
 			// Upload the data to the hardware, but only if we've precached startup sounds already
 			AudioDevice->Precache(this);
 		}
-		// remove bulk data if no AudioDevice is used and no sounds were initialized
-		else if (IsRunningGame())
-		{
-#if WITH_EDITORONLY_DATA
-			FScopeLock ScopeLock(&RawDataCriticalSection);
-#endif
-			RawData.RemoveBulkData();
-		}
 	}
 
 	// Only add this streaming sound if the platform supports streaming
