UE-145287 - Deprecate Make Literal Double node

Instead of having two float-related literal nodes, we should only have one. The single-precision variant needs to be deprecated in favor of the double-precision one.

diff --git a/Engine/Config/BaseEngine.ini b/Engine/Config/BaseEngine.ini
--- a/Engine/Config/BaseEngine.ini
+++ b/Engine/Config/BaseEngine.ini
@@ -1057,6 +1057,8 @@ DefaultWorldPartitionHLODLayer=/Engine/Maps/Templates/HLODs/HLODLayer_Instanced.
 ;5.1
 
 ; Blueprint real number support
++FunctionRedirects=(OldName="KismetSystemLibrary.MakeLiteralFloat", NewName="MakeLiteralDouble")
+
 +FunctionRedirects=(OldName="KismetTextLibrary.Conv_FloatToText", NewName="Conv_DoubleToText")
 +PropertyRedirects=(OldName="KismetTextLibrary.Conv_DoubleToText.InDouble", NewName="Conv_DoubleToText.Value")
 
diff --git a/Engine/Source/Runtime/Engine/Classes/Kismet/KismetSystemLibrary.h b/Engine/Source/Runtime/Engine/Classes/Kismet/KismetSystemLibrary.h
--- a/Engine/Source/Runtime/Engine/Classes/Kismet/KismetSystemLibrary.h
+++ b/Engine/Source/Runtime/Engine/Classes/Kismet/KismetSystemLibrary.h
@@ -349,15 +349,15 @@ class ENGINE_API UKismetSystemLibrary : public UBlueprintFunctionLibrary
 	 * @param	Value	value to set the float to
 	 * @return	The literal float
 	 */
-	UFUNCTION(BlueprintPure, Category="Math|Float", meta=(BlueprintThreadSafe))
+	UE_DEPRECATED(5.1, "This method has been deprecated and will be removed.")
 	static float MakeLiteralFloat(float Value);
 
 	/**
-	 * Creates a literal double
-	 * @param	Value	value to set the double to
-	 * @return	The literal double
+	 * Creates a literal float (double-precision)
+	 * @param	Value	value to set the float (double-precision) to
+	 * @return	The literal float (double-precision)
 	 */
-	UFUNCTION(BlueprintPure, Category = "Math|Double", meta = (BlueprintThreadSafe))
+	UFUNCTION(BlueprintPure, Category = "Math|Float", meta = (BlueprintThreadSafe, DisplayName = "Make Literal Float"))
 	static double MakeLiteralDouble(double Value);
 
 	/**
