Fix bug in FTextureLayout3d
- Allocating space from IdealNode could sometimes (incorrectly) fail even though the node has enough space without requiring the texture to grow.
- When this happened the implementation would fall back to allocating from the ContainingNode.
- However, in this case the incorrect UnusedLeaves entry would be removed causing the layout to be corrupted.

- Also cleaned up some code.

diff --git a/Engine/Source/Runtime/Engine/Public/TextureLayout3d.h b/Engine/Source/Runtime/Engine/Public/TextureLayout3d.h
--- a/Engine/Source/Runtime/Engine/Public/TextureLayout3d.h
+++ b/Engine/Source/Runtime/Engine/Public/TextureLayout3d.h
@@ -35,7 +35,7 @@ public:
 		bAlignByFour(bInAlignByFour),
 		bAllowShrink(bInAllowShrink)
 	{
-		checkSlow(MaxSizeX < USHRT_MAX && MaxSizeY < USHRT_MAX && MaxSizeZ < USHRT_MAX);
+		check(MaxSizeX < USHRT_MAX && MaxSizeY < USHRT_MAX && MaxSizeZ < USHRT_MAX);
 		new(Nodes) FTextureLayoutNode3d(0, 0, 0, MaxSizeX, MaxSizeY, MaxSizeZ, INDEX_NONE);
 		new (UnusedLeaves) FUnusedLeaf(0, 0, 0, MaxSizeX, MaxSizeY, MaxSizeZ, 0);
 	}
@@ -101,8 +101,9 @@ public:
 		if (IdealNodeIndex != INDEX_NONE)
 		{
 			NodeIndex = AddSurfaceInner(IdealNodeIndex, ElementSizeX, ElementSizeY, ElementSizeZ, false);
+			check(NodeIndex != INDEX_NONE);
 		}
-		if (NodeIndex == INDEX_NONE && ContainingNodeIndex != INDEX_NONE)
+		else if (ContainingNodeIndex != INDEX_NONE)
 		{
 			// Try allocating space which might enlarge the texture.
 			NodeIndex = AddSurfaceInner(ContainingNodeIndex, ElementSizeX, ElementSizeY, ElementSizeZ, true);
@@ -110,7 +111,7 @@ public:
 
 		if (NodeIndex != INDEX_NONE)
 		{
-			FTextureLayoutNode3d&	Node = Nodes[NodeIndex];
+			FTextureLayoutNode3d& Node = Nodes[NodeIndex];
 			Node.bUsed = true;
 			OutBaseX = Node.MinX;
 			OutBaseY = Node.MinY;
@@ -246,8 +247,6 @@ private:
 				SizeZ;
 		bool	bUsed;
 
-		FTextureLayoutNode3d() {}
-
 		FTextureLayoutNode3d(uint16 InMinX, uint16 InMinY, uint16 InMinZ, uint16 InSizeX, uint16 InSizeY, uint16 InSizeZ, int32 InParent):
 			ChildA(INDEX_NONE),
 			ChildB(INDEX_NONE),
@@ -336,7 +335,7 @@ private:
 			{
 				// Reject this node if this is an attempt to allocate space without enlarging the texture, 
 				// And this node cannot hold the element without enlarging the texture.
-				if (CurrentNodePtr->MinX + ElementSizeX > SizeX || CurrentNodePtr->MinY + ElementSizeY > SizeY || CurrentNodePtr->MinZ + ElementSizeY > SizeZ)
+				if (CurrentNodePtr->MinX + ElementSizeX > SizeX || CurrentNodePtr->MinY + ElementSizeY > SizeY || CurrentNodePtr->MinZ + ElementSizeZ > SizeZ)
 				{
 					return INDEX_NONE;
 				}
@@ -355,24 +354,25 @@ private:
 			// The pointer to the current node may be invalidated below, be explicit!
 			CurrentNodePtr = 0;
 
-			int32 ChildIndex;
+			// Store a copy of the current node on the stack for easier debugging.
+			// Can't store a pointer to the current node since Nodes may be reallocated in this function.
+			const FTextureLayoutNode3d CurrentNode = Nodes[NodeIndex];
+
+			// Update the child indices
+			int32 ChildANodeIndex = FreeNodeIndices.Num() ? FreeNodeIndices.Pop() : Nodes.AddUninitialized();
+			int32 ChildBNodeIndex = FreeNodeIndices.Num() ? FreeNodeIndices.Pop() : Nodes.AddUninitialized();
+
+			Nodes[NodeIndex].ChildA = ChildANodeIndex;
+			Nodes[NodeIndex].ChildB = ChildBNodeIndex;
 
 			// Add new nodes, and link them as children of the current node.
 			if (ExcessWidth > ExcessHeight)
 			{
-				// Store a copy of the current node on the stack for easier debugging.
-				// Can't store a pointer to the current node since Nodes may be reallocated in this function.
-				const FTextureLayoutNode3d CurrentNode = Nodes[NodeIndex];
-
 				if (ExcessWidth > ExcessDepth)
 				{
-					// Update the child indices
-					ChildIndex = FreeNodeIndices.Num() ? FreeNodeIndices.Pop() : Nodes.AddUninitialized();
-					Nodes[NodeIndex].ChildA = ChildIndex;
-
 					// Create a child with the same width as the element being placed.
 					// The height may not be the same as the element height yet, in that case another subdivision will occur when traversing this child node.
-					new(&Nodes[ChildIndex]) FTextureLayoutNode3d(
+					new(&Nodes[ChildANodeIndex]) FTextureLayoutNode3d(
 						CurrentNode.MinX,
 						CurrentNode.MinY,
 						CurrentNode.MinZ,
@@ -382,9 +382,7 @@ private:
 						NodeIndex);
 
 					// Create a second child to contain the leftover area in the X direction
-					ChildIndex = FreeNodeIndices.Num() ? FreeNodeIndices.Pop() : Nodes.AddUninitialized();
-					Nodes[NodeIndex].ChildB = ChildIndex;
-					new(&Nodes[ChildIndex]) FTextureLayoutNode3d(
+					new(&Nodes[ChildBNodeIndex]) FTextureLayoutNode3d(
 						CurrentNode.MinX + ElementSizeX,
 						CurrentNode.MinY,
 						CurrentNode.MinZ,
@@ -395,10 +393,7 @@ private:
 				}
 				else
 				{
-					ChildIndex = FreeNodeIndices.Num() ? FreeNodeIndices.Pop() : Nodes.AddUninitialized();
-					Nodes[NodeIndex].ChildA = ChildIndex;
-
-					new(&Nodes[ChildIndex]) FTextureLayoutNode3d(
+					new(&Nodes[ChildANodeIndex]) FTextureLayoutNode3d(
 						CurrentNode.MinX,
 						CurrentNode.MinY,
 						CurrentNode.MinZ,
@@ -407,9 +402,7 @@ private:
 						ElementSizeZ,
 						NodeIndex);
 
-					ChildIndex = FreeNodeIndices.Num() ? FreeNodeIndices.Pop() : Nodes.AddUninitialized();
-					Nodes[NodeIndex].ChildB = ChildIndex;
-					new(&Nodes[ChildIndex]) FTextureLayoutNode3d(
+					new(&Nodes[ChildBNodeIndex]) FTextureLayoutNode3d(
 						CurrentNode.MinX,
 						CurrentNode.MinY,
 						CurrentNode.MinZ + ElementSizeZ,
@@ -421,15 +414,9 @@ private:
 			}
 			else
 			{
-				// Store a copy of the current node on the stack for easier debugging.
-				// Can't store a pointer to the current node since Nodes may be reallocated in this function.
-				const FTextureLayoutNode3d CurrentNode = Nodes[NodeIndex];
-
 				if (ExcessHeight > ExcessDepth)
 				{
-					ChildIndex = FreeNodeIndices.Num() ? FreeNodeIndices.Pop() : Nodes.AddUninitialized();
-					Nodes[NodeIndex].ChildA = ChildIndex;
-					new(&Nodes[ChildIndex]) FTextureLayoutNode3d(
+					new(&Nodes[ChildANodeIndex]) FTextureLayoutNode3d(
 						CurrentNode.MinX,
 						CurrentNode.MinY,
 						CurrentNode.MinZ,
@@ -438,9 +425,7 @@ private:
 						CurrentNode.SizeZ,
 						NodeIndex);
 
-					ChildIndex = FreeNodeIndices.Num() ? FreeNodeIndices.Pop() : Nodes.AddUninitialized();
-					Nodes[NodeIndex].ChildB = ChildIndex;
-					new(&Nodes[ChildIndex]) FTextureLayoutNode3d(
+					new(&Nodes[ChildBNodeIndex]) FTextureLayoutNode3d(
 						CurrentNode.MinX,
 						CurrentNode.MinY + ElementSizeY,
 						CurrentNode.MinZ,
@@ -451,10 +436,7 @@ private:
 				}
 				else
 				{
-					ChildIndex = FreeNodeIndices.Num() ? FreeNodeIndices.Pop() : Nodes.AddUninitialized();
-					Nodes[NodeIndex].ChildA = ChildIndex;
-
-					new(&Nodes[ChildIndex]) FTextureLayoutNode3d(
+					new(&Nodes[ChildANodeIndex]) FTextureLayoutNode3d(
 						CurrentNode.MinX,
 						CurrentNode.MinY,
 						CurrentNode.MinZ,
@@ -463,9 +445,7 @@ private:
 						ElementSizeZ,
 						NodeIndex);
 
-					ChildIndex = FreeNodeIndices.Num() ? FreeNodeIndices.Pop() : Nodes.AddUninitialized();
-					Nodes[NodeIndex].ChildB = ChildIndex;
-					new(&Nodes[ChildIndex]) FTextureLayoutNode3d(
+					new(&Nodes[ChildBNodeIndex]) FTextureLayoutNode3d(
 						CurrentNode.MinX,
 						CurrentNode.MinY,
 						CurrentNode.MinZ + ElementSizeZ,
@@ -476,11 +456,12 @@ private:
 				}
 			}
 
-			const int32 ChildBIndex = Nodes[NodeIndex].ChildB;
-			new (UnusedLeaves) FUnusedLeaf(Nodes[ChildBIndex].MinX, Nodes[ChildBIndex].MinY, Nodes[ChildBIndex].MinZ, Nodes[ChildBIndex].SizeX, Nodes[ChildBIndex].SizeY, Nodes[ChildBIndex].SizeZ, ChildBIndex);
+			const FTextureLayoutNode3d& ChildBNode = Nodes[ChildBNodeIndex];
+
+			new (UnusedLeaves) FUnusedLeaf(ChildBNode.MinX, ChildBNode.MinY, ChildBNode.MinZ, ChildBNode.SizeX, ChildBNode.SizeY, ChildBNode.SizeZ, ChildBNodeIndex);
 
 			// Only traversing ChildA, since ChildA is always the newly created node that matches the element size
-			return AddSurfaceInner(Nodes[NodeIndex].ChildA, ElementSizeX, ElementSizeY, ElementSizeZ, bAllowTextureEnlargement);
+			return AddSurfaceInner(ChildANodeIndex, ElementSizeX, ElementSizeY, ElementSizeZ, bAllowTextureEnlargement);
 		}
 	}
 
@@ -522,52 +503,50 @@ private:
 	void RemoveChildren(int32 NodeIndex)
 	{
 		// Traverse the children depth first
-		if (Nodes[NodeIndex].ChildA != INDEX_NONE)
+		if (Nodes[NodeIndex].ChildA == INDEX_NONE)
 		{
-			RemoveChildren(Nodes[NodeIndex].ChildA);
-		}
-		if (Nodes[NodeIndex].ChildB != INDEX_NONE)
-		{
-			RemoveChildren(Nodes[NodeIndex].ChildB);
+			check(Nodes[NodeIndex].ChildB == INDEX_NONE);
+			return;
 		}
 
-		if (Nodes[NodeIndex].ChildA != INDEX_NONE)
+		check(Nodes[NodeIndex].ChildB != INDEX_NONE);
+
+		// Traverse the children depth first
+		RemoveChildren(Nodes[NodeIndex].ChildA);
+		RemoveChildren(Nodes[NodeIndex].ChildB);
+
 		{
-			// Store off the index of the child since it may be changed in the code below
-			const int32 OldChildA = Nodes[NodeIndex].ChildA;
+			const int32 ChildNodeIndex = Nodes[NodeIndex].ChildA;
 
 			// Remove the child from the Nodes array
-			FreeNodeIndices.Add(OldChildA);
+			FreeNodeIndices.Add(ChildNodeIndex);
 
 			for (int32 Index = 0; Index < UnusedLeaves.Num(); ++Index)
 			{
-				if (UnusedLeaves[Index].NodeIndex == OldChildA)
+				if (UnusedLeaves[Index].NodeIndex == ChildNodeIndex)
 				{
 					UnusedLeaves.RemoveAtSwap(Index);
 					break;
 				}
 			}
-
-			// Mark the node as not having a ChildA
-			Nodes[NodeIndex].ChildA = INDEX_NONE;
 		}
 
-		if (Nodes[NodeIndex].ChildB != INDEX_NONE)
 		{
-			const int32 OldChildB = Nodes[NodeIndex].ChildB;
-			FreeNodeIndices.Add(OldChildB);
+			const int32 ChildNodeIndex = Nodes[NodeIndex].ChildB;
+			FreeNodeIndices.Add(ChildNodeIndex);
 
 			for (int32 Index = 0; Index < UnusedLeaves.Num(); ++Index)
 			{
-				if (UnusedLeaves[Index].NodeIndex == OldChildB)
+				if (UnusedLeaves[Index].NodeIndex == ChildNodeIndex)
 				{
 					UnusedLeaves.RemoveAtSwap(Index);
 					break;
 				}
 			}
-
-			Nodes[NodeIndex].ChildB = INDEX_NONE;
 		}
+
+		Nodes[NodeIndex].ChildA = INDEX_NONE;
+		Nodes[NodeIndex].ChildB = INDEX_NONE;
 	}
 };
 
