Texture: Migrated to the new DDC API

#rnx

diff --git a/Engine/Source/Runtime/Engine/Private/Streaming/StreamingManagerTexture.cpp b/Engine/Source/Runtime/Engine/Private/Streaming/StreamingManagerTexture.cpp
--- a/Engine/Source/Runtime/Engine/Private/Streaming/StreamingManagerTexture.cpp
+++ b/Engine/Source/Runtime/Engine/Private/Streaming/StreamingManagerTexture.cpp
@@ -205,10 +205,6 @@ TArray<FStreamingRenderAsset>& FRenderAssetStreamingManager::GetStreamingRenderA
 void FRenderAssetStreamingManager::OnPreGarbageCollect()
 {
 	TRACE_CPUPROFILER_EVENT_SCOPE(FRenderAssetStreamingManager::OnPreGarbageCollect);
-#if WITH_EDITORONLY_DATA
-	void PurgeAbandonedDDCHandles();
-	PurgeAbandonedDDCHandles();
-#endif
 
 	FScopeLock ScopeLock(&CriticalSection);
 	if (StreamingRenderAssetsSyncEvent.IsValid())
diff --git a/Engine/Source/Runtime/Engine/Private/Streaming/Texture2DMipDataProvider_DDC.cpp b/Engine/Source/Runtime/Engine/Private/Streaming/Texture2DMipDataProvider_DDC.cpp
--- a/Engine/Source/Runtime/Engine/Private/Streaming/Texture2DMipDataProvider_DDC.cpp
+++ b/Engine/Source/Runtime/Engine/Private/Streaming/Texture2DMipDataProvider_DDC.cpp
@@ -5,16 +5,14 @@ Texture2DMipDataProvider_DDC.cpp : Implementation of FTextureMipDataProvider usi
 =============================================================================*/
 
 #include "Texture2DMipDataProvider_DDC.h"
-#include "DerivedDataCacheInterface.h"
 #include "Engine/Texture.h"
-#include "TextureResource.h"
-#include "Streaming/Texture2DStreamIn_DDC.h" // for FAbandonedDDCHandleManager
+#include "Rendering/StreamableTextureResource.h"
 #include "Streaming/TextureStreamingHelpers.h"
+#include "TextureResource.h"
 
 #if WITH_EDITORONLY_DATA
 
 #include "DerivedDataCache.h"
-#include "DerivedDataCacheInterface.h"
 #include "DerivedDataCacheKey.h"
 
 FTexture2DMipDataProvider_DDC::FTexture2DMipDataProvider_DDC(const UTexture* Texture)
@@ -25,14 +23,12 @@ FTexture2DMipDataProvider_DDC::FTexture2DMipDataProvider_DDC(const UTexture* Tex
 
 FTexture2DMipDataProvider_DDC::~FTexture2DMipDataProvider_DDC()
 {
-	check(!DDCHandles.Num());
 }
 
 void FTexture2DMipDataProvider_DDC::Init(const FTextureUpdateContext& Context, const FTextureUpdateSyncOptions& SyncOptions)
 {
-	if (!DDCHandles.Num())
+	if (DDCBuffers.IsEmpty())
 	{
-		DDCHandles.AddZeroed(CurrentFirstLODIdx);
 		DDCBuffers.AddZeroed(CurrentFirstLODIdx);
 
 		FTexturePlatformData*const * PtrPlatformData = const_cast<UTexture*>(Context.Texture)->GetRunningPlatformData();
@@ -41,6 +37,13 @@ void FTexture2DMipDataProvider_DDC::Init(const FTextureUpdateContext& Context, c
 			const FTexturePlatformData* PlatformData = *PtrPlatformData;
 			const int32 LODBias = static_cast<int32>(Context.MipsView.GetData() - PlatformData->Mips.GetData());
 
+			using namespace UE::DerivedData;
+			TArray<FCacheGetChunkRequest> MipKeys;
+
+			TStringBuilder<256> MipNameBuilder;
+			Context.Texture->GetPathName(nullptr, MipNameBuilder);
+			const int32 TextureNameLen = MipNameBuilder.Len();
+
 			if (PlatformData->DerivedDataKey.IsType<FString>())
 			{
 				for (int32 MipIndex = PendingFirstLODIdx; MipIndex < CurrentFirstLODIdx; ++MipIndex)
@@ -48,19 +51,17 @@ void FTexture2DMipDataProvider_DDC::Init(const FTextureUpdateContext& Context, c
 					const FTexture2DMipMap& OwnerMip = *Context.MipsView[MipIndex];
 					if (OwnerMip.IsPagedToDerivedData())
 					{
-						DDCHandles[MipIndex] = GetDerivedDataCacheRef().GetAsynchronous(*PlatformData->GetDerivedDataMipKeyString(MipIndex + LODBias, OwnerMip), Context.Texture->GetPathName());
+						FCacheGetChunkRequest& Request = MipKeys.AddDefaulted_GetRef();
+						MipNameBuilder.Appendf(TEXT(" [MIP %d]"), MipIndex + LODBias);
+						Request.Name = MipNameBuilder;
+						Request.Key = ConvertLegacyCacheKey(PlatformData->GetDerivedDataMipKeyString(MipIndex + LODBias, OwnerMip));
+						Request.UserData = MipIndex;
+						MipNameBuilder.RemoveSuffix(MipNameBuilder.Len() - TextureNameLen);
 					}
 				}
 			}
 			else if (PlatformData->DerivedDataKey.IsType<UE::DerivedData::FCacheKeyProxy>())
 			{
-				using namespace UE::DerivedData;
-				TArray<FCacheGetChunkRequest> MipKeys;
-
-				TStringBuilder<256> MipNameBuilder;
-				Context.Texture->GetPathName(nullptr, MipNameBuilder);
-				const int32 TextureNameLen = MipNameBuilder.Len();
-
 				const FCacheKey& Key = *PlatformData->DerivedDataKey.Get<UE::DerivedData::FCacheKeyProxy>().AsCacheKey();
 				for (int32 MipIndex = PendingFirstLODIdx; MipIndex < CurrentFirstLODIdx; ++MipIndex)
 				{
@@ -68,7 +69,7 @@ void FTexture2DMipDataProvider_DDC::Init(const FTextureUpdateContext& Context, c
 					if (MipMap.IsPagedToDerivedData())
 					{
 						FCacheGetChunkRequest& Request = MipKeys.AddDefaulted_GetRef();
-						MipNameBuilder.Appendf(TEXT(" [MIP 0]"), MipIndex + LODBias);
+						MipNameBuilder.Appendf(TEXT(" [MIP %d]"), MipIndex + LODBias);
 						Request.Name = MipNameBuilder;
 						Request.Key = Key;
 						Request.Id = FTexturePlatformData::MakeMipId(MipIndex + LODBias);
@@ -76,44 +77,37 @@ void FTexture2DMipDataProvider_DDC::Init(const FTextureUpdateContext& Context, c
 						MipNameBuilder.RemoveSuffix(MipNameBuilder.Len() - TextureNameLen);
 					}
 				}
-
-				if (MipKeys.Num())
-				{
-					GetCache().GetChunks(MipKeys, DDCRequestOwner, [this](FCacheGetChunkResponse&& Response)
-					{
-						if (Response.Status == EStatus::Ok)
-						{
-							const int32 MipIndex = int32(Response.UserData);
-							check(!DDCBuffers[MipIndex]);
-							DDCBuffers[MipIndex] = MoveTemp(Response.RawData);
-						}
-					});
-				}
 			}
 			else
 			{
-				UE_LOG(LogTexture, Error, TEXT("Attempting to stream in mips for texture that has not generated a supported derived data key format."));
+				UE_LOG(LogTexture, Error, TEXT("Attempting to stream data in an unsupported cache format for mips [%d, %d) of %s."),
+					PendingFirstLODIdx, CurrentFirstLODIdx, *Context.Texture->GetPathName());
+			}
+
+			if (MipKeys.Num())
+			{
+				FRequestBarrier Barrier(DDCRequestOwner);
+				GetCache().GetChunks(MipKeys, DDCRequestOwner, [this](FCacheGetChunkResponse&& Response)
+				{
+					if (Response.Status == EStatus::Ok)
+					{
+						const int32 MipIndex = int32(Response.UserData);
+						check(!DDCBuffers[MipIndex]);
+						DDCBuffers[MipIndex] = MoveTemp(Response.RawData);
+					}
+				});
 			}
 		}
 		else
 		{
-			UE_LOG(LogTexture, Error, TEXT("Attempting to stream in mips for texture that has not generated derived data yet."));
+			UE_LOG(LogTexture, Error, TEXT("Attempting to stream data that has not been generated yet for mips [%d, %d) of %s."),
+				PendingFirstLODIdx, CurrentFirstLODIdx, *Context.Texture->GetPathName());
 		}
 
-
 		*SyncOptions.bSnooze = true;
 	}
-	else // The DDC request have been issued, only check whether they are ready (since no good sync option is used).
+	else // The cache request has been issued, only check whether they are ready.
 	{
-		for (uint32 Handle : DDCHandles)
-		{
-			if (Handle && !GetDerivedDataCacheRef().PollAsynchronousCompletion(Handle))
-			{
-				*SyncOptions.bSnooze = true;
-				return;
-			}
-		}
-
 		if (!DDCRequestOwner.Poll())
 		{
 			*SyncOptions.bSnooze = true;
@@ -134,7 +128,8 @@ bool FTexture2DMipDataProvider_DDC::SerializeMipInfo(const FTextureUpdateContext
 	}
 	else if (uint64(MipSize) < OutMipInfo.DataSize && DepthOrArraySize > 1 && OutMipInfo.DataSize % DepthOrArraySize == 0 && MipSize % DepthOrArraySize == 0)
 	{
-		UE_LOG(LogTexture, Verbose, TEXT("DDC mip size smaller than streaming buffer size. (%s, Mip %d): %d KB / %d KB."), *Context.Resource->GetTextureName().ToString(), ResourceState.MaxNumLODs - MipIndex, OutMipInfo.DataSize / 1024, MipSize / 1024);
+		UE_LOG(LogTexture, Verbose, TEXT("Cached mip size is smaller than streaming buffer size for mip %d of %s: %d KiB / %d KiB."),
+			ResourceState.MaxNumLODs - MipIndex, *Context.Resource->GetTextureName().ToString(), OutMipInfo.DataSize / 1024, MipSize / 1024);
 
 		const uint64 SourceSubSize = MipSize / DepthOrArraySize;
 		const uint64 DestSubSize = OutMipInfo.DataSize / DepthOrArraySize;
@@ -151,7 +146,8 @@ bool FTexture2DMipDataProvider_DDC::SerializeMipInfo(const FTextureUpdateContext
 		return true;
 	}
 
-	UE_LOG(LogTexture, Warning, TEXT("Mismatch between DDC mip size and streaming buffer size. (%s, Mip %d): %d KB / %d KB."), *Context.Resource->GetTextureName().ToString(), ResourceState.MaxNumLODs - MipIndex, OutMipInfo.DataSize / 1024, MipSize / 1024);
+	UE_LOG(LogTexture, Warning, TEXT("Mismatch between cached mip size and streaming buffer size for mip %d of %s: %d KB / %d KB."),
+		ResourceState.MaxNumLODs - MipIndex, *Context.Resource->GetTextureName().ToString(), OutMipInfo.DataSize / 1024, MipSize / 1024);
 	FMemory::Memzero(OutMipInfo.DestData, OutMipInfo.DataSize);
 	return false;
 }
@@ -168,68 +164,31 @@ int32 FTexture2DMipDataProvider_DDC::GetMips(
 		const FTexturePlatformData* PlatformData = *PtrPlatformData;
 		const int32 LODBias = static_cast<int32>(Context.MipsView.GetData() - PlatformData->Mips.GetData());
 
-		if (PlatformData->DerivedDataKey.IsType<FString>())
+		for (int32 MipIndex = StartingMipIndex; MipIndex < CurrentFirstLODIdx; ++MipIndex)
 		{
-			for (int32 MipIndex = StartingMipIndex; MipIndex < CurrentFirstLODIdx; ++MipIndex)
+			bool bSuccess = false;
+			if (DDCBuffers[MipIndex])
 			{
-				const uint32 Handle = DDCHandles[MipIndex];
-				bool bSuccess = false;
-				if (Handle)
+				// The result must be read from a memory reader!
+				FMemoryReaderView Ar(DDCBuffers[MipIndex], true);
+				if (SerializeMipInfo(Context, Ar, MipIndex, DDCBuffers[MipIndex].GetSize(), MipInfos[MipIndex]))
 				{
-					DDCHandles[MipIndex] = 0; // Clear the handle.
-
-					TArray64<uint8> DerivedMipData;
-					if (GetDerivedDataCacheRef().GetAsynchronousResults(Handle, DerivedMipData))
-					{
-						const FTextureMipInfo& MipInfo = MipInfos[MipIndex];
-
-						// The result must be read from a memory reader!
-						FMemoryReaderView Ar(MakeMemoryView(DerivedMipData), true);
-						if (SerializeMipInfo(Context, Ar, MipIndex, DerivedMipData.Num(), MipInfo))
-						{
-							bSuccess = true;
-						}
-					}
-				}
-
-				if (!bSuccess)
-				{
-					AdvanceTo(ETickState::CleanUp, ETickThread::Async);
-					return MipIndex; // We failed at getting this mip. Cancel will be called.
+					bSuccess = true;
 				}
+				DDCBuffers[MipIndex].Reset();
 			}
-		}
-		else if (PlatformData->DerivedDataKey.IsType<UE::DerivedData::FCacheKeyProxy>())
-		{
-			for (int32 MipIndex = StartingMipIndex; MipIndex < CurrentFirstLODIdx; ++MipIndex)
-			{
-				bool bSuccess = false;
-				if (DDCBuffers[MipIndex])
-				{
-					// The result must be read from a memory reader!
-					FMemoryReaderView Ar(DDCBuffers[MipIndex], true);
-					if (SerializeMipInfo(Context, Ar, MipIndex, DDCBuffers[MipIndex].GetSize(), MipInfos[MipIndex]))
-					{
-						bSuccess = true;
-					}
-					DDCBuffers[MipIndex].Reset();
-				}
 
-				if (!bSuccess)
-				{
-					AdvanceTo(ETickState::CleanUp, ETickThread::Async);
-					return MipIndex; // We failed at getting this mip. Cancel will be called.
-				}
+			if (!bSuccess)
+			{
+				AdvanceTo(ETickState::CleanUp, ETickThread::Async);
+				return MipIndex; // We failed at getting this mip. Cancel will be called.
 			}
 		}
-		else
-		{
-			UE_LOG(LogTexture, Error, TEXT("Attempting to stream in mips for texture that has not generated a supported derived data key format."));
-		}
 	}
 	else
 	{
-		UE_LOG(LogTexture, Error, TEXT("Attempting to stream in mips for texture that has not generated derived data yet."));
+		UE_LOG(LogTexture, Error, TEXT("Attempting to stream data that has not been generated yet for mips [%d, %d) of %s."),
+			PendingFirstLODIdx, CurrentFirstLODIdx, *Context.Texture->GetPathName());
 	}
 
 	AdvanceTo(ETickState::CleanUp, ETickThread::Async);
@@ -255,27 +214,11 @@ void FTexture2DMipDataProvider_DDC::Cancel(const FTextureUpdateSyncOptions& Sync
 
 FTextureMipDataProvider::ETickThread FTexture2DMipDataProvider_DDC::GetCancelThread() const
 {
-	if (DDCHandles.Num() || DDCBuffers.Num())
-	{
-		return ETickThread::Async;
-	}
-	else
-	{
-		return ETickThread::None;
-	}
+	return DDCBuffers.Num() ? ETickThread::Async : ETickThread::None;
 }
 
 void FTexture2DMipDataProvider_DDC::ReleaseDDCResources()
 {
-	for (uint32& Handle : DDCHandles)
-	{
-		if (Handle)
-		{
-			GAbandonedDDCHandleManager.Add(Handle);
-			Handle = 0;
-		}
-	}
-	DDCHandles.Empty();
 	DDCRequestOwner.Cancel();
 	DDCBuffers.Empty();
 }
diff --git a/Engine/Source/Runtime/Engine/Private/Streaming/Texture2DMipDataProvider_DDC.h b/Engine/Source/Runtime/Engine/Private/Streaming/Texture2DMipDataProvider_DDC.h
--- a/Engine/Source/Runtime/Engine/Private/Streaming/Texture2DMipDataProvider_DDC.h
+++ b/Engine/Source/Runtime/Engine/Private/Streaming/Texture2DMipDataProvider_DDC.h
@@ -37,15 +37,12 @@ public:
 
 protected:
 
-	// Helper to route incomplete DDC requests to the FAbandonedDDCHandleManager.
 	void ReleaseDDCResources();
 
 	bool SerializeMipInfo(const FTextureUpdateContext& Context, FArchive& Ar, int32 MipIndex, int64 MipSize, const FTextureMipInfo& OutMipInfo);
 
 private:
 
-	// The list of DDC handle for each mips (0 based on the highest mip).
-	TArray<uint32, TInlineAllocator<MAX_TEXTURE_MIP_COUNT>> DDCHandles;
 	TArray<FSharedBuffer, TInlineAllocator<MAX_TEXTURE_MIP_COUNT> > DDCBuffers;
 	UE::DerivedData::FRequestOwner DDCRequestOwner;
 };
diff --git a/Engine/Source/Runtime/Engine/Private/Streaming/Texture2DStreamIn_DDC.cpp b/Engine/Source/Runtime/Engine/Private/Streaming/Texture2DStreamIn_DDC.cpp
--- a/Engine/Source/Runtime/Engine/Private/Streaming/Texture2DStreamIn_DDC.cpp
+++ b/Engine/Source/Runtime/Engine/Private/Streaming/Texture2DStreamIn_DDC.cpp
@@ -5,10 +5,10 @@ Texture2DStreamIn_DDC.cpp: Stream in helper for 2D textures loading DDC files.
 =============================================================================*/
 
 #include "Streaming/Texture2DStreamIn_DDC.h"
-#include "Streaming/TextureStreamingHelpers.h"
+#include "Rendering/Texture2DResource.h"
 #include "RenderUtils.h"
 #include "Serialization/MemoryReader.h"
-#include "Rendering/Texture2DResource.h"
+#include "Streaming/TextureStreamingHelpers.h"
 
 #if WITH_EDITORONLY_DATA
 
@@ -16,84 +16,16 @@ Texture2DStreamIn_DDC.cpp: Stream in helper for 2D textures loading DDC files.
 #include "DerivedDataCacheInterface.h"
 #include "DerivedDataCacheKey.h"
 #include "DerivedDataRequestOwner.h"
+#include "Misc/ScopeExit.h"
 
 int32 GStreamingUseAsyncRequestsForDDC = 1;
 static FAutoConsoleVariableRef CVarStreamingDDCPendingSleep(
 	TEXT("r.Streaming.UseAsyncRequestsForDDC"),
 	GStreamingUseAsyncRequestsForDDC,
-	TEXT("Whether to use async DDC requets in order to react quickly to cancel and suspend rendering requests (default=0)"),
+	TEXT("Whether to use async DDC requests in order to react quickly to cancel and suspend rendering requests (default=0)"),
 	ECVF_Default
 );
 
-int32 GStreamingAbandonedDDCHandlePurgeFrequency = 150;
-static FAutoConsoleVariableRef CVarStreamingAbandonedDDCHandlePurgeFrequency(
-	TEXT("r.Streaming.AbandonedDDCHandlePurgeFrequency"),
-	GStreamingAbandonedDDCHandlePurgeFrequency,
-	TEXT("The number of abandonned handle at which a purge will be triggered (default=150)"),
-	ECVF_Default
-);
-
-// ******************************************
-// ******* FAbandonedDDCHandleManager *******
-// ******************************************
-
-FAbandonedDDCHandleManager GAbandonedDDCHandleManager;
-
-void FAbandonedDDCHandleManager::Add(uint32 InHandle)
-{
-	check(InHandle);
-
-	bool bPurge = false;
-	{
-		FScopeLock ScopeLock(&CS);
-		Handles.Add(InHandle);
-
-		++TotalAdd;
-		bPurge = (TotalAdd % GStreamingAbandonedDDCHandlePurgeFrequency == 0);
-	}
-
-	if (bPurge)
-	{
-		Purge();
-	}
-}
-
-void FAbandonedDDCHandleManager::Purge()
-{
-	TArray<uint32> TempHandles;	
-	{
-		FScopeLock ScopeLock(&CS);
-		FMemory::Memswap(&TempHandles, &Handles, sizeof(TempHandles));
-	}
-
-	FDerivedDataCacheInterface& DDCInterface = GetDerivedDataCacheRef();
-	TArray64<uint8> Data;
-
-	for (int32 Index = 0; Index < TempHandles.Num(); ++Index)
-	{
-		uint32 Handle = TempHandles[Index];
-		if (DDCInterface.PollAsynchronousCompletion(Handle))
-		{
-			DDCInterface.GetAsynchronousResults(Handle, Data);
-			Data.Reset();
-
-			TempHandles.RemoveAtSwap(Index);
-			--Index;
-		}
-	}
-
-	if (TempHandles.Num())
-	{
-		FScopeLock ScopeLock(&CS);
-		Handles.Append(TempHandles);
-	}
-}
-
-void PurgeAbandonedDDCHandles()
-{
-	GAbandonedDDCHandleManager.Purge();
-}
-
 // ******************************************
 // ********* FTexture2DStreamIn_DDC *********
 // ******************************************
@@ -102,22 +34,11 @@ FTexture2DStreamIn_DDC::FTexture2DStreamIn_DDC(UTexture2D* InTexture)
 	: FTexture2DStreamIn(InTexture)
 	, DDCRequestOwner(UE::DerivedData::EPriority::Normal)
 {
-	DDCHandles.AddZeroed(ResourceState.MaxNumLODs);
 	DDCMipRequestStatus.AddZeroed(ResourceState.MaxNumLODs);
 }
 
 FTexture2DStreamIn_DDC::~FTexture2DStreamIn_DDC()
 {
-	// On cancellation, we don't wait for DDC requests to complete before releasing the object.
-	// This prevents GC from being stalled when texture are deleted.
-	for (uint32& AbandonedHandle : DDCHandles)
-	{
-		if (AbandonedHandle)
-		{
-			GAbandonedDDCHandleManager.Add(AbandonedHandle);
-			AbandonedHandle = 0;
-		}
-	}
 }
 
 void FTexture2DStreamIn_DDC::DoCreateAsyncDDCRequests(const FContext& Context)
@@ -131,45 +52,45 @@ void FTexture2DStreamIn_DDC::DoCreateAsyncDDCRequests(const FContext& Context)
 	{
 		const int32 LODBias = static_cast<int32>(Context.MipsView.GetData() - PlatformData->Mips.GetData());
 
+		using namespace UE::DerivedData;
+		TArray<FCacheGetChunkRequest> MipKeys;
+
+		TStringBuilder<256> MipNameBuilder;
+		Context.Texture->GetPathName(nullptr, MipNameBuilder);
+		const int32 TextureNameLen = MipNameBuilder.Len();
+
 		if (PlatformData->DerivedDataKey.IsType<FString>())
 		{
 			for (int32 MipIndex = PendingFirstLODIdx; MipIndex < CurrentFirstLODIdx && !IsCancelled(); ++MipIndex)
 			{
 				const FTexture2DMipMap& MipMap = *Context.MipsView[MipIndex];
-				if (MipMap.IsPagedToDerivedData())
+				FMipRequestStatus& Status = DDCMipRequestStatus[MipIndex];
+				if (!MipMap.IsPagedToDerivedData())
 				{
-					check(!DDCHandles[MipIndex]);
-					DDCHandles[MipIndex] = GetDerivedDataCacheRef().GetAsynchronous(*PlatformData->GetDerivedDataMipKeyString(MipIndex + LODBias, MipMap), Context.Texture->GetPathName());
-
-#if !UE_BUILD_SHIPPING
-					// On some platforms the IO is too fast to test cancelation requests timing issues.
-					if (FRenderAssetStreamingSettings::ExtraIOLatency > 0 && TaskSynchronization.GetValue() == 0)
-					{
-						FPlatformProcess::Sleep(FRenderAssetStreamingSettings::ExtraIOLatency * .001f); // Slow down the streaming.
-					}
-#endif
+					UE_LOG(LogTexture, Error, TEXT("Attempting to stream data that is already loaded for mip %d of %s."),
+						MipIndex, *Context.Texture->GetPathName());
+					MarkAsCancelled();
 				}
-				else
+				else if (!Status.bRequestIssued && !Status.Buffer)
 				{
-					UE_LOG(LogTexture, Error, TEXT("Attempting to stream in a mip that is already present."));
-					MarkAsCancelled();
+					FCacheGetChunkRequest& Request = MipKeys.AddDefaulted_GetRef();
+					MipNameBuilder.Appendf(TEXT(" [MIP %d]"), MipIndex + LODBias);
+					Request.Name = MipNameBuilder;
+					Request.Key = ConvertLegacyCacheKey(PlatformData->GetDerivedDataMipKeyString(MipIndex + LODBias, MipMap));
+					Request.UserData = MipIndex;
+					MipNameBuilder.RemoveSuffix(MipNameBuilder.Len() - TextureNameLen);
+					Status.bRequestIssued = true;
 				}
 			}
 		}
 		else if (PlatformData->DerivedDataKey.IsType<UE::DerivedData::FCacheKeyProxy>())
 		{
-			using namespace UE::DerivedData;
-			TArray<FCacheGetChunkRequest> MipKeys;
-
-			TStringBuilder<256> MipNameBuilder;
-			Context.Texture->GetPathName(nullptr, MipNameBuilder);
-			const int32 TextureNameLen = MipNameBuilder.Len();
-
 			const FCacheKey& Key = *PlatformData->DerivedDataKey.Get<UE::DerivedData::FCacheKeyProxy>().AsCacheKey();
 			for (int32 MipIndex = PendingFirstLODIdx; MipIndex < CurrentFirstLODIdx && !IsCancelled(); ++MipIndex)
 			{
 				const FTexture2DMipMap& MipMap = *Context.MipsView[MipIndex];
-				if (MipMap.IsPagedToDerivedData())
+				FMipRequestStatus& Status = DDCMipRequestStatus[MipIndex];
+				if (MipMap.IsPagedToDerivedData() && !Status.bRequestIssued && !Status.Buffer)
 				{
 					FCacheGetChunkRequest& Request = MipKeys.AddDefaulted_GetRef();
 					MipNameBuilder.Appendf(TEXT(" [MIP %d]"), MipIndex + LODBias);
@@ -178,46 +99,52 @@ void FTexture2DStreamIn_DDC::DoCreateAsyncDDCRequests(const FContext& Context)
 					Request.Id = FTexturePlatformData::MakeMipId(MipIndex + LODBias);
 					Request.UserData = MipIndex;
 					MipNameBuilder.RemoveSuffix(MipNameBuilder.Len() - TextureNameLen);
-					DDCMipRequestStatus[MipIndex].bRequestIssued = true;
+					Status.bRequestIssued = true;
 				}
 			}
-
-			if (MipKeys.Num())
-			{
-				GetCache().GetChunks(MipKeys, DDCRequestOwner, [this](FCacheGetChunkResponse&& Response)
-				{
-					if (Response.Status == EStatus::Ok)
-					{
-						const int32 MipIndex = int32(Response.UserData);
-						check(!DDCMipRequestStatus[MipIndex].Buffer);
-						DDCMipRequestStatus[MipIndex].Buffer = MoveTemp(Response.RawData);
-					}
-				});
-			}
 		}
 		else
 		{
-			UE_LOG(LogTexture, Error, TEXT("Attempting to stream in mips for texture that has not generated a supported derived data key format."));
+			UE_LOG(LogTexture, Error, TEXT("Attempting to stream data in an unsupported cache format for mips [%d, %d) of %s."),
+				PendingFirstLODIdx, CurrentFirstLODIdx, *Context.Texture->GetPathName());
 			MarkAsCancelled();
 		}
+
+		if (MipKeys.Num())
+		{
+		#if !UE_BUILD_SHIPPING
+			// On some platforms the IO is too fast to test cancellation requests timing issues.
+			if (FRenderAssetStreamingSettings::ExtraIOLatency > 0 && TaskSynchronization.GetValue() == 0)
+			{
+				FPlatformProcess::Sleep(MipKeys.Num() * FRenderAssetStreamingSettings::ExtraIOLatency * 0.001f); // Slow down the streaming.
+			}
+		#endif
+
+			FRequestBarrier Barrier(DDCRequestOwner);
+			GetCache().GetChunks(MipKeys, DDCRequestOwner, [this](FCacheGetChunkResponse&& Response)
+			{
+				if (Response.Status == EStatus::Ok)
+				{
+					const int32 MipIndex = int32(Response.UserData);
+					FMipRequestStatus& Status = DDCMipRequestStatus[MipIndex];
+					check(!Status.Buffer);
+					Status.Buffer = MoveTemp(Response.RawData);
+					check(Status.bRequestIssued);
+					Status.bRequestIssued = false;
+				}
+			});
+		}
 	}
 	else
 	{
-		UE_LOG(LogTexture, Error, TEXT("Attempting to stream in mips for texture that has not generated derived data yet."));
+		UE_LOG(LogTexture, Error, TEXT("Attempting to stream data that has not been generated yet for mips [%d, %d) of %s."),
+			PendingFirstLODIdx, CurrentFirstLODIdx, *Context.Texture->GetPathName());
 		MarkAsCancelled();
 	}
 }
 
 bool FTexture2DStreamIn_DDC::DoPoolDDCRequests(const FContext& Context) 
 {
-	for (int32 MipIndex = PendingFirstLODIdx; MipIndex < CurrentFirstLODIdx && !IsCancelled(); ++MipIndex)
-	{
-		const uint32 Handle = DDCHandles[MipIndex];
-		if (Handle && !GetDerivedDataCacheRef().PollAsynchronousCompletion(Handle))
-		{
-			return false;
-		}
-	}
 	return DDCRequestOwner.Poll();
 }
 
@@ -228,167 +155,51 @@ void FTexture2DStreamIn_DDC::DoLoadNewMipsFromDDC(const FContext& Context)
 		return;
 	}
 
-	if (const FTexturePlatformData* PlatformData = Context.Texture->GetPlatformData())
+	using namespace UE::DerivedData;
+	const EPriority OriginalPriority = DDCRequestOwner.GetPriority();
+	ON_SCOPE_EXIT { DDCRequestOwner.SetPriority(OriginalPriority); };
+	DDCRequestOwner.SetPriority(EPriority::Blocking);
+	DoCreateAsyncDDCRequests(Context);
+	DDCRequestOwner.Wait();
+
+	for (int32 MipIndex = PendingFirstLODIdx; MipIndex < CurrentFirstLODIdx && !IsCancelled(); ++MipIndex)
 	{
-		const int32 LODBias = static_cast<int32>(Context.MipsView.GetData() - PlatformData->Mips.GetData());
+		const FTexture2DMipMap& MipMap = *Context.MipsView[MipIndex];
+		check(MipData[MipIndex].Data);
 
-		if (PlatformData->DerivedDataKey.IsType<FString>())
+		FMipRequestStatus& Status = DDCMipRequestStatus[MipIndex];
+		if (Status.Buffer)
 		{
-			for (int32 MipIndex = PendingFirstLODIdx; MipIndex < CurrentFirstLODIdx && !IsCancelled(); ++MipIndex)
-			{
-				const FTexture2DMipMap& MipMap = *Context.MipsView[MipIndex];
-				check(MipData[MipIndex].Data != nullptr);
-
-				if (MipMap.IsPagedToDerivedData())
-				{
-					// The overhead of doing 2 copy of each mip data (from GetSynchronous() and FMemoryReader) in hidden by other texture DDC ops happening at the same time.
-					TArray64<uint8> DerivedMipData;
-					bool bDDCValid = true;
-
-					const uint32 Handle = DDCHandles[MipIndex];
-					if (Handle)
-					{
-						bDDCValid = GetDerivedDataCacheRef().GetAsynchronousResults(Handle, DerivedMipData);
-						DDCHandles[MipIndex] = 0;
-					}
-					else
-					{
-						bDDCValid = GetDerivedDataCacheRef().GetSynchronous(*PlatformData->GetDerivedDataMipKeyString(MipIndex + LODBias, MipMap), DerivedMipData, Context.Texture->GetPathName());
-					}
+			const SIZE_T ExpectedMipSize = CalcTextureMipMapSize(MipMap.SizeX, MipMap.SizeY, Context.Resource->GetPixelFormat(), 0);
 
-					if (bDDCValid)
-					{
-						const SIZE_T ExpectedMipSize = CalcTextureMipMapSize(MipMap.SizeX, MipMap.SizeY, Context.Resource->GetPixelFormat(), 0);
-						FMemoryReaderView Ar(MakeMemoryView(DerivedMipData), true);
+			UE_LOG(LogTexture, Verbose, TEXT("FTexture2DStreamIn_DDC::DoLoadNewMipsFromDDC Size=%dx%d ExpectedMipSize=%" SIZE_T_FMT " DerivedMipSize=%" SIZE_T_FMT),
+				MipMap.SizeX, MipMap.SizeY, ExpectedMipSize, Status.Buffer.GetSize());
 
-						UE_LOG(LogTextureUpload,Verbose,TEXT("FTexture2DStreamIn_DDC::DoLoadNewMipsFromDDC Size: %dx%d , ExpectedMipSize=%d DerivedMipData=%d"),
-							MipMap.SizeX,MipMap.SizeY,
-							ExpectedMipSize,(int)DerivedMipData.Num());
+			// Serializes directly into MipData so it must be tight packed pitch
+			const uint32 DestPitch = MipData[MipIndex].Pitch;
+			FTexture2DResource::WarnRequiresTightPackedMip(MipMap.SizeX, MipMap.SizeY, Context.Resource->GetPixelFormat(), DestPitch);
 
-						// Serializes directly into MipData so it must be tight packed pitch :
-						uint32 DestPitch = MipData[MipIndex].Pitch;
-						FTexture2DResource::WarnRequiresTightPackedMip(MipMap.SizeX, MipMap.SizeY, Context.Resource->GetPixelFormat(), DestPitch);
-
-						if (DerivedMipData.Num() == ExpectedMipSize)
-						{
-							// Pitch ignored
-							// to copy and respect Pitch, use CopyTextureData2D
-							// MipData[] should have size but doesn't
-							Ar.Serialize(MipData[MipIndex].Data, DerivedMipData.Num());
-						}
-						else
-						{
-							UE_LOG(LogTexture, Error, TEXT("DDC mip size (%" SIZE_T_FMT ") not as expected (%" SIZE_T_FMT ") for mip %d of %s."),
-								static_cast<SIZE_T>(DerivedMipData.Num()), ExpectedMipSize, MipIndex, *Context.Texture->GetPathName());
-							MarkAsCancelled();
-						}
-					}
-					else
-					{
-						MarkAsCancelled();
-					}
-				}
-				else
-				{
-					UE_LOG(LogTexture, Error, TEXT("Attempting to stream in a mip that is already present."));
-					MarkAsCancelled();
-				}
+			if (Status.Buffer.GetSize() == ExpectedMipSize)
+			{
+				// Pitch ignored
+				// to copy and respect Pitch, use CopyTextureData2D
+				// MipData[] should have size but doesn't
+				FMemory::Memcpy(MipData[MipIndex].Data, Status.Buffer.GetData(), Status.Buffer.GetSize());
 			}
-		}
-		else if (PlatformData->DerivedDataKey.IsType<UE::DerivedData::FCacheKeyProxy>())
-		{
-			using namespace UE::DerivedData;
-			const FCacheKey& Key = *PlatformData->DerivedDataKey.Get<FCacheKeyProxy>().AsCacheKey();
-			TStringBuilder<256> MipNameBuilder;
-			Context.Texture->GetPathName(nullptr, MipNameBuilder);
-			const int32 TextureNameLen = MipNameBuilder.Len();
-
-			for (int32 MipIndex = PendingFirstLODIdx; MipIndex < CurrentFirstLODIdx && !IsCancelled(); ++MipIndex)
+			else
 			{
-				const FTexture2DMipMap& MipMap = *Context.MipsView[MipIndex];
-				check(MipData[MipIndex].Data != nullptr);
-
-				if (MipMap.IsPagedToDerivedData())
-				{
-					FSharedBuffer MipResult;
-					FMipRequestStatus& MipRequestStatus = DDCMipRequestStatus[MipIndex];
-					if (MipRequestStatus.bRequestIssued)
-					{
-						MipResult = MoveTemp(MipRequestStatus.Buffer);
-						MipRequestStatus.bRequestIssued = false;
-					}
-					else
-					{
-						FCacheGetChunkRequest Request;
-						MipNameBuilder.Appendf(TEXT(" [MIP %d]"), MipIndex + LODBias);
-						Request.Name = MipNameBuilder;
-						Request.Key = Key;
-						Request.Id = FTexturePlatformData::MakeMipId(MipIndex + LODBias);
-						Request.UserData = MipIndex;
-						MipNameBuilder.RemoveSuffix(MipNameBuilder.Len() - TextureNameLen);
-
-						FRequestOwner BlockingRequestOwner(EPriority::Blocking);
-						GetCache().GetChunks({Request}, BlockingRequestOwner, [MipIndex, &MipResult](FCacheGetChunkResponse&& Response)
-						{
-							if (Response.Status == EStatus::Ok)
-							{
-								MipResult = MoveTemp(Response.RawData);
-							}
-						});
-						BlockingRequestOwner.Wait();
-
-					}
-
-					if (MipResult)
-					{
-						const SIZE_T ExpectedMipSize = CalcTextureMipMapSize(MipMap.SizeX, MipMap.SizeY, Context.Resource->GetPixelFormat(), 0);
-						
-						UE_LOG(LogTextureUpload,Verbose,TEXT("FTexture2DStreamIn_DDC::DoLoadNewMipsFromDDC Size: %dx%d , ExpectedMipSize=%d MipResult=%d"),
-							MipMap.SizeX,MipMap.SizeY,
-							ExpectedMipSize,(int)MipResult.GetSize());
-						
-						// uses Memcpy instead of CopyTextureData2D, so Pitch must be tight packed
-						uint32 DestPitch = MipData[MipIndex].Pitch;
-						FTexture2DResource::WarnRequiresTightPackedMip(MipMap.SizeX, MipMap.SizeY, Context.Resource->GetPixelFormat(), DestPitch);
-
-						if (MipResult.GetSize() == ExpectedMipSize)
-						{
-							// Pitch ignored
-							// to memcpy and respect Pitch, use CopyTextureData2D
-							// check MipData[] Size & Pitch
-							FMemory::Memcpy(MipData[MipIndex].Data, MipResult.GetData(), MipResult.GetSize());
-						}
-						else
-						{
-							UE_LOG(LogTexture, Error, TEXT("DDC mip size (%" SIZE_T_FMT ") not as expected (%" SIZE_T_FMT ") for mip %d of %s."),
-								static_cast<SIZE_T>(MipResult.GetSize()), ExpectedMipSize, MipIndex, *Context.Texture->GetPathName());
-							MarkAsCancelled();
-						}
-					}
-					else
-					{
-						MarkAsCancelled();
-					}
-				}
-				else
-				{
-					UE_LOG(LogTexture, Error, TEXT("Attempting to stream in a mip that is already present."));
-					MarkAsCancelled();
-				}
+				UE_LOG(LogTexture, Error, TEXT("Cached mip size (%" SIZE_T_FMT ") not as expected (%" SIZE_T_FMT ") for mip %d of %s."),
+					static_cast<SIZE_T>(Status.Buffer.GetSize()), ExpectedMipSize, MipIndex, *Context.Texture->GetPathName());
+				MarkAsCancelled();
 			}
 		}
 		else
 		{
-			UE_LOG(LogTexture, Error, TEXT("Attempting to stream in mips for texture that has not generated a supported derived data key format."));
 			MarkAsCancelled();
 		}
-		FPlatformMisc::MemoryBarrier();
-	}
-	else
-	{
-		UE_LOG(LogTexture, Error, TEXT("Attempting to stream in mips for texture that has not generated derived data yet."));
-		MarkAsCancelled();
 	}
+
+	FPlatformMisc::MemoryBarrier();
 }
 
-#endif
\ No newline at end of file
+#endif
diff --git a/Engine/Source/Runtime/Engine/Private/Streaming/Texture2DStreamIn_DDC.h b/Engine/Source/Runtime/Engine/Private/Streaming/Texture2DStreamIn_DDC.h
--- a/Engine/Source/Runtime/Engine/Private/Streaming/Texture2DStreamIn_DDC.h
+++ b/Engine/Source/Runtime/Engine/Private/Streaming/Texture2DStreamIn_DDC.h
@@ -30,8 +30,6 @@ protected:
 		bool bRequestIssued = false;
 	};
 
-	// StreamIn_Default : Locked mips of the intermediate textures, used as disk load destination.
-	TArray<uint32, TInlineAllocator<MAX_TEXTURE_MIP_COUNT> > DDCHandles;
 	TArray<FMipRequestStatus, TInlineAllocator<MAX_TEXTURE_MIP_COUNT> > DDCMipRequestStatus;
 	UE::DerivedData::FRequestOwner DDCRequestOwner;
 
@@ -39,32 +37,14 @@ protected:
 	// ********* Helpers **********
 	// ****************************
 
-	// Create DDC load requests (into DDCHandles)
+	// Create DDC load requests
 	void DoCreateAsyncDDCRequests(const FContext& Context);
 
-	// Create DDC load requests (into DDCHandles)
+	// Create DDC load requests
 	bool DoPoolDDCRequests(const FContext& Context);
 
 	// Load from DDC into MipData
 	void DoLoadNewMipsFromDDC(const FContext& Context);
 };
 
-/**
-* This class provides a helper to release DDC handles that haven't been waited for.
-* This is to get around limitations of FDerivedDataCacheInterface.
-*/
-
-class FAbandonedDDCHandleManager
-{
-public:
-	void Add(uint32 InHandle);
-	void Purge();
-private:
-	TArray<uint32> Handles;	
-	FCriticalSection CS;
-	uint32 TotalAdd = 0;
-};
-
-extern FAbandonedDDCHandleManager GAbandonedDDCHandleManager;
-
 #endif // WITH_EDITORONLY_DATA
diff --git a/Engine/Source/Runtime/Engine/Private/Texture2D.cpp b/Engine/Source/Runtime/Engine/Private/Texture2D.cpp
--- a/Engine/Source/Runtime/Engine/Private/Texture2D.cpp
+++ b/Engine/Source/Runtime/Engine/Private/Texture2D.cpp
@@ -6,7 +6,6 @@
 
 #include "Engine/Texture2D.h"
 
-#include "Serialization/MemoryWriter.h"
 #include "Misc/App.h"
 #include "HAL/PlatformFileManager.h"
 #include "HAL/FileManager.h"
@@ -23,7 +22,8 @@
 #include "EngineUtils.h"
 #include "DeviceProfiles/DeviceProfile.h"
 #include "DeviceProfiles/DeviceProfileManager.h"
-#include "DerivedDataCacheInterface.h"
+#include "DerivedDataCache.h"
+#include "DerivedDataRequestOwner.h"
 #include "Engine/TextureStreamingTypes.h"
 #include "Streaming/TextureStreamingHelpers.h"
 #include "Streaming/Texture2DStreamOut_AsyncReallocate.h"
@@ -164,21 +164,22 @@ void FTexture2DMipMap::Serialize(FArchive& Ar, UObject* Owner, int32 MipIdx)
 #if WITH_EDITORONLY_DATA
 int64 FTexture2DMipMap::StoreInDerivedDataCache(const FString& InDerivedDataKey, const FStringView& TextureName, bool bReplaceExistingDDC)
 {
+	using namespace UE::DerivedData;
+
 	int64 BulkDataSizeInBytes = BulkData.GetBulkDataSize();
 	check(BulkDataSizeInBytes > 0);
 
-	TArray64<uint8> DerivedData;
-	FMemoryWriter64 Ar(DerivedData, /*bIsPersistent=*/ true);
-	{
-		void* BulkMipData = BulkData.Lock(LOCK_READ_ONLY);
-		Ar.Serialize(BulkMipData, BulkDataSizeInBytes);
-		BulkData.Unlock();
-	}
-	const int64 Result = DerivedData.Num();
-	GetDerivedDataCacheRef().Put(*InDerivedDataKey, DerivedData, TextureName, bReplaceExistingDDC);
-	bPagedToDerivedData = true;
+	FValue DerivedData;
+	DerivedData = FValue::Compress(FSharedBuffer::MakeView(BulkData.Lock(LOCK_READ_ONLY), BulkDataSizeInBytes));
+	BulkData.Unlock();
+
+	FRequestOwner AsyncOwner(EPriority::Normal);
+	GetCache().PutValue({{{TextureName}, ConvertLegacyCacheKey(InDerivedDataKey), MoveTemp(DerivedData)}}, AsyncOwner);
+	AsyncOwner.KeepAlive();
+
+	SetPagedToDerivedData(true);
 	BulkData.RemoveBulkData();
-	return Result;
+	return BulkDataSizeInBytes;
 }
 #endif // #if WITH_EDITORONLY_DATA
 
diff --git a/Engine/Source/Runtime/Engine/Private/TextureDerivedData.cpp b/Engine/Source/Runtime/Engine/Private/TextureDerivedData.cpp
--- a/Engine/Source/Runtime/Engine/Private/TextureDerivedData.cpp
+++ b/Engine/Source/Runtime/Engine/Private/TextureDerivedData.cpp
@@ -36,10 +36,7 @@
 
 #include "DerivedDataCache.h"
 #include "DerivedDataCacheInterface.h"
-#include "DerivedDataCacheKey.h"
 #include "DerivedDataRequestOwner.h"
-#include "DerivedDataRequestTypes.h"
-#include "DerivedDataValue.h"
 #include "Interfaces/ITargetPlatform.h"
 #include "Interfaces/ITargetPlatformManagerModule.h"
 #include "Interfaces/ITextureFormat.h"
@@ -1086,7 +1083,7 @@ int64 PutDerivedDataInCache(FTexturePlatformData* DerivedData, const FString& De
 		if (UE_LOG_ACTIVE(LogTexture,Verbose) || bDDCError)
 		{
 			if (LogString.IsEmpty())
-		{
+			{
 				LogString = FString::Printf(
 					TEXT("Storing texture in DDC:\n  Name: %s\n  Key: %s\n  Format: %s\n"),
 					*FString(TextureName),
@@ -1144,9 +1141,17 @@ int64 PutDerivedDataInCache(FTexturePlatformData* DerivedData, const FString& De
 	// At this point we've stored all the non-inline data in the DDC, so this will only serialize and store the TexturePlatformData metadata and any inline mips
 	FMemoryWriter64 Ar(RawDerivedData, /*bIsPersistent=*/ true);
 	DerivedData->Serialize(Ar, NULL);
-	TotalBytesPut += RawDerivedData.Num();
-	GetDerivedDataCacheRef().Put(*DerivedDataKey, RawDerivedData, TextureName, bReplaceExistingDerivedDataDDC);
-	UE_LOG(LogTexture,Verbose,TEXT("%s  Derived Data: %d bytes"),*LogString,RawDerivedData.Num());
+	const int64 RawDerivedDataSize = RawDerivedData.Num();
+	TotalBytesPut += RawDerivedDataSize;
+
+	using namespace UE::DerivedData;
+	FRequestOwner AsyncOwner(EPriority::Normal);
+	FValue Value = FValue::Compress(MakeSharedBufferFromArray(MoveTemp(RawDerivedData)));
+	const ECachePolicy Policy = bReplaceExistingDerivedDataDDC ? ECachePolicy::Store : ECachePolicy::Default;
+	GetCache().PutValue({{{TextureName}, ConvertLegacyCacheKey(DerivedDataKey), MoveTemp(Value), Policy}}, AsyncOwner);
+	AsyncOwner.KeepAlive();
+
+	UE_LOG(LogTexture, Verbose, TEXT("%s  Derived Data: %" INT64_FMT " bytes"), *LogString, RawDerivedDataSize);
 	return TotalBytesPut;
 }
 
@@ -1337,34 +1342,36 @@ typedef TArray<uint32> FAsyncVTChunkHandles;
  * Executes async DDC gets for mips stored in the derived data cache.
  * @param Mip - Mips to retrieve.
  * @param FirstMipToLoad - Index of the first mip to retrieve.
- * @param OutHandles - Handles to the asynchronous DDC gets.
+ * @param Callback - Callback invoked for each mip as it loads.
  * 
  * This function must be called after the initial DDC fetch is complete,
  * so we know what our in-use key is. This might be on the worker immediately
  * after the fetch completes.
  */
-static bool BeginLoadDerivedMips(FTexturePlatformData& PlatformData, int32 FirstMipToLoad, FStringView DebugContext, FAsyncMipHandles& OutHandles, TUniqueFunction<void (int32 MipIndex, FSharedBuffer MipData)> Callback)
+static bool LoadDerivedStreamingMips(FTexturePlatformData& PlatformData, int32 FirstMipToLoad, FStringView DebugContext, TFunctionRef<void (int32 MipIndex, FSharedBuffer MipData)> Callback)
 {
 	using namespace UE::DerivedData;
+	TArray<FCacheGetChunkRequest, TInlineAllocator<MAX_TEXTURE_MIP_COUNT>> Requests;
 
 	TIndirectArray<FTexture2DMipMap>& Mips = PlatformData.Mips;
 	if (PlatformData.DerivedDataKey.IsType<FString>())
 	{
-		FDerivedDataCacheInterface& DDC = GetDerivedDataCacheRef();
-		OutHandles.AddZeroed(Mips.Num());
 		for (int32 MipIndex = FirstMipToLoad; MipIndex < Mips.Num(); ++MipIndex)
 		{
 			const FTexture2DMipMap& Mip = Mips[MipIndex];
 			if (Mip.IsPagedToDerivedData())
 			{
-				OutHandles[MipIndex] = DDC.GetAsynchronous(*PlatformData.GetDerivedDataMipKeyString(MipIndex, Mip), DebugContext);
+				TStringBuilder<256> MipNameBuilder;
+				MipNameBuilder.Append(DebugContext).Appendf(TEXT(" [MIP %d]"), MipIndex);
+				FCacheGetChunkRequest& Request = Requests.AddDefaulted_GetRef();
+				Request.Name = MipNameBuilder;
+				Request.Key = ConvertLegacyCacheKey(PlatformData.GetDerivedDataMipKeyString(MipIndex, Mip));
+				Request.UserData = MipIndex;
 			}
 		}
 	}
 	else if (PlatformData.DerivedDataKey.IsType<FCacheKeyProxy>())
 	{
-		TArray<FCacheGetChunkRequest> MipKeys;
-
 		const FCacheKey& Key = *PlatformData.DerivedDataKey.Get<UE::DerivedData::FCacheKeyProxy>().AsCacheKey();
 		for (int32 MipIndex = FirstMipToLoad; MipIndex < Mips.Num(); ++MipIndex)
 		{
@@ -1372,58 +1379,89 @@ static bool BeginLoadDerivedMips(FTexturePlatformData& PlatformData, int32 First
 			if (Mip.IsPagedToDerivedData())
 			{
 				TStringBuilder<256> MipNameBuilder;
-				MipNameBuilder.Append(DebugContext).Appendf(TEXT(" [MIP 0]"), MipIndex);
-				FCacheGetChunkRequest& Request = MipKeys.AddDefaulted_GetRef();
+				MipNameBuilder.Append(DebugContext).Appendf(TEXT(" [MIP %d]"), MipIndex);
+				FCacheGetChunkRequest& Request = Requests.AddDefaulted_GetRef();
 				Request.Name = MipNameBuilder;
 				Request.Key = Key;
 				Request.Id = FTexturePlatformData::MakeMipId(MipIndex);
 				Request.UserData = MipIndex;
 			}
 		}
+	}
+	else
+	{
+		UE_LOG(LogTexture, Error, TEXT("Attempting to stream in mips for texture that has not generated a supported derived data key format."));
+	}
 
-		if (!MipKeys.IsEmpty())
+	bool bMiss = false;
+
+	if (!Requests.IsEmpty())
+	{
+		COOK_STAT(auto Timer = TextureCookStats::StreamingMipUsageStats.TimeSyncWork());
+		uint64 Size = 0;
+		FRequestOwner BlockingOwner(EPriority::Blocking);
+		GetCache().GetChunks(Requests, BlockingOwner, [Callback = MoveTemp(Callback), &Size, &bMiss](FCacheGetChunkResponse&& Response)
 		{
-			check(Callback);
-			bool bMiss = false;
-			FRequestOwner RequestOwner(EPriority::Blocking);
-			GetCache().GetChunks(MipKeys, RequestOwner, [Callback = MoveTemp(Callback), &bMiss](FCacheGetChunkResponse&& Response)
+			Size += Response.RawSize;
+			if (Response.Status == EStatus::Ok)
 			{
-				if (Response.Status == EStatus::Ok)
-				{
-					Callback(int32(Response.UserData), MoveTemp(Response.RawData));
-				}
-				else
-				{
-					bMiss = true;
-				}
-			});
-			RequestOwner.Wait();
-
-			if (bMiss)
+				Callback(int32(Response.UserData), MoveTemp(Response.RawData));
+			}
+			else
 			{
-				return false;
+				bMiss = true;
 			}
-		}
+		});
+		BlockingOwner.Wait();
+		COOK_STAT(Timer.AddHitOrMiss(!bMiss ? FCookStats::CallStats::EHitOrMiss::Hit : FCookStats::CallStats::EHitOrMiss::Miss, int64(Size)));
 	}
-	else
-	{
-		UE_LOG(LogTexture, Error, TEXT("Attempting to stream in mips for texture that has not generated a supported derived data key format."));
-	}
-	return true;
+
+	return !bMiss;
 }
 
-static void BeginLoadDerivedVTChunks(const TArray<FVirtualTextureDataChunk>& Chunks, FStringView DebugContext, FAsyncVTChunkHandles& OutHandles)
+static bool LoadDerivedStreamingVTChunks(const TArray<FVirtualTextureDataChunk>& Chunks, FStringView DebugContext, TFunctionRef<void (int32 ChunkIndex, FSharedBuffer ChunkData)> Callback)
 {
-	FDerivedDataCacheInterface& DDC = GetDerivedDataCacheRef();
-	OutHandles.AddZeroed(Chunks.Num());
+	using namespace UE::DerivedData;
+	TArray<FCacheGetChunkRequest> Requests;
+
 	for (int32 ChunkIndex = 0; ChunkIndex < Chunks.Num(); ++ChunkIndex)
 	{
 		const FVirtualTextureDataChunk& Chunk = Chunks[ChunkIndex];
-		if (Chunk.DerivedDataKey.IsEmpty() == false)
+		if (!Chunk.DerivedDataKey.IsEmpty())
 		{
-			OutHandles[ChunkIndex] = DDC.GetAsynchronous(*Chunk.DerivedDataKey, DebugContext);
+			TStringBuilder<256> ChunkNameBuilder;
+			ChunkNameBuilder.Append(DebugContext).Appendf(TEXT(" [Chunk %d]"), ChunkIndex);
+			FCacheGetChunkRequest& Request = Requests.AddDefaulted_GetRef();
+			Request.Name = ChunkNameBuilder;
+			Request.Key = ConvertLegacyCacheKey(Chunk.DerivedDataKey);
+			Request.UserData = ChunkIndex;
 		}
 	}
+
+	bool bMiss = false;
+
+	if (!Requests.IsEmpty())
+	{
+		COOK_STAT(auto Timer = TextureCookStats::StreamingMipUsageStats.TimeSyncWork());
+		uint64 Size = 0;
+		FRequestOwner BlockingOwner(EPriority::Blocking);
+		GetCache().GetChunks(Requests, BlockingOwner, [Callback = MoveTemp(Callback), &Size, &bMiss](FCacheGetChunkResponse&& Response)
+		{
+			Size += Response.RawSize;
+			if (Response.Status == EStatus::Ok)
+			{
+				Callback(int32(Response.UserData), MoveTemp(Response.RawData));
+			}
+			else
+			{
+				bMiss = true;
+			}
+		});
+		BlockingOwner.Wait();
+		COOK_STAT(Timer.AddHitOrMiss(!bMiss ? FCookStats::CallStats::EHitOrMiss::Hit : FCookStats::CallStats::EHitOrMiss::Miss, int64(Size)));
+	}
+
+	return !bMiss;
 }
 
 /** Logs a warning that MipSize is correct for the mipmap. */
@@ -1455,91 +1493,71 @@ static void CheckMipSize(FTexture2DMipMap& Mip, EPixelFormat PixelFormat, int64
 //
 static bool FetchAllTextureDataSynchronous(FTexturePlatformData* PlatformData, FStringView DebugContext, TArray<TArray64<uint8>>& OutMipData, TArray<TArray64<uint8>>& OutVTChunkData)
 {
-	OutMipData.Empty();
-	OutVTChunkData.Empty();
-
-	using namespace UE::DerivedData;
-	FDerivedDataCacheInterface& DDC = GetDerivedDataCacheRef();
+	const auto StoreMip = [&OutMipData](int32 MipIndex, FSharedBuffer MipBuffer)
+	{
+		OutMipData[MipIndex].Append(static_cast<const uint8*>(MipBuffer.GetData()), MipBuffer.GetSize());
+	};
 
-	OutMipData.AddDefaulted(PlatformData->Mips.Num());
+	const int32 MipCount = PlatformData->Mips.Num();
+	OutMipData.Empty(MipCount);
+	OutMipData.AddDefaulted(MipCount);
 
-	// This only handles non-vt mips that are paged to derived data.
-	// (some mips are inline). Doesn't handle excluded mips due to 
-	// platform settings.
-	FAsyncMipHandles MipHandles;
-	if (!BeginLoadDerivedMips(*PlatformData, 0, DebugContext, MipHandles,
-			[&OutMipData](int32 MipIndex, FSharedBuffer MipBuffer)
-				{
-					OutMipData[MipIndex].Append((uint8*)MipBuffer.GetData(), MipBuffer.GetSize());
-				}
-			)
-		)
+	if (!LoadDerivedStreamingMips(*PlatformData, 0, DebugContext, StoreMip))
 	{
 		return false;
 	}
 
-	// DDC1 fetches are async, so we need to wait on the handles,
-	// and if it's not paged to derived data we need to copy the
-	// data from bulk
-	for (int32 MipIndex = 0; MipIndex < PlatformData->Mips.Num(); ++MipIndex)
+	for (int32 MipIndex = 0; MipIndex < MipCount; ++MipIndex)
 	{
 		FTexture2DMipMap& Mip = PlatformData->Mips[MipIndex];
-		if (Mip.IsPagedToDerivedData() == false)
+		if (OutMipData[MipIndex].IsEmpty())
 		{
-			OutMipData[MipIndex].Append((uint8*)Mip.BulkData.LockReadOnly(), Mip.BulkData.GetBulkDataSize());
-			Mip.BulkData.Unlock();
-			continue;
+			if (Mip.BulkData.IsBulkDataLoaded())
+			{
+				OutMipData[MipIndex].Append((uint8*)Mip.BulkData.LockReadOnly(), Mip.BulkData.GetBulkDataSize());
+				Mip.BulkData.Unlock();
+			}
+			else
+			{
+				return false;
+			}
 		}
+	}
 
-		// Here we either got the data synchronously with the DDC2 path in BeginLoadDerivedMips,
-		// or we didn't if we are DDC1. BeginLoadDerivedMips only allocates the async handles
-		// in the ddc1 path, so we use that to tell.
-		if (MipHandles.Num() == 0)
-		{
-			// skip, ddc1. This is the same for all mips but we still need to check for paged
-			// data above.
-			continue;
-		}
+	const auto StoreChunk = [&OutVTChunkData](int32 ChunkIndex, FSharedBuffer ChunkBuffer)
+	{
+		OutVTChunkData[ChunkIndex].Append(static_cast<const uint8*>(ChunkBuffer.GetData()), ChunkBuffer.GetSize());
+	};
+
+	const int32 ChunkCount = PlatformData->VTData ? PlatformData->VTData->Chunks.Num() : 0;
+	OutVTChunkData.Empty(ChunkCount);
+	if (ChunkCount)
+	{
+		OutVTChunkData.AddDefaulted(ChunkCount);
 
-		uint32 AsyncHandle = MipHandles[MipIndex];
-		DDC.WaitAsynchronousCompletion(AsyncHandle);
-		bool bLoadedFromDDC = DDC.GetAsynchronousResults(AsyncHandle, OutMipData[MipIndex]);
-		if (bLoadedFromDDC == false)
+		if (!LoadDerivedStreamingVTChunks(PlatformData->VTData->Chunks, DebugContext, StoreChunk))
 		{
 			return false;
 		}
-	}
-	
-	if (PlatformData->VTData)
-	{
-		OutVTChunkData.AddDefaulted(PlatformData->VTData->Chunks.Num());
-
-		FAsyncVTChunkHandles AsyncVTHandles;
-		BeginLoadDerivedVTChunks(PlatformData->VTData->Chunks, DebugContext, AsyncVTHandles);
 
-		for (int32 ChunkIndex = 0; ChunkIndex < PlatformData->VTData->Chunks.Num(); ++ChunkIndex)
+		for (int32 ChunkIndex = 0; ChunkIndex < ChunkCount; ++ChunkIndex)
 		{
 			FVirtualTextureDataChunk& Chunk = PlatformData->VTData->Chunks[ChunkIndex];
-			if (Chunk.DerivedDataKey.IsEmpty())
+			if (OutVTChunkData[ChunkIndex].IsEmpty())
 			{
-				// The data is resident and we can just copy it.
-				OutVTChunkData[ChunkIndex].Append((uint8*)Chunk.BulkData.LockReadOnly(), Chunk.BulkData.GetBulkDataSize());
-				Chunk.BulkData.Unlock();
-			}
-			else
-			{
-				// The data was fetched and we need to wait on the result.
-				uint32 AsyncHandle = AsyncVTHandles[ChunkIndex];
-				DDC.WaitAsynchronousCompletion(AsyncHandle);
-				bool bLoadedFromDDC = DDC.GetAsynchronousResults(AsyncHandle, OutVTChunkData[ChunkIndex]);
-
-				if (bLoadedFromDDC == false)
+				if (Chunk.BulkData.IsBulkDataLoaded())
+				{
+					// The data is resident and we can just copy it.
+					OutVTChunkData[ChunkIndex].Append((uint8*)Chunk.BulkData.LockReadOnly(), Chunk.BulkData.GetBulkDataSize());
+					Chunk.BulkData.Unlock();
+				}
+				else
 				{
 					return false;
 				}
 			}
-		} // end each vt chunk
-	} // end if virtual texture
+		}
+	}
 
 	return true;
 }
@@ -1760,95 +1778,36 @@ bool FTexturePlatformData::TryInlineMipData(int32 FirstMipToLoad, FStringView De
 {
 	TRACE_CPUPROFILER_EVENT_SCOPE(FTexturePlatformData::TryInlineMipData);
 
-	FAsyncMipHandles AsyncHandles;
-	FAsyncVTChunkHandles AsyncVTHandles;
-	TArray64<uint8> TempData;
-	FDerivedDataCacheInterface& DDC = GetDerivedDataCacheRef();
-
+	const auto StoreMip = [this](int32 MipIndex, FSharedBuffer MipBuffer)
+	{
+		FTexture2DMipMap& Mip = Mips[MipIndex];
+		Mip.BulkData.Lock(LOCK_READ_WRITE);
+		void* MipData = Mip.BulkData.Realloc(int64(MipBuffer.GetSize()));
+		FMemory::Memcpy(MipData, MipBuffer.GetData(), MipBuffer.GetSize());
+		Mip.BulkData.Unlock();
+		Mip.SetPagedToDerivedData(false);
+	};
 
-	if (!BeginLoadDerivedMips(*this, FirstMipToLoad, DebugContext, AsyncHandles,
-		[this](int32 MipIndex, FSharedBuffer MipBuffer)
-		{
-			FTexture2DMipMap& Mip = Mips[MipIndex];
-			Mip.BulkData.Lock(LOCK_READ_WRITE);
-			void* MipData = Mip.BulkData.Realloc(int64(MipBuffer.GetSize()));
-			FMemory::Memcpy(MipData, MipBuffer.GetData(), MipBuffer.GetSize());
-			Mip.BulkData.Unlock();
-			Mip.SetPagedToDerivedData(false);
-		}))
+	if (!LoadDerivedStreamingMips(*this, FirstMipToLoad, DebugContext, StoreMip))
 	{
 		return false;
 	}
 
-	if (VTData != nullptr)
+	const auto StoreChunk = [this](int32 ChunkIndex, FSharedBuffer ChunkBuffer)
 	{
-		BeginLoadDerivedVTChunks(VTData->Chunks, DebugContext, AsyncVTHandles);
-	}
+		FVirtualTextureDataChunk& Chunk = VTData->Chunks[ChunkIndex];
+		Chunk.BulkData.Lock(LOCK_READ_WRITE);
+		void* ChunkData = Chunk.BulkData.Realloc(int64(ChunkBuffer.GetSize()));
+		FMemory::Memcpy(ChunkData, ChunkBuffer.GetData(), ChunkBuffer.GetSize());
+		Chunk.BulkData.Unlock();
+		Chunk.DerivedDataKey.Empty();
+	};
 
-	// Process regular mips
-	for (int32 MipIndex = FirstMipToLoad; MipIndex < Mips.Num(); ++MipIndex)
+	if (VTData && !LoadDerivedStreamingVTChunks(VTData->Chunks, DebugContext, StoreChunk))
 	{
-		FTexture2DMipMap& Mip = Mips[MipIndex];
-		if (Mip.IsPagedToDerivedData())
-		{
-			uint32 AsyncHandle = AsyncHandles[MipIndex];
-			bool bLoadedFromDDC = false;
-			{
-				COOK_STAT(auto Timer = TextureCookStats::StreamingMipUsageStats.TimeAsyncWait());
-				DDC.WaitAsynchronousCompletion(AsyncHandle);
-				bLoadedFromDDC = DDC.GetAsynchronousResults(AsyncHandle, TempData);
-				COOK_STAT(Timer.AddHitOrMiss(bLoadedFromDDC ? FCookStats::CallStats::EHitOrMiss::Hit : FCookStats::CallStats::EHitOrMiss::Miss, TempData.Num()));
-			}
-			if (bLoadedFromDDC)
-			{
-				FMemoryReaderView Ar(MakeMemoryView(TempData), /*bIsPersistent=*/ true);
-
-				Mip.BulkData.Lock(LOCK_READ_WRITE);
-				void* MipData = Mip.BulkData.Realloc(TempData.Num());
-				Ar.Serialize(MipData, TempData.Num());
-				Mip.BulkData.Unlock();
-				Mip.SetPagedToDerivedData(false);
-			}
-			else
-			{
-				return false;
-			}
-			TempData.Reset();
-		}
+		return false;
 	}
 
-	// Process VT mips
-	if (VTData != nullptr)
-	{
-		for (int32 ChunkIndex = 0; ChunkIndex < VTData->Chunks.Num(); ++ChunkIndex)
-		{
-			FVirtualTextureDataChunk& Chunk = VTData->Chunks[ChunkIndex];
-			if (Chunk.DerivedDataKey.IsEmpty() == false)
-			{
-				uint32 AsyncHandle = AsyncVTHandles[ChunkIndex];
-				bool bLoadedFromDDC = false;
-				{
-					COOK_STAT(auto Timer = TextureCookStats::StreamingMipUsageStats.TimeAsyncWait());
-					DDC.WaitAsynchronousCompletion(AsyncHandle);
-					bLoadedFromDDC = DDC.GetAsynchronousResults(AsyncHandle, TempData);
-					COOK_STAT(Timer.AddHitOrMiss(bLoadedFromDDC ? FCookStats::CallStats::EHitOrMiss::Hit : FCookStats::CallStats::EHitOrMiss::Miss, TempData.Num()));
-				}
-				if (bLoadedFromDDC)
-				{
-					Chunk.BulkData.Lock(LOCK_READ_WRITE);
-					void* ChunkData = Chunk.BulkData.Realloc(TempData.Num());
-					FMemory::Memcpy(ChunkData, TempData.GetData(), TempData.Num());
-					Chunk.BulkData.Unlock();
-					Chunk.DerivedDataKey.Empty();
-				}
-				else
-				{
-					return false;
-				}
-				TempData.Reset();
-			}
-		}
-	}
 	return true;
 }
 
@@ -1908,10 +1867,10 @@ bool FTexturePlatformData::TryLoadMips(int32 FirstMipToLoad, void** OutMipData,
 	// TryLoadMips fills mip pointers but not sizes
 	//  dangerous, not robust, use TryLoadMipsWithSizes instead
 
-	return TryLoadMipsWithSizes(FirstMipToLoad,OutMipData,nullptr,DebugContext);
+	return TryLoadMipsWithSizes(FirstMipToLoad, OutMipData, nullptr, DebugContext);
 }
 
-bool FTexturePlatformData::TryLoadMipsWithSizes(int32 FirstMipToLoad, void** OutMipData, int64 * OutMipSize, FStringView DebugContext)
+bool FTexturePlatformData::TryLoadMipsWithSizes(int32 FirstMipToLoad, void** OutMipData, int64* OutMipSize, FStringView DebugContext)
 {
 	TRACE_CPUPROFILER_EVENT_SCOPE(FTexturePlatformData::TryLoadMips);
 
@@ -1920,29 +1879,26 @@ bool FTexturePlatformData::TryLoadMipsWithSizes(int32 FirstMipToLoad, void** Out
 	check(LoadableMips >= 0);
 
 #if WITH_EDITOR
+	const auto StoreMip = [this, OutMipData, OutMipSize, FirstMipToLoad, &NumMipsCached](int32 MipIndex, FSharedBuffer MipBuffer)
+	{
+		FTexture2DMipMap& Mip = Mips[MipIndex];
 
-	TArray64<uint8> TempData;
-	FAsyncMipHandles AsyncHandles;
-	FDerivedDataCacheInterface& DDC = GetDerivedDataCacheRef();
-	if (!BeginLoadDerivedMips(*this, FirstMipToLoad, DebugContext, AsyncHandles,
-		[this, OutMipData,OutMipSize, FirstMipToLoad, &NumMipsCached](int32 MipIndex, FSharedBuffer MipBuffer)
-		{
-			FTexture2DMipMap& Mip = Mips[MipIndex];
+		const int64 MipSize = static_cast<int64>(MipBuffer.GetSize());
+		CheckMipSize(Mip, PixelFormat, MipSize);
+		NumMipsCached++;
 
-			const int64 MipSize = static_cast<int64>(MipBuffer.GetSize());
-			CheckMipSize(Mip, PixelFormat, MipSize);
-			NumMipsCached++;
+		if (OutMipData)
+		{
+			OutMipData[MipIndex - FirstMipToLoad] = FMemory::Malloc(MipSize);
+			FMemory::Memcpy(OutMipData[MipIndex - FirstMipToLoad], MipBuffer.GetData(), MipSize);
+		}
+		if ( OutMipSize ) 
+		{
+			OutMipSize[MipIndex - FirstMipToLoad] = MipSize;
+		}
+	};
 
-			if (OutMipData)
-			{
-				OutMipData[MipIndex - FirstMipToLoad] = FMemory::Malloc(MipSize);
-				FMemory::Memcpy(OutMipData[MipIndex - FirstMipToLoad], MipBuffer.GetData(), MipSize);
-			}
-			if ( OutMipSize ) 
-			{
-				OutMipSize[MipIndex - FirstMipToLoad] = MipSize;
-			}
-		}))
+	if (!LoadDerivedStreamingMips(*this, FirstMipToLoad, DebugContext, StoreMip))
 	{
 		return false;
 	}
@@ -1967,7 +1923,7 @@ bool FTexturePlatformData::TryLoadMipsWithSizes(int32 FirstMipToLoad, void** Out
 		const int64 BulkDataSize = Mip.BulkData.GetBulkDataSize();
 		if (BulkDataSize > 0)
 		{
-			if (OutMipData != nullptr)
+			if (OutMipData)
 			{
 #if PLATFORM_SUPPORTS_TEXTURE_STREAMING
 				// We want to make sure that any non-streamed mips are coming from the texture asset file, and not from an external bulk file.
@@ -1982,7 +1938,7 @@ bool FTexturePlatformData::TryLoadMipsWithSizes(int32 FirstMipToLoad, void** Out
 #endif
 				Mip.BulkData.GetCopy(&OutMipData[MipIndex - FirstMipToLoad], true);
 			}
-			if ( OutMipSize )
+			if (OutMipSize)
 			{
 				OutMipSize[MipIndex - FirstMipToLoad] = BulkDataSize;
 			}
@@ -1990,45 +1946,6 @@ bool FTexturePlatformData::TryLoadMipsWithSizes(int32 FirstMipToLoad, void** Out
 		}
 	}
 
-#if WITH_EDITOR
-	if (DerivedDataKey.IsType<FString>())
-	{
-		// Wait for async DDC gets.
-		for (int32 MipIndex = FirstMipToLoad; MipIndex < LoadableMips; ++MipIndex)
-		{
-			FTexture2DMipMap& Mip = Mips[MipIndex];
-			if (Mip.IsPagedToDerivedData())
-			{
-				uint32 AsyncHandle = AsyncHandles[MipIndex];
-				DDC.WaitAsynchronousCompletion(AsyncHandle);
-				if (DDC.GetAsynchronousResults(AsyncHandle, TempData))
-				{
-					FMemoryReaderView Ar(MakeMemoryView(TempData), /*bIsPersistent=*/ true);
-					CheckMipSize(Mip, PixelFormat, TempData.Num());
-					NumMipsCached++;
-
-					if (OutMipData)
-					{
-						OutMipData[MipIndex - FirstMipToLoad] = FMemory::Malloc(TempData.Num());
-						Ar.Serialize(OutMipData[MipIndex - FirstMipToLoad], TempData.Num());
-					}
-					if ( OutMipSize )
-					{
-						OutMipSize[MipIndex - FirstMipToLoad] = TempData.Num();
-					}
-				}
-				else
-				{
-					UE_LOG(LogTexture, Verbose, TEXT("DDC.GetAsynchronousResults() failed for %.*s, MipIndex: %d"),
-						DebugContext.Len(), DebugContext.GetData(),
-						MipIndex);
-				}
-				TempData.Reset();
-			}
-		}
-	}
-#endif // #if WITH_EDITOR
-
 	if (NumMipsCached != (LoadableMips - FirstMipToLoad))
 	{
 		UE_LOG(LogTexture, Verbose, TEXT("TryLoadMips failed for %.*s, NumMipsCached: %d, LoadableMips: %d, FirstMipToLoad: %d"),
@@ -2192,86 +2109,103 @@ bool FTexturePlatformData::CanUseCookedDataPath() const
 #if WITH_EDITOR
 bool FTexturePlatformData::AreDerivedMipsAvailable(FStringView Context) const
 {
-	TArray<int32, TInlineAllocator<16>> PagedMipIndices;
-	for (int32 MipIndex = 0; MipIndex < Mips.Num(); ++MipIndex)
+	if (DerivedDataKey.IsType<FString>())
 	{
-		if (Mips[MipIndex].IsPagedToDerivedData())
+		using namespace UE::DerivedData;
+		TArray<FCacheGetValueRequest, TInlineAllocator<16>> MipRequests;
+
+		int32 MipIndex = 0;
+		const FSharedString SharedContext = Context;
+		for (const FTexture2DMipMap& Mip : Mips)
 		{
-			PagedMipIndices.Add(MipIndex);
+			if (Mip.IsPagedToDerivedData())
+			{
+				const FCacheKey MipKey = ConvertLegacyCacheKey(GetDerivedDataMipKeyString(MipIndex, Mip));
+				const ECachePolicy ExistsPolicy = ECachePolicy::Query | ECachePolicy::SkipData;
+				MipRequests.Add({SharedContext, MipKey, ExistsPolicy});
+			}
+			++MipIndex;
 		}
-	}
-
-	if (PagedMipIndices.IsEmpty())
-	{
-		return true;
-	}
-
-	bool bAreDerivedMipsAvailable = false;
 
-	if (DerivedDataKey.IsType<FString>())
-	{
-		TArray<FString, TInlineAllocator<16>> MipKeys;
-		for (int32 PagedMipIndex : PagedMipIndices)
+		if (MipRequests.IsEmpty())
 		{
-			MipKeys.Add(GetDerivedDataMipKeyString(PagedMipIndex, Mips[PagedMipIndex]));
+			return true;
 		}
 
-		bAreDerivedMipsAvailable = GetDerivedDataCacheRef().AllCachedDataProbablyExists(MipKeys);
-
-		// When using a shared DDC and performing async loading, 
-		// prefetch the lowest mip to avoid high prio request stalls from the render thread
-		if (bAreDerivedMipsAvailable && !IsInGameThread())
+		// When performing async loading, prefetch the lowest streaming mip into local caches
+		// to avoid high priority request stalls from the render thread.
+		if (!IsInGameThread())
 		{
-			TRACE_CPUPROFILER_EVENT_SCOPE(PrefetchSmallestMip);
-			GetDerivedDataCacheRef().TryToPrefetch({MipKeys.Last()}, Context);
+			MipRequests.Last().Policy |= ECachePolicy::StoreLocal;
 		}
+
+		bool bAreDerivedMipsAvailable = true;
+		FRequestOwner BlockingOwner(EPriority::Blocking);
+		GetCache().GetValue(MipRequests, BlockingOwner, [&bAreDerivedMipsAvailable](FCacheGetValueResponse&& Response)
+		{
+			bAreDerivedMipsAvailable &= Response.Status == EStatus::Ok;
+		});
+		BlockingOwner.Wait();
+		return bAreDerivedMipsAvailable;
 	}
 	else if (DerivedDataKey.IsType<UE::DerivedData::FCacheKeyProxy>())
 	{
-		bAreDerivedMipsAvailable = true;
+		return true;
 	}
 
-	return bAreDerivedMipsAvailable;
+	return false;
 }
 
 bool FTexturePlatformData::AreDerivedVTChunksAvailable(FStringView Context) const
 {
 	check(VTData);
-	TArray<FString, TInlineAllocator<16>> ChunkKeys;
+
+	using namespace UE::DerivedData;
+	TArray<FCacheGetValueRequest, TInlineAllocator<16>> ChunkRequests;
+
+	int32 ChunkIndex = 0;
+	const FSharedString SharedContext = Context;
 	for (const FVirtualTextureDataChunk& Chunk : VTData->Chunks)
 	{
 		if (!Chunk.DerivedDataKey.IsEmpty())
 		{
-			ChunkKeys.Add(Chunk.DerivedDataKey);
+			const FCacheKey ChunkKey = ConvertLegacyCacheKey(Chunk.DerivedDataKey);
+			const ECachePolicy ExistsPolicy = ECachePolicy::Query | ECachePolicy::SkipData;
+			ChunkRequests.Add({SharedContext, ChunkKey, ExistsPolicy});
 		}
+		++ChunkIndex;
 	}
 
-	if (ChunkKeys.IsEmpty())
+	if (ChunkRequests.IsEmpty())
 	{
 		return true;
 	}
 
-	const bool bAreDerivedChunksAvailable = GetDerivedDataCacheRef().AllCachedDataProbablyExists(ChunkKeys);
-
-	// When using a shared DDC and performing async loading, 
-	// prefetch the lowest mip to avoid high prio request stalls from the render thread
-	if (bAreDerivedChunksAvailable && !IsInGameThread())
+	// When performing async loading, prefetch the last chunk into local caches
+	// to avoid high priority request stalls from the render thread.
+	if (!IsInGameThread())
 	{
-		TRACE_CPUPROFILER_EVENT_SCOPE(PrefetchSmallestDerivedVTChunk);
-		GetDerivedDataCacheRef().TryToPrefetch({ ChunkKeys.Last() }, Context);
+		ChunkRequests.Last().Policy |= ECachePolicy::StoreLocal;
 	}
 
+	bool bAreDerivedChunksAvailable = true;
+	FRequestOwner BlockingOwner(EPriority::Blocking);
+	GetCache().GetValue(ChunkRequests, BlockingOwner, [&bAreDerivedChunksAvailable](FCacheGetValueResponse&& Response)
+	{
+		bAreDerivedChunksAvailable &= Response.Status == EStatus::Ok;
+	});
+	BlockingOwner.Wait();
 	return bAreDerivedChunksAvailable;
 }
 
 bool FTexturePlatformData::AreDerivedMipsAvailable() const
 {
-	return AreDerivedMipsAvailable(TEXT("DerivedMips"_SV));
+	return AreDerivedMipsAvailable(TEXTVIEW("DerivedMips"));
 }
 
 bool FTexturePlatformData::AreDerivedVTChunksAvailable() const
 {
-	return AreDerivedVTChunksAvailable(TEXT("DerivedVTChunks"_SV));
+	return AreDerivedVTChunksAvailable(TEXTVIEW("DerivedVTChunks"));
 }
 
 #endif // #if WITH_EDITOR
@@ -3151,7 +3085,6 @@ void UTexture::ForceRebuildPlatformData(uint8 InEncodeSpeedOverride /* =255 ETex
 			EncodeSpeed = GetDesiredEncodeSpeed();
 		}
 
-
 		TArray<FTextureBuildSettings> BuildSettingsFetch;
 		TArray<FTextureBuildSettings> BuildSettingsFetchOrBuild;
 		TArray<FTexturePlatformData::FTextureEncodeResultMetadata> ResultMetadataFetch;
@@ -3181,53 +3114,8 @@ void UTexture::ForceRebuildPlatformData(uint8 InEncodeSpeedOverride /* =255 ETex
 	}
 }
 
-
-
 void UTexture::MarkPlatformDataTransient()
 {
-	FDerivedDataCacheInterface& DDC = GetDerivedDataCacheRef();
-
-	FTexturePlatformData** RunningPlatformData = GetRunningPlatformData();
-	if (RunningPlatformData)
-	{
-		FTexturePlatformData* PlatformData = *RunningPlatformData;
-		if (PlatformData)
-		{
-			if (const FString* KeyString = PlatformData->DerivedDataKey.TryGet<FString>())
-			{
-				for (int32 MipIndex = 0; MipIndex < PlatformData->Mips.Num(); ++MipIndex)
-				{
-					FTexture2DMipMap& Mip = PlatformData->Mips[MipIndex];
-					if (Mip.IsPagedToDerivedData())
-					{
-						DDC.MarkTransient(*PlatformData->GetDerivedDataMipKeyString(MipIndex, Mip));
-					}
-				}
-				DDC.MarkTransient(**KeyString);
-			}
-		}
-	}
-
-	TMap<FString, FTexturePlatformData*> *CookedPlatformData = GetCookedPlatformData();
-	if ( CookedPlatformData )
-	{
-		for ( auto It : *CookedPlatformData )
-		{
-			FTexturePlatformData* PlatformData = It.Value;
-			if (const FString* KeyString = PlatformData->DerivedDataKey.TryGet<FString>())
-			{
-				for (int32 MipIndex = 0; MipIndex < PlatformData->Mips.Num(); ++MipIndex)
-				{
-					FTexture2DMipMap& Mip = PlatformData->Mips[MipIndex];
-					if (Mip.IsPagedToDerivedData())
-					{
-						DDC.MarkTransient(*PlatformData->GetDerivedDataMipKeyString(MipIndex, Mip));
-					}
-				}
-				DDC.MarkTransient(**KeyString);
-			}
-		}
-	}
 }
 #endif // #if WITH_EDITOR
 
diff --git a/Engine/Source/Runtime/Engine/Private/TextureDerivedDataTask.cpp b/Engine/Source/Runtime/Engine/Private/TextureDerivedDataTask.cpp
--- a/Engine/Source/Runtime/Engine/Private/TextureDerivedDataTask.cpp
+++ b/Engine/Source/Runtime/Engine/Private/TextureDerivedDataTask.cpp
@@ -30,11 +30,8 @@
 #include "DerivedDataBuildInputs.h"
 #include "DerivedDataBuildOutput.h"
 #include "DerivedDataBuildSession.h"
-#include "DerivedDataCacheInterface.h"
-#include "DerivedDataCacheKey.h"
+#include "DerivedDataCache.h"
 #include "DerivedDataRequestOwner.h"
-#include "DerivedDataSharedString.h"
-#include "DerivedDataValue.h"
 #include "Engine/TextureCube.h"
 #include "Engine/VolumeTexture.h"
 #include "ImageCoreUtils.h"
@@ -728,6 +725,8 @@ FTextureCacheDerivedDataWorker::FTextureCacheDerivedDataWorker(
 // DDC1 primary fetch/build work function
 void FTextureCacheDerivedDataWorker::DoWork()
 {
+	using namespace UE::DerivedData;
+
 	TRACE_CPUPROFILER_EVENT_SCOPE(FTextureCacheDerivedDataWorker::DoWork);
 
 	const bool bForceRebuild = EnumHasAnyFlags(CacheFlags, ETextureCacheFlags::ForceRebuild);
@@ -737,7 +736,7 @@ void FTextureCacheDerivedDataWorker::DoWork()
 	const bool bValidateVirtualTextureCompression = CVarVTValidateCompressionOnLoad.GetValueOnAnyThread() != 0;
 	bool bInvalidVirtualTextureCompression = false;
 
-	TArray64<uint8> RawDerivedData;
+	FSharedBuffer RawDerivedData;
 
 	// Can't have a fetch first if we are rebuilding
 	if (bForceRebuild)
@@ -762,8 +761,15 @@ void FTextureCacheDerivedDataWorker::DoWork()
 		{
 			FString FetchFirstKey;
 			GetTextureDerivedDataKeyFromSuffix(FetchFirstKeySuffix, FetchFirstKey);
-		
-			bLoadedFromDDC = GetDerivedDataCacheRef().GetSynchronous(*FetchFirstKey, RawDerivedData, TexturePathName);
+
+			FRequestOwner BlockingOwner(EPriority::Blocking);
+			GetCache().GetValue({{{TexturePathName}, ConvertLegacyCacheKey(FetchFirstKey)}}, BlockingOwner, [&RawDerivedData](FCacheGetValueResponse&& Response)
+			{
+				RawDerivedData = Response.Value.GetData().Decompress();
+			});
+			BlockingOwner.Wait();
+
+			bLoadedFromDDC = !RawDerivedData.IsNull();
 			if (bLoadedFromDDC)
 			{
 				bUsedFetchFirst = true;
@@ -778,7 +784,14 @@ void FTextureCacheDerivedDataWorker::DoWork()
 		// Didn't get the initial fetch, so we're using fetch/build.
 		LocalDerivedDataKeySuffix = MoveTemp(FetchOrBuildKeySuffix);
 		GetTextureDerivedDataKeyFromSuffix(LocalDerivedDataKeySuffix, LocalDerivedDataKey);
-		bLoadedFromDDC = GetDerivedDataCacheRef().GetSynchronous(*LocalDerivedDataKey, RawDerivedData, TexturePathName);
+
+		FRequestOwner BlockingOwner(EPriority::Blocking);
+		GetCache().GetValue({{{TexturePathName}, ConvertLegacyCacheKey(LocalDerivedDataKey)}}, BlockingOwner, [&RawDerivedData](FCacheGetValueResponse&& Response)
+		{
+			RawDerivedData = Response.Value.GetData().Decompress();
+		});
+		BlockingOwner.Wait();
+		bLoadedFromDDC = !RawDerivedData.IsNull();
 	}
 
 	KeySuffix = LocalDerivedDataKeySuffix;
@@ -791,8 +804,8 @@ void FTextureCacheDerivedDataWorker::DoWork()
 		const bool bInlineMips = EnumHasAnyFlags(CacheFlags, ETextureCacheFlags::InlineMips);
 		const bool bForDDC = EnumHasAnyFlags(CacheFlags, ETextureCacheFlags::ForDDCBuild);
 
-		BytesCached = RawDerivedData.Num();
-		FMemoryReaderView Ar(MakeMemoryView(RawDerivedData), /*bIsPersistent=*/ true);
+		BytesCached = RawDerivedData.GetSize();
+		FMemoryReaderView Ar(RawDerivedData.GetView(), /*bIsPersistent=*/ true);
 		DerivedData->Serialize(Ar, NULL);
 		bSucceeded = true;
 		// Load any streaming (not inline) mips that are necessary for our platform.
@@ -805,15 +818,22 @@ void FTextureCacheDerivedDataWorker::DoWork()
 				if (DerivedData->VTData != nullptr &&
 					DerivedData->VTData->IsInitialized())
 				{
-					TArray<FString, TInlineAllocator<16>> ChunkKeys;
+					FCacheGetValueRequest Request;
+					Request.Name = TexturePathName;
+					Request.Policy = ECachePolicy::Default | ECachePolicy::SkipData;
+
+					TArray<FCacheGetValueRequest, TInlineAllocator<16>> ChunkKeys;
 					for (const FVirtualTextureDataChunk& Chunk : DerivedData->VTData->Chunks)
 					{
 						if (!Chunk.DerivedDataKey.IsEmpty())
 						{
-							ChunkKeys.Add(Chunk.DerivedDataKey);
+							ChunkKeys.Add_GetRef(Request).Key = ConvertLegacyCacheKey(Chunk.DerivedDataKey);
 						}
 					}
-					GetDerivedDataCacheRef().TryToPrefetch(ChunkKeys, TexturePathName);
+
+					FRequestOwner BlockingOwner(EPriority::Blocking);
+					GetCache().GetValue(ChunkKeys, BlockingOwner, [](FCacheGetValueResponse&&){});
+					BlockingOwner.Wait();
 				}
 			}
 
diff --git a/Engine/Source/Runtime/Engine/Private/VT/VirtualTextureBuiltData.cpp b/Engine/Source/Runtime/Engine/Private/VT/VirtualTextureBuiltData.cpp
--- a/Engine/Source/Runtime/Engine/Private/VT/VirtualTextureBuiltData.cpp
+++ b/Engine/Source/Runtime/Engine/Private/VT/VirtualTextureBuiltData.cpp
@@ -2,7 +2,8 @@
 
 #include "VirtualTextureBuiltData.h"
 
-#include "DerivedDataCacheInterface.h"
+#include "DerivedDataCache.h"
+#include "DerivedDataRequestOwner.h"
 #include "Misc/CoreMisc.h"
 #include "Serialization/CustomVersion.h"
 #include "Serialization/MemoryWriter.h"
@@ -408,7 +409,7 @@ bool FVirtualTextureBuiltData::ValidateData(FStringView const& InDDCDebugContext
 
 	const uint32 TilePixelSize = GetPhysicalTileSize();
 	TArray<uint8> UncompressedResult;
-	TArray64<uint8> ChunkDataDDC;
+	FSharedBuffer ChunkDataDDC;
 
 	const FString TextureName(InDDCDebugContext);
 
@@ -418,20 +419,27 @@ bool FVirtualTextureBuiltData::ValidateData(FStringView const& InDDCDebugContext
 		const FVirtualTextureDataChunk& Chunk = Chunks[ChunkIndex];
 		FBulkDataLockedScope BulkDataLockedScope;
 
-		const uint8* ChunkData = nullptr;
+		const void* ChunkData = nullptr;
 		uint32 ChunkDataSize = 0u;
 		if (Chunk.BulkData.GetBulkDataSize() > 0)
 		{
-			ChunkDataSize = (uint32)Chunk.BulkData.GetBulkDataSize();
-			ChunkData = (uint8*)Chunk.BulkData.LockReadOnly();
+			ChunkDataSize = IntCastChecked<uint32>(Chunk.BulkData.GetBulkDataSize());
+			ChunkData = Chunk.BulkData.LockReadOnly();
 			BulkDataLockedScope.BulkData = &Chunk.BulkData;
 		}
 #if WITH_EDITORONLY_DATA
 		else
 		{
 			ChunkDataDDC.Reset();
-			const bool bDDCResult = GetDerivedDataCacheRef().GetSynchronous(*Chunk.DerivedDataKey, ChunkDataDDC, InDDCDebugContext);
-			if (!bDDCResult)
+			using namespace UE::DerivedData;
+			FRequestOwner BlockingOwner(EPriority::Blocking);
+			GetCache().GetValue({{{InDDCDebugContext}, ConvertLegacyCacheKey(Chunk.DerivedDataKey)}}, BlockingOwner, [&ChunkDataDDC](FCacheGetValueResponse&& Response)
+			{
+				ChunkDataDDC = Response.Value.GetData().Decompress();
+			});
+			BlockingOwner.Wait();
+
+			if (ChunkDataDDC.IsNull())
 			{
 				UE_LOG(LogTexture, Log, TEXT("Virtual Texture %s failed to retrieve DDC data (%s) for chunk %d"), *TextureName, *Chunk.DerivedDataKey, ChunkIndex);
 				bResult = false;
@@ -439,13 +447,13 @@ bool FVirtualTextureBuiltData::ValidateData(FStringView const& InDDCDebugContext
 			}
 
 			ChunkData = ChunkDataDDC.GetData();
-			ChunkDataSize = (uint32)ChunkDataDDC.Num();
+			ChunkDataSize = IntCastChecked<uint32>(ChunkDataDDC.GetSize());
 		}
 #endif // WITH_EDITORONLY_DATA
 
 		if (!ChunkData || ChunkDataSize < sizeof(FVirtualTextureChunkHeader))
 		{
-			UE_LOG(LogTexture, Error, TEXT("Virtual Texture %s has invalid size %d for chunk %d"), *TextureName, ChunkDataSize, ChunkIndex);
+			UE_LOG(LogTexture, Error, TEXT("Virtual Texture %s has invalid size %u for chunk %d"), *TextureName, ChunkDataSize, ChunkIndex);
 			bResult = false;
 			break;
 		}
@@ -483,7 +491,7 @@ bool FVirtualTextureBuiltData::ValidateData(FStringView const& InDDCDebugContext
 							const uint32 CompressedTileSize = NextTileOffset - TileOffset;
 
 							UncompressedResult.SetNumUninitialized(PackedOutputSize, false);
-							const bool bUncompressResult = FCompression::UncompressMemory(NAME_Zlib, UncompressedResult.GetData(), PackedOutputSize, &ChunkData[TileOffset], CompressedTileSize);
+							const bool bUncompressResult = FCompression::UncompressMemory(NAME_Zlib, UncompressedResult.GetData(), PackedOutputSize, static_cast<const uint8*>(ChunkData) + TileOffset, CompressedTileSize);
 							if (!bUncompressResult)
 							{
 								UE_LOG(LogTexture, Error, TEXT("Virtual Texture %s failed to validate compression for chunk %d"), *TextureName, ChunkIndex);
@@ -542,11 +550,15 @@ int64 FVirtualTextureDataChunk::StoreInDerivedDataCache(const FString& InDerived
 	int64 BulkDataSizeInBytes = BulkData.GetBulkDataSize();
 	check(BulkDataSizeInBytes > 0);
 
-	{
-		const uint8* BulkChunkData = (uint8*)BulkData.Lock(LOCK_READ_ONLY);
-		GetDerivedDataCacheRef().Put(*InDerivedDataKey, TConstArrayView64<uint8>(BulkChunkData, BulkDataSizeInBytes), TextureName, bReplaceExistingDDC);
-		BulkData.Unlock();
-	}
+	using namespace UE::DerivedData;
+
+	FValue DerivedData = FValue::Compress(FSharedBuffer::MakeView(BulkData.Lock(LOCK_READ_ONLY), BulkDataSizeInBytes));
+	BulkData.Unlock();
+
+	FRequestOwner AsyncOwner(EPriority::Normal);
+	const ECachePolicy Policy = bReplaceExistingDDC ? ECachePolicy::Store : ECachePolicy::Default;
+	GetCache().PutValue({{{TextureName}, ConvertLegacyCacheKey(InDerivedDataKey), MoveTemp(DerivedData), Policy}}, AsyncOwner);
+	AsyncOwner.KeepAlive();
 
 	DerivedDataKey = InDerivedDataKey;
 	ShortenKey(DerivedDataKey, ShortDerivedDataKey);
diff --git a/Engine/Source/Runtime/Engine/Private/VT/VirtualTextureChunkDDCCache.cpp b/Engine/Source/Runtime/Engine/Private/VT/VirtualTextureChunkDDCCache.cpp
--- a/Engine/Source/Runtime/Engine/Private/VT/VirtualTextureChunkDDCCache.cpp
+++ b/Engine/Source/Runtime/Engine/Private/VT/VirtualTextureChunkDDCCache.cpp
@@ -9,7 +9,8 @@
 #include "Misc/Paths.h"
 #include "Async/AsyncWork.h"
 #include "Misc/CoreMisc.h"
-#include "DerivedDataCacheInterface.h"
+#include "DerivedDataCache.h"
+#include "DerivedDataRequestOwner.h"
 #include "HAL/FileManager.h"
 #include "Misc/PackageName.h"
 #include "Misc/FileHelper.h"
@@ -59,7 +60,6 @@ public:
 		TRACE_CPUPROFILER_EVENT_SCOPE(FAsyncFillCacheWorker::DoWork);
 
 		IPlatformFile& PlatformFile = FPlatformFileManager::Get().GetPlatformFile();
-		FDerivedDataCacheInterface& DDC = GetDerivedDataCacheRef();
 
 		// Limit size of valid hash
 		static const uint32 kMaxHashSize = 32 * 1024;
@@ -104,23 +104,29 @@ public:
 		}
 
 		// Fetch data from DDC
-		TArray64<uint8> Results;
-		//TODO(ddebaets) this sync request seems to be blocking here while it uses the job pool. add overload to perform it on this thread?
-		const bool DDCResult = DDC.GetSynchronous(*Chunk->DerivedDataKey, Results, FinalFilename);
-		if (DDCResult == false)
+		using namespace UE::DerivedData;
+		FSharedBuffer Results;
+		FRequestOwner BlockingOwner(EPriority::Blocking);
+		GetCache().GetValue({{{FinalFilename}, ConvertLegacyCacheKey(Chunk->DerivedDataKey)}}, BlockingOwner, [&Results](FCacheGetValueResponse&& Response)
+		{
+			Results = Response.Value.GetData().Decompress();
+		});
+		BlockingOwner.Wait();
+
+		if (Results.IsNull())
 		{
 			UE_LOG(LogVTDiskCache, Log, TEXT("Failed to fetch data from DDC (key: %s)"), *Chunk->DerivedDataKey);
 			return;
 		}
 
-		if (Results.Num() <= sizeof(FVirtualTextureChunkHeader))
+		if (Results.GetSize() <= sizeof(FVirtualTextureChunkHeader))
 		{
-			UE_LOG(LogVTDiskCache, Error, TEXT("VT DDC data is too small %" INT64_FMT " (key: %s)"), Results.Num(), *Chunk->DerivedDataKey);
+			UE_LOG(LogVTDiskCache, Error, TEXT("VT DDC data is too small %" UINT64_FMT " (key: %s)"), Results.GetSize(), *Chunk->DerivedDataKey);
 			return;
 		}
 
-		const uint8* ChunkData = Results.GetData();
-		const int64 ChunkDataSize = Results.Num();
+		const void* ChunkData = Results.GetData();
+		const uint64 ChunkDataSize = Results.GetSize();
 
 		FSHAHash Hash;
 		FSHA1::HashBuffer(ChunkData, ChunkDataSize, Hash.Hash);
@@ -147,7 +153,7 @@ public:
 			*Ar << Header;
 			const int64 ArchiveOffset = Ar->Tell();
 			check(ArchiveOffset == sizeof(FVirtualTextureFileHeader));
-			Ar->Serialize(const_cast<uint8*>(ChunkData), ChunkDataSize);
+			Ar->Serialize(const_cast<void*>(ChunkData), ChunkDataSize);
 		}
 
 		if (PlatformFile.MoveFile(*FinalFilename, *TempFilename))
