Better handling of implicit scalar replication throughout various parts of new HLSL translator

diff --git a/Engine/Source/Runtime/Engine/Classes/Materials/MaterialExpressionNamedReroute.h b/Engine/Source/Runtime/Engine/Classes/Materials/MaterialExpressionNamedReroute.h
--- a/Engine/Source/Runtime/Engine/Classes/Materials/MaterialExpressionNamedReroute.h
+++ b/Engine/Source/Runtime/Engine/Classes/Materials/MaterialExpressionNamedReroute.h
@@ -74,6 +74,7 @@ public:
 	virtual void SetEditableName(const FString& NewName) override;
 
 	virtual void PostCopyNode(const TArray<UMaterialExpression*>& CopiedExpressions) override;
+	virtual bool GenerateHLSLExpression(FMaterialHLSLGenerator& Generator, UE::HLSLTree::FScope& Scope, int32 OutputIndex, UE::HLSLTree::FExpression const*& OutExpression) const override;
 #endif // WITH_EDITOR
 	//~ End UMaterialExpression Interface
 
@@ -114,6 +115,7 @@ public:
 	virtual bool MatchesSearchQuery(const TCHAR* SearchQuery) override;
 	
 	virtual void PostCopyNode(const TArray<UMaterialExpression*>& CopiedExpressions) override;
+	virtual bool GenerateHLSLExpression(FMaterialHLSLGenerator& Generator, UE::HLSLTree::FScope& Scope, int32 OutputIndex, UE::HLSLTree::FExpression const*& OutExpression) const override;
 #endif // WITH_EDITOR
 	//~ End UMaterialExpression Interface
 
diff --git a/Engine/Source/Runtime/Engine/Private/HLSLTree/HLSLTree.cpp b/Engine/Source/Runtime/Engine/Private/HLSLTree/HLSLTree.cpp
--- a/Engine/Source/Runtime/Engine/Private/HLSLTree/HLSLTree.cpp
+++ b/Engine/Source/Runtime/Engine/Private/HLSLTree/HLSLTree.cpp
@@ -539,6 +539,10 @@ FRequestedType::FRequestedType(const Shader::FType& InType, bool bDefaultRequest
 	{
 		StructType = InType.StructType;
 	}
+	else if (InType.IsObject())
+	{
+		ObjectType = InType.ObjectType;
+	}
 	else
 	{
 		const Shader::FValueTypeDescription TypeDesc = Shader::GetValueTypeDescription(InType);
@@ -651,6 +655,15 @@ EExpressionEvaluation FPreparedComponent::GetEvaluation(const FEmitScope& Scope)
 	return Result;
 }
 
+void FPreparedComponent::SetLoopEvaluation(FEmitScope& Scope)
+{
+	Evaluation = MakeLoopEvaluation(Evaluation);
+	if (IsLoopEvaluation(Evaluation))
+	{
+		LoopScope = FEmitScope::FindSharedParent(&Scope, LoopScope);
+	}
+}
+
 FPreparedComponent CombineComponents(const FPreparedComponent& Lhs, const FPreparedComponent& Rhs)
 {
 	if (Lhs.IsNone())
@@ -720,7 +733,7 @@ int32 FPreparedType::GetNumComponents() const
 	}
 	else if (ValueComponentType != Shader::EValueComponentType::Void)
 	{
-		auto Predicate = [](const FPreparedComponent& InComponent) { return InComponent.IsRequested(); };
+		auto Predicate = [](const FPreparedComponent& InComponent) { return !InComponent.IsNone(); };
 		const int32 MaxComponentIndex = PreparedComponents.FindLastByPredicate(Predicate);
 		return (MaxComponentIndex != INDEX_NONE) ? MaxComponentIndex + 1 : 1;
 	}
@@ -748,10 +761,13 @@ Shader::FType FPreparedType::GetType() const
 
 FRequestedType FPreparedType::GetRequestedType() const
 {
+	const int32 NumComponents = GetNumComponents();
 	FRequestedType Result;
 	Result.ValueComponentType = ValueComponentType;
 	Result.StructType = StructType;
-	for (int32 Index = 0; Index < PreparedComponents.Num(); ++Index)
+	Result.ObjectType = ObjectType;
+	Result.RequestedComponents.Init(false, NumComponents);
+	for (int32 Index = 0; Index < NumComponents; ++Index)
 	{
 		const FPreparedComponent& Component = PreparedComponents[Index];
 		if (Component.IsRequested())
@@ -775,7 +791,7 @@ EExpressionEvaluation FPreparedType::GetEvaluation(const FEmitScope& Scope) cons
 EExpressionEvaluation FPreparedType::GetEvaluation(const FEmitScope& Scope, const FRequestedType& RequestedType) const
 {
 	EExpressionEvaluation Result = EExpressionEvaluation::ConstantZero;
-	for (int32 Index = 0; Index < PreparedComponents.Num(); ++Index)
+	for (int32 Index = 0; Index < RequestedType.RequestedComponents.Num(); ++Index)
 	{
 		if (RequestedType.IsComponentRequested(Index))
 		{
@@ -811,15 +827,14 @@ FPreparedComponent FPreparedType::GetMergedComponent() const
 
 FPreparedComponent FPreparedType::GetComponent(int32 Index) const
 {
-	if (PreparedComponents.IsValidIndex(Index))
+	const int32 NumComponents = GetNumComponents();
+	const int32 ComponentIndex = (NumComponents == 1 && Index >= 0 && Index < 4) ? 0 : Index;
+	if (PreparedComponents.IsValidIndex(ComponentIndex))
 	{
-		return PreparedComponents[Index];
+		return PreparedComponents[ComponentIndex];
 	}
 
-	// Return 'ConstantZero' for unprepared components
-	// TODO - do we ever want to return 'None' instead?
-	const int32 NumComponents = GetNumComponents();
-	return (NumComponents == 1) ? PreparedComponents[0] : FPreparedComponent(EExpressionEvaluation::ConstantZero);
+	return FPreparedComponent(EExpressionEvaluation::None);
 }
 
 void FPreparedType::EnsureNumComponents(int32 NumComponents)
@@ -894,15 +909,21 @@ void FPreparedType::MergeEvaluation(EExpressionEvaluation Evaluation)
 
 void FPreparedType::SetLoopEvaluation(FEmitScope& Scope, const FRequestedType& RequestedType)
 {
-	for (int32 Index = 0; Index < PreparedComponents.Num(); ++Index)
+	const int32 NumRequestedComponents = RequestedType.GetNumComponents();
+	const int32 NumComponents = GetNumComponents();
+	
+	if (NumComponents == 1 && FMath::IsWithinInclusive(NumRequestedComponents, 1, 4))
 	{
-		if (RequestedType.IsComponentRequested(Index))
+		// If we're a scalar type, set loop evaluation if any of xyzw are requested
+		PreparedComponents[0].SetLoopEvaluation(Scope);
+	}
+	else
+	{
+		for (int32 Index = 0; Index < NumComponents; ++Index)
 		{
-			FPreparedComponent& Component = PreparedComponents[Index];
-			Component.Evaluation = MakeLoopEvaluation(Component.Evaluation);
-			if (IsLoopEvaluation(Component.Evaluation))
+			if (RequestedType.IsComponentRequested(Index))
 			{
-				Component.LoopScope = FEmitScope::FindSharedParent(&Scope, Component.LoopScope);
+				PreparedComponents[Index].SetLoopEvaluation(Scope);
 			}
 		}
 	}
@@ -1016,7 +1037,8 @@ FPreparedType MakeNonLWCType(const FPreparedType& Type)
 bool FPrepareValueResult::TryMergePreparedType(FEmitContext& Context,
 	const Shader::FStructType* StructType,
 	FName ObjectType,
-	Shader::EValueComponentType ComponentType)
+	Shader::EValueComponentType ComponentType,
+	int32 NumComponents)
 {
 	if (!StructType && ObjectType.IsNone() && ComponentType == Shader::EValueComponentType::Void)
 	{
@@ -1051,6 +1073,11 @@ bool FPrepareValueResult::TryMergePreparedType(FEmitContext& Context,
 	}
 	else
 	{
+		const int32 PrevNumComponents = PreparedType.GetNumComponents();
+		if (NumComponents != PrevNumComponents)
+		{
+			return Context.Error(TEXT("Invalid type"));
+		}
 		PreparedType.ValueComponentType = Shader::CombineComponentTypes(PreparedType.ValueComponentType, ComponentType);
 	}
 
@@ -1068,21 +1095,29 @@ bool FPrepareValueResult::SetTypeVoid()
 
 bool FPrepareValueResult::SetType(FEmitContext& Context, const FRequestedType& RequestedType, EExpressionEvaluation Evaluation, const Shader::FType& Type)
 {
-	if (TryMergePreparedType(Context, Type.StructType, Type.ObjectType, Shader::GetValueTypeDescription(Type.ValueType).ComponentType))
+	const int32 NumComponents = Type.GetNumComponents();
+	if (TryMergePreparedType(Context, Type.StructType, Type.ObjectType, Shader::GetValueTypeDescription(Type.ValueType).ComponentType, NumComponents))
 	{
 		if (Evaluation != EExpressionEvaluation::None)
 		{
-			const int32 NumComponents = Type.GetNumComponents();
-			for (int32 Index = 0; Index < NumComponents; ++Index)
+			const int32 NumRequestedComponents = RequestedType.GetNumComponents();
+			if (NumComponents == 1 && NumRequestedComponents > 0 && NumRequestedComponents <= 4)
 			{
-				if (RequestedType.IsComponentRequested(Index))
-				{
-					PreparedType.MergeComponent(Index, Evaluation);
-				}
-				else
+				PreparedType.MergeComponent(0, Evaluation);
+			}
+			else
+			{
+				for (int32 Index = 0; Index < NumComponents; ++Index)
 				{
-					// If component wasn't requested, it can be replaced with constant-0
-					PreparedType.MergeComponent(Index, EExpressionEvaluation::ConstantZero);
+					if (RequestedType.IsComponentRequested(Index))
+					{
+						PreparedType.MergeComponent(Index, Evaluation);
+					}
+					else
+					{
+						// If component wasn't requested, it can be replaced with constant-0
+						PreparedType.MergeComponent(Index, EExpressionEvaluation::ConstantZero);
+					}
 				}
 			}
 		}
@@ -1103,20 +1138,28 @@ bool FPrepareValueResult::SetType(FEmitContext& Context, const FRequestedType& R
 		}
 	}
 
-	if (TryMergePreparedType(Context, Type.StructType, Type.ObjectType, ComponentType))
+	const int32 NumComponents = Type.GetNumComponents();
+	if (TryMergePreparedType(Context, Type.StructType, Type.ObjectType, ComponentType, NumComponents))
 	{
-		const int32 NumComponents = Type.GetNumComponents();
-		for (int32 Index = 0; Index < NumComponents; ++Index)
+		const int32 NumRequestedComponents = RequestedType.GetNumComponents();
+		if (NumComponents == 1 && NumRequestedComponents > 0 && NumRequestedComponents <= 4)
 		{
-			const FPreparedComponent& Component = Type.GetComponent(Index);
-			if (RequestedType.IsComponentRequested(Index))
-			{
-				PreparedType.MergeComponent(Index, Component);
-			}
-			else if(!Component.IsNone())
+			PreparedType.MergeComponent(0, Type.GetComponent(0));
+		}
+		else
+		{
+			for (int32 Index = 0; Index < NumComponents; ++Index)
 			{
-				// If component wasn't requested, it can be replaced with constant-0
-				PreparedType.MergeComponent(Index, EExpressionEvaluation::ConstantZero);
+				const FPreparedComponent& Component = Type.GetComponent(Index);
+				if (RequestedType.IsComponentRequested(Index))
+				{
+					PreparedType.MergeComponent(Index, Component);
+				}
+				else if (!Component.IsNone())
+				{
+					// If component wasn't requested, it can be replaced with constant-0
+					PreparedType.MergeComponent(Index, EExpressionEvaluation::ConstantZero);
+				}
 			}
 		}
 		return true;
diff --git a/Engine/Source/Runtime/Engine/Private/HLSLTree/HLSLTreeCommon.cpp b/Engine/Source/Runtime/Engine/Private/HLSLTree/HLSLTreeCommon.cpp
--- a/Engine/Source/Runtime/Engine/Private/HLSLTree/HLSLTreeCommon.cpp
+++ b/Engine/Source/Runtime/Engine/Private/HLSLTree/HLSLTreeCommon.cpp
@@ -18,10 +18,10 @@ struct FPreshaderLoopScope
 
 FSwizzleParameters::FSwizzleParameters(int8 InR, int8 InG, int8 InB, int8 InA) : NumComponents(0), bHasSwizzle(true)
 {
-	ComponentIndex[0] = InR;
-	ComponentIndex[1] = InG;
-	ComponentIndex[2] = InB;
-	ComponentIndex[3] = InA;
+	SwizzleComponentIndex[0] = InR;
+	SwizzleComponentIndex[1] = InG;
+	SwizzleComponentIndex[2] = InB;
+	SwizzleComponentIndex[3] = InA;
 
 	if (InA >= 0)
 	{
@@ -40,11 +40,14 @@ FSwizzleParameters::FSwizzleParameters(int8 InR, int8 InG, int8 InB, int8 InA) :
 	{
 		check(InG <= 3);
 		++NumComponents;
+		check(InR >= 0);
 	}
 
-	// At least one proper index
-	check(InR >= 0 && InR <= 3);
-	++NumComponents;
+	if (InR >= 0)
+	{
+		check(InR <= 3);
+		++NumComponents;
+	}
 }
 
 FSwizzleParameters MakeSwizzleMask(bool bInR, bool bInG, bool bInB, bool bInA)
@@ -439,11 +442,11 @@ FRequestedType GetRequestedSwizzleType(const FSwizzleParameters& Params, const F
 {
 	FRequestedType RequestedInputType;
 	RequestedInputType.ValueComponentType = RequestedType.ValueComponentType;
-	for (int32 Index = 0; Index < Params.NumComponents; ++Index)
+	for (int32 Index = 0; Index < RequestedType.RequestedComponents.Num(); ++Index)
 	{
 		if (RequestedType.IsComponentRequested(Index))
 		{
-			const int32 SwizzledComponentIndex = Params.ComponentIndex[Index];
+			const int32 SwizzledComponentIndex = Params.GetSwizzleComponentIndex(Index);
 			RequestedInputType.SetComponentRequest(SwizzledComponentIndex);
 		}
 	}
@@ -484,7 +487,7 @@ bool FExpressionSwizzle::PrepareValue(FEmitContext& Context, FEmitScope& Scope,
 	if (RequestedInputType.IsVoid())
 	{
 		// All the requested components are outside the swizzle, so just return 0
-		ResultType.ValueComponentType = Shader::EValueComponentType::Float;
+		ResultType.ValueComponentType = RequestedType.ValueComponentType;
 		for (int32 ComponentIndex = 0; ComponentIndex < Parameters.NumComponents; ++ComponentIndex)
 		{
 			ResultType.SetComponent(ComponentIndex, EExpressionEvaluation::ConstantZero);
@@ -499,7 +502,7 @@ bool FExpressionSwizzle::PrepareValue(FEmitContext& Context, FEmitScope& Scope,
 		{
 			if (RequestedType.IsComponentRequested(ComponentIndex))
 			{
-				const int32 SwizzledComponentIndex = Parameters.ComponentIndex[ComponentIndex];
+				const int32 SwizzledComponentIndex = Parameters.GetSwizzleComponentIndex(ComponentIndex);
 				ResultType.SetComponent(ComponentIndex, InputType.GetComponent(SwizzledComponentIndex));
 			}
 			else
@@ -514,66 +517,61 @@ bool FExpressionSwizzle::PrepareValue(FEmitContext& Context, FEmitScope& Scope,
 
 void FExpressionSwizzle::EmitValueShader(FEmitContext& Context, FEmitScope& Scope, const FRequestedType& RequestedType, FEmitValueShaderResult& OutResult) const
 {
-	static const TCHAR ComponentName[] = { 'x', 'y', 'z', 'w' };
-	TCHAR Swizzle[5] = TEXT("");
-	TCHAR LWCSwizzle[10] = TEXT("");
-	bool bHasSwizzleReorder = false;
+	const FRequestedType RequestedInputType = Private::GetRequestedSwizzleType(Parameters, RequestedType);
+	FEmitShaderExpression* EmitInput = Input->GetValueShader(Context, Scope, RequestedInputType);
+	const Shader::FValueTypeDescription InputTypeDesc = Shader::GetValueTypeDescription(EmitInput->Type);
 
-	const int32 NumComponents = FMath::Min<int32>(RequestedType.GetNumComponents(), Parameters.NumComponents);
-	FRequestedType RequestedInputType;
-	RequestedInputType.ValueComponentType = RequestedType.ValueComponentType;
-	for (int32 ComponentIndex = 0; ComponentIndex < NumComponents; ++ComponentIndex)
+	bool bSkipSwizzle = (InputTypeDesc.NumComponents == Parameters.NumComponents);
+	if (bSkipSwizzle)
 	{
-		// If component wasn't requested, we just refer to 'x' component of input, since that should always be present
-		int32 SwizzledComponentIndex = 0;
-
-		if (RequestedType.IsComponentRequested(ComponentIndex))
-		{
-			SwizzledComponentIndex = Parameters.ComponentIndex[ComponentIndex];
-			RequestedInputType.SetComponentRequest(SwizzledComponentIndex);
-		}
-
-		Swizzle[ComponentIndex] = ComponentName[SwizzledComponentIndex];
-		LWCSwizzle[ComponentIndex * 2 + 0] = TEXT(',');
-		LWCSwizzle[ComponentIndex * 2 + 1] = TEXT('0') + SwizzledComponentIndex;
-
-		if (SwizzledComponentIndex != ComponentIndex)
+		for (int32 ComponentIndex = 0; ComponentIndex < Parameters.NumComponents; ++ComponentIndex)
 		{
-			bHasSwizzleReorder = true;
+			if (Parameters.GetSwizzleComponentIndex(ComponentIndex) != ComponentIndex)
+			{
+				bSkipSwizzle = false;
+				break;
+			}
 		}
 	}
 
-	FEmitShaderExpression* InputValue = Input->GetValueShader(Context, Scope, RequestedInputType);
-	const Shader::FValueTypeDescription InputTypeDesc = Shader::GetValueTypeDescription(InputValue->Type);
-
-	if (bHasSwizzleReorder || NumComponents != InputTypeDesc.NumComponents)
+	if (bSkipSwizzle)
+	{
+		OutResult.Code = EmitInput;
+	}
+	else
 	{
-		const Shader::EValueType ResultType = Shader::MakeValueType(InputTypeDesc.ComponentType, NumComponents);
-		const int32 NumRequestedComponents = RequestedInputType.GetNumComponents();
-		if (NumRequestedComponents > InputTypeDesc.NumComponents)
+		const bool bIsLWC = (InputTypeDesc.ComponentType == Shader::EValueComponentType::Double);
+		TStringBuilder<256> FormattedCode;
+		if (bIsLWC)
+		{
+			FormattedCode.Appendf(TEXT("LWCSwizzle(%s"), EmitInput->Reference);
+		}
+		else
 		{
-			// Zero-extend our input if needed, so we can acccess all the given components
-			InputValue = Context.EmitCast(Scope, InputValue, Shader::MakeValueType(InputTypeDesc.ComponentType, NumRequestedComponents), EEmitCastFlags::ZeroExtendScalar);
+			FormattedCode.Appendf(TEXT("%s."), EmitInput->Reference);
 		}
 
-		if (InputTypeDesc.ComponentType == Shader::EValueComponentType::Double)
+		static const TCHAR ComponentName[] = { 'x', 'y', 'z', 'w' };
+		for (int32 ComponentIndex = 0; ComponentIndex < Parameters.NumComponents; ++ComponentIndex)
 		{
-			check(LWCSwizzle[0]);
-			OutResult.Code = Context.EmitInlineExpression(Scope, ResultType, TEXT("LWCSwizzle(%%)"),
-				InputValue,
-				LWCSwizzle);
+			const int32 SwizzledComponentIndex = (InputTypeDesc.NumComponents > 1) ? Parameters.GetSwizzleComponentIndex(ComponentIndex) : 0;
+			if (bIsLWC)
+			{
+				FormattedCode.Appendf(TEXT(", %d"), SwizzledComponentIndex);
+			}
+			else
+			{
+				FormattedCode.AppendChar(ComponentName[SwizzledComponentIndex]);
+			}
 		}
-		else
+
+		if (bIsLWC)
 		{
-			check(Swizzle[0]);
-			OutResult.Code = Context.EmitInlineExpression(Scope, ResultType, TEXT("%.%"),
-				InputValue,
-				Swizzle);
+			FormattedCode.Append(TEXT(")"));
 		}
-	}
-	else
-	{
-		OutResult.Code = InputValue;
+
+		const Shader::FType ResultType = Shader::MakeValueType(InputTypeDesc.ComponentType, Parameters.NumComponents);
+		OutResult.Code = Context.EmitInlineExpressionWithDependency(Scope, EmitInput, ResultType, FormattedCode.ToView());
 	}
 }
 
@@ -593,10 +591,10 @@ void FExpressionSwizzle::EmitValuePreshader(FEmitContext& Context, FEmitScope& S
 
 		OutResult.Preshader.WriteOpcode(Shader::EPreshaderOpcode::ComponentSwizzle)
 			.Write((uint8)Parameters.NumComponents)
-			.Write((uint8)Parameters.ComponentIndex[0])
-			.Write((uint8)Parameters.ComponentIndex[1])
-			.Write((uint8)Parameters.ComponentIndex[2])
-			.Write((uint8)Parameters.ComponentIndex[3]);
+			.Write((uint8)Parameters.GetSwizzleComponentIndex(0))
+			.Write((uint8)Parameters.GetSwizzleComponentIndex(1))
+			.Write((uint8)Parameters.GetSwizzleComponentIndex(2))
+			.Write((uint8)Parameters.GetSwizzleComponentIndex(3));
 		OutResult.Type = Shader::MakeValueType(InputTypeDesc.ComponentType, Parameters.NumComponents);
 	}
 }
@@ -605,42 +603,45 @@ namespace Private
 {
 struct FAppendTypes
 {
-	Shader::EValueType ResultType;
+	FPreparedType ResultType;
 	Shader::EValueType LhsType;
 	Shader::EValueType RhsType;
 	FRequestedType LhsRequestedType;
 	FRequestedType RhsRequestedType;
 	bool bIsLWC;
 };
-FAppendTypes GetAppendTypes(const FRequestedType& RequestedType, Shader::EValueType LhsType, Shader::EValueType RhsType)
+FAppendTypes GetAppendTypes(const FRequestedType& RequestedType, const FPreparedType& LhsType, const FPreparedType& RhsType)
 {
-	const Shader::FValueTypeDescription LhsTypeDesc = Shader::GetValueTypeDescription(LhsType);
-	const Shader::FValueTypeDescription RhsTypeDesc = Shader::GetValueTypeDescription(RhsType);
-	const Shader::EValueComponentType ComponentType = Shader::CombineComponentTypes(LhsTypeDesc.ComponentType, RhsTypeDesc.ComponentType);
-	const int32 NumComponents = FMath::Min(LhsTypeDesc.NumComponents + RhsTypeDesc.NumComponents, 4);
+	const Shader::EValueComponentType ComponentType = Shader::CombineComponentTypes(LhsType.ValueComponentType, RhsType.ValueComponentType);
+	const int32 NumLhsComponents = LhsType.GetNumComponents();
+	const int32 NumRhsComponents = RhsType.GetNumComponents();
+	const int32 NumComponents = FMath::Min(NumLhsComponents + NumRhsComponents, 4);
 
 	FAppendTypes Types;
+	Types.ResultType.ValueComponentType = ComponentType;
 	Types.LhsRequestedType.ValueComponentType = ComponentType;
 	Types.RhsRequestedType.ValueComponentType = ComponentType;
-	for (int32 Index = 0; Index < LhsTypeDesc.NumComponents; ++Index)
+	Types.LhsType = Shader::MakeValueType(ComponentType, NumLhsComponents);
+	Types.RhsType = Shader::MakeValueType(ComponentType, NumComponents - NumLhsComponents);
+	Types.bIsLWC = ComponentType == Shader::EValueComponentType::Double;
+
+	for (int32 Index = 0; Index < NumLhsComponents; ++Index)
 	{
+		Types.ResultType.SetComponent(Index, LhsType.GetComponent(Index));
 		if (RequestedType.IsComponentRequested(Index))
 		{
 			Types.LhsRequestedType.SetComponentRequest(Index);
 		}
 	}
-	for (int32 Index = LhsTypeDesc.NumComponents; Index < NumComponents; ++Index)
+	for (int32 Index = NumLhsComponents; Index < NumComponents; ++Index)
 	{
+		Types.ResultType.SetComponent(Index, RhsType.GetComponent(Index - NumLhsComponents));
 		if (RequestedType.IsComponentRequested(Index))
 		{
-			Types.RhsRequestedType.SetComponentRequest(Index - LhsTypeDesc.NumComponents);
+			Types.RhsRequestedType.SetComponentRequest(Index - NumLhsComponents);
 		}
 	}
 
-	Types.ResultType = Shader::MakeValueType(ComponentType, NumComponents);
-	Types.LhsType = Shader::MakeValueType(ComponentType, LhsTypeDesc.NumComponents);
-	Types.RhsType = Shader::MakeValueType(ComponentType, NumComponents - LhsTypeDesc.NumComponents);
-	Types.bIsLWC = ComponentType == Shader::EValueComponentType::Double;
 	return Types;
 }
 } // namespace Private
@@ -665,18 +666,11 @@ const FExpression* FExpressionAppend::ComputePreviousFrame(FTree& Tree, const FR
 bool FExpressionAppend::PrepareValue(FEmitContext& Context, FEmitScope& Scope, const FRequestedType& RequestedType, FPrepareValueResult& OutResult) const
 {
 	const FPreparedType& LhsType = Context.PrepareExpression(Lhs, Scope, RequestedType);
-	const int32 NumRequestedComponents = RequestedType.GetNumComponents();
-	const int32 NumLhsComponents = FMath::Min(LhsType.GetNumComponents(), NumRequestedComponents);
-
-	FPreparedType ResultType(LhsType.ValueComponentType);
-	for (int32 Index = 0; Index < NumLhsComponents; ++Index)
-	{
-		ResultType.SetComponent(Index, LhsType.GetComponent(Index));
-	}
+	const int32 NumLhsComponents = LhsType.GetNumComponents();
 
 	FRequestedType RhsRequestedType;
 	RhsRequestedType.ValueComponentType = RequestedType.ValueComponentType;
-	for (int32 Index = NumLhsComponents; Index < NumRequestedComponents; ++Index)
+	for (int32 Index = NumLhsComponents; Index < RequestedType.RequestedComponents.Num(); ++Index)
 	{
 		RhsRequestedType.SetComponentRequest(Index - NumLhsComponents, RequestedType.IsComponentRequested(Index));
 	}
@@ -689,24 +683,19 @@ bool FExpressionAppend::PrepareValue(FEmitContext& Context, FEmitScope& Scope, c
 		{
 			return Context.Error(TEXT("Type mismatch"));
 		}
-
-		const int32 NumRhsComponents = FMath::Min(RhsType.GetNumComponents(), NumRequestedComponents - NumLhsComponents);
-		for (int32 Index = 0; Index < NumRhsComponents; ++Index)
-		{
-			ResultType.SetComponent(NumLhsComponents + Index, RhsType.GetComponent(Index));
-		}
 	}
 
-	const Private::FAppendTypes Types = Private::GetAppendTypes(RequestedType, LhsType.GetType(), RhsType.GetType());
+	const Private::FAppendTypes Types = Private::GetAppendTypes(RequestedType, LhsType, RhsType);
 	Context.MarkInputType(Lhs, Types.LhsType);
 	Context.MarkInputType(Rhs, Types.RhsType);
 
-	return OutResult.SetType(Context, RequestedType, ResultType);
+	return OutResult.SetType(Context, RequestedType, Types.ResultType);
 }
 
 void FExpressionAppend::EmitValueShader(FEmitContext& Context, FEmitScope& Scope, const FRequestedType& RequestedType, FEmitValueShaderResult& OutResult) const
 {
-	const Private::FAppendTypes Types = Private::GetAppendTypes(RequestedType, Context.GetType(Lhs), Context.GetType(Rhs));
+	const Private::FAppendTypes Types = Private::GetAppendTypes(RequestedType, Context.GetPreparedType(Lhs), Context.GetPreparedType(Rhs));
+	const Shader::FType ResultType = Types.ResultType.GetType();
 	FEmitShaderExpression* LhsValue = Lhs->GetValueShader(Context, Scope, Types.LhsRequestedType, Types.LhsType);
 
 	if (Types.RhsType == Shader::EValueType::Void)
@@ -718,14 +707,14 @@ void FExpressionAppend::EmitValueShader(FEmitContext& Context, FEmitScope& Scope
 		FEmitShaderExpression* RhsValue = Rhs->GetValueShader(Context, Scope, Types.RhsRequestedType, Types.RhsType);
 		if (Types.bIsLWC)
 		{
-			OutResult.Code = Context.EmitInlineExpression(Scope, Types.ResultType, TEXT("MakeLWCVector(%, %)"),
+			OutResult.Code = Context.EmitInlineExpression(Scope, ResultType, TEXT("MakeLWCVector(%, %)"),
 				LhsValue,
 				RhsValue);
 		}
 		else
 		{
-			OutResult.Code = Context.EmitInlineExpression(Scope, Types.ResultType, TEXT("%(%, %)"),
-				Shader::GetValueTypeDescription(Types.ResultType).Name,
+			OutResult.Code = Context.EmitInlineExpression(Scope, ResultType, TEXT("%(%, %)"),
+				Shader::GetValueTypeDescription(ResultType).Name,
 				LhsValue,
 				RhsValue);
 		}
@@ -734,7 +723,7 @@ void FExpressionAppend::EmitValueShader(FEmitContext& Context, FEmitScope& Scope
 
 void FExpressionAppend::EmitValuePreshader(FEmitContext& Context, FEmitScope& Scope, const FRequestedType& RequestedType, FEmitValuePreshaderResult& OutResult) const
 {
-	const Private::FAppendTypes Types = Private::GetAppendTypes(RequestedType, Context.GetType(Lhs), Context.GetType(Rhs));
+	const Private::FAppendTypes Types = Private::GetAppendTypes(RequestedType, Context.GetPreparedType(Lhs), Context.GetPreparedType(Rhs));
 	Lhs->GetValuePreshader(Context, Scope, Types.LhsRequestedType, OutResult.Preshader);
 	if (Types.RhsType != Shader::EValueType::Void)
 	{
@@ -745,7 +734,7 @@ void FExpressionAppend::EmitValuePreshader(FEmitContext& Context, FEmitScope& Sc
 
 		OutResult.Preshader.WriteOpcode(Shader::EPreshaderOpcode::AppendVector);
 	}
-	OutResult.Type = Types.ResultType;
+	OutResult.Type = Types.ResultType.GetType();
 }
 
 void FExpressionSwitchBase::ComputeAnalyticDerivatives(FTree& Tree, FExpressionDerivatives& OutResult) const
diff --git a/Engine/Source/Runtime/Engine/Private/HLSLTree/HLSLTreeEmit.cpp b/Engine/Source/Runtime/Engine/Private/HLSLTree/HLSLTreeEmit.cpp
--- a/Engine/Source/Runtime/Engine/Private/HLSLTree/HLSLTreeEmit.cpp
+++ b/Engine/Source/Runtime/Engine/Private/HLSLTree/HLSLTreeEmit.cpp
@@ -436,7 +436,7 @@ FPreparedType FEmitContext::PrepareExpression(const FExpression* InExpression, F
 
 void FEmitContext::MarkInputType(const FExpression* InExpression, const Shader::FType& Type)
 {
-	if (OwnerStack.Num() > 0)
+	if (OwnerStack.Num() > 0 && !Type.IsVoid())
 	{
 		const FOwnedNode* InputOwner = OwnerStack.Last();
 		for (UObject* InputObject : InputOwner->GetOwners())
@@ -945,7 +945,12 @@ FEmitShaderExpression* FEmitContext::EmitPreshaderOrConstant(FEmitScope& Scope,
 		const Shader::EValueType FieldType = ResultType.GetFlatFieldType(FieldIndex);
 		const Shader::FValueTypeDescription TypeDesc = Shader::GetValueTypeDescription(FieldType);
 		const int32 NumFieldComponents = TypeDesc.NumComponents;
-		const EExpressionEvaluation FieldEvaluation = PreparedType.GetFieldEvaluation(Scope, RequestedType, ComponentIndex, NumFieldComponents);
+
+		// If this is a struct, use GetFieldEvaluation()
+		// If this isn't a struct, use GetEvaluation() instead, which will properly handle promoting scalar types
+		const EExpressionEvaluation FieldEvaluation = Type.IsStruct()
+			? PreparedType.GetFieldEvaluation(Scope, RequestedType, ComponentIndex, NumFieldComponents)
+			: PreparedType.GetEvaluation(Scope, RequestedType);
 
 		if (FieldEvaluation == EExpressionEvaluation::Preshader)
 		{
diff --git a/Engine/Source/Runtime/Engine/Private/Materials/MaterialExpressionHLSL.cpp b/Engine/Source/Runtime/Engine/Private/Materials/MaterialExpressionHLSL.cpp
--- a/Engine/Source/Runtime/Engine/Private/Materials/MaterialExpressionHLSL.cpp
+++ b/Engine/Source/Runtime/Engine/Private/Materials/MaterialExpressionHLSL.cpp
@@ -7,6 +7,7 @@
 #include "Materials/MaterialExpressionExecBegin.h"
 #include "Materials/MaterialExpressionExecEnd.h"
 #include "Materials/MaterialExpressionReroute.h"
+#include "Materials/MaterialExpressionNamedReroute.h"
 #include "Materials/MaterialExpressionGenericConstant.h"
 #include "Materials/MaterialExpressionConstant.h"
 #include "Materials/MaterialExpressionConstant2Vector.h"
@@ -164,6 +165,18 @@ bool UMaterialExpressionReroute::GenerateHLSLExpression(FMaterialHLSLGenerator&
 	return OutExpression != nullptr;
 }
 
+bool UMaterialExpressionNamedRerouteDeclaration::GenerateHLSLExpression(FMaterialHLSLGenerator& Generator, UE::HLSLTree::FScope& Scope, int32 OutputIndex, UE::HLSLTree::FExpression const*& OutExpression) const
+{
+	OutExpression = Input.TryAcquireHLSLExpression(Generator, Scope);
+	return OutExpression != nullptr;
+}
+
+bool UMaterialExpressionNamedRerouteUsage::GenerateHLSLExpression(FMaterialHLSLGenerator& Generator, UE::HLSLTree::FScope& Scope, int32 OutputIndex, UE::HLSLTree::FExpression const*& OutExpression) const
+{
+	OutExpression = Generator.AcquireExpression(Scope, 0, Declaration, 0, UE::HLSLTree::FSwizzleParameters());
+	return OutExpression != nullptr;
+}
+
 bool UMaterialExpressionGenericConstant::GenerateHLSLExpression(FMaterialHLSLGenerator& Generator, UE::HLSLTree::FScope& Scope, int32 OutputIndex, UE::HLSLTree::FExpression const*& OutExpression) const
 {
 	OutExpression = Generator.NewConstant(GetConstantValue());
diff --git a/Engine/Source/Runtime/Engine/Private/Shader/Preshader.cpp b/Engine/Source/Runtime/Engine/Private/Shader/Preshader.cpp
--- a/Engine/Source/Runtime/Engine/Private/Shader/Preshader.cpp
+++ b/Engine/Source/Runtime/Engine/Private/Shader/Preshader.cpp
@@ -175,7 +175,7 @@ void FPreshaderData::WriteValue(const FValue& Value)
 	{
 		const EValueComponentType ComponentType = Value.Type.GetComponentType(Index);
 		const int32 ComponentSize = GetComponentTypeSizeInBytes(ComponentType);
-		const FValueComponent Component = Value.TryGetComponent(Index);
+		const FValueComponent Component = Value.GetComponent(Index);
 		WriteData(&Component.Packed, ComponentSize);
 	}
 }
@@ -453,16 +453,16 @@ static void EvaluateComponentSwizzle(FPreshaderStack& Stack, FPreshaderDataConte
 	switch (NumElements)
 	{
 	case 4:
-		Result.Component[3] = Value.TryGetComponent(IndexA);
+		Result.Component[3] = Value.GetComponent(IndexA);
 		// Fallthrough...
 	case 3:
-		Result.Component[2] = Value.TryGetComponent(IndexB);
+		Result.Component[2] = Value.GetComponent(IndexB);
 		// Fallthrough...
 	case 2:
-		Result.Component[1] = Value.TryGetComponent(IndexG);
+		Result.Component[1] = Value.GetComponent(IndexG);
 		// Fallthrough...
 	case 1:
-		Result.Component[0] = Value.TryGetComponent(IndexR);
+		Result.Component[0] = Value.GetComponent(IndexR);
 		break;
 	default:
 		UE_LOG(LogMaterial, Fatal, TEXT("Invalid number of swizzle elements: %d"), NumElements);
diff --git a/Engine/Source/Runtime/Engine/Public/HLSLTree/HLSLTree.h b/Engine/Source/Runtime/Engine/Public/HLSLTree/HLSLTree.h
--- a/Engine/Source/Runtime/Engine/Public/HLSLTree/HLSLTree.h
+++ b/Engine/Source/Runtime/Engine/Public/HLSLTree/HLSLTree.h
@@ -220,6 +220,8 @@ struct FPreparedComponent
 
 	EExpressionEvaluation GetEvaluation(const FEmitScope& Scope) const;
 
+	void SetLoopEvaluation(FEmitScope& Scope);
+
 	FEmitScope* LoopScope = nullptr;
 	Shader::FComponentBounds Bounds;
 	EExpressionEvaluation Evaluation = EExpressionEvaluation::None;
@@ -318,7 +320,8 @@ private:
 	bool TryMergePreparedType(FEmitContext& Context,
 		const Shader::FStructType* StructType,
 		FName ObjectType,
-		Shader::EValueComponentType ComponentType);
+		Shader::EValueComponentType ComponentType,
+		int32 NumComponents);
 
 	FPreparedType PreparedType;
 	bool bPreparingValue = false;
diff --git a/Engine/Source/Runtime/Engine/Public/HLSLTree/HLSLTreeHash.h b/Engine/Source/Runtime/Engine/Public/HLSLTree/HLSLTreeHash.h
--- a/Engine/Source/Runtime/Engine/Public/HLSLTree/HLSLTreeHash.h
+++ b/Engine/Source/Runtime/Engine/Public/HLSLTree/HLSLTreeHash.h
@@ -84,9 +84,9 @@ inline void AppendHash(FHasher& Hasher, const Shader::FType& Type)
 inline void AppendHash(FHasher& Hasher, const Shader::FValue& Value)
 {
 	AppendHash(Hasher, Value.Type);
-	for (int32 i = 0; i < Value.Type.GetNumComponents(); ++i)
+	for (int32 i = 0; i < Value.GetNumComponents(); ++i)
 	{
-		AppendHash(Hasher, Value.TryGetComponent(i));
+		AppendHash(Hasher, Value.GetComponent(i));
 	}
 }
 
diff --git a/Engine/Source/Runtime/Engine/Public/HLSLTree/HLSLTreeTypes.h b/Engine/Source/Runtime/Engine/Public/HLSLTree/HLSLTreeTypes.h
--- a/Engine/Source/Runtime/Engine/Public/HLSLTree/HLSLTreeTypes.h
+++ b/Engine/Source/Runtime/Engine/Public/HLSLTree/HLSLTreeTypes.h
@@ -139,10 +139,16 @@ FOperationDescription GetOperationDescription(EOperation Op);
 
 struct FSwizzleParameters
 {
-	FSwizzleParameters() : NumComponents(0), bHasSwizzle(false) { ComponentIndex[0] = ComponentIndex[1] = ComponentIndex[2] = ComponentIndex[3] = INDEX_NONE; }
+	FSwizzleParameters() : NumComponents(0), bHasSwizzle(false) { SwizzleComponentIndex[0] = SwizzleComponentIndex[1] = SwizzleComponentIndex[2] = SwizzleComponentIndex[3] = INDEX_NONE; }
 	explicit FSwizzleParameters(int8 IndexR, int8 IndexG = INDEX_NONE, int8 IndexB = INDEX_NONE, int8 IndexA = INDEX_NONE);
 
-	int8 ComponentIndex[4];
+	int32 GetSwizzleComponentIndex(int32 Index) const
+	{
+		const int32 ComponentIndex = (NumComponents == 1) ? 0 : Index;
+		return SwizzleComponentIndex[ComponentIndex];
+	}
+
+	int8 SwizzleComponentIndex[4];
 	int8 NumComponents;
 	bool bHasSwizzle;
 };
diff --git a/Engine/Source/Runtime/Engine/Public/Shader/ShaderTypes.h b/Engine/Source/Runtime/Engine/Public/Shader/ShaderTypes.h
--- a/Engine/Source/Runtime/Engine/Public/Shader/ShaderTypes.h
+++ b/Engine/Source/Runtime/Engine/Public/Shader/ShaderTypes.h
@@ -451,16 +451,17 @@ struct FValue
 	}
 
 	inline const FType& GetType() const { return Type; }
-	inline const FValueComponent& GetComponent(int32 i) const
-	{
-		checkf(Component.IsValidIndex(i), TEXT("Invalid component %d, of type '%s'"), i, Type.GetName());
-		return Component[i];
-	}
+	inline int32 GetNumComponents() const { return Type.GetNumComponents(); }
 
-	/** returns 0 for invalid components */
-	inline FValueComponent TryGetComponent(int32 i) const
+	inline FValueComponent GetComponent(int32 Index) const
 	{
-		return Component.IsValidIndex(i) ? Component[i] : FValueComponent();
+		// Return scalar value for any request to xyzw
+		const int32 ComponentIndex = (GetNumComponents() == 1 && Index >= 0 && Index < 4) ? 0 : Index;
+		if (Component.IsValidIndex(ComponentIndex))
+		{
+			return Component[ComponentIndex];
+		}
+		return FValueComponent();
 	}
 
 	static FValue FromMemoryImage(EValueType Type, const void* Data, uint32* OutSizeInBytes = nullptr);
