New HLSL translator fixes
MarkLiveValues will properly consider scopes/statements
Allow type of expression to change when preparing multiple times, as long as it doesn't change from scalar to non-scalar

diff --git a/Engine/Source/Runtime/Engine/Private/HLSLTree/HLSLTree.cpp b/Engine/Source/Runtime/Engine/Private/HLSLTree/HLSLTree.cpp
--- a/Engine/Source/Runtime/Engine/Private/HLSLTree/HLSLTree.cpp
+++ b/Engine/Source/Runtime/Engine/Private/HLSLTree/HLSLTree.cpp
@@ -791,7 +791,7 @@ EExpressionEvaluation FPreparedType::GetEvaluation(const FEmitScope& Scope) cons
 EExpressionEvaluation FPreparedType::GetEvaluation(const FEmitScope& Scope, const FRequestedType& RequestedType) const
 {
 	EExpressionEvaluation Result = EExpressionEvaluation::ConstantZero;
-	for (int32 Index = 0; Index < RequestedType.RequestedComponents.Num(); ++Index)
+	for (int32 Index = 0; Index < RequestedType.GetNumComponents(); ++Index)
 	{
 		if (RequestedType.IsComponentRequested(Index))
 		{
@@ -953,7 +953,7 @@ Shader::FComponentBounds FPreparedType::GetComponentBounds(int32 Index) const
 Shader::FComponentBounds FPreparedType::GetBounds(const FRequestedType& RequestedType) const
 {
 	Shader::FComponentBounds Result(Shader::EComponentBound::DoubleMax, Shader::EComponentBound::NegDoubleMax);
-	for (int32 Index = 0; Index < RequestedType.RequestedComponents.Num(); ++Index)
+	for (int32 Index = 0; Index < RequestedType.GetNumComponents(); ++Index)
 	{
 		if (RequestedType.IsComponentRequested(Index))
 		{
@@ -1074,9 +1074,11 @@ bool FPrepareValueResult::TryMergePreparedType(FEmitContext& Context,
 	else
 	{
 		const int32 PrevNumComponents = PreparedType.GetNumComponents();
-		if (NumComponents != PrevNumComponents)
+		if (PrevNumComponents == 1 && NumComponents > 1)
 		{
-			return Context.Error(TEXT("Invalid type"));
+			// Once an expression is prepared as a scalar type, it's not valid to change later
+			// Other expressions may have already assumed this will be a scalar type, and its components will replicate
+			return Context.Error(TEXT("Type was already prepared as scalar"));
 		}
 		PreparedType.ValueComponentType = Shader::CombineComponentTypes(PreparedType.ValueComponentType, ComponentType);
 	}
diff --git a/Engine/Source/Runtime/Engine/Private/HLSLTree/HLSLTreeCommon.cpp b/Engine/Source/Runtime/Engine/Private/HLSLTree/HLSLTreeCommon.cpp
--- a/Engine/Source/Runtime/Engine/Private/HLSLTree/HLSLTreeCommon.cpp
+++ b/Engine/Source/Runtime/Engine/Private/HLSLTree/HLSLTreeCommon.cpp
@@ -442,7 +442,7 @@ FRequestedType GetRequestedSwizzleType(const FSwizzleParameters& Params, const F
 {
 	FRequestedType RequestedInputType;
 	RequestedInputType.ValueComponentType = RequestedType.ValueComponentType;
-	for (int32 Index = 0; Index < RequestedType.RequestedComponents.Num(); ++Index)
+	for (int32 Index = 0; Index < RequestedType.GetNumComponents(); ++Index)
 	{
 		if (RequestedType.IsComponentRequested(Index))
 		{
@@ -670,7 +670,7 @@ bool FExpressionAppend::PrepareValue(FEmitContext& Context, FEmitScope& Scope, c
 
 	FRequestedType RhsRequestedType;
 	RhsRequestedType.ValueComponentType = RequestedType.ValueComponentType;
-	for (int32 Index = NumLhsComponents; Index < RequestedType.RequestedComponents.Num(); ++Index)
+	for (int32 Index = NumLhsComponents; Index < RequestedType.GetNumComponents(); ++Index)
 	{
 		RhsRequestedType.SetComponentRequest(Index - NumLhsComponents, RequestedType.IsComponentRequested(Index));
 	}
@@ -684,6 +684,13 @@ bool FExpressionAppend::PrepareValue(FEmitContext& Context, FEmitScope& Scope, c
 			return Context.Error(TEXT("Type mismatch"));
 		}
 	}
+	else
+	{
+		// We never want the result of an 'Append' operation to be a scalar value
+		// Ensure that we include at least a single component for the rhs, even if it's not currently requested
+		RhsType.ValueComponentType = RequestedType.ValueComponentType;
+		RhsType.SetComponent(0, EExpressionEvaluation::ConstantZero);
+	}
 
 	const Private::FAppendTypes Types = Private::GetAppendTypes(RequestedType, LhsType, RhsType);
 	Context.MarkInputType(Lhs, Types.LhsType);
diff --git a/Engine/Source/Runtime/Engine/Private/HLSLTree/HLSLTreeEmit.cpp b/Engine/Source/Runtime/Engine/Private/HLSLTree/HLSLTreeEmit.cpp
--- a/Engine/Source/Runtime/Engine/Private/HLSLTree/HLSLTreeEmit.cpp
+++ b/Engine/Source/Runtime/Engine/Private/HLSLTree/HLSLTreeEmit.cpp
@@ -507,6 +507,9 @@ FEmitScope* FEmitContext::InternalPrepareScope(FScope* InScope, FScope* InParent
 	for (int32 Index = 0; Index < EmitScopeChain.Num(); ++Index)
 	{
 		FEmitScope* EmitScope = EmitScopeChain[Index];
+		FEmitScope* EmitParentScope = EmitScope->ParentScope;
+		FStatement* Statement = EmitScope->OwnerStatement;
+
 		if (bMarkDead)
 		{
 			EmitScope->Evaluation = EExpressionEvaluation::None;
@@ -514,12 +517,9 @@ FEmitScope* FEmitContext::InternalPrepareScope(FScope* InScope, FScope* InParent
 		}
 		else if (EmitScope->State == EEmitScopeState::Initializing)
 		{
-			FStatement* Statement = EmitScope->OwnerStatement;
 			if (Statement)
 			{
 				FEmitOwnerScope OwnerScope(*this, Statement);
-
-				FEmitScope* EmitParentScope = EmitScope->ParentScope;
 				check(EmitParentScope);
 				if (Statement->Prepare(*this, *EmitParentScope))
 				{
@@ -546,6 +546,13 @@ FEmitScope* FEmitContext::InternalPrepareScope(FScope* InScope, FScope* InParent
 				EmitScope->State = EEmitScopeState::Live;
 			}
 		}
+		else if (Statement && EmitScope->State == EEmitScopeState::Live && bMarkLiveValues)
+		{
+			// Need to Prepare statement a second time to mark live values
+			FEmitOwnerScope OwnerScope(*this, Statement);
+			check(EmitParentScope);
+			verify(Statement->Prepare(*this, *EmitParentScope));
+		}
 
 		if (EmitScope->State == EEmitScopeState::Dead)
 		{
diff --git a/Engine/Source/Runtime/Engine/Private/HLSLTree/HLSLTreeOperation.cpp b/Engine/Source/Runtime/Engine/Private/HLSLTree/HLSLTreeOperation.cpp
--- a/Engine/Source/Runtime/Engine/Private/HLSLTree/HLSLTreeOperation.cpp
+++ b/Engine/Source/Runtime/Engine/Private/HLSLTree/HLSLTreeOperation.cpp
@@ -591,7 +591,6 @@ bool FExpressionOperation::PrepareValue(FEmitContext& Context, FEmitScope& Scope
 		Context.MarkInputType(Inputs[Index], Types.InputType[Index]);
 	}
 
-	bool bMarkInputsLive = Context.bMarkLiveValues;
 	switch (Op)
 	{
 	case EOperation::Mul:
diff --git a/Engine/Source/Runtime/Engine/Private/Materials/MaterialHLSLTree.cpp b/Engine/Source/Runtime/Engine/Private/Materials/MaterialHLSLTree.cpp
--- a/Engine/Source/Runtime/Engine/Private/Materials/MaterialHLSLTree.cpp
+++ b/Engine/Source/Runtime/Engine/Private/Materials/MaterialHLSLTree.cpp
@@ -1249,7 +1249,7 @@ void FExpressionVertexInterpolator::EmitValueShader(FEmitContext& Context, FEmit
 		FRequestedType RequestedPreshaderType;
 		RequestedPreshaderType.ValueComponentType = RequestedType.ValueComponentType;
 		RequestedPreshaderType.StructType = RequestedType.StructType;
-		for (int32 Index = 0; Index < RequestedType.RequestedComponents.Num(); ++Index)
+		for (int32 Index = 0; Index < RequestedType.GetNumComponents(); ++Index)
 		{
 			// Requested components that are *not* requested by the interpolator will be either constant or preshader evaluation
 			if (RequestedType.IsComponentRequested(Index) &&
@@ -1290,17 +1290,17 @@ int32 FEmitData::FindInterpolatorIndex(const FExpression* Expression) const
 
 void FEmitData::AddInterpolator(const FExpression* Expression, const FRequestedType& RequestedType, const FPreparedType& PreparedType)
 {
-	int32 Index = FindInterpolatorIndex(Expression);
-	if (Index == INDEX_NONE)
+	int32 InterpolatorIndex = FindInterpolatorIndex(Expression);
+	if (InterpolatorIndex == INDEX_NONE)
 	{
-		Index = VertexInterpolators.Emplace(Expression);
+		InterpolatorIndex = VertexInterpolators.Emplace(Expression);
 	}
 
-	FVertexInterpolator& Interpolator = VertexInterpolators[Index];
+	FVertexInterpolator& Interpolator = VertexInterpolators[InterpolatorIndex];
 	Interpolator.PreparedType = PreparedType;
-	for (int32 ComponentIndex = 0; ComponentIndex < RequestedType.RequestedComponents.Num(); ++ComponentIndex)
+	for (int32 ComponentIndex = 0; ComponentIndex < RequestedType.GetNumComponents(); ++ComponentIndex)
 	{
-		if (RequestedType.IsComponentRequested(Index))
+		if (RequestedType.IsComponentRequested(ComponentIndex))
 		{
 			const EExpressionEvaluation ComponentEvaluation = PreparedType.GetComponent(ComponentIndex).Evaluation;
 			if (ComponentEvaluation == EExpressionEvaluation::Shader)
