Try to make static analyzer happy, use TUniquePtr rather than raw pointers, so multiple assignments aren't a (potential) problem

diff --git a/Engine/Source/Runtime/Engine/Private/Materials/MaterialInstanceConstant.cpp b/Engine/Source/Runtime/Engine/Private/Materials/MaterialInstanceConstant.cpp
--- a/Engine/Source/Runtime/Engine/Private/Materials/MaterialInstanceConstant.cpp
+++ b/Engine/Source/Runtime/Engine/Private/Materials/MaterialInstanceConstant.cpp
@@ -202,8 +202,8 @@ void UMaterialInstanceConstant::UpdateCachedData()
 	if (!bLoadedCachedExpressionData)
 	{
 		const bool bUsingNewHLSLGenerator = IsUsingNewHLSLGenerator();
-		FMaterialCachedExpressionData* LocalCachedExpressionData = nullptr;
-		FMaterialCachedHLSLTree* LocalCachedTree = nullptr;
+		TUniquePtr<FMaterialCachedExpressionData> LocalCachedExpressionData;
+		TUniquePtr<FMaterialCachedHLSLTree> LocalCachedTree;
 
 		// If we have overriden material layers, need to create a local cached expression data
 		// Otherwise we can leave it as null, and use cached data from our parent
@@ -213,30 +213,30 @@ void UMaterialInstanceConstant::UpdateCachedData()
 			UMaterial* BaseMaterial = GetMaterial();
 			if (bUsingNewHLSLGenerator)
 			{
-				LocalCachedTree = new FMaterialCachedHLSLTree();
+				LocalCachedTree.Reset(new FMaterialCachedHLSLTree());
 				LocalCachedTree->GenerateTree(BaseMaterial, &LocalStaticParameters.MaterialLayers, nullptr);
 			}
 			else
 			{
 				FMaterialCachedExpressionContext Context;
 				Context.LayerOverrides = &LocalStaticParameters.MaterialLayers;
-				LocalCachedExpressionData = new FMaterialCachedExpressionData();
+				LocalCachedExpressionData.Reset(new FMaterialCachedExpressionData());
 				LocalCachedExpressionData->Reset();
 				LocalCachedExpressionData->UpdateForExpressions(Context, BaseMaterial->Expressions, GlobalParameter, INDEX_NONE);
 			}
 		}
 
-		CachedHLSLTree.Reset(LocalCachedTree);
+		CachedHLSLTree = MoveTemp(LocalCachedTree);
 
 		if (bUsingNewHLSLGenerator && bHasStaticPermutationResource)
 		{
 			check(!LocalCachedExpressionData);
-			LocalCachedExpressionData = new FMaterialCachedExpressionData();
+			LocalCachedExpressionData.Reset(new FMaterialCachedExpressionData());
 			LocalCachedExpressionData->Reset();
 			LocalCachedExpressionData->UpdateForCachedHLSLTree(GetCachedHLSLTree(), &LocalStaticParameters);
 		}
 
-		CachedExpressionData.Reset(LocalCachedExpressionData);
+		CachedExpressionData = MoveTemp(LocalCachedExpressionData);
 
 		FObjectCacheEventSink::NotifyReferencedTextureChanged_Concurrent(this);
 	}
