Support static component mask parameters with new HLSL translator

diff --git a/Engine/Source/Runtime/Engine/Classes/Materials/MaterialExpressionStaticComponentMaskParameter.h b/Engine/Source/Runtime/Engine/Classes/Materials/MaterialExpressionStaticComponentMaskParameter.h
--- a/Engine/Source/Runtime/Engine/Classes/Materials/MaterialExpressionStaticComponentMaskParameter.h
+++ b/Engine/Source/Runtime/Engine/Classes/Materials/MaterialExpressionStaticComponentMaskParameter.h
@@ -65,6 +65,7 @@ public:
 		}
 		return false;
 	}
+	virtual bool GenerateHLSLExpression(FMaterialHLSLGenerator& Generator, UE::HLSLTree::FScope& Scope, int32 OutputIndex, UE::HLSLTree::FExpression const*& OutExpression) const override;
 #endif
 	//~ End UMaterialExpression Interface
 
diff --git a/Engine/Source/Runtime/Engine/Private/HLSLTree/HLSLTreeCommon.cpp b/Engine/Source/Runtime/Engine/Private/HLSLTree/HLSLTreeCommon.cpp
--- a/Engine/Source/Runtime/Engine/Private/HLSLTree/HLSLTreeCommon.cpp
+++ b/Engine/Source/Runtime/Engine/Private/HLSLTree/HLSLTreeCommon.cpp
@@ -479,7 +479,14 @@ const FExpression* FExpressionSwizzle::ComputePreviousFrame(FTree& Tree, const F
 	return Tree.NewSwizzle(Parameters, Tree.GetPreviousFrame(Input, RequestedInputType));
 }
 
-bool FExpressionSwizzle::PrepareValue(FEmitContext& Context, FEmitScope& Scope, const FRequestedType& RequestedType, FPrepareValueResult& OutResult) const
+namespace Private
+{
+bool SwizzlePrepareValue(FEmitContext& Context,
+	FEmitScope& Scope,
+	const FRequestedType& RequestedType,
+	const FSwizzleParameters& Parameters,
+	const FExpression* Input,
+	FPrepareValueResult& OutResult)
 {
 	const FRequestedType RequestedInputType = Private::GetRequestedSwizzleType(Parameters, RequestedType);
 
@@ -515,7 +522,12 @@ bool FExpressionSwizzle::PrepareValue(FEmitContext& Context, FEmitScope& Scope,
 	return OutResult.SetType(Context, RequestedType, ResultType);
 }
 
-void FExpressionSwizzle::EmitValueShader(FEmitContext& Context, FEmitScope& Scope, const FRequestedType& RequestedType, FEmitValueShaderResult& OutResult) const
+void SwizzleEmitValueShader(FEmitContext& Context,
+	FEmitScope& Scope,
+	const FRequestedType& RequestedType,
+	const FSwizzleParameters& Parameters,
+	const FExpression* Input,
+	FEmitValueShaderResult& OutResult)
 {
 	const FRequestedType RequestedInputType = Private::GetRequestedSwizzleType(Parameters, RequestedType);
 	FEmitShaderExpression* EmitInput = Input->GetValueShader(Context, Scope, RequestedInputType);
@@ -575,7 +587,12 @@ void FExpressionSwizzle::EmitValueShader(FEmitContext& Context, FEmitScope& Scop
 	}
 }
 
-void FExpressionSwizzle::EmitValuePreshader(FEmitContext& Context, FEmitScope& Scope, const FRequestedType& RequestedType, FEmitValuePreshaderResult& OutResult) const
+void SwizzleEmitValuePreshader(FEmitContext& Context,
+	FEmitScope& Scope,
+	const FRequestedType& RequestedType,
+	const FSwizzleParameters& Parameters,
+	const FExpression* Input,
+	FEmitValuePreshaderResult& OutResult)
 {
 	const FRequestedType RequestedInputType = Private::GetRequestedSwizzleType(Parameters, RequestedType);
 	if (RequestedInputType.IsVoid())
@@ -599,6 +616,72 @@ void FExpressionSwizzle::EmitValuePreshader(FEmitContext& Context, FEmitScope& S
 	}
 }
 
+} // namespace Private
+
+bool FExpressionSwizzle::PrepareValue(FEmitContext& Context, FEmitScope& Scope, const FRequestedType& RequestedType, FPrepareValueResult& OutResult) const
+{
+	return Private::SwizzlePrepareValue(Context, Scope, RequestedType, Parameters, Input, OutResult);
+}
+
+void FExpressionSwizzle::EmitValueShader(FEmitContext& Context, FEmitScope& Scope, const FRequestedType& RequestedType, FEmitValueShaderResult& OutResult) const
+{
+	return Private::SwizzleEmitValueShader(Context, Scope, RequestedType, Parameters, Input, OutResult);
+}
+
+void FExpressionSwizzle::EmitValuePreshader(FEmitContext& Context, FEmitScope& Scope, const FRequestedType& RequestedType, FEmitValuePreshaderResult& OutResult) const
+{
+	return Private::SwizzleEmitValuePreshader(Context, Scope, RequestedType, Parameters, Input, OutResult);
+}
+
+void FExpressionComponentMask::ComputeAnalyticDerivatives(FTree& Tree, FExpressionDerivatives& OutResult) const
+{
+	const FExpressionDerivatives InputDerivatives = Tree.GetAnalyticDerivatives(Input);
+	if (InputDerivatives.IsValid())
+	{
+		OutResult.ExpressionDdx = Tree.NewExpression<FExpressionComponentMask>(InputDerivatives.ExpressionDdx, Mask);
+		OutResult.ExpressionDdy = Tree.NewExpression<FExpressionComponentMask>(InputDerivatives.ExpressionDdy, Mask);
+	}
+}
+
+const FExpression* FExpressionComponentMask::ComputePreviousFrame(FTree& Tree, const FRequestedType& RequestedType) const
+{
+	return Tree.NewExpression<FExpressionComponentMask>(Tree.GetPreviousFrame(Input, RequestedType), Mask);
+}
+
+namespace Private
+{
+FSwizzleParameters GetComponentMaskSwizzle(FEmitContext& Context, FEmitScope& Scope, const FExpression* Mask)
+{
+	const Shader::FBoolValue MaskValue = Mask->GetValueConstant(Context, Scope, Shader::EValueType::Bool4).AsBool();
+	return MakeSwizzleMask(MaskValue[0], MaskValue[1], MaskValue[2], MaskValue[3]);
+}
+}
+
+bool FExpressionComponentMask::PrepareValue(FEmitContext& Context, FEmitScope& Scope, const FRequestedType& RequestedType, FPrepareValueResult& OutResult) const
+{
+	const FPreparedType& PreparedMaskType = Context.PrepareExpression(Mask, Scope, Shader::EValueType::Bool4);
+	const EExpressionEvaluation MaskEvaluation = PreparedMaskType.GetEvaluation(Scope, Shader::EValueType::Bool4);
+	if (!IsConstantEvaluation(MaskEvaluation))
+	{
+		return Context.Error(TEXT("Mask must be constant"));
+	}
+
+	const FSwizzleParameters Swizzle = Private::GetComponentMaskSwizzle(Context, Scope, Mask);
+	return Private::SwizzlePrepareValue(Context, Scope, RequestedType, Swizzle, Input, OutResult);
+}
+
+void FExpressionComponentMask::EmitValueShader(FEmitContext& Context, FEmitScope& Scope, const FRequestedType& RequestedType, FEmitValueShaderResult& OutResult) const
+{
+	const FSwizzleParameters Swizzle = Private::GetComponentMaskSwizzle(Context, Scope, Mask);
+	return Private::SwizzleEmitValueShader(Context, Scope, RequestedType, Swizzle, Input, OutResult);
+}
+
+void FExpressionComponentMask::EmitValuePreshader(FEmitContext& Context, FEmitScope& Scope, const FRequestedType& RequestedType, FEmitValuePreshaderResult& OutResult) const
+{
+	const FSwizzleParameters Swizzle = Private::GetComponentMaskSwizzle(Context, Scope, Mask);
+	return Private::SwizzleEmitValuePreshader(Context, Scope, RequestedType, Swizzle, Input, OutResult);
+}
+
 namespace Private
 {
 struct FAppendTypes
diff --git a/Engine/Source/Runtime/Engine/Private/Materials/MaterialExpressionHLSL.cpp b/Engine/Source/Runtime/Engine/Private/Materials/MaterialExpressionHLSL.cpp
--- a/Engine/Source/Runtime/Engine/Private/Materials/MaterialExpressionHLSL.cpp
+++ b/Engine/Source/Runtime/Engine/Private/Materials/MaterialExpressionHLSL.cpp
@@ -19,6 +19,7 @@
 #include "Materials/MaterialExpressionParameter.h"
 #include "Materials/MaterialExpressionCurveAtlasRowParameter.h"
 #include "Materials/MaterialExpressionStaticSwitchParameter.h"
+#include "Materials/MaterialExpressionStaticComponentMaskParameter.h"
 #include "Materials/MaterialExpressionPixelDepth.h"
 #include "Materials/MaterialExpressionWorldPosition.h"
 #include "Materials/MaterialExpressionCameraPositionWS.h"
@@ -396,6 +397,27 @@ bool UMaterialExpressionStaticSwitchParameter::GenerateHLSLExpression(FMaterialH
 	return true;
 }
 
+bool UMaterialExpressionStaticComponentMaskParameter::GenerateHLSLExpression(FMaterialHLSLGenerator& Generator, UE::HLSLTree::FScope& Scope, int32 OutputIndex, UE::HLSLTree::FExpression const*& OutExpression) const
+{
+	using namespace UE::HLSLTree;
+
+	FMaterialParameterMetadata ParameterMeta;
+	if (!GetParameterValue(ParameterMeta))
+	{
+		return Generator.Error(TEXT("Failed to get parameter value"));
+	}
+
+	const FExpression* ExpressionInput = Input.AcquireHLSLExpression(Generator, Scope);
+	if (!ExpressionInput)
+	{
+		return false;
+	}
+
+	const FExpression* ExpressionMask = Generator.GenerateMaterialParameter(ParameterName, ParameterMeta);
+	OutExpression = Generator.GetTree().NewExpression<FExpressionComponentMask>(ExpressionInput, ExpressionMask);
+	return true;
+}
+
 bool UMaterialExpressionPixelDepth::GenerateHLSLExpression(FMaterialHLSLGenerator& Generator, UE::HLSLTree::FScope& Scope, int32 OutputIndex, UE::HLSLTree::FExpression const*& OutExpression) const
 {
 	using namespace UE::HLSLTree::Material;
diff --git a/Engine/Source/Runtime/Engine/Private/Materials/MaterialHLSLTree.cpp b/Engine/Source/Runtime/Engine/Private/Materials/MaterialHLSLTree.cpp
--- a/Engine/Source/Runtime/Engine/Private/Materials/MaterialHLSLTree.cpp
+++ b/Engine/Source/Runtime/Engine/Private/Materials/MaterialHLSLTree.cpp
@@ -499,18 +499,36 @@ void FExpressionParameter::EmitValuePreshader(FEmitContext& Context, FEmitScope&
 
 	Context.PreshaderStackPosition++;
 	OutResult.Type = GetShaderValueType(ParameterType);
-	if (ParameterType == EMaterialParameterType::StaticSwitch)
+	if (IsStaticMaterialParameter(ParameterType))
 	{
 		Shader::FValue Value = DefaultValue;
 		if (EmitMaterialData.StaticParameters)
 		{
-			for (const FStaticSwitchParameter& Parameter : EmitMaterialData.StaticParameters->StaticSwitchParameters)
+			switch (ParameterType)
 			{
-				if (Parameter.ParameterInfo == ParameterInfo)
+			case EMaterialParameterType::StaticSwitch:
+				for (const FStaticSwitchParameter& Parameter : EmitMaterialData.StaticParameters->StaticSwitchParameters)
+				{
+					if (Parameter.ParameterInfo == ParameterInfo)
+					{
+						Value = Parameter.Value;
+						break;
+					}
+				}
+				break;
+			case EMaterialParameterType::StaticComponentMask:
+				for (const FStaticComponentMaskParameter& Parameter : EmitMaterialData.StaticParameters->StaticComponentMaskParameters)
 				{
-					Value = Parameter.Value;
-					break;
+					if (Parameter.ParameterInfo == ParameterInfo)
+					{
+						Value = Shader::FValue(Parameter.R, Parameter.G, Parameter.B, Parameter.A);
+						break;
+					}
 				}
+				break;
+			default:
+				checkNoEntry();
+				break;
 			}
 		}
 		OutResult.Preshader.WriteOpcode(Shader::EPreshaderOpcode::Constant).Write(Value);
diff --git a/Engine/Source/Runtime/Engine/Public/HLSLTree/HLSLTreeCommon.h b/Engine/Source/Runtime/Engine/Public/HLSLTree/HLSLTreeCommon.h
--- a/Engine/Source/Runtime/Engine/Public/HLSLTree/HLSLTreeCommon.h
+++ b/Engine/Source/Runtime/Engine/Public/HLSLTree/HLSLTreeCommon.h
@@ -154,6 +154,24 @@ public:
 	virtual void EmitValuePreshader(FEmitContext& Context, FEmitScope& Scope, const FRequestedType& RequestedType, FEmitValuePreshaderResult& OutResult) const override;
 };
 
+class FExpressionComponentMask : public FExpression
+{
+public:
+	FExpressionComponentMask(const FExpression* InInput, const FExpression* InMask)
+		: Input(InInput)
+		, Mask(InMask)
+	{}
+
+	const FExpression* Input;
+	const FExpression* Mask;
+
+	virtual void ComputeAnalyticDerivatives(FTree& Tree, FExpressionDerivatives& OutResult) const override;
+	virtual const FExpression* ComputePreviousFrame(FTree& Tree, const FRequestedType& RequestedType) const override;
+	virtual bool PrepareValue(FEmitContext& Context, FEmitScope& Scope, const FRequestedType& RequestedType, FPrepareValueResult& OutResult) const override;
+	virtual void EmitValueShader(FEmitContext& Context, FEmitScope& Scope, const FRequestedType& RequestedType, FEmitValueShaderResult& OutResult) const override;
+	virtual void EmitValuePreshader(FEmitContext& Context, FEmitScope& Scope, const FRequestedType& RequestedType, FEmitValuePreshaderResult& OutResult) const override;
+};
+
 class FExpressionAppend : public FExpression
 {
 public:
