Implement single shader file compilation for cooked platform ODSC.

For example: recompileshaders /Engine/Private/Lumen/LumenVisualizeHardwareRayTracing.usf

1. Game sends the shader file string across the network to the cooking server.
2. Cooking server finds all shaders of that type.
3. Cooking server flushes those specific types from the global shader map.
4. Cooking server recompiles just those shaders and sends them back to the game.

# Todo
* Refactor duplicate code in RecompileShadersForRemote
* Implement the ability for material shader maps to compile just specific shader types requested (this doesn't exist).

diff --git a/Engine/Source/Editor/UnrealEd/Private/Cooker/IoStoreCookOnTheFlyRequestManager.cpp b/Engine/Source/Editor/UnrealEd/Private/Cooker/IoStoreCookOnTheFlyRequestManager.cpp
--- a/Engine/Source/Editor/UnrealEd/Private/Cooker/IoStoreCookOnTheFlyRequestManager.cpp
+++ b/Engine/Source/Editor/UnrealEd/Private/Cooker/IoStoreCookOnTheFlyRequestManager.cpp
@@ -1067,6 +1067,7 @@ private:
 			int32 iShaderPlatform = static_cast<int32>(RecompileData.ShaderPlatform);
 			TUniquePtr<FArchive> Ar = Request.ReadBody();
 			*Ar << RecompileData.MaterialsToLoad;
+			*Ar << RecompileData.ShaderTypesToLoad;
 			*Ar << iShaderPlatform;
 			RecompileData.ShaderPlatform = static_cast<EShaderPlatform>(iShaderPlatform);
 			*Ar << RecompileData.CommandType;
diff --git a/Engine/Source/Runtime/Engine/Private/ODSC/ODSCManager.cpp b/Engine/Source/Runtime/Engine/Private/ODSC/ODSCManager.cpp
--- a/Engine/Source/Runtime/Engine/Private/ODSC/ODSCManager.cpp
+++ b/Engine/Source/Runtime/Engine/Private/ODSC/ODSCManager.cpp
@@ -70,11 +70,11 @@ bool FODSCManager::Tick(float DeltaSeconds)
 	return false;
 }
 
-void FODSCManager::AddThreadedRequest(const TArray<FString>& MaterialsToCompile, EShaderPlatform ShaderPlatform, ODSCRecompileCommand RecompileCommandType)
+void FODSCManager::AddThreadedRequest(const TArray<FString>& MaterialsToCompile, const FString& ShaderTypesToLoad, EShaderPlatform ShaderPlatform, ODSCRecompileCommand RecompileCommandType)
 {
 	if (Thread)
 	{
-		Thread->AddRequest(MaterialsToCompile, ShaderPlatform, RecompileCommandType);
+		Thread->AddRequest(MaterialsToCompile, ShaderTypesToLoad, ShaderPlatform, RecompileCommandType);
 	}
 }
 
diff --git a/Engine/Source/Runtime/Engine/Private/ODSC/ODSCThread.cpp b/Engine/Source/Runtime/Engine/Private/ODSC/ODSCThread.cpp
--- a/Engine/Source/Runtime/Engine/Private/ODSC/ODSCThread.cpp
+++ b/Engine/Source/Runtime/Engine/Private/ODSC/ODSCThread.cpp
@@ -20,8 +20,9 @@ FODSCMessageHandler::FODSCMessageHandler(EShaderPlatform InShaderPlatform, ODSCR
 {
 }
 
-FODSCMessageHandler::FODSCMessageHandler(const TArray<FString>& InMaterials, EShaderPlatform InShaderPlatform, ODSCRecompileCommand InRecompileCommandType) :
+FODSCMessageHandler::FODSCMessageHandler(const TArray<FString>& InMaterials, const FString& ShaderTypesToLoad, EShaderPlatform InShaderPlatform, ODSCRecompileCommand InRecompileCommandType) :
 	MaterialsToLoad(std::move(InMaterials)),
+	ShaderTypesToLoad(ShaderTypesToLoad),
 	ShaderPlatform(InShaderPlatform),
 	RecompileCommandType(InRecompileCommandType)
 {
@@ -33,6 +34,7 @@ void FODSCMessageHandler::FillPayload(FArchive& Payload)
 	RequestStartTime = FPlatformTime::Seconds();
 
 	Payload << MaterialsToLoad;
+	Payload << ShaderTypesToLoad;
 	uint32 ConvertedShaderPlatform = (uint32)ShaderPlatform;
 	Payload << ConvertedShaderPlatform;
 	Payload << RecompileCommandType;
@@ -108,9 +110,9 @@ void FODSCThread::Tick()
 	Process();
 }
 
-void FODSCThread::AddRequest(const TArray<FString>& MaterialsToCompile, EShaderPlatform ShaderPlatform, ODSCRecompileCommand RecompileCommandType)
+void FODSCThread::AddRequest(const TArray<FString>& MaterialsToCompile, const FString& ShaderTypesToLoad, EShaderPlatform ShaderPlatform, ODSCRecompileCommand RecompileCommandType)
 {
-	PendingMaterialThreadedRequests.Enqueue(new FODSCMessageHandler(MaterialsToCompile, ShaderPlatform, RecompileCommandType));
+	PendingMaterialThreadedRequests.Enqueue(new FODSCMessageHandler(MaterialsToCompile, ShaderTypesToLoad, ShaderPlatform, RecompileCommandType));
 }
 
 void FODSCThread::AddShaderPipelineRequest(EShaderPlatform ShaderPlatform, const FString& MaterialName, const FString& VertexFactoryName, const FString& PipelineName, const TArray<FString>& ShaderTypeNames)
diff --git a/Engine/Source/Runtime/Engine/Private/ODSC/ODSCThread.h b/Engine/Source/Runtime/Engine/Private/ODSC/ODSCThread.h
--- a/Engine/Source/Runtime/Engine/Private/ODSC/ODSCThread.h
+++ b/Engine/Source/Runtime/Engine/Private/ODSC/ODSCThread.h
@@ -17,7 +17,7 @@ class FODSCMessageHandler : public IPlatformFile::IFileServerMessageHandler
 {
 public:
 	FODSCMessageHandler(EShaderPlatform InShaderPlatform, ODSCRecompileCommand InRecompileCommandType);
-	FODSCMessageHandler(const TArray<FString>& InMaterials, EShaderPlatform InShaderPlatform, ODSCRecompileCommand InRecompileCommandType);
+	FODSCMessageHandler(const TArray<FString>& InMaterials, const FString& ShaderTypesToLoad, EShaderPlatform InShaderPlatform, ODSCRecompileCommand InRecompileCommandType);
 	/** Subclass fills out an archive to send to the server */
 	virtual void FillPayload(FArchive& Payload) override;
 
@@ -38,6 +38,9 @@ private:
 	/** The materials we send over the network and expect maps for on the return */
 	TArray<FString> MaterialsToLoad;
 
+	/** The names of shader type file names to compile shaders for. */
+	FString ShaderTypesToLoad;
+
 	/** Which shader platform we are compiling for */
 	EShaderPlatform ShaderPlatform;
 
@@ -85,12 +88,13 @@ public:
 	 * Add a shader compile request to be processed by this thread.
 	 *
 	 * @param MaterialsToCompile - List of material names to submit compiles for.
+	 * @param ShaderTypesToLoad - List of shader types to submit compiles for.
 	 * @param ShaderPlatform - Which shader platform to compile for.
 	 * @param RecompileCommandType - Whether we should recompile changed or global shaders.
 	 *
 	 * @return false if no longer needs ticking
 	 */
-	void AddRequest(const TArray<FString>& MaterialsToCompile, EShaderPlatform ShaderPlatform, ODSCRecompileCommand RecompileCommandType);
+	void AddRequest(const TArray<FString>& MaterialsToCompile, const FString& ShaderTypesToLoad, EShaderPlatform ShaderPlatform, ODSCRecompileCommand RecompileCommandType);
 
 	/**
 	 * Add a request to compile a pipeline (VS/PS) of shaders.  The results are submitted and processed in an async manner.
diff --git a/Engine/Source/Runtime/Engine/Private/ShaderCompiler/ShaderCompiler.cpp b/Engine/Source/Runtime/Engine/Private/ShaderCompiler/ShaderCompiler.cpp
--- a/Engine/Source/Runtime/Engine/Private/ShaderCompiler/ShaderCompiler.cpp
+++ b/Engine/Source/Runtime/Engine/Private/ShaderCompiler/ShaderCompiler.cpp
@@ -6028,7 +6028,7 @@ protected:
 
 namespace
 {
-	ODSCRecompileCommand ParseRecompileCommandString(const TCHAR* CmdString, TArray<FString>& OutMaterialsToLoad)
+	ODSCRecompileCommand ParseRecompileCommandString(const TCHAR* CmdString, TArray<FString>& OutMaterialsToLoad, FString& OutShaderTypesToLoad)
 	{
 		FString CmdName = FParse::Token(CmdString, 0);
 
@@ -6067,7 +6067,7 @@ namespace
 				OutMaterialsToLoad.Add(It->GetPathName());
 			}
 		}
-		else
+		else if (FCString::Stricmp(*CmdName, TEXT("All")) == 0)
 		{
 			CommandType = ODSCRecompileCommand::Material;
 
@@ -6077,6 +6077,19 @@ namespace
 				OutMaterialsToLoad.Add(It->GetPathName());
 			}
 		}
+		else
+		{
+			CommandType = ODSCRecompileCommand::SingleShader;
+
+			OutShaderTypesToLoad = CmdName;
+
+			// tell other side which materials to load and compile the single
+			// shader for.
+			for (TObjectIterator<UMaterialInterface> It; It; ++It)
+			{
+				OutMaterialsToLoad.Add(It->GetPathName());
+			}
+		}
 
 		return CommandType;
 	}
@@ -6204,8 +6217,9 @@ bool RecompileShaders(const TCHAR* Cmd, FOutputDevice& Ar)
 	{
 #if WITH_ODSC
 		TArray<FString> MaterialsToLoad;
-		ODSCRecompileCommand CommandType = ParseRecompileCommandString(Cmd, MaterialsToLoad);
-		GODSCManager->AddThreadedRequest(MaterialsToLoad, GMaxRHIShaderPlatform, CommandType);
+		FString ShaderTypesToLoad;
+		ODSCRecompileCommand CommandType = ParseRecompileCommandString(Cmd, MaterialsToLoad, ShaderTypesToLoad);
+		GODSCManager->AddThreadedRequest(MaterialsToLoad, ShaderTypesToLoad, GMaxRHIShaderPlatform, CommandType);
 #endif
 		return true;
 	}
@@ -7365,9 +7379,42 @@ void RecompileShadersForRemote(
 			// Only compile for the desired platform if requested
 			if (ShaderPlatform == Args.ShaderPlatform || Args.ShaderPlatform == SP_NumPlatforms)
 			{
-				// If we are explicitly wanting to recompile global or if shaders have changed.
-				if (Args.CommandType == ODSCRecompileCommand::Global ||
-					Args.CommandType == ODSCRecompileCommand::Changed)
+				if (Args.CommandType == ODSCRecompileCommand::SingleShader &&
+					Args.ShaderTypesToLoad.Len() > 0)
+				{
+					UE_LOG(LogShaders, Display, TEXT("Recompiling single shader."));
+
+					TArray<const FShaderType*> ShaderTypes = FShaderType::GetShaderTypesByFilename(*Args.ShaderTypesToLoad);
+					TArray<const FShaderPipelineType*> ShaderPipelineTypes = FShaderPipelineType::GetShaderPipelineTypesByFilename(*Args.ShaderTypesToLoad);
+
+					for (const FShaderType* ShaderType : ShaderTypes)
+					{
+						UE_LOG(LogShaders, Display, TEXT("\t%s..."), ShaderType->GetName());
+					}
+
+					if (ShaderTypes.Num() > 0 || ShaderPipelineTypes.Num() > 0)
+					{
+						BeginRecompileGlobalShaders(ShaderTypes, ShaderPipelineTypes, ShaderPlatform, TargetPlatform);
+						FinishRecompileGlobalShaders();
+					}
+
+					// write the shader compilation info to memory, converting fnames to strings
+					FMemoryWriter MemWriter(*Args.GlobalShaderMap, true);
+					FNameAsStringProxyArchive Ar(MemWriter);
+
+					TOptional<FArchiveCookData> CookData;
+					FArchiveCookContext CookContext(nullptr /*InPackage*/, FArchiveCookContext::ECookTypeUnknown);
+					if (TargetPlatform != nullptr)
+					{
+						CookData.Emplace(*TargetPlatform, CookContext);
+					}
+					Ar.SetCookData(CookData.GetPtrOrNull());
+
+					// save out the global shader map to the byte array
+					SaveGlobalShadersForRemoteRecompile(Ar, ShaderPlatform);
+				}
+				else if (Args.CommandType == ODSCRecompileCommand::Global ||
+						 Args.CommandType == ODSCRecompileCommand::Changed)
 				{
 					UE_LOG(LogShaders, Display, TEXT("Recompiling global shaders."));
 
diff --git a/Engine/Source/Runtime/Engine/Public/ODSC/ODSCManager.h b/Engine/Source/Runtime/Engine/Public/ODSC/ODSCManager.h
--- a/Engine/Source/Runtime/Engine/Public/ODSC/ODSCManager.h
+++ b/Engine/Source/Runtime/Engine/Public/ODSC/ODSCManager.h
@@ -51,7 +51,7 @@ public:
 	 *
 	 * @return false if no longer needs ticking
 	 */
-	void AddThreadedRequest(const TArray<FString>& MaterialsToCompile, EShaderPlatform ShaderPlatform, ODSCRecompileCommand RecompileCommandType);
+	void AddThreadedRequest(const TArray<FString>& MaterialsToCompile, const FString& ShaderTypesToLoad, EShaderPlatform ShaderPlatform, ODSCRecompileCommand RecompileCommandType);
 
 	/**
 	 * Add a request to compile a pipeline (VS/PS) of shaders.  The results are submitted and processed in an async manner.
diff --git a/Engine/Source/Runtime/Engine/Public/ShaderCompiler.h b/Engine/Source/Runtime/Engine/Public/ShaderCompiler.h
--- a/Engine/Source/Runtime/Engine/Public/ShaderCompiler.h
+++ b/Engine/Source/Runtime/Engine/Public/ShaderCompiler.h
@@ -1267,7 +1267,8 @@ enum class ODSCRecompileCommand
 	None,
 	Changed,
 	Global,
-	Material
+	Material,
+	SingleShader
 };
 
 struct FShaderRecompileData
@@ -1287,6 +1288,9 @@ struct FShaderRecompileData
 	/** Materials to load. */
 	TArray<FString> MaterialsToLoad;
 
+	/** The names of shader type file names to compile shaders for. */
+	FString ShaderTypesToLoad;
+
 	/** What type of shaders to recompile. All, Changed, Global, or Material? */
 	ODSCRecompileCommand CommandType = ODSCRecompileCommand::Changed;
 
@@ -1312,6 +1316,7 @@ struct FShaderRecompileData
 		ModifiedFiles = Other.ModifiedFiles;
 		MeshMaterialMaps = Other.MeshMaterialMaps;
 		MaterialsToLoad = Other.MaterialsToLoad;
+		ShaderTypesToLoad = Other.ShaderTypesToLoad;
 		CommandType = Other.CommandType;
 		GlobalShaderMap = Other.GlobalShaderMap;
 
diff --git a/Engine/Source/Runtime/NetworkFileSystem/Private/NetworkFileServerConnection.cpp b/Engine/Source/Runtime/NetworkFileSystem/Private/NetworkFileServerConnection.cpp
--- a/Engine/Source/Runtime/NetworkFileSystem/Private/NetworkFileServerConnection.cpp
+++ b/Engine/Source/Runtime/NetworkFileSystem/Private/NetworkFileServerConnection.cpp
@@ -1298,7 +1298,7 @@ void FNetworkFileServerClientConnection::ProcessRecompileShaders( FArchive& In,
 
 	// tell other side all the materials to load, by pathname
 	In << RecompileData.MaterialsToLoad;
-
+	In << RecompileData.ShaderTypesToLoad;
 	int32 iShaderPlatform = static_cast<int32>(RecompileData.ShaderPlatform);
 	In << iShaderPlatform;
 	RecompileData.ShaderPlatform = static_cast<EShaderPlatform>(iShaderPlatform);
