Fixed indenting in PostLoad.

diff --git a/Engine/Source/Runtime/Engine/Private/Components/SceneCaptureComponent.cpp b/Engine/Source/Runtime/Engine/Private/Components/SceneCaptureComponent.cpp
--- a/Engine/Source/Runtime/Engine/Private/Components/SceneCaptureComponent.cpp
+++ b/Engine/Source/Runtime/Engine/Private/Components/SceneCaptureComponent.cpp
@@ -79,29 +79,29 @@ void ASceneCapture::PostLoad()
 		if (IsTemplate())
 		{
 			if (UBlueprintGeneratedClass* BPClass = Cast<UBlueprintGeneratedClass>(GetClass()))
-{
+			{
 				for (USCS_Node* RootNode : BPClass->SimpleConstructionScript->GetRootNodes())
-{
+				{
 					static const FName OldMeshName(TEXT("CamMesh0"));
 					static const FName OldFrustumName(TEXT("DrawFrust0"));
 					static const FName NewRootName(TEXT("SceneComponent"));
 					if (RootNode->ParentComponentOrVariableName == OldMeshName || RootNode->ParentComponentOrVariableName == OldFrustumName)
-	{
+					{
 						RootNode->ParentComponentOrVariableName = NewRootName;
 					}
 				}
-	}
-}
+			}
+		}
 
 		PRAGMA_DISABLE_DEPRECATION_WARNINGS
 		if (MeshComp_DEPRECATED)
 		{
 			MeshComp_DEPRECATED->SetStaticMesh(nullptr);
-			}
-		PRAGMA_ENABLE_DEPRECATION_WARNINGS
 		}
-#endif
+		PRAGMA_ENABLE_DEPRECATION_WARNINGS
 	}
+#endif
+}
 
 void ASceneCapture::Serialize(FArchive& Ar)
 {
