UpdateTexture2D and UpdateFromBufferTexture2D bug fixes
- Support SrcX, SrcY offset
- Handle texture formats that don't support buffer read; copy to temporary texture instead
- Only alias texture descriptor for block compressed formats to avoid breaking metadata buffers that might be bound
- Disable tiling of 32_32_32 formats as this is not supported
UpdateTextureShader.usf
- Removed unused shaders UpdateTexture2DSubresourceCS, CopyTexture2DCS, and TCopyTexture2DCS
- TUpdateTexture2DSubresourceCS now supports source offset xy
UpdateTextureShaders.cpp,*.h
- Removed unused shader binding code
- Add TUpdateTexture2DSubresourceCS and FUpdateTexture2DSubresourceCS to support ValueType template variation for float, uint, and sint srv/uav access
- Fixed FUpdateTexture3DSubresourceCS name

diff --git a/Engine/Shaders/Private/UpdateTextureShaders.usf b/Engine/Shaders/Private/UpdateTextureShaders.usf
--- a/Engine/Shaders/Private/UpdateTextureShaders.usf
+++ b/Engine/Shaders/Private/UpdateTextureShaders.usf
@@ -7,27 +7,11 @@
 #include "Common.ush"
 
 Buffer<int> SrcBuffer;
-RWTexture2D<int4> DestTexture;
-uint2 SrcPitch; //x = number of SrcBuffer elements in a row of the source area, y = Number of components per texel in the source.
-uint4 DestPosSize; // xy = pos, zw = Size of the destination sub rect
-
-[numthreads(8,8,1)]
-void UpdateTexture2DSubresourceCS( uint2 ThreadId : SV_DispatchThreadID )
-{
-	if( all( ThreadId.xy < DestPosSize.zw ) )
-	{
-		int2 TexturePixelOffset = ThreadId.xy + DestPosSize.xy;
-
-		//supports textures with 1-4 components.  Textures with  < 4 components will simply ignore writes to the extra components via channel masking.		
-		int BufferOffset = (ThreadId.x * SrcPitch.y) + (ThreadId.y * SrcPitch.x);
-		DestTexture[TexturePixelOffset] = int4(SrcBuffer[BufferOffset].x, SrcBuffer[BufferOffset + 1].x, SrcBuffer[BufferOffset + 2].x, SrcBuffer[BufferOffset + 3].x);
-	}
-}
-
 RWTexture3D<int4> DestTexture3D;
 uint SrcDepthPitch; //number of SrcBuffer entries between z slices of the source volume
 uint4 DestPos; // xyz = starting offset in destination volume texture
 uint4 DestSize; //syz = size of the volume to update
+uint2 SrcPitch; //x = number of SrcBuffer elements in a row of the source area, y = Number of components per texel in the source.
 
 [numthreads(8, 8, 1)]
 void UpdateTexture3DSubresourceCS(uint3 ThreadId : SV_DispatchThreadID)
@@ -42,17 +26,6 @@ void UpdateTexture3DSubresourceCS(uint3 ThreadId : SV_DispatchThreadID)
 	}
 }
 
-Texture2D<int4> SrcTexture;
-
-[numthreads(8,8,1)]
-void CopyTexture2DCS( uint3 ThreadId : SV_DispatchThreadID )
-{
-	if (all(ThreadId.xy >= DestPosSize.xy) && all(ThreadId.xy < DestPosSize.zw))
-	{
-		DestTexture[ThreadId.xy] = SrcTexture.Load(int3(ThreadId.x, ThreadId.y, 0));
-	}
-}
-
 Buffer<uint4> SrcCopyBuffer;
 RWBuffer<uint4> DestBuffer;
 
@@ -79,28 +52,18 @@ void CopyData1CS(uint3 ThreadId : SV_DispatchThreadID)
 
 Buffer<COMPONENT_TYPE> TSrcBuffer;
 RWTexture2D<COMPONENT_TYPE> TDestTexture;
-uint TSrcPitch;
+uint3 TSrcPosPitch;
+uint4 TDestPosSize; // xy = pos, zw = Size of the destination sub rect
 
 [numthreads(8,8,1)]
 void TUpdateTexture2DSubresourceCS( uint2 ThreadId : SV_DispatchThreadID )
 {
-	if( all( ThreadId.xy < DestPosSize.zw ) )
-	{
-		int2 TexturePixelOffset = ThreadId.xy + DestPosSize.xy;
-		int BufferOffset = (ThreadId.y * TSrcPitch) + ThreadId.x;
-		TDestTexture[TexturePixelOffset] = TSrcBuffer[BufferOffset];
-	}
-}
-
-Texture2D<COMPONENT_TYPE> TSrcTexture;
-uint2 SrcPos;
-
-[numthreads(8,8,1)]
-void TCopyTexture2DCS( uint3 ThreadId : SV_DispatchThreadID )
-{
-	if (all(ThreadId.xy < DestPosSize.zw))
+	if( all( ThreadId.xy < TDestPosSize.zw ) )
 	{
-		TDestTexture[DestPosSize.xy + ThreadId.xy] = TSrcTexture.Load(int3(SrcPos + ThreadId.xy, 0));
+		int2 DstTextureOffset = ThreadId.xy + TDestPosSize.xy;
+		int2 SrcTextureOffset = ThreadId.xy + TSrcPosPitch.xy;
+		int BufferOffset = (SrcTextureOffset.y * TSrcPosPitch.z) + SrcTextureOffset.x;
+		TDestTexture[DstTextureOffset] = TSrcBuffer[BufferOffset];
 	}
 }
 
diff --git a/Engine/Source/Runtime/Engine/Private/UpdateTextureShaders.cpp b/Engine/Source/Runtime/Engine/Private/UpdateTextureShaders.cpp
--- a/Engine/Source/Runtime/Engine/Private/UpdateTextureShaders.cpp
+++ b/Engine/Source/Runtime/Engine/Private/UpdateTextureShaders.cpp
@@ -2,21 +2,16 @@
 
 #include "UpdateTextureShaders.h"
 
-IMPLEMENT_SHADER_TYPE4_WITH_TEMPLATE_PREFIX(, ENGINE_API, FUpdateTexture2DSubresouceCS, SF_Compute);
-IMPLEMENT_SHADER_TYPE4_WITH_TEMPLATE_PREFIX(, ENGINE_API, FUpdateTexture3DSubresouceCS, SF_Compute);
+IMPLEMENT_TYPE_LAYOUT(FUpdateTexture2DSubresourceCS);
 
-IMPLEMENT_SHADER_TYPE(,FCopyTexture2DCS,TEXT("/Engine/Private/UpdateTextureShaders.usf"),TEXT("CopyTexture2DCS"),SF_Compute);
+IMPLEMENT_SHADER_TYPE4_WITH_TEMPLATE_PREFIX(, ENGINE_API, FUpdateTexture3DSubresourceCS, SF_Compute);
 
-IMPLEMENT_SHADER_TYPE(template<>, TCopyTexture2DCS<1u>, TEXT("/Engine/Private/UpdateTextureShaders.usf"), TEXT("TCopyTexture2DCS"), SF_Compute);
-IMPLEMENT_SHADER_TYPE(template<>, TCopyTexture2DCS<2u>, TEXT("/Engine/Private/UpdateTextureShaders.usf"), TEXT("TCopyTexture2DCS"), SF_Compute);
-IMPLEMENT_SHADER_TYPE(template<>, TCopyTexture2DCS<4u>, TEXT("/Engine/Private/UpdateTextureShaders.usf"), TEXT("TCopyTexture2DCS"), SF_Compute);
-typedef TCopyTexture2DCS< 3u, float > FCopyTexture2DCS_Float3;
-IMPLEMENT_SHADER_TYPE(template<>, FCopyTexture2DCS_Float3, TEXT("/Engine/Private/UpdateTextureShaders.usf"), TEXT("TCopyTexture2DCS"), SF_Compute);
-
-IMPLEMENT_SHADER_TYPE4_WITH_TEMPLATE_PREFIX(template<>, ENGINE_API, TUpdateTexture2DSubresouceCS<1u>, SF_Compute);
-IMPLEMENT_SHADER_TYPE4_WITH_TEMPLATE_PREFIX(template<>, ENGINE_API, TUpdateTexture2DSubresouceCS<2u>, SF_Compute);
-IMPLEMENT_SHADER_TYPE4_WITH_TEMPLATE_PREFIX(template<>, ENGINE_API, TUpdateTexture2DSubresouceCS<3u>, SF_Compute);
-IMPLEMENT_SHADER_TYPE4_WITH_TEMPLATE_PREFIX(template<>, ENGINE_API, TUpdateTexture2DSubresouceCS<4u>, SF_Compute);
+#define IMPLEMENT_UPDATE_TEX2D_SUBRESOURCE_SHADER(ValueType)\
+	typedef TUpdateTexture2DSubresourceCS<EUpdateTextureValueType::ValueType> FUpdateTexture2DSubresourceCS_##ValueType;\
+	IMPLEMENT_SHADER_TYPE4_WITH_TEMPLATE_PREFIX(template<>, ENGINE_API, FUpdateTexture2DSubresourceCS_##ValueType, SF_Compute);
+IMPLEMENT_UPDATE_TEX2D_SUBRESOURCE_SHADER(Float);
+IMPLEMENT_UPDATE_TEX2D_SUBRESOURCE_SHADER(Int32);
+IMPLEMENT_UPDATE_TEX2D_SUBRESOURCE_SHADER(Uint32);
 
 IMPLEMENT_SHADER_TYPE4_WITH_TEMPLATE_PREFIX(template<>, ENGINE_API, TCopyDataCS<2>, SF_Compute);
-IMPLEMENT_SHADER_TYPE4_WITH_TEMPLATE_PREFIX(template<>, ENGINE_API, TCopyDataCS<1>, SF_Compute);
\ No newline at end of file
+IMPLEMENT_SHADER_TYPE4_WITH_TEMPLATE_PREFIX(template<>, ENGINE_API, TCopyDataCS<1>, SF_Compute);
diff --git a/Engine/Source/Runtime/Engine/Public/UpdateTextureShaders.h b/Engine/Source/Runtime/Engine/Public/UpdateTextureShaders.h
--- a/Engine/Source/Runtime/Engine/Public/UpdateTextureShaders.h
+++ b/Engine/Source/Runtime/Engine/Public/UpdateTextureShaders.h
@@ -8,49 +8,48 @@
 #include "Shader.h"
 #include "GlobalShader.h"
 
-class FUpdateTexture2DSubresouceCS : public FGlobalShader
+enum class EUpdateTextureValueType
 {
-	DECLARE_EXPORTED_SHADER_TYPE(FUpdateTexture2DSubresouceCS, Global, ENGINE_API);
+	Float,
+	Int32,
+	Uint32
+};
+
+class FUpdateTexture2DSubresourceCS : public FGlobalShader
+{
+	DECLARE_EXPORTED_SHADER_TYPE(FUpdateTexture2DSubresourceCS, Global, ENGINE_API);
 public:
-	FUpdateTexture2DSubresouceCS() {}
-	FUpdateTexture2DSubresouceCS(const ShaderMetaType::CompiledShaderInitializerType& Initializer)
+	FUpdateTexture2DSubresourceCS() {}
+	FUpdateTexture2DSubresourceCS(const ShaderMetaType::CompiledShaderInitializerType& Initializer)
 		: FGlobalShader( Initializer )
 	{
-		SrcPitchParameter.Bind(Initializer.ParameterMap, TEXT("SrcPitch"), SPF_Mandatory);
-		SrcBuffer.Bind(Initializer.ParameterMap, TEXT("SrcBuffer"), SPF_Mandatory);
-		DestPosSizeParameter.Bind(Initializer.ParameterMap, TEXT("DestPosSize"), SPF_Mandatory);
-		DestTexture.Bind(Initializer.ParameterMap, TEXT("DestTexture"), SPF_Mandatory);
-	}
-
-	static const TCHAR* GetSourceFilename() { return TEXT("/Engine/Private/UpdateTextureShaders.usf"); }
-	
-	static const TCHAR* GetFunctionName() { return TEXT("UpdateTexture2DSubresourceCS"); }
-
-	static bool ShouldCompilePermutation(const FGlobalShaderPermutationParameters& Parameters)
-	{
-		return IsFeatureLevelSupported(Parameters.Platform, ERHIFeatureLevel::SM5);
+		SrcPosPitchParameter.Bind(Initializer.ParameterMap, TEXT("TSrcPosPitch"), SPF_Mandatory);
+		SrcBuffer.Bind(Initializer.ParameterMap, TEXT("TSrcBuffer"), SPF_Mandatory);
+		DestPosSizeParameter.Bind(Initializer.ParameterMap, TEXT("TDestPosSize"), SPF_Mandatory);
+		DestTexture.Bind(Initializer.ParameterMap, TEXT("TDestTexture"), SPF_Mandatory);
 	}
 
-	LAYOUT_FIELD(FShaderParameter, SrcPitchParameter);
+	LAYOUT_FIELD(FShaderParameter, SrcPosPitchParameter);
 	LAYOUT_FIELD(FShaderResourceParameter, SrcBuffer);
 	LAYOUT_FIELD(FShaderParameter, DestPosSizeParameter);
 	LAYOUT_FIELD(FShaderResourceParameter, DestTexture);
+
+	static inline TShaderRef<FUpdateTexture2DSubresourceCS> SelectShader(FGlobalShaderMap* GlobalShaderMap, EUpdateTextureValueType ValueType);
+
+protected:
+	template<typename ShaderType>
+	static inline TShaderRef<FUpdateTexture2DSubresourceCS> SelectShader(FGlobalShaderMap* GlobalShaderMap);
 };
 
-template<uint32 NumComponents>
-class TUpdateTexture2DSubresouceCS : public FGlobalShader
+template<EUpdateTextureValueType ValueType>
+class TUpdateTexture2DSubresourceCS : public FUpdateTexture2DSubresourceCS
 {
-	DECLARE_EXPORTED_SHADER_TYPE(TUpdateTexture2DSubresouceCS, Global, ENGINE_API);
+	DECLARE_EXPORTED_SHADER_TYPE(TUpdateTexture2DSubresourceCS, Global, ENGINE_API);
 public:
-	TUpdateTexture2DSubresouceCS() {}
-	TUpdateTexture2DSubresouceCS(const ShaderMetaType::CompiledShaderInitializerType& Initializer)
-		: FGlobalShader(Initializer)
-	{
-		SrcPitchParameter.Bind(Initializer.ParameterMap, TEXT("TSrcPitch"), SPF_Mandatory);
-		SrcBuffer.Bind(Initializer.ParameterMap, TEXT("TSrcBuffer"), SPF_Mandatory);
-		DestPosSizeParameter.Bind(Initializer.ParameterMap, TEXT("DestPosSize"), SPF_Mandatory);
-		DestTexture.Bind(Initializer.ParameterMap, TEXT("TDestTexture"), SPF_Mandatory);
-	}
+	TUpdateTexture2DSubresourceCS() {}
+	TUpdateTexture2DSubresourceCS(const ShaderMetaType::CompiledShaderInitializerType& Initializer)
+		: FUpdateTexture2DSubresourceCS(Initializer)
+	{}
 
 	static const TCHAR* GetSourceFilename() { return TEXT("/Engine/Private/UpdateTextureShaders.usf"); }
 	
@@ -59,35 +58,45 @@ public:
 	static void ModifyCompilationEnvironment(const FGlobalShaderPermutationParameters& Parameters, FShaderCompilerEnvironment& OutEnvironment)
 	{
 		const TCHAR* Type = nullptr;
-		static_assert(NumComponents >= 1u && NumComponents <= 4u, "Invalid NumComponents");
-		switch (NumComponents)
+		switch (ValueType)
 		{
-		case 1u: Type = TEXT("uint"); break;
-		case 2u: Type = TEXT("uint2"); break;
-		case 3u: Type = TEXT("uint3"); break;
-		case 4u: Type = TEXT("uint4"); break;
+		default: checkNoEntry();
+		case EUpdateTextureValueType::Float: OutEnvironment.SetDefine(TEXT("COMPONENT_TYPE"), TEXT("float4")); break;
+		case EUpdateTextureValueType::Int32: OutEnvironment.SetDefine(TEXT("COMPONENT_TYPE"), TEXT("int4")); break;
+		case EUpdateTextureValueType::Uint32: OutEnvironment.SetDefine(TEXT("COMPONENT_TYPE"), TEXT("uint4")); break;
 		}
-
-		OutEnvironment.SetDefine(TEXT("COMPONENT_TYPE"), Type);
 	}
 
 	static bool ShouldCompilePermutation(const FGlobalShaderPermutationParameters& Parameters)
 	{
 		return IsFeatureLevelSupported(Parameters.Platform, ERHIFeatureLevel::SM5);
 	}
-
-	LAYOUT_FIELD(FShaderParameter, SrcPitchParameter);
-	LAYOUT_FIELD(FShaderResourceParameter, SrcBuffer);
-	LAYOUT_FIELD(FShaderParameter, DestPosSizeParameter);
-	LAYOUT_FIELD(FShaderResourceParameter, DestTexture);
 };
 
-class FUpdateTexture3DSubresouceCS : public FGlobalShader
+template <typename ShaderType>
+inline TShaderRef<FUpdateTexture2DSubresourceCS> FUpdateTexture2DSubresourceCS::SelectShader(FGlobalShaderMap* GlobalShaderMap)
+{
+	TShaderRef<ShaderType> Shader = GlobalShaderMap->GetShader<ShaderType>();
+	return Shader;
+}
+
+inline TShaderRef<FUpdateTexture2DSubresourceCS> FUpdateTexture2DSubresourceCS::SelectShader(FGlobalShaderMap* GlobalShaderMap, EUpdateTextureValueType ValueType)
 {
-	DECLARE_EXPORTED_SHADER_TYPE(FUpdateTexture3DSubresouceCS, Global, ENGINE_API);
+	switch (ValueType)
+	{
+	default: checkNoEntry();
+	case EUpdateTextureValueType::Float: return FUpdateTexture2DSubresourceCS::SelectShader<TUpdateTexture2DSubresourceCS<EUpdateTextureValueType::Float>>(GlobalShaderMap);
+	case EUpdateTextureValueType::Int32: return FUpdateTexture2DSubresourceCS::SelectShader<TUpdateTexture2DSubresourceCS<EUpdateTextureValueType::Int32>>(GlobalShaderMap);
+	case EUpdateTextureValueType::Uint32: return FUpdateTexture2DSubresourceCS::SelectShader<TUpdateTexture2DSubresourceCS<EUpdateTextureValueType::Uint32>>(GlobalShaderMap);
+	}
+}
+
+class FUpdateTexture3DSubresourceCS : public FGlobalShader
+{
+	DECLARE_EXPORTED_SHADER_TYPE(FUpdateTexture3DSubresourceCS, Global, ENGINE_API);
 public:
-	FUpdateTexture3DSubresouceCS() {}
-	FUpdateTexture3DSubresouceCS(const ShaderMetaType::CompiledShaderInitializerType& Initializer)
+	FUpdateTexture3DSubresourceCS() {}
+	FUpdateTexture3DSubresourceCS(const ShaderMetaType::CompiledShaderInitializerType& Initializer)
 		: FGlobalShader(Initializer)
 	{
 		SrcPitchParameter.Bind(Initializer.ParameterMap, TEXT("SrcPitch"), SPF_Mandatory);
@@ -120,85 +129,6 @@ public:
 	LAYOUT_FIELD(FShaderResourceParameter, DestTexture3D)
 };
 
-class FCopyTexture2DCS : public FGlobalShader
-{
-	DECLARE_SHADER_TYPE(FCopyTexture2DCS,Global);
-public:
-	FCopyTexture2DCS() {}
-	FCopyTexture2DCS(const ShaderMetaType::CompiledShaderInitializerType& Initializer)
-		: FGlobalShader( Initializer )
-	{
-		SrcTexture.Bind(Initializer.ParameterMap, TEXT("SrcTexture"), SPF_Mandatory);
-		DestTexture.Bind(Initializer.ParameterMap, TEXT("DestTexture"), SPF_Mandatory);
-		DestPosSize.Bind(Initializer.ParameterMap, TEXT("DestPosSize"), SPF_Mandatory);
-	}
-
-	static bool ShouldCompilePermutation(const FGlobalShaderPermutationParameters& Parameters)
-	{
-		return IsFeatureLevelSupported(Parameters.Platform, ERHIFeatureLevel::SM5);
-	}
-
-	LAYOUT_FIELD(FShaderResourceParameter, SrcTexture);
-	LAYOUT_FIELD(FShaderResourceParameter, DestTexture);
-	LAYOUT_FIELD(FShaderParameter, DestPosSize);
-};
-
-template<uint32 NumComponents, typename ComponentType = uint32>
-class TCopyTexture2DCS : public FGlobalShader
-{
-	DECLARE_SHADER_TYPE(TCopyTexture2DCS, Global);
-public:
-	TCopyTexture2DCS() {}
-	TCopyTexture2DCS(const ShaderMetaType::CompiledShaderInitializerType& Initializer)
-		: FGlobalShader(Initializer)
-	{
-		SrcTexture.Bind(Initializer.ParameterMap, TEXT("TSrcTexture"), SPF_Mandatory);
-		DestTexture.Bind(Initializer.ParameterMap, TEXT("TDestTexture"), SPF_Mandatory);
-		SrcPosParameter.Bind(Initializer.ParameterMap, TEXT("SrcPos"), SPF_Mandatory);
-		DestPosSizeParameter.Bind(Initializer.ParameterMap, TEXT("DestPosSize"), SPF_Mandatory);
-	}
-
-	static const TCHAR* GetTypename(uint32 Dummy)
-	{
-		switch (NumComponents)
-		{
-		case 1u: return TEXT("uint"); break;
-		case 2u: return TEXT("uint2"); break;
-		case 3u: return TEXT("uint3"); break;
-		case 4u: return TEXT("uint4"); break;
-		}
-		return nullptr;
-	}
-
-	static const TCHAR* GetTypename(float Dummy)
-	{
-		switch (NumComponents)
-		{
-		case 1u: return TEXT("float"); break;
-		case 2u: return TEXT("float2"); break;
-		case 3u: return TEXT("float3"); break;
-		case 4u: return TEXT("float4"); break;
-		}
-		return nullptr;
-	}
-
-	static void ModifyCompilationEnvironment(const FGlobalShaderPermutationParameters& Parameters, FShaderCompilerEnvironment& OutEnvironment)
-	{
-		static_assert(NumComponents >= 1u && NumComponents <= 4u, "Invalid NumComponents");
-		OutEnvironment.SetDefine(TEXT("COMPONENT_TYPE"), GetTypename((ComponentType)0));
-	}
-
-	static bool ShouldCompilePermutation(const FGlobalShaderPermutationParameters& Parameters)
-	{
-		return IsFeatureLevelSupported(Parameters.Platform, ERHIFeatureLevel::SM5);
-	}
-
-	LAYOUT_FIELD(FShaderResourceParameter, SrcTexture);
-	LAYOUT_FIELD(FShaderResourceParameter, DestTexture);
-	LAYOUT_FIELD(FShaderParameter, SrcPosParameter);
-	LAYOUT_FIELD(FShaderParameter, DestPosSizeParameter);
-};
-
 template<uint32 ElementsPerThread>
 class TCopyDataCS : public FGlobalShader
 {
