Add better output to COTF + ODSC.

- Print out what command is being run (global, material, changed, singleshader)
- Print out exactly how many shaders were compiled.
- Added a function to convert command type enum to a string.

diff --git a/Engine/Source/Runtime/Engine/Private/ShaderCompiler/ShaderCompiler.cpp b/Engine/Source/Runtime/Engine/Private/ShaderCompiler/ShaderCompiler.cpp
--- a/Engine/Source/Runtime/Engine/Private/ShaderCompiler/ShaderCompiler.cpp
+++ b/Engine/Source/Runtime/Engine/Private/ShaderCompiler/ShaderCompiler.cpp
@@ -7296,6 +7296,24 @@ FShaderRecompileData::FShaderRecompileData(const FString& InPlatformName, EShade
 {
 }
 
+extern ENGINE_API const TCHAR* ODSCCmdEnumToString(ODSCRecompileCommand Cmd)
+{
+	switch (Cmd)
+	{
+	case ODSCRecompileCommand::None:
+		return TEXT("None");
+	case ODSCRecompileCommand::Changed:
+		return TEXT("Change");
+	case ODSCRecompileCommand::Global:
+		return TEXT("Global");
+	case ODSCRecompileCommand::Material:
+		return TEXT("Material");
+	case ODSCRecompileCommand::SingleShader:
+		return TEXT("SingleShader");
+	}
+	ensure(false);
+	return TEXT("Unknown");
+}
 
 #if WITH_EDITOR
 void RecompileShadersForRemote(
@@ -7317,7 +7335,7 @@ void RecompileShadersForRemote(
 
 	UE_LOG(LogShaders, Display, TEXT(""));
 	UE_LOG(LogShaders, Display, TEXT("********************************"));
-	UE_LOG(LogShaders, Display, TEXT("Received compile shader request."));
+	UE_LOG(LogShaders, Display, TEXT("Received compile shader request %s."), ODSCCmdEnumToString(Args.CommandType));
 
 	const bool bPreviousState = GShaderCompilingManager->IsShaderCompilationSkipped();
 	GShaderCompilingManager->SkipShaderCompilation(false);
@@ -7335,7 +7353,9 @@ void RecompileShadersForRemote(
 		MaterialsToCompile.Add(LoadObject<UMaterialInterface>(NULL, *Args.MaterialsToLoad[Index]));
 	}
 
-	UE_LOG(LogShaders, Verbose, TEXT("  Done!"))
+	UE_LOG(LogShaders, Verbose, TEXT("  Done!"));
+
+	const uint32 StartTotalShadersCompiled = GShaderCompilerStats->GetTotalShadersCompiled();
 
 	// figure out which shaders are out of date
 	TArray<const FShaderType*> OutdatedShaderTypes;
@@ -7394,8 +7414,6 @@ void RecompileShadersForRemote(
 				if (Args.CommandType == ODSCRecompileCommand::SingleShader &&
 					Args.ShaderTypesToLoad.Len() > 0)
 				{
-					UE_LOG(LogShaders, Display, TEXT("Recompiling single shader."));
-
 					TArray<const FShaderType*> ShaderTypes = FShaderType::GetShaderTypesByFilename(*Args.ShaderTypesToLoad);
 					TArray<const FShaderPipelineType*> ShaderPipelineTypes = FShaderPipelineType::GetShaderPipelineTypesByFilename(*Args.ShaderTypesToLoad);
 
@@ -7503,10 +7521,7 @@ void RecompileShadersForRemote(
 	}
 
 	UE_LOG(LogShaders, Display, TEXT(""));
-	UE_LOG(LogShaders, Display, TEXT("Finished shader compile request in %.2f seconds."), FPlatformTime::Seconds() - StartTime);
-
-	// Provide a log of what happened.
-	GShaderCompilingManager->PrintStats(true);
+	UE_LOG(LogShaders, Display, TEXT("Compiled %u shaders in %.2f seconds."), GShaderCompilerStats->GetTotalShadersCompiled() - StartTotalShadersCompiled, FPlatformTime::Seconds() - StartTime);
 
 	// Restore compilation state.
 	GShaderCompilingManager->SkipShaderCompilation(bPreviousState);
diff --git a/Engine/Source/Runtime/Engine/Public/ShaderCompiler.h b/Engine/Source/Runtime/Engine/Public/ShaderCompiler.h
--- a/Engine/Source/Runtime/Engine/Public/ShaderCompiler.h
+++ b/Engine/Source/Runtime/Engine/Public/ShaderCompiler.h
@@ -1271,6 +1271,8 @@ enum class ODSCRecompileCommand
 	SingleShader
 };
 
+extern ENGINE_API const TCHAR* ODSCCmdEnumToString(ODSCRecompileCommand Cmd);
+
 struct FShaderRecompileData
 {
 	/** The platform name to compile for. */
