- Fix for merge code not updating existing parameters properly for SoundCues.
- Fix for instanced variables not caching when any parameterized sound type (Cue or MetaSound) is replayed

diff --git a/Engine/Source/Runtime/AudioExtensions/Private/AudioParameter.cpp b/Engine/Source/Runtime/AudioExtensions/Private/AudioParameter.cpp
--- a/Engine/Source/Runtime/AudioExtensions/Private/AudioParameter.cpp
+++ b/Engine/Source/Runtime/AudioExtensions/Private/AudioParameter.cpp
@@ -220,8 +220,7 @@ void FAudioParameter::Merge(TArray<FAudioParameter>&& InParams, TArray<FAudioPar
 			FAudioParameter NewParam = InParams.Pop(bAllowShrinking);
 			if (NewParam.ParamName == OutParam.ParamName)
 			{
-				NewParam.Merge(OutParam);
-				OutParam = MoveTemp(NewParam);
+				OutParam.Merge(NewParam);
 			}
 			else
 			{
diff --git a/Engine/Source/Runtime/Engine/Private/Components/AudioComponent.cpp b/Engine/Source/Runtime/Engine/Private/Components/AudioComponent.cpp
--- a/Engine/Source/Runtime/Engine/Private/Components/AudioComponent.cpp
+++ b/Engine/Source/Runtime/Engine/Private/Components/AudioComponent.cpp
@@ -808,8 +808,8 @@ void UAudioComponent::PlayInternal(const PlayInternalRequestData& InPlayRequestD
 		AudioDevice->SetCanHaveMultipleActiveSounds(AudioComponentID, bCanPlayMultipleInstances);
 	}
 
-	TArray<FAudioParameter> SoundParams = InstanceParameters;
-
+	TArray<FAudioParameter> SoundParams = DefaultParameters;
+	
 	if (AActor* Owner = GetOwner())
 	{
 		TArray<FAudioParameter> ActorParams;
@@ -821,8 +821,8 @@ void UAudioComponent::PlayInternal(const PlayInternalRequestData& InPlayRequestD
 		UE_LOG(LogAudio, VeryVerbose, TEXT("AudioComponent : '%s' playing sound '%s' has no owner"), *GetFullName(), *Sound->GetName());
 	}
 
-	TArray<FAudioParameter> DefaultParamsCopy = DefaultParameters;
-	FAudioParameter::Merge(MoveTemp(DefaultParamsCopy), SoundParams);
+	TArray<FAudioParameter> InstanceParamsCopy = InstanceParameters;
+	FAudioParameter::Merge(MoveTemp(InstanceParamsCopy), SoundParams);
 
 	AudioDevice->AddNewActiveSound(NewActiveSound, MoveTemp(SoundParams));
 
