Properly return failure with TextureFormatModule missing.

diff --git a/Engine/Source/Runtime/Engine/Private/TextureDerivedDataBuildUtils.cpp b/Engine/Source/Runtime/Engine/Private/TextureDerivedDataBuildUtils.cpp
--- a/Engine/Source/Runtime/Engine/Private/TextureDerivedDataBuildUtils.cpp
+++ b/Engine/Source/Runtime/Engine/Private/TextureDerivedDataBuildUtils.cpp
@@ -246,20 +246,25 @@ UE::DerivedData::FUtf8SharedString FindTextureBuildFunction(const FName TextureF
 {
 	using namespace UE::DerivedData;
 
-	if (FReadScopeLock Lock(GTextureBuildFunctionLock); const FUtf8SharedString* Function = GTextureBuildFunctionMap.Find(TextureFormatName))
 	{
-		return *Function;
+		FReadScopeLock Lock(GTextureBuildFunctionLock);
+		if (const FUtf8SharedString* Function = GTextureBuildFunctionMap.Find(TextureFormatName))
+		{
+			return *Function;
+		}
 	}
 
 	FName TextureFormatModuleName;
+	ITextureFormatManagerModule* TFM = GetTextureFormatManager();
+	if (TFM == nullptr)
+	{
+		return {};
+	}
 
-	if (ITextureFormatManagerModule* TFM = GetTextureFormatManager())
+	ITextureFormatModule* TextureFormatModule = nullptr;
+	if (!TFM->FindTextureFormatAndModule(TextureFormatName, TextureFormatModuleName, TextureFormatModule))
 	{
-		ITextureFormatModule* TextureFormatModule = nullptr;
-		if (!TFM->FindTextureFormatAndModule(TextureFormatName, TextureFormatModuleName, TextureFormatModule))
-		{
-			return {};
-		}
+		return {};
 	}
 
 	TStringBuilder<128> FunctionName;
