Mark struct internals as SaveGame

#rnx

diff --git a/Engine/Source/Runtime/Engine/Classes/Engine/MemberReference.h b/Engine/Source/Runtime/Engine/Classes/Engine/MemberReference.h
--- a/Engine/Source/Runtime/Engine/Classes/Engine/MemberReference.h
+++ b/Engine/Source/Runtime/Engine/Classes/Engine/MemberReference.h
@@ -71,27 +71,27 @@ protected:
 	 * if it is a native delegate signature function (declared globally). Should 
 	 * be NULL if bSelfContext is true.  
 	 */
-	UPROPERTY()
+	UPROPERTY(SaveGame)
 	mutable TObjectPtr<UObject> MemberParent;
 
 	/**  */
-	UPROPERTY()
+	UPROPERTY(SaveGame)
 	mutable FString MemberScope;
 
 	/** Name of variable */
-	UPROPERTY()
+	UPROPERTY(SaveGame)
 	mutable FName MemberName;
 
 	/** The Guid of the variable */
-	UPROPERTY()
+	UPROPERTY(SaveGame)
 	mutable FGuid MemberGuid;
 
 	/** Whether or not this should be a "self" context */
-	UPROPERTY()
+	UPROPERTY(SaveGame)
 	mutable bool bSelfContext;
 
 	/** Whether or not this property has been deprecated */
-	UPROPERTY()
+	UPROPERTY(SaveGame)
 	mutable bool bWasDeprecated;
 	
 public:
