Fixed Performance Warning: Console object named 'r.SkinCache.SkipCompilingGPUSkinVF' shows many (500) FindConsoleObject() calls

diff --git a/Engine/Source/Runtime/Engine/Private/SkeletalRenderGPUSkin.cpp b/Engine/Source/Runtime/Engine/Private/SkeletalRenderGPUSkin.cpp
--- a/Engine/Source/Runtime/Engine/Private/SkeletalRenderGPUSkin.cpp
+++ b/Engine/Source/Runtime/Engine/Private/SkeletalRenderGPUSkin.cpp
@@ -1378,7 +1378,7 @@ const FVertexFactory* FSkeletalMeshObjectGPUSkin::GetSkinVertexFactory(const FSc
 	}
 
 	// If we have not compiled GPU Skin vertex factory variants
-	const auto CVar = IConsoleManager::Get().FindConsoleVariable(TEXT("r.SkinCache.SkipCompilingGPUSkinVF"));
+	static const auto CVar = IConsoleManager::Get().FindConsoleVariable(TEXT("r.SkinCache.SkipCompilingGPUSkinVF"));
 	if (FeatureLevel != ERHIFeatureLevel::ES3_1 && CVar && CVar->GetBool() == true)
 	{
 		ensureMsgf(false, TEXT("We are attempting to render with a GPU Skin Vertex Factory, but r.SkinCache.SkipCompilingGPUSkinVF=1 so we don't have shaders.  Skeletal meshes will draw in ref pose.  Either disable r.SkinCache.SkipCompilingGPUSkinVF or increase the r.SkinCache.SceneMemoryLimitInMB size."));
