Added profiling instrumentation for ConditionalCollectGarbage to catch some non-trivial work

diff --git a/Engine/Source/Runtime/Engine/Private/UnrealEngine.cpp b/Engine/Source/Runtime/Engine/Private/UnrealEngine.cpp
--- a/Engine/Source/Runtime/Engine/Private/UnrealEngine.cpp
+++ b/Engine/Source/Runtime/Engine/Private/UnrealEngine.cpp
@@ -1633,6 +1633,8 @@ void UEngine::ConditionalCollectGarbage()
 {
 	if (GFrameCounter != LastGCFrame)
 	{
+		QUICK_SCOPE_CYCLE_COUNTER(STAT_ConditionalCollectGarbage);
+
 #if !(UE_BUILD_SHIPPING || UE_BUILD_TEST)
 		if (CVarStressTestGCWhileStreaming.GetValueOnGameThread() && IsAsyncLoading())
 		{
