Don't recalculate LOD screen sizes in cooked editor.
#rnx

diff --git a/Engine/Source/Runtime/Engine/Private/StaticMesh.cpp b/Engine/Source/Runtime/Engine/Private/StaticMesh.cpp
--- a/Engine/Source/Runtime/Engine/Private/StaticMesh.cpp
+++ b/Engine/Source/Runtime/Engine/Private/StaticMesh.cpp
@@ -1809,7 +1809,12 @@ void FStaticMeshRenderData::Serialize(FArchive& Ar, UStaticMesh* Owner, bool bCo
 void FStaticMeshRenderData::InitResources(ERHIFeatureLevel::Type InFeatureLevel, UStaticMesh* Owner)
 {
 #if WITH_EDITOR
-	ResolveSectionInfo(Owner);
+	// Init the section info only for uncooked editor.
+	// Cooked packages don't need this and don't want any LOD screen size changes that it applies.
+	if (!Owner->GetPackage()->HasAnyPackageFlags(PKG_FilterEditorOnly))
+	{
+		ResolveSectionInfo(Owner);
+	}
 #endif // #if WITH_EDITOR
 
 	for (int32 LODIndex = 0; LODIndex < LODResources.Num(); ++LODIndex)
