Added object stat in WaitForRHIThreadFenceForDynamicData to diagnose waits.

diff --git a/Engine/Source/Runtime/Engine/Private/SkeletalRenderGPUSkin.cpp b/Engine/Source/Runtime/Engine/Private/SkeletalRenderGPUSkin.cpp
--- a/Engine/Source/Runtime/Engine/Private/SkeletalRenderGPUSkin.cpp
+++ b/Engine/Source/Runtime/Engine/Private/SkeletalRenderGPUSkin.cpp
@@ -338,7 +338,6 @@ void FSkeletalMeshObjectGPUSkin::ReleaseResources()
 			FGPUSkinCacheEntry* LocalSkinCacheEntryForRayTracing = SkinCacheEntryForRayTracing;
 			FGPUSkinCache::Release(LocalSkinCacheEntryForRayTracing);
 
-			FScopeCycleCounter Context(MeshObject->GetStatId());
 			MeshObject->WaitForRHIThreadFenceForDynamicData();
 			*PtrSkinCacheEntry = nullptr;
 			SkinCacheEntryForRayTracing = nullptr;
@@ -578,6 +577,7 @@ void FSkeletalMeshObjectGPUSkin::WaitForRHIThreadFenceForDynamicData()
 	// we should be done with the old data at this point
 	if (RHIThreadFenceForDynamicData.GetReference())
 	{
+		FScopeCycleCounter Context(GetStatId());
 		FRHICommandListExecutor::WaitOnRHIThreadFence(RHIThreadFenceForDynamicData);
 		RHIThreadFenceForDynamicData = nullptr;
 	}
