Quartz: Add a cvar to force play requests going to the engine right away (instead of caching the reqeust on the component until closer to the deadline)

diff --git a/Engine/Source/Runtime/Engine/Private/Components/AudioComponent.cpp b/Engine/Source/Runtime/Engine/Private/Components/AudioComponent.cpp
--- a/Engine/Source/Runtime/Engine/Private/Components/AudioComponent.cpp
+++ b/Engine/Source/Runtime/Engine/Private/Components/AudioComponent.cpp
@@ -33,7 +33,7 @@ FAutoConsoleVariableRef CVarPrimeSoundOnAudioComponentSpawn(
 	ECVF_Default);
 
 //CVar for how long voiceslots should be taken up when queuing sounds
-static int32 TimeToTakeUpVoiceSlotCVar = int32(EQuartzCommandQuantization::EighthNote);
+static int32 TimeToTakeUpVoiceSlotCVar = int32(EQuartzCommandQuantization::HalfNote);
 FAutoConsoleVariableRef CVarTimeToTakeUpVoiceSlot(
 	TEXT("au.Quartz.TimeToTakeUpVoiceSlot"),
 	TimeToTakeUpVoiceSlotCVar,
@@ -41,6 +41,16 @@ FAutoConsoleVariableRef CVarTimeToTakeUpVoiceSlot(
 	TEXT("Value: The EQuartzCommandQuantization index of the desired duration"),
 	ECVF_Default);
 
+
+//CVar to disable the gamethread-side caching of play requests to minimize time spent rendering a silent voice before the Quartz Deadline
+static int32 bAlwaysTakeVoiceSlotCVar = 0;
+FAutoConsoleVariableRef bCVarAlwaysTakeVoiceSlot(
+	TEXT("au.Quartz.bAlwaysTakeVoiceSlot"),
+	bAlwaysTakeVoiceSlotCVar,
+	TEXT("Always take voice slot immediately without trying to cache the request on the component\n")
+	TEXT("default = 0: cache play requests on the component until closer to the deadline. - 1: always forward the request to the audio engine immediately."),
+	ECVF_Default);
+
 static int32 WorldlessGetAudioTimeBehaviorCVar = 0;
 FAutoConsoleVariableRef CVarWorldlessGetAudioTimeBehavior(
 	TEXT("au.WorldlessGetAudioTimeBehavior"),
@@ -501,7 +511,7 @@ void UAudioComponent::PlayQuantized(
 	double NumFramesForDesiredTime = OutTickRate.GetFramesPerDuration(InQuantizationBoundary.Quantization) * InQuantizationBoundary.Multiplier;
 
 	// If the desired quantization time is less than our min time, just execute immediately
-	const bool bStealVoiceSlot = NumFramesForDesiredTime <= NumFramesBeforeMinTime;
+	const bool bStealVoiceSlot = bAlwaysTakeVoiceSlotCVar || (NumFramesForDesiredTime <= NumFramesBeforeMinTime);
 	const bool bClockIsNotRunning = !InClockHandle->IsClockRunning(WorldContextObject);
 	const bool bCommandResetsClock = InQuantizationBoundary.bResetClockOnQueued;
 	if (bStealVoiceSlot || bClockIsNotRunning || bCommandResetsClock)
