Static analysis fixes

diff --git a/Engine/Source/Runtime/Engine/Private/RepLayout.cpp b/Engine/Source/Runtime/Engine/Private/RepLayout.cpp
--- a/Engine/Source/Runtime/Engine/Private/RepLayout.cpp
+++ b/Engine/Source/Runtime/Engine/Private/RepLayout.cpp
@@ -2732,7 +2732,7 @@ void FRepLayout::SendProperties_r(
 		}
 
 #if !(UE_BUILD_SHIPPING || UE_BUILD_TEST)
-		if (GDoReplicationContextString> 0)
+		if (GDoReplicationContextString > 0 && Writer.PackageMap)
 		{
 			Writer.PackageMap->SetDebugContextString(FString::Printf(TEXT("%s - %s"), *Owner->GetPathName(), *Cmd.Property->GetPathName()));
 		}
@@ -2836,7 +2836,7 @@ void FRepLayout::SendProperties_r(
 		}
 
 #if !(UE_BUILD_SHIPPING || UE_BUILD_TEST)
-		if (GDoReplicationContextString> 0)
+		if (GDoReplicationContextString > 0 && Writer.PackageMap)
 		{
 			Writer.PackageMap->ClearDebugContextString();
 		}
@@ -2944,7 +2944,7 @@ static FORCEINLINE void WriteProperty_BackwardsCompatible(
 	const bool bDoChecksum)
 {
 #if !(UE_BUILD_SHIPPING || UE_BUILD_TEST)
-	if (GDoReplicationContextString> 0)
+	if (GDoReplicationContextString > 0 && Writer.PackageMap)
 	{
 		Writer.PackageMap->SetDebugContextString(FString::Printf(TEXT("%s - %s"), *Owner->GetPathName(), *Cmd.Property->GetPathName()));
 	}
@@ -2968,7 +2968,7 @@ static FORCEINLINE void WriteProperty_BackwardsCompatible(
 	NETWORK_PROFILER(GNetworkProfiler.TrackReplicateProperty(Cmd.Property, NumEndBits - NumStartBits, nullptr));
 
 #if !(UE_BUILD_SHIPPING || UE_BUILD_TEST)
-	if (GDoReplicationContextString> 0)
+	if (GDoReplicationContextString > 0 && Writer.PackageMap)
 	{
 		Writer.PackageMap->ClearDebugContextString();
 	}
