Strata - now using force cast of data when sent to backend fucntions to have a behavior close to the root node. (2d TexCoord can be used as Emissive color for instance)

diff --git a/Engine/Shaders/Shared/StrataDefinitions.h b/Engine/Shaders/Shared/StrataDefinitions.h
--- a/Engine/Shaders/Shared/StrataDefinitions.h
+++ b/Engine/Shaders/Shared/StrataDefinitions.h
@@ -8,7 +8,7 @@
 #pragma once
 
 // Change this to force recompilation of all strata dependent shaders (use https://www.random.org/cgi-bin/randbyte?nbytes=4&format=h)
-#define STRATA_SHADER_VERSION 0x3ec8acbb 
+#define STRATA_SHADER_VERSION 0x2a533622
 
 
 
diff --git a/Engine/Source/Runtime/Engine/Private/Materials/HLSLMaterialTranslator.cpp b/Engine/Source/Runtime/Engine/Private/Materials/HLSLMaterialTranslator.cpp
--- a/Engine/Source/Runtime/Engine/Private/Materials/HLSLMaterialTranslator.cpp
+++ b/Engine/Source/Runtime/Engine/Private/Materials/HLSLMaterialTranslator.cpp
@@ -10264,6 +10264,12 @@ int32 FHLSLMaterialTranslator::StrataCreateAndRegisterNullMaterial()
 	return OutputCodeChunk;
 }
 
+FString FHLSLMaterialTranslator::StrataGetCastParameterCode(int32 Index, EMaterialValueType DestType)
+{
+	int32 CastParameter = ForceCast(Index, DestType);
+	return GetParameterCode(CastParameter);
+}
+
 int32 FHLSLMaterialTranslator::StrataSlabBSDF(
 	int32 UseMetalness,
 	int32 BaseColor, int32 EdgeColor, int32 Specular, int32 Metallic,
@@ -10289,17 +10295,24 @@ int32 FHLSLMaterialTranslator::StrataSlabBSDF(
 		}
 		return AddCodeChunk(
 			MCT_Strata, TEXT("PromoteParameterBlendedBSDFToOperator(GetStrataSlabBSDF(Parameters.StrataPixelFootprint, %s, %s, %s, %s, %s, %s, %s, %s, %s, %s, %s, %s, %s, %s, %s, %s, %s, %s, %s, Parameters.SharedLocalBases.Types) /* Normal = %s ; Tangent = %s */, Parameters.StrataTree, %u, %u, %u, %u)"),
-			*GetParameterCode(UseMetalness),
-			*GetParameterCode(BaseColor), *GetParameterCode(EdgeColor), *GetParameterCode(Specular), *GetParameterCode(Metallic),
-			*GetParameterCode(DiffuseAlbedo), *GetParameterCode(F0), *GetParameterCode(F90),
-			*GetParameterCode(Roughness), *GetParameterCode(Anisotropy),
-			*GetParameterCode(SSSProfileId),
-			*GetParameterCode(SSSMFP),
-			*GetParameterCode(SSSMFPScale),
-			*GetParameterCode(EmissiveColor),
-			*GetParameterCode(Haziness),
-			*GetParameterCode(FuzzAmount), *GetParameterCode(FuzzColor),
-			*GetParameterCode(Thickness),
+			*StrataGetCastParameterCode(UseMetalness,	MCT_Float),
+			*StrataGetCastParameterCode(BaseColor,		MCT_Float3),
+			*StrataGetCastParameterCode(EdgeColor,		MCT_Float3),
+			*StrataGetCastParameterCode(Specular,		MCT_Float),
+			*StrataGetCastParameterCode(Metallic,		MCT_Float),
+			*StrataGetCastParameterCode(DiffuseAlbedo,	MCT_Float3),
+			*StrataGetCastParameterCode(F0,				MCT_Float3),
+			*StrataGetCastParameterCode(F90,			MCT_Float3),
+			*StrataGetCastParameterCode(Roughness,		MCT_Float),
+			*StrataGetCastParameterCode(Anisotropy,		MCT_Float),
+			*StrataGetCastParameterCode(SSSProfileId,	MCT_Float),
+			*StrataGetCastParameterCode(SSSMFP,			MCT_Float3),
+			*StrataGetCastParameterCode(SSSMFPScale,	MCT_Float),
+			*StrataGetCastParameterCode(EmissiveColor,	MCT_Float3),
+			*StrataGetCastParameterCode(Haziness,		MCT_Float),
+			*StrataGetCastParameterCode(FuzzAmount,		MCT_Float),
+			*StrataGetCastParameterCode(FuzzColor,		MCT_Float3),
+			*StrataGetCastParameterCode(Thickness,		MCT_Float),
 			*SharedLocalBasisIndexMacro,
 			*NormalCode,
 			*TangentCode,
@@ -10308,20 +10321,27 @@ int32 FHLSLMaterialTranslator::StrataSlabBSDF(
 			PromoteToOperator->LayerDepth,
 			PromoteToOperator->bIsBottom ? 1 : 0);
 	}
-
+	
 	return AddCodeChunk(
 		MCT_Strata, TEXT("GetStrataSlabBSDF(Parameters.StrataPixelFootprint, %s, %s, %s, %s, %s, %s, %s, %s, %s, %s, %s, %s, %s, %s, %s, %s, %s, %s, %s, Parameters.SharedLocalBases.Types) /* Normal = %s ; Tangent = %s */"),
-		* GetParameterCode(UseMetalness),
-		*GetParameterCode(BaseColor),		*GetParameterCode(EdgeColor),	*GetParameterCode(Specular), *GetParameterCode(Metallic),
-		*GetParameterCode(DiffuseAlbedo),	*GetParameterCode(F0),			*GetParameterCode(F90),
-		*GetParameterCode(Roughness),		*GetParameterCode(Anisotropy),
-		*GetParameterCode(SSSProfileId),
-		*GetParameterCode(SSSMFP),
-		*GetParameterCode(SSSMFPScale),
-		*GetParameterCode(EmissiveColor),
-		*GetParameterCode(Haziness),
-		*GetParameterCode(FuzzAmount),		*GetParameterCode(FuzzColor),
-		*GetParameterCode(Thickness),
+		*StrataGetCastParameterCode(UseMetalness,	MCT_Float),
+		*StrataGetCastParameterCode(BaseColor,		MCT_Float3),
+		*StrataGetCastParameterCode(EdgeColor,		MCT_Float3),
+		*StrataGetCastParameterCode(Specular,		MCT_Float),
+		*StrataGetCastParameterCode(Metallic,		MCT_Float),
+		*StrataGetCastParameterCode(DiffuseAlbedo,	MCT_Float3),
+		*StrataGetCastParameterCode(F0,				MCT_Float3),
+		*StrataGetCastParameterCode(F90,			MCT_Float3),
+		*StrataGetCastParameterCode(Roughness,		MCT_Float),
+		*StrataGetCastParameterCode(Anisotropy,		MCT_Float),
+		*StrataGetCastParameterCode(SSSProfileId,	MCT_Float),
+		*StrataGetCastParameterCode(SSSMFP,			MCT_Float3),
+		*StrataGetCastParameterCode(SSSMFPScale,	MCT_Float),
+		*StrataGetCastParameterCode(EmissiveColor,	MCT_Float3),
+		*StrataGetCastParameterCode(Haziness,		MCT_Float),
+		*StrataGetCastParameterCode(FuzzAmount,		MCT_Float),
+		*StrataGetCastParameterCode(FuzzColor,		MCT_Float3),
+		*StrataGetCastParameterCode(Thickness,		MCT_Float),
 		*SharedLocalBasisIndexMacro,
 		*NormalCode,
 		*TangentCode
@@ -10351,14 +10371,22 @@ int32 FHLSLMaterialTranslator::StrataConversionFromLegacy(
 	return AddCodeChunk(
 		MCT_Strata, TEXT("StrataConvertLegacyMaterial%s(Parameters.StrataPixelFootprint, %s, %s, %s, %s, %s, %s, %s, %s, %s, %s, %s, %s, %s, %s, %s, %s, %s, %s, %s, Parameters.SharedLocalBases.Types, Parameters.StrataTree) /* Normal = %s ; Tangent = %s ; ClearCoat_Normal = %s ; ClearCoat_Tangent = %s */"),
 		bHasDynamicShadingModels ? TEXT("Dynamic") : TEXT("Static"),
-		*GetParameterCode(BaseColor), *GetParameterCode(Specular), *GetParameterCode(Metallic),
-		*GetParameterCode(Roughness), *GetParameterCode(Anisotropy),
-		*GetParameterCode(SubSurfaceColor), *GetParameterCode(SubSurfaceProfileId),
-		*GetParameterCode(ClearCoat), *GetParameterCode(ClearCoatRoughness),
-		*GetParameterCode(EmissiveColor),
-		*GetParameterCode(Opacity),
-		*GetParameterCode(TransmittanceColor),
-		*GetParameterCode(WaterScatteringCoefficients), *GetParameterCode(WaterAbsorptionCoefficients), *GetParameterCode(WaterPhaseG), *GetParameterCode(ColorScaleBehindWater),
+		*StrataGetCastParameterCode(BaseColor,						MCT_Float3),
+		*StrataGetCastParameterCode(Specular,						MCT_Float),
+		*StrataGetCastParameterCode(Metallic,						MCT_Float),
+		*StrataGetCastParameterCode(Roughness,						MCT_Float),
+		*StrataGetCastParameterCode(Anisotropy,						MCT_Float),
+		*StrataGetCastParameterCode(SubSurfaceColor,				MCT_Float3),
+		*StrataGetCastParameterCode(SubSurfaceProfileId,			MCT_Float),
+		*StrataGetCastParameterCode(ClearCoat,						MCT_Float),
+		*StrataGetCastParameterCode(ClearCoatRoughness,				MCT_Float),
+		*StrataGetCastParameterCode(EmissiveColor,					MCT_Float3),
+		*StrataGetCastParameterCode(Opacity,						MCT_Float),
+		*StrataGetCastParameterCode(TransmittanceColor,				MCT_Float3),
+		*StrataGetCastParameterCode(WaterScatteringCoefficients,	MCT_Float3),
+		*StrataGetCastParameterCode(WaterAbsorptionCoefficients,	MCT_Float3),
+		*StrataGetCastParameterCode(WaterPhaseG,					MCT_Float),
+		*StrataGetCastParameterCode(ColorScaleBehindWater,			MCT_Float3),
 		*GetParameterCode(ShadingModel),
 		*SharedLocalBasisIndexMacro,
 		*ClearCoat_SharedLocalBasisIndexMacro,
@@ -10375,10 +10403,10 @@ int32 FHLSLMaterialTranslator::StrataVolumetricFogCloudBSDF(int32 Albedo, int32
 {
 	return AddCodeChunk(
 		MCT_Strata, TEXT("GetStrataVolumeFogCloudBSDF(%s, %s, %s, %s)"),
-		*GetParameterCode(Albedo),
-		*GetParameterCode(Extinction),
-		*GetParameterCode(EmissiveColor),
-		*GetParameterCode(AmbientOcclusion)
+		*StrataGetCastParameterCode(Albedo,				MCT_Float3),
+		*StrataGetCastParameterCode(Extinction,			MCT_Float3),
+		*StrataGetCastParameterCode(EmissiveColor,		MCT_Float3),
+		*StrataGetCastParameterCode(AmbientOcclusion,	MCT_Float)
 	);
 }
 
@@ -10386,8 +10414,8 @@ int32 FHLSLMaterialTranslator::StrataUnlitBSDF(int32 EmissiveColor, int32 Transm
 {
 	return AddCodeChunk(
 		MCT_Strata, TEXT("GetStrataUnlitBSDF(%s, %s)"),
-		*GetParameterCode(EmissiveColor),
-		*GetParameterCode(TransmittanceColor)
+		*StrataGetCastParameterCode(EmissiveColor,		MCT_Float3),
+		*StrataGetCastParameterCode(TransmittanceColor,	MCT_Float3)
 	);
 }
 
@@ -10395,12 +10423,12 @@ int32 FHLSLMaterialTranslator::StrataHairBSDF(int32 BaseColor, int32 Scatter, in
 {
 	return AddCodeChunk(
 		MCT_Strata, TEXT("GetStrataHairBSDF(%s, %s, %s, %s, %s, %s, %s) /* %s */"),
-		*GetParameterCode(BaseColor),
-		*GetParameterCode(Scatter),
-		*GetParameterCode(Specular),
-		*GetParameterCode(Roughness),
-		*GetParameterCode(Backlit),
-		*GetParameterCode(EmissiveColor),
+		*StrataGetCastParameterCode(BaseColor,			MCT_Float3),
+		*StrataGetCastParameterCode(Scatter,			MCT_Float),
+		*StrataGetCastParameterCode(Specular,			MCT_Float),
+		*StrataGetCastParameterCode(Roughness,			MCT_Float),
+		*StrataGetCastParameterCode(Backlit,			MCT_Float),
+		*StrataGetCastParameterCode(EmissiveColor,		MCT_Float3),
 		*SharedLocalBasisIndexMacro,
 		*GetParameterCode(Tangent)
 	);
@@ -10413,16 +10441,16 @@ int32 FHLSLMaterialTranslator::StrataSingleLayerWaterBSDF(
 {
 	return AddCodeChunk(
 		MCT_Strata, TEXT("PromoteParameterBlendedBSDFToOperator(GetStrataSingleLayerWaterBSDF(%s, %s, %s, %s, %s, %s, %s, %s, %s, %s, %s), Parameters.StrataTree, %u, %u, %u, %u) /* %s */"),
-		*GetParameterCode(BaseColor),
-		*GetParameterCode(Metallic),
-		*GetParameterCode(Specular),
-		*GetParameterCode(Roughness),
-		*GetParameterCode(EmissiveColor),
-		*GetParameterCode(TopMaterialOpacity),
-		*GetParameterCode(WaterAlbedo),
-		*GetParameterCode(WaterExtinction),
-		*GetParameterCode(WaterPhaseG),
-		*GetParameterCode(ColorScaleBehindWater),
+		*StrataGetCastParameterCode(BaseColor,				MCT_Float3),
+		*StrataGetCastParameterCode(Metallic,				MCT_Float),
+		*StrataGetCastParameterCode(Specular,				MCT_Float),
+		*StrataGetCastParameterCode(Roughness,				MCT_Float),
+		*StrataGetCastParameterCode(EmissiveColor,			MCT_Float3),
+		*StrataGetCastParameterCode(TopMaterialOpacity,		MCT_Float),
+		*StrataGetCastParameterCode(WaterAlbedo,			MCT_Float3),
+		*StrataGetCastParameterCode(WaterExtinction,		MCT_Float3),
+		*StrataGetCastParameterCode(WaterPhaseG,			MCT_Float),
+		*StrataGetCastParameterCode(ColorScaleBehindWater,	MCT_Float3),
 		*SharedLocalBasisIndexMacro,
 		PromoteToOperator->Index,
 		PromoteToOperator->BSDFIndex,
diff --git a/Engine/Source/Runtime/Engine/Private/Materials/HLSLMaterialTranslator.h b/Engine/Source/Runtime/Engine/Private/Materials/HLSLMaterialTranslator.h
--- a/Engine/Source/Runtime/Engine/Private/Materials/HLSLMaterialTranslator.h
+++ b/Engine/Source/Runtime/Engine/Private/Materials/HLSLMaterialTranslator.h
@@ -645,6 +645,7 @@ protected:
 
 	bool StrataGenerateDerivedMaterialOperatorData();
 	void StrataEvaluateSharedLocalBases(uint8& UsedSharedLocalBasesCount, FString* OutStrataPixelNormalInitializerValues, FShaderCompilerEnvironment* OutEnvironment);
+	FString StrataGetCastParameterCode(int32 Index, EMaterialValueType DestType);
 
 	// FMaterialCompiler interface.
 
