fixed AudioThread ResumeEvent not being completed on destruction (now this asserts in FTaskBase destructor)

diff --git a/Engine/Source/Runtime/Engine/Private/AudioThread.cpp b/Engine/Source/Runtime/Engine/Private/AudioThread.cpp
--- a/Engine/Source/Runtime/Engine/Private/AudioThread.cpp
+++ b/Engine/Source/Runtime/Engine/Private/AudioThread.cpp
@@ -125,7 +125,7 @@ double FAudioThread::LongestAudioThreadTimeMsec = 0.0;
 
 #if UE_AUDIO_THREAD_AS_PIPE
 
-UE::Tasks::FTaskEvent FAudioThread::ResumeEvent{ UE_SOURCE_LOCATION };
+TUniquePtr<UE::Tasks::FTaskEvent> FAudioThread::ResumeEvent;
 int32 FAudioThread::SuspendCount{ 0 };
 
 void FAudioThread::SuspendAudioThread()
@@ -152,8 +152,8 @@ void FAudioThread::SuspendAudioThread()
 			GIsAudioThreadSuspended.store(true, std::memory_order_release);
 			SuspendEvent->Trigger();
 
-			ResumeEvent.BusyWait(); // don't block one of workers while audio pipe is suspended
-			ResumeEvent = UE::Tasks::FTaskEvent{ UE_SOURCE_LOCATION}; // thread-safe because happens "inside" a suspend call and 
+			ResumeEvent->BusyWait(); // don't block one of workers while audio pipe is suspended
+			ResumeEvent = MakeUnique<UE::Tasks::FTaskEvent>(UE_SOURCE_LOCATION); // thread-safe because happens "inside" a suspend call and 
 			// can't happen concurrently with a resume call
 		}
 	);
@@ -182,8 +182,8 @@ void FAudioThread::ResumeAudioThread()
 	check(GIsAudioThreadSuspended.load(std::memory_order_relaxed));
 	GIsAudioThreadSuspended.store(false, std::memory_order_release);
 
-	check(!ResumeEvent.IsCompleted());
-	ResumeEvent.Trigger();
+	check(!ResumeEvent->IsCompleted());
+	ResumeEvent->Trigger();
 }
 
 #else // UE_AUDIO_THREAD_AS_PIPE
@@ -235,7 +235,7 @@ void FAudioThread::SuspendAudioThread()
 {
 	TRACE_CPUPROFILER_EVENT_SCOPE(SuspendAudioThread);
 	check(FPlatformTLS::GetCurrentThreadId() == GGameThreadId);
-	check(!GIsAudioThreadSuspended.Load() || GCVarSuspendAudioThread != 0);
+	check(!GIsAudioThreadSuspended.load() || GCVarSuspendAudioThread != 0);
 
 	if (IsAudioThreadRunning())
 	{
@@ -255,7 +255,7 @@ void FAudioThread::ResumeAudioThread()
 PRAGMA_DISABLE_DEPRECATION_WARNINGS
 
 	check(FPlatformTLS::GetCurrentThreadId() == GGameThreadId);
-	if(GIsAudioThreadSuspended.Load() && GCVarSuspendAudioThread == 0)
+	if(GIsAudioThreadSuspended.load() && GCVarSuspendAudioThread == 0)
 	{
 		GIsAudioThreadSuspended = false;
 		FPlatformMisc::MemoryBarrier();
@@ -630,6 +630,9 @@ void FAudioThread::StartAudioThread()
 	PreGCDestroy = FCoreUObjectDelegates::PreGarbageCollectConditionalBeginDestroy.AddStatic(&FAudioThread::SuspendAudioThread);
 	PostGCDestroy = FCoreUObjectDelegates::PostGarbageCollectConditionalBeginDestroy.AddStatic(&FAudioThread::ResumeAudioThread);
 
+	check(!ResumeEvent.IsValid());
+	ResumeEvent = MakeUnique<UE::Tasks::FTaskEvent>(UE_SOURCE_LOCATION);
+
 	GIsAudioThreadRunning.store(true, std::memory_order_release);
 }
 
@@ -650,6 +653,9 @@ void FAudioThread::StopAudioThread()
 	FCoreUObjectDelegates::PostGarbageCollectConditionalBeginDestroy.Remove(PostGCDestroy);
 
 	GIsAudioThreadRunning.store(false, std::memory_order_release);
+
+	check(ResumeEvent.IsValid());
+	ResumeEvent->Trigger(); // every FTaskEvent must be triggered before destruction to pass the check for completion
 }
 
 FAudioCommandFence::~FAudioCommandFence()
@@ -733,7 +739,7 @@ PRAGMA_ENABLE_DEPRECATION_WARNINGS
 void FAudioThread::StopAudioThread()
 {
 	check(FPlatformTLS::GetCurrentThreadId() == GGameThreadId);
-	check(!GIsAudioThreadSuspended.Load() || GCVarSuspendAudioThread != 0);
+	check(!GIsAudioThreadSuspended.load() || GCVarSuspendAudioThread != 0);
 
 	if (!IsAudioThreadRunning())
 	{
diff --git a/Engine/Source/Runtime/Engine/Public/AudioThread.h b/Engine/Source/Runtime/Engine/Public/AudioThread.h
--- a/Engine/Source/Runtime/Engine/Public/AudioThread.h
+++ b/Engine/Source/Runtime/Engine/Public/AudioThread.h
@@ -38,7 +38,7 @@ private:
 
 #if UE_AUDIO_THREAD_AS_PIPE
 
-	static UE::Tasks::FTaskEvent ResumeEvent;
+	static TUniquePtr<UE::Tasks::FTaskEvent> ResumeEvent;
 	static int32 SuspendCount;  // accessed only from GT
 
 	// GC callback handles
