Fixed old static analysis nullptr check - merging early.

diff --git a/Engine/Source/Runtime/Engine/Private/NetConnection.cpp b/Engine/Source/Runtime/Engine/Private/NetConnection.cpp
--- a/Engine/Source/Runtime/Engine/Private/NetConnection.cpp
+++ b/Engine/Source/Runtime/Engine/Private/NetConnection.cpp
@@ -3290,8 +3290,13 @@ void UNetConnection::ReceivedPacket( FBitReader& Reader, bool bIsReinjectedPacke
 
 				// Dispatch the raw, unsequenced bunch to the channel.
 				bool bLocalSkipAck = false;
-				//warning: May destroy channel.
-				Channel->ReceivedRawBunch( Bunch, bLocalSkipAck ); //-V595
+
+				if (Channel != nullptr)
+				{
+					// Warning: May destroy channel
+					Channel->ReceivedRawBunch(Bunch, bLocalSkipAck);
+				}
+
 				if ( bLocalSkipAck )
 				{
 					bSkipAck = true;
