Fixing static analysis warnings in AudioDevice.cpp

diff --git a/Engine/Source/Runtime/Engine/Private/AudioDevice.cpp b/Engine/Source/Runtime/Engine/Private/AudioDevice.cpp
--- a/Engine/Source/Runtime/Engine/Private/AudioDevice.cpp
+++ b/Engine/Source/Runtime/Engine/Private/AudioDevice.cpp
@@ -4187,8 +4187,10 @@ void FAudioDevice::StopSources(TArray<FWaveInstance*>& WaveInstances, int32 Firs
 	// being finished which might reset it being finished.
 	for (int32 InstanceIndex = 0; InstanceIndex < FirstActiveIndex; InstanceIndex++)
 	{
-		FWaveInstance* WaveInstance = WaveInstances[InstanceIndex];
-		WaveInstance->StopWithoutNotification();
+		if (FWaveInstance* WaveInstance = WaveInstances[InstanceIndex])
+		{
+			WaveInstance->StopWithoutNotification();
+		}
 	}
 
 #if ENABLE_AUDIO_DEBUG
