FRequestedType/FPreparedType store full Shader::FType, rather then just an EValueComponentType.  This allows correctly converting to either Double4x4 or DoubleInverse4x4.  This context is lost with only the component type

diff --git a/Engine/Source/Editor/MaterialEditor/Private/MaterialEditorUtilities.cpp b/Engine/Source/Editor/MaterialEditor/Private/MaterialEditorUtilities.cpp
--- a/Engine/Source/Editor/MaterialEditor/Private/MaterialEditorUtilities.cpp
+++ b/Engine/Source/Editor/MaterialEditor/Private/MaterialEditorUtilities.cpp
@@ -234,7 +234,23 @@ void FMaterialEditorUtilities::UpdateSearchResults(const class UEdGraph* Graph)
 
 void FMaterialEditorUtilities::GetVisibleMaterialParameters(const UMaterial* Material, UMaterialInstance* MaterialInstance, TArray<FMaterialParameterInfo>& VisibleExpressions)
 {
+	check(Material);
+	check(MaterialInstance);
+
 	VisibleExpressions.Empty();
+	if (Material->IsUsingNewHLSLGenerator())
+	{
+		// When using the new HLSL generator, MI parameter list will already have unused parameters culled
+		// We can assume that any remaining parameters are visible
+		TArray<FMaterialParameterInfo> ParameterInfo;
+		TArray<FGuid> ParamterGuid;
+		for (int32 TypeIndex = 0; TypeIndex < NumMaterialParameterTypes; ++TypeIndex)
+		{
+			MaterialInstance->GetAllParameterInfoOfType((EMaterialParameterType)TypeIndex, ParameterInfo, ParamterGuid);
+			VisibleExpressions.Append(ParameterInfo);
+		}
+		return;
+	}
 
 	TUniquePtr<FGetVisibleMaterialParametersFunctionState> FunctionState = MakeUnique<FGetVisibleMaterialParametersFunctionState>(nullptr);
 	TArray<FGetVisibleMaterialParametersFunctionState*> FunctionStack;
diff --git a/Engine/Source/Runtime/Engine/Private/HLSLTree/HLSLTree.cpp b/Engine/Source/Runtime/Engine/Private/HLSLTree/HLSLTree.cpp
--- a/Engine/Source/Runtime/Engine/Private/HLSLTree/HLSLTree.cpp
+++ b/Engine/Source/Runtime/Engine/Private/HLSLTree/HLSLTree.cpp
@@ -523,47 +523,36 @@ void FExpressionFunctionCall::EmitValuePreshader(FEmitContext& Context, FEmitSco
 	OutResult.Type = Function->OutputExpressions[OutputIndex]->GetValuePreshader(Context, Scope, RequestedType, OutResult.Preshader);
 }
 
-FRequestedType::FRequestedType(Shader::EValueType InType, bool bDefaultRequest)
+
+FRequestedType::FRequestedType(const Shader::FType& InType, bool bDefaultRequest) : Type(InType)
 {
-	const Shader::FValueTypeDescription TypeDesc = Shader::GetValueTypeDescription(InType);
-	ValueComponentType = TypeDesc.ComponentType;
 	if (bDefaultRequest)
 	{
-		RequestedComponents.Init(true, TypeDesc.NumComponents);
+		RequestedComponents.Init(true, InType.GetNumComponents());
 	}
 }
 
-FRequestedType::FRequestedType(const Shader::FType& InType, bool bDefaultRequest)
+FRequestedType::FRequestedType(const FRequestedType& InType, bool bDefaultRequest)
+	: Type(InType.Type)
+	, DefaultType(InType.DefaultType)
 {
-	if (InType.IsStruct())
-	{
-		StructType = InType.StructType;
-	}
-	else if (InType.IsObject())
-	{
-		ObjectType = InType.ObjectType;
-	}
-	else
-	{
-		const Shader::FValueTypeDescription TypeDesc = Shader::GetValueTypeDescription(InType);
-		ValueComponentType = TypeDesc.ComponentType;
-	}
 	if (bDefaultRequest)
 	{
-		RequestedComponents.Init(true, InType.GetNumComponents());
+		RequestedComponents = InType.RequestedComponents;
 	}
 }
 
-FRequestedType::FRequestedType(FName InObjectType) : ObjectType(InObjectType)
-{
-	RequestedComponents.Init(true, 1);
-}
-
-int32 FRequestedType::GetNumComponents() const
+int32 FRequestedType::GetNumComponentsRequested(int32 InDefaultNumComponents) const
 {
 	if (DefaultType == EDefaultRequestedType::Any)
 	{
-		return 4; // TODO?
+		return InDefaultNumComponents;
+	}
+
+	if (Type.IsVoid())
+	{
+		check(!RequestedComponents.Contains(true));
+		return 0;
 	}
 
 	const int32 MaxComponentIndex = RequestedComponents.FindLast(true);
@@ -571,25 +560,44 @@ int32 FRequestedType::GetNumComponents() const
 	{
 		return 0;
 	}
+
+	if (!Type.IsNumericVector())
+	{
+		return Type.GetNumComponents();
+	}
+
 	return MaxComponentIndex + 1;
 }
 
-Shader::FType FRequestedType::GetType() const
+Shader::FType FRequestedType::GetType(const Shader::FType& InDefaultType) const
 {
-	if (StructType)
+	if (DefaultType == EDefaultRequestedType::Any)
 	{
-		return StructType;
+		return InDefaultType;
 	}
-	else if (!ObjectType.IsNone())
+
+	const int32 NumComponents = GetNumComponentsRequested(0);
+	if (NumComponents == 0)
 	{
-		return ObjectType;
+		return Shader::EValueType::Void;
 	}
-	else if (DefaultType == EDefaultRequestedType::Any)
+
+	if (!Type.IsNumericVector())
 	{
-		return FName(TEXT("Any"));
+		return Type;
 	}
 
-	return Shader::MakeValueType(ValueComponentType, GetNumComponents());
+	const Shader::FValueTypeDescription TypeDesc = Shader::GetValueTypeDescription(Type);
+	return Shader::MakeValueType(TypeDesc.ComponentType, NumComponents);
+}
+
+Shader::EValueComponentType FRequestedType::GetValueComponentType() const
+{
+	if (Type.IsNumeric())
+	{
+		return Shader::GetValueTypeDescription(Type).ComponentType;
+	}
+	return Shader::EValueComponentType::Void;
 }
 
 void FRequestedType::SetComponentRequest(int32 Index, bool bRequested)
@@ -616,12 +624,27 @@ void FRequestedType::SetFieldRequested(const Shader::FStructField* Field, bool b
 	}
 }
 
-void FRequestedType::SetField(const Shader::FStructField* Field, const FRequestedType& InRequest)
+void FRequestedType::SetFieldRequested(const Shader::FStructField* Field, const FRequestedType& InRequest)
 {
-	const int32 NumComponents = Field->GetNumComponents();
-	for (int32 Index = 0; Index < NumComponents; ++Index)
+	const int32 NumFieldComponents = Field->GetNumComponents();
+	if (NumFieldComponents == 1)
 	{
-		SetComponentRequest(Field->ComponentIndex + Index, InRequest.IsComponentRequested(Index));
+		for (int32 Index = 0; Index < 4; ++Index)
+		{
+			// Scalar fields can replicate, so field is requested if we request any of xyzw
+			if (InRequest.IsComponentRequested(Index))
+			{
+				SetComponentRequest(Field->ComponentIndex, true);
+				break;
+			}
+		}
+	}
+	else
+	{
+		for (int32 Index = 0; Index < NumFieldComponents; ++Index)
+		{
+			SetComponentRequest(Field->ComponentIndex + Index, InRequest.IsComponentRequested(Index));
+		}
 	}
 }
 
@@ -694,22 +717,8 @@ FPreparedComponent CombineComponents(const FPreparedComponent& Lhs, const FPrepa
 	}
 }
 
-FPreparedType::FPreparedType(Shader::EValueType InType, const FPreparedComponent& InComponent)
-	: FPreparedType(Shader::FType(InType), InComponent)
+FPreparedType::FPreparedType(const Shader::FType& InType, const FPreparedComponent& InComponent) : Type(InType)
 {
-}
-
-FPreparedType::FPreparedType(const Shader::FType& InType, const FPreparedComponent& InComponent)
-{
-	if (InType.IsStruct())
-	{
-		StructType = InType.StructType;
-	}
-	else
-	{
-		ValueComponentType = Shader::GetValueTypeDescription(InType).ComponentType;
-	}
-
 	if (!InComponent.IsNone())
 	{
 		const int32 NumComponents = InType.GetNumComponents();
@@ -723,21 +732,26 @@ FPreparedType::FPreparedType(const Shader::FType& InType, const FPreparedCompone
 
 int32 FPreparedType::GetNumComponents() const
 {
-	if (StructType)
+	auto Predicate = [](const FPreparedComponent& InComponent) { return !InComponent.IsNone(); };
+
+	if (Type.IsVoid())
 	{
-		return StructType->ComponentTypes.Num();
+		check(!PreparedComponents.ContainsByPredicate(Predicate));
+		return 0;
 	}
-	else if (!ObjectType.IsNone())
+
+	const int32 MaxComponentIndex = PreparedComponents.FindLastByPredicate(Predicate);
+	if (MaxComponentIndex == INDEX_NONE)
 	{
-		return 1;
+		return 0;
 	}
-	else if (ValueComponentType != Shader::EValueComponentType::Void)
+
+	if (!Type.IsNumericVector())
 	{
-		auto Predicate = [](const FPreparedComponent& InComponent) { return !InComponent.IsNone(); };
-		const int32 MaxComponentIndex = PreparedComponents.FindLastByPredicate(Predicate);
-		return (MaxComponentIndex != INDEX_NONE) ? MaxComponentIndex + 1 : 1;
+		return Type.GetNumComponents();
 	}
-	return 0;
+
+	return MaxComponentIndex + 1;
 }
 
 bool FPreparedType::IsVoid() const
@@ -747,25 +761,35 @@ bool FPreparedType::IsVoid() const
 
 Shader::FType FPreparedType::GetType() const
 {
-	if (IsStruct())
+	const int32 NumComponents = GetNumComponents();
+	if (NumComponents == 0)
 	{
-		return StructType;
+		return Shader::EValueType::Void;
 	}
-	else if (IsObject())
+
+	if (!Type.IsNumericVector())
 	{
-		return ObjectType;
+		return Type;
 	}
 
-	return Shader::MakeValueType(ValueComponentType, GetNumComponents());
+	const Shader::FValueTypeDescription TypeDesc = Shader::GetValueTypeDescription(Type);
+	return Shader::MakeValueType(TypeDesc.ComponentType, NumComponents);
+}
+
+Shader::EValueComponentType FPreparedType::GetValueComponentType() const
+{
+	if (Type.IsNumeric())
+	{
+		return Shader::GetValueTypeDescription(Type).ComponentType;
+	}
+	return Shader::EValueComponentType::Void;
 }
 
 FRequestedType FPreparedType::GetRequestedType() const
 {
 	const int32 NumComponents = GetNumComponents();
 	FRequestedType Result;
-	Result.ValueComponentType = ValueComponentType;
-	Result.StructType = StructType;
-	Result.ObjectType = ObjectType;
+	Result.Type = Type;
 	Result.RequestedComponents.Init(false, NumComponents);
 	for (int32 Index = 0; Index < NumComponents; ++Index)
 	{
@@ -791,7 +815,7 @@ EExpressionEvaluation FPreparedType::GetEvaluation(const FEmitScope& Scope) cons
 EExpressionEvaluation FPreparedType::GetEvaluation(const FEmitScope& Scope, const FRequestedType& RequestedType) const
 {
 	EExpressionEvaluation Result = EExpressionEvaluation::ConstantZero;
-	for (int32 Index = 0; Index < RequestedType.GetNumComponents(); ++Index)
+	for (int32 Index = 0; Index < RequestedType.GetNumComponentsRequested(PreparedComponents.Num()); ++Index)
 	{
 		if (RequestedType.IsComponentRequested(Index))
 		{
@@ -909,9 +933,9 @@ void FPreparedType::MergeEvaluation(EExpressionEvaluation Evaluation)
 
 void FPreparedType::SetLoopEvaluation(FEmitScope& Scope, const FRequestedType& RequestedType)
 {
-	const int32 NumRequestedComponents = RequestedType.GetNumComponents();
 	const int32 NumComponents = GetNumComponents();
-	
+	const int32 NumRequestedComponents = RequestedType.GetNumComponentsRequested(NumComponents);
+
 	if (NumComponents == 1 && FMath::IsWithinInclusive(NumRequestedComponents, 1, 4))
 	{
 		// If we're a scalar type, set loop evaluation if any of xyzw are requested
@@ -941,7 +965,7 @@ Shader::FComponentBounds FPreparedType::GetComponentBounds(int32 Index) const
 		return Shader::FComponentBounds(Shader::EComponentBound::Zero, Shader::EComponentBound::Zero);
 	}
 
-	const Shader::EValueComponentType ComponentType = StructType ? StructType->ComponentTypes[Index] : ValueComponentType;
+	const Shader::EValueComponentType ComponentType = GetType().GetComponentType(Index);
 	check(ComponentType != Shader::EValueComponentType::Void);
 
 	const Shader::FValueComponentTypeDescription ComponentTypeDesc = Shader::GetValueComponentTypeDescription(ComponentType);
@@ -952,8 +976,11 @@ Shader::FComponentBounds FPreparedType::GetComponentBounds(int32 Index) const
 
 Shader::FComponentBounds FPreparedType::GetBounds(const FRequestedType& RequestedType) const
 {
+	const int32 NumComponents = GetNumComponents();
+	const int32 NumRequestedComponents = RequestedType.GetNumComponentsRequested(NumComponents);
+
 	Shader::FComponentBounds Result(Shader::EComponentBound::DoubleMax, Shader::EComponentBound::NegDoubleMax);
-	for (int32 Index = 0; Index < RequestedType.GetNumComponents(); ++Index)
+	for (int32 Index = 0; Index < NumRequestedComponents; ++Index)
 	{
 		if (RequestedType.IsComponentRequested(Index))
 		{
@@ -967,7 +994,8 @@ Shader::FComponentBounds FPreparedType::GetBounds(const FRequestedType& Requeste
 
 void FPreparedType::SetField(const Shader::FStructField* Field, const FPreparedType& FieldType)
 {
-	for (int32 Index = 0; Index < Field->GetNumComponents(); ++Index)
+	const int32 NumFieldComponents = Field->GetNumComponents();
+	for (int32 Index = 0; Index < NumFieldComponents; ++Index)
 	{
 		SetComponent(Field->ComponentIndex + Index, FieldType.GetComponent(Index));
 	}
@@ -997,20 +1025,23 @@ FPreparedType MergePreparedTypes(const FPreparedType& Lhs, const FPreparedType&
 
 	int32 NumComponents = 0;
 	FPreparedType Result;
-	if (Lhs.IsStruct() || Rhs.IsStruct())
+	if (Lhs.IsNumeric() && Rhs.IsNumeric())
 	{
-		if (Lhs.StructType != Rhs.StructType)
-		{
-			// Mismatched structs
-			return Result;
-		}
-		Result.StructType = Lhs.StructType;
-		NumComponents = Result.StructType->ComponentTypes.Num();
+		const Shader::FValueTypeDescription LhsDesc = Shader::GetValueTypeDescription(Lhs.Type);
+		const Shader::FValueTypeDescription RhsDesc = Shader::GetValueTypeDescription(Rhs.Type);
+		const Shader::EValueComponentType ComponentType = Shader::CombineComponentTypes(LhsDesc.ComponentType, RhsDesc.ComponentType);
+		NumComponents = FMath::Max(Lhs.GetNumComponents(), Rhs.GetNumComponents());
+		Result.Type = Shader::MakeValueType(ComponentType, NumComponents);
 	}
 	else
 	{
-		Result.ValueComponentType = Shader::CombineComponentTypes(Lhs.ValueComponentType, Rhs.ValueComponentType);
-		NumComponents = FMath::Max(Lhs.GetNumComponents(), Rhs.GetNumComponents());
+		if (Lhs.Type != Rhs.Type)
+		{
+			// Mismatched non-numeric
+			return Result;
+		}
+		Result.Type = Lhs.Type;
+		NumComponents = Result.Type.GetNumComponents();
 	}
 
 	Result.PreparedComponents.Reset(NumComponents);
@@ -1029,18 +1060,14 @@ FPreparedType MakeNonLWCType(const FPreparedType& Type)
 	FPreparedType Result(Type);
 	if (Result.IsNumeric())
 	{
-		Result.ValueComponentType = Shader::MakeNonLWCType(Result.ValueComponentType);
+		Result.Type = Shader::MakeNonLWCType(Result.Type);
 	}
 	return Result;
 }
 
-bool FPrepareValueResult::TryMergePreparedType(FEmitContext& Context,
-	const Shader::FStructType* StructType,
-	FName ObjectType,
-	Shader::EValueComponentType ComponentType,
-	int32 NumComponents)
+bool FPrepareValueResult::TryMergePreparedType(FEmitContext& Context, const Shader::FType& Type)
 {
-	if (!StructType && ObjectType.IsNone() && ComponentType == Shader::EValueComponentType::Void)
+	if (Type.IsVoid())
 	{
 		return false;
 	}
@@ -1048,39 +1075,32 @@ bool FPrepareValueResult::TryMergePreparedType(FEmitContext& Context,
 	if (!PreparedType.IsInitialized())
 	{
 		PreparedType.PreparedComponents.Reset();
-		PreparedType.ValueComponentType = ComponentType;
-		PreparedType.ObjectType = ObjectType;
-		PreparedType.StructType = StructType;
+		PreparedType.Type = Type;
 		return true;
 	}
 
-	if (StructType)
-	{
-		check(ComponentType == Shader::EValueComponentType::Void);
-		check(ObjectType.IsNone());
-		if (StructType != PreparedType.StructType)
-		{
-			return Context.Error(TEXT("Invalid type"));
-		}
-	}
-	else if (!ObjectType.IsNone())
+	if (!Type.IsNumericVector())
 	{
-		check(ComponentType == Shader::EValueComponentType::Void);
-		if (ObjectType != PreparedType.ObjectType)
+		if (Type != PreparedType.Type)
 		{
 			return Context.Error(TEXT("Invalid type"));
 		}
 	}
 	else
 	{
+		const Shader::FValueTypeDescription TypeDesc = Shader::GetValueTypeDescription(Type);
 		const int32 PrevNumComponents = PreparedType.GetNumComponents();
-		if (PrevNumComponents == 1 && NumComponents > 1)
+		if (PrevNumComponents == 1 && TypeDesc.NumComponents > 1)
 		{
 			// Once an expression is prepared as a scalar type, it's not valid to change later
 			// Other expressions may have already assumed this will be a scalar type, and its components will replicate
 			return Context.Error(TEXT("Type was already prepared as scalar"));
 		}
-		PreparedType.ValueComponentType = Shader::CombineComponentTypes(PreparedType.ValueComponentType, ComponentType);
+
+		const Shader::FValueTypeDescription PrevTypeDesc = Shader::GetValueTypeDescription(PreparedType.Type);
+		const Shader::EValueComponentType ComponentType = Shader::CombineComponentTypes(PrevTypeDesc.ComponentType, TypeDesc.ComponentType);
+		const int32 NumComponents = FMath::Max<int32>(PrevNumComponents, TypeDesc.NumComponents);
+		PreparedType.Type = Shader::MakeValueType(ComponentType, NumComponents);
 	}
 
 	return true;
@@ -1089,20 +1109,19 @@ bool FPrepareValueResult::TryMergePreparedType(FEmitContext& Context,
 bool FPrepareValueResult::SetTypeVoid()
 {
 	PreparedType.PreparedComponents.Reset();
-	PreparedType.ValueComponentType = Shader::EValueComponentType::Void;
-	PreparedType.ObjectType = FName();
-	PreparedType.StructType = nullptr;
+	PreparedType.Type = Shader::EValueType::Void;
 	return false;
 }
 
 bool FPrepareValueResult::SetType(FEmitContext& Context, const FRequestedType& RequestedType, EExpressionEvaluation Evaluation, const Shader::FType& Type)
 {
-	const int32 NumComponents = Type.GetNumComponents();
-	if (TryMergePreparedType(Context, Type.StructType, Type.ObjectType, Shader::GetValueTypeDescription(Type.ValueType).ComponentType, NumComponents))
+	if (TryMergePreparedType(Context, Type))
 	{
 		if (Evaluation != EExpressionEvaluation::None)
 		{
-			const int32 NumRequestedComponents = RequestedType.GetNumComponents();
+			const int32 NumComponents = Type.GetNumComponents();
+			const int32 NumRequestedComponents = RequestedType.GetNumComponentsRequested(NumComponents);
+			
 			if (NumComponents == 1 && NumRequestedComponents > 0 && NumRequestedComponents <= 4)
 			{
 				PreparedType.MergeComponent(0, Evaluation);
@@ -1128,31 +1147,32 @@ bool FPrepareValueResult::SetType(FEmitContext& Context, const FRequestedType& R
 	return false;
 }
 
-bool FPrepareValueResult::SetType(FEmitContext& Context, const FRequestedType& RequestedType, const FPreparedType& Type)
+bool FPrepareValueResult::SetType(FEmitContext& Context, const FRequestedType& RequestedType, const FPreparedType& InPreparedType)
 {
-	Shader::EValueComponentType ComponentType = Type.ValueComponentType;
-	if (ComponentType == Shader::EValueComponentType::Double)
+	Shader::FType Type = InPreparedType.GetType();
+	if (Type.IsNumericLWC())
 	{
-		const Shader::FComponentBounds Bounds = Type.GetBounds(RequestedType);
+		const Shader::FComponentBounds Bounds = InPreparedType.GetBounds(RequestedType);
 		if (Shader::IsWithinBounds(Bounds, Shader::GetValueComponentTypeDescription(Shader::EValueComponentType::Float).Bounds))
 		{
-			ComponentType = Shader::EValueComponentType::Float;
+			Type = Shader::MakeNonLWCType(Type);
 		}
 	}
 
-	const int32 NumComponents = Type.GetNumComponents();
-	if (TryMergePreparedType(Context, Type.StructType, Type.ObjectType, ComponentType, NumComponents))
+	if (TryMergePreparedType(Context, Type))
 	{
-		const int32 NumRequestedComponents = RequestedType.GetNumComponents();
+		const int32 NumComponents = Type.GetNumComponents();
+		const int32 NumRequestedComponents = RequestedType.GetNumComponentsRequested(NumComponents);
+		
 		if (NumComponents == 1 && NumRequestedComponents > 0 && NumRequestedComponents <= 4)
 		{
-			PreparedType.MergeComponent(0, Type.GetComponent(0));
+			PreparedType.MergeComponent(0, InPreparedType.GetComponent(0));
 		}
 		else
 		{
 			for (int32 Index = 0; Index < NumComponents; ++Index)
 			{
-				const FPreparedComponent& Component = Type.GetComponent(Index);
+				const FPreparedComponent& Component = InPreparedType.GetComponent(Index);
 				if (RequestedType.IsComponentRequested(Index))
 				{
 					PreparedType.MergeComponent(Index, Component);
@@ -1376,14 +1396,14 @@ Shader::FValue FExpression::GetValueConstant(FEmitContext& Context, FEmitScope&
 bool FExpression::GetValueObject(FEmitContext& Context, FEmitScope& Scope, const FName& ObjectTypeName, void* OutObjectBase) const
 {
 	const FPreparedType PreparedType = Context.GetPreparedType(this);
-	check(PreparedType.ObjectType == ObjectTypeName);
+	check(PreparedType.Type.ObjectType == ObjectTypeName);
 	return EmitValueObject(Context, Scope, ObjectTypeName, OutObjectBase);
 }
 
 bool FExpression::CheckObjectSupportsCustomHLSL(FEmitContext& Context, FEmitScope& Scope, const FName& ObjectTypeName) const
 {
 	const FPreparedType PreparedType = Context.GetPreparedType(this);
-	check(PreparedType.ObjectType == ObjectTypeName);
+	check(PreparedType.Type.ObjectType == ObjectTypeName);
 
 	FEmitCustomHLSLParameterResult UnusedResult;
 	return EmitCustomHLSLParameter(Context, Scope, ObjectTypeName, nullptr, UnusedResult);
diff --git a/Engine/Source/Runtime/Engine/Private/HLSLTree/HLSLTreeCommon.cpp b/Engine/Source/Runtime/Engine/Private/HLSLTree/HLSLTreeCommon.cpp
--- a/Engine/Source/Runtime/Engine/Private/HLSLTree/HLSLTreeCommon.cpp
+++ b/Engine/Source/Runtime/Engine/Private/HLSLTree/HLSLTreeCommon.cpp
@@ -166,17 +166,17 @@ void FExpressionGetStructField::ComputeAnalyticDerivatives(FTree& Tree, FExpress
 const FExpression* FExpressionGetStructField::ComputePreviousFrame(FTree& Tree, const FRequestedType& RequestedType) const
 {
 	FRequestedType RequestedStructType(StructType, false);
-	RequestedStructType.SetField(Field, RequestedType);
+	RequestedStructType.SetFieldRequested(Field, RequestedType);
 	return Tree.NewExpression<FExpressionGetStructField>(StructType, Field, Tree.GetPreviousFrame(StructExpression, RequestedStructType));
 }
 
 bool FExpressionGetStructField::PrepareValue(FEmitContext& Context, FEmitScope& Scope, const FRequestedType& RequestedType, FPrepareValueResult& OutResult) const
 {
 	FRequestedType RequestedStructType(StructType, false);
-	RequestedStructType.SetField(Field, RequestedType);
+	RequestedStructType.SetFieldRequested(Field, RequestedType);
 
 	FPreparedType StructPreparedType = Context.PrepareExpression(StructExpression, Scope, RequestedStructType);
-	if (!StructPreparedType.IsVoid() && StructPreparedType.StructType != StructType)
+	if (!StructPreparedType.IsVoid() && StructPreparedType.Type.StructType != StructType)
 	{
 		return Context.Errorf(TEXT("Expected type %s"), StructType->Name);
 	}
@@ -187,7 +187,7 @@ bool FExpressionGetStructField::PrepareValue(FEmitContext& Context, FEmitScope&
 void FExpressionGetStructField::EmitValueShader(FEmitContext& Context, FEmitScope& Scope, const FRequestedType& RequestedType, FEmitValueShaderResult& OutResult) const
 {
 	FRequestedType RequestedStructType(StructType, false);
-	RequestedStructType.SetField(Field, RequestedType);
+	RequestedStructType.SetFieldRequested(Field, RequestedType);
 
 	FEmitShaderExpression* StructValue = StructExpression->GetValueShader(Context, Scope, RequestedStructType);
 
@@ -199,7 +199,7 @@ void FExpressionGetStructField::EmitValueShader(FEmitContext& Context, FEmitScop
 void FExpressionGetStructField::EmitValuePreshader(FEmitContext& Context, FEmitScope& Scope, const FRequestedType& RequestedType, FEmitValuePreshaderResult& OutResult) const
 {
 	FRequestedType RequestedStructType(StructType, false);
-	RequestedStructType.SetField(Field, RequestedType);
+	RequestedStructType.SetFieldRequested(Field, RequestedType);
 
 	StructExpression->GetValuePreshader(Context, Scope, RequestedStructType, OutResult.Preshader);
 	OutResult.Preshader.WriteOpcode(Shader::EPreshaderOpcode::GetField).Write(Field->Type).Write(Field->ComponentIndex);
@@ -241,7 +241,7 @@ bool FExpressionSetStructField::PrepareValue(FEmitContext& Context, FEmitScope&
 	RequestedStructType.ClearFieldRequested(Field);
 
 	FPreparedType StructPreparedType = Context.PrepareExpression(StructExpression, Scope, RequestedStructType);
-	if (!StructPreparedType.IsVoid() && StructPreparedType.StructType != StructType)
+	if (!StructPreparedType.IsVoid() && StructPreparedType.Type.StructType != StructType)
 	{
 		return Context.Errorf(TEXT("Expected type %s"), StructType->Name);
 	}
@@ -342,14 +342,17 @@ bool FExpressionSelect::PrepareValue(FEmitContext& Context, FEmitScope& Scope, c
 
 	const FPreparedType& LhsType = Context.PrepareExpression(FalseExpression, Scope, RequestedType);
 	const FPreparedType& RhsType = Context.PrepareExpression(TrueExpression, Scope, RequestedType);
+	if (LhsType.IsVoid() || RhsType.IsVoid())
+	{
+		return false;
+	}
 	
-	if (LhsType.ValueComponentType != RhsType.ValueComponentType ||
-		LhsType.StructType != RhsType.StructType)
+	FPreparedType ResultType = MergePreparedTypes(LhsType, RhsType);
+	if (ResultType.IsVoid())
 	{
 		return Context.Error(TEXT("Type mismatch"));
 	}
 
-	FPreparedType ResultType = MergePreparedTypes(LhsType, RhsType);
 	ResultType.MergeEvaluation(ConditionEvaluation);
 	ResultType.MergeEvaluation(EExpressionEvaluation::Shader); // TODO - support preshader
 	return OutResult.SetType(Context, RequestedType, ResultType);
@@ -404,7 +407,7 @@ bool FExpressionDerivative::PrepareValue(FEmitContext& Context, FEmitScope& Scop
 		return false;
 	}
 
-	ResultType.ValueComponentType = Shader::MakeNonLWCType(ResultType.ValueComponentType);
+	ResultType.Type = Shader::MakeNonLWCType(ResultType.Type);
 
 	const EExpressionEvaluation InputEvaluation = ResultType.GetEvaluation(Scope, RequestedType);
 	if (InputEvaluation != EExpressionEvaluation::Shader)
@@ -440,9 +443,8 @@ namespace Private
 {
 FRequestedType GetRequestedSwizzleType(const FSwizzleParameters& Params, const FRequestedType& RequestedType)
 {
-	FRequestedType RequestedInputType;
-	RequestedInputType.ValueComponentType = RequestedType.ValueComponentType;
-	for (int32 Index = 0; Index < RequestedType.GetNumComponents(); ++Index)
+	FRequestedType RequestedInputType(RequestedType, false);
+	for (int32 Index = 0; Index < RequestedType.GetNumComponentsRequested(Params.NumComponents); ++Index)
 	{
 		if (RequestedType.IsComponentRequested(Index))
 		{
@@ -494,7 +496,7 @@ bool SwizzlePrepareValue(FEmitContext& Context,
 	if (RequestedInputType.IsVoid())
 	{
 		// All the requested components are outside the swizzle, so just return 0
-		ResultType.ValueComponentType = RequestedType.ValueComponentType;
+		ResultType.Type = RequestedType.Type;
 		for (int32 ComponentIndex = 0; ComponentIndex < Parameters.NumComponents; ++ComponentIndex)
 		{
 			ResultType.SetComponent(ComponentIndex, EExpressionEvaluation::ConstantZero);
@@ -504,7 +506,7 @@ bool SwizzlePrepareValue(FEmitContext& Context,
 	{
 		const FPreparedType& InputType = Context.PrepareExpression(Input, Scope, RequestedInputType);
 
-		ResultType.ValueComponentType = InputType.ValueComponentType;
+		ResultType.Type = InputType.Type;
 		for (int32 ComponentIndex = 0; ComponentIndex < Parameters.NumComponents; ++ComponentIndex)
 		{
 			if (RequestedType.IsComponentRequested(ComponentIndex))
@@ -695,17 +697,19 @@ struct FAppendTypes
 };
 FAppendTypes GetAppendTypes(const FRequestedType& RequestedType, const FPreparedType& LhsType, const FPreparedType& RhsType)
 {
-	const Shader::EValueComponentType ComponentType = Shader::CombineComponentTypes(LhsType.ValueComponentType, RhsType.ValueComponentType);
+	const Shader::FValueTypeDescription LhsTypeDesc = Shader::GetValueTypeDescription(LhsType.Type);
+	const Shader::FValueTypeDescription RhsTypeDesc = Shader::GetValueTypeDescription(RhsType.Type);
+	const Shader::EValueComponentType ComponentType = Shader::CombineComponentTypes(LhsTypeDesc.ComponentType, RhsTypeDesc.ComponentType);
 	const int32 NumLhsComponents = LhsType.GetNumComponents();
 	const int32 NumRhsComponents = RhsType.GetNumComponents();
 	const int32 NumComponents = FMath::Min(NumLhsComponents + NumRhsComponents, 4);
 
 	FAppendTypes Types;
-	Types.ResultType.ValueComponentType = ComponentType;
-	Types.LhsRequestedType.ValueComponentType = ComponentType;
-	Types.RhsRequestedType.ValueComponentType = ComponentType;
+	Types.ResultType.Type = Shader::MakeValueType(ComponentType, NumComponents);
 	Types.LhsType = Shader::MakeValueType(ComponentType, NumLhsComponents);
 	Types.RhsType = Shader::MakeValueType(ComponentType, NumComponents - NumLhsComponents);
+	Types.LhsRequestedType.Type = Types.LhsType;
+	Types.RhsRequestedType.Type = Types.RhsType;
 	Types.bIsLWC = ComponentType == Shader::EValueComponentType::Double;
 
 	for (int32 Index = 0; Index < NumLhsComponents; ++Index)
@@ -751,9 +755,8 @@ bool FExpressionAppend::PrepareValue(FEmitContext& Context, FEmitScope& Scope, c
 	const FPreparedType& LhsType = Context.PrepareExpression(Lhs, Scope, RequestedType);
 	const int32 NumLhsComponents = LhsType.GetNumComponents();
 
-	FRequestedType RhsRequestedType;
-	RhsRequestedType.ValueComponentType = RequestedType.ValueComponentType;
-	for (int32 Index = NumLhsComponents; Index < RequestedType.GetNumComponents(); ++Index)
+	FRequestedType RhsRequestedType(RequestedType, false);
+	for (int32 Index = NumLhsComponents; Index < RequestedType.GetNumComponentsRequested(4 - NumLhsComponents); ++Index)
 	{
 		RhsRequestedType.SetComponentRequest(Index - NumLhsComponents, RequestedType.IsComponentRequested(Index));
 	}
@@ -762,7 +765,7 @@ bool FExpressionAppend::PrepareValue(FEmitContext& Context, FEmitScope& Scope, c
 	if (!RhsRequestedType.IsVoid())
 	{
 		RhsType = Context.PrepareExpression(Rhs, Scope, RhsRequestedType);
-		if (LhsType.ValueComponentType != RhsType.ValueComponentType)
+		if (LhsType.GetValueComponentType() != RhsType.GetValueComponentType())
 		{
 			return Context.Error(TEXT("Type mismatch"));
 		}
@@ -771,7 +774,7 @@ bool FExpressionAppend::PrepareValue(FEmitContext& Context, FEmitScope& Scope, c
 	{
 		// We never want the result of an 'Append' operation to be a scalar value
 		// Ensure that we include at least a single component for the rhs, even if it's not currently requested
-		RhsType.ValueComponentType = RequestedType.ValueComponentType;
+		RhsType.Type = RequestedType.Type;
 		RhsType.SetComponent(0, EExpressionEvaluation::ConstantZero);
 	}
 
@@ -942,7 +945,7 @@ bool FExpressionCustomHLSL::PrepareValue(FEmitContext& Context, FEmitScope& Scop
 
 		if (InputType.IsObject())
 		{
-			if (!Input.Expression->CheckObjectSupportsCustomHLSL(Context, Scope, InputType.ObjectType))
+			if (!Input.Expression->CheckObjectSupportsCustomHLSL(Context, Scope, InputType.Type.ObjectType))
 			{
 				return Context.Error(TEXT("Object type not supported for custom HLSL"));
 			}
diff --git a/Engine/Source/Runtime/Engine/Private/HLSLTree/HLSLTreeEmit.cpp b/Engine/Source/Runtime/Engine/Private/HLSLTree/HLSLTreeEmit.cpp
--- a/Engine/Source/Runtime/Engine/Private/HLSLTree/HLSLTreeEmit.cpp
+++ b/Engine/Source/Runtime/Engine/Private/HLSLTree/HLSLTreeEmit.cpp
@@ -414,8 +414,6 @@ FPreparedType FEmitContext::PrepareExpression(const FExpression* InExpression, F
 		return Result->PreparedType;
 	}
 
-	MarkInputType(InExpression, RequestedType.GetType());
-
 	bool bResult = false;
 	{
 		FEmitOwnerScope OwnerScope(*this, InExpression);
@@ -430,6 +428,7 @@ FPreparedType FEmitContext::PrepareExpression(const FExpression* InExpression, F
 	{
 		ResultType = Result->PreparedType;
 		check(!ResultType.IsVoid());
+		MarkInputType(InExpression, RequestedType.GetType(ResultType.GetType()));
 	}
 	return ResultType;
 }
@@ -1178,6 +1177,9 @@ FEmitShaderExpression* FEmitContext::EmitCast(FEmitScope& Scope, FEmitShaderExpr
 		{
 			const bool bReplicateScalar = (SourceTypeDesc.NumComponents == 1) && !EnumHasAnyFlags(Flags, EEmitCastFlags::ZeroExtendScalar);
 
+			check(SourceTypeDesc.NumComponents <= 4);
+			check(DestTypeDesc.NumComponents <= 4);
+
 			int32 NumComponents = 0;
 			bool bNeedClosingParen = false;
 			if (bIsLWC)
diff --git a/Engine/Source/Runtime/Engine/Private/HLSLTree/HLSLTreeOperation.cpp b/Engine/Source/Runtime/Engine/Private/HLSLTree/HLSLTreeOperation.cpp
--- a/Engine/Source/Runtime/Engine/Private/HLSLTree/HLSLTreeOperation.cpp
+++ b/Engine/Source/Runtime/Engine/Private/HLSLTree/HLSLTreeOperation.cpp
@@ -161,24 +161,26 @@ FOperationRequestedTypes GetOperationRequestedTypes(EOperation Op, const FReques
 	case EOperation::Length:
 	case EOperation::Normalize:
 	case EOperation::Sum:
-		Types.InputType[0] = Shader::MakeValueType(RequestedType.ValueComponentType, 4);
+		Types.InputType[0] = Shader::MakeValueType(RequestedType.GetValueComponentType(), 4);
 		break;
 	case EOperation::VecMulMatrix3:
 		Types.bIsMatrixOperation = true;
-		Types.InputType[0] = Shader::MakeValueType(RequestedType.ValueComponentType, 3);
-		Types.InputType[1] = Shader::MakeValueType(RequestedType.ValueComponentType, 16);
+		Types.InputType[0] = Shader::EValueType::Float3;
+		Types.InputType[1] = Shader::EValueType::Float4x4;
 		break;
 	case EOperation::VecMulMatrix4:
+		// float3 * FLWCMatrix -> FLWCVector3
+		// Note that we'll never explicitly request the 'FLWCVector3 * FLWCInverseMatrix -> float3' case
 		Types.bIsMatrixOperation = true;
-		Types.InputType[0] = Shader::MakeValueType(RequestedType.ValueComponentType, 3);
-		Types.InputType[1] = Shader::MakeValueType(RequestedType.ValueComponentType, 16);
+		Types.InputType[0] = Shader::EValueType::Float3;
+		Types.InputType[1] = Shader::MakeValueType(RequestedType.GetValueComponentType(), 16);
 		break;
 	case EOperation::Matrix3MulVec:
 	case EOperation::Matrix4MulVec:
 		// No LWC for transpose matrices
 		Types.bIsMatrixOperation = true;
-		Types.InputType[0] = Shader::MakeValueType(RequestedType.ValueComponentType, 16);
-		Types.InputType[1] = Shader::MakeValueType(RequestedType.ValueComponentType, 3);
+		Types.InputType[0] = Shader::EValueType::Float4x4;
+		Types.InputType[1] = Shader::EValueType::Float3;
 		break;
 	default:
 		break;
@@ -195,11 +197,9 @@ FOperationTypes GetOperationTypes(EOperation Op, TConstArrayView<FPreparedType>
 		Op == EOperation::Matrix4MulVec)
 	{
 		FPreparedComponent IntermediateComponent;
-		Shader::EValueComponentType IntermediateComponentType = Shader::EValueComponentType::Void;
 		for (int32 Index = 0; Index < InputPreparedType.Num(); ++Index)
 		{
 			IntermediateComponent = CombineComponents(IntermediateComponent, InputPreparedType[Index].GetMergedComponent());
-			IntermediateComponentType = Shader::CombineComponentTypes(IntermediateComponentType, InputPreparedType[Index].ValueComponentType);
 		}
 
 		switch (Op)
@@ -211,9 +211,27 @@ FOperationTypes GetOperationTypes(EOperation Op, TConstArrayView<FPreparedType>
 			Types.ResultType = FPreparedType(Shader::EValueType::Float3, IntermediateComponent);
 			break;
 		case EOperation::VecMulMatrix4:
-			Types.InputType[0] = Shader::MakeValueType(IntermediateComponentType, 3);
-			Types.InputType[1] = Shader::MakeValueType(IntermediateComponentType, 16);
-			Types.ResultType = FPreparedType(Shader::MakeValueType(IntermediateComponentType, 3), IntermediateComponent);
+			switch (InputPreparedType[1].GetType().ValueType)
+			{
+			case Shader::EValueType::Double4x4:
+				// float3 * FLWCMatrix -> FLWCVector3
+				Types.InputType[0] = Shader::EValueType::Float3;
+				Types.InputType[1] = Shader::EValueType::Double4x4;
+				Types.ResultType = FPreparedType(Shader::EValueType::Double3, IntermediateComponent);
+				break;
+			case Shader::EValueType::DoubleInverse4x4:
+				// FLWCVector3 * FLWCInverseMatrix -> float3
+				Types.InputType[0] = Shader::EValueType::Double3;
+				Types.InputType[1] = Shader::EValueType::DoubleInverse4x4;
+				Types.ResultType = FPreparedType(Shader::EValueType::Float3, IntermediateComponent);
+				break;
+			default:
+				// float3 * float4x4 -> float3
+				Types.InputType[0] = Shader::EValueType::Float3;
+				Types.InputType[1] = Shader::EValueType::Float4x4;
+				Types.ResultType = FPreparedType(Shader::EValueType::Float3, IntermediateComponent);
+				break;
+			}
 			break;
 		case EOperation::Matrix3MulVec:
 		case EOperation::Matrix4MulVec:
@@ -232,11 +250,11 @@ FOperationTypes GetOperationTypes(EOperation Op, TConstArrayView<FPreparedType>
 			IntermediateType = MergePreparedTypes(IntermediateType, InputPreparedType[Index]);
 		}
 
-		const Shader::EValueComponentType IntermediateComponentType = IntermediateType.ValueComponentType;
+		const Shader::EValueComponentType IntermediateComponentType = IntermediateType.GetValueComponentType();
 		const int32 NumIntermediateComponents = IntermediateType.GetNumComponents();
 		for (int32 Index = 0; Index < InputPreparedType.Num(); ++Index)
 		{
-			Shader::EValueComponentType InputComponentType = InputPreparedType[Index].ValueComponentType;
+			Shader::EValueComponentType InputComponentType = InputPreparedType[Index].GetValueComponentType();
 			if (InputComponentType != IntermediateComponentType)
 			{
 				switch (IntermediateComponentType)
@@ -305,7 +323,7 @@ FOperationTypes GetOperationTypes(EOperation Op, TConstArrayView<FPreparedType>
 		case EOperation::Greater:
 		case EOperation::LessEqual:
 		case EOperation::GreaterEqual:
-			Types.ResultType.ValueComponentType = Shader::EValueComponentType::Bool;
+			Types.ResultType.Type = Shader::MakeValueType(Shader::EValueComponentType::Bool, NumIntermediateComponents);
 			break;
 		case EOperation::Fmod:
 			Types.InputType[1] = Shader::MakeNonLWCType(Types.InputType[1]);
diff --git a/Engine/Source/Runtime/Engine/Private/Materials/MaterialCachedHLSLTree.cpp b/Engine/Source/Runtime/Engine/Private/Materials/MaterialCachedHLSLTree.cpp
--- a/Engine/Source/Runtime/Engine/Private/Materials/MaterialCachedHLSLTree.cpp
+++ b/Engine/Source/Runtime/Engine/Private/Materials/MaterialCachedHLSLTree.cpp
@@ -219,7 +219,7 @@ bool FMaterialCachedHLSLTree::IsAttributeUsed(UE::HLSLTree::FEmitContext& Contex
 		return false;
 	}
 
-	FRequestedType RequestedType;
+	FRequestedType RequestedType(GetMaterialAttributesType(), false);
 	RequestedType.SetFieldRequested(PropertyField);
 
 	const int32 NumComponents = PropertyField->GetNumComponents();
diff --git a/Engine/Source/Runtime/Engine/Private/Materials/MaterialHLSLEmitter.cpp b/Engine/Source/Runtime/Engine/Private/Materials/MaterialHLSLEmitter.cpp
--- a/Engine/Source/Runtime/Engine/Private/Materials/MaterialHLSLEmitter.cpp
+++ b/Engine/Source/Runtime/Engine/Private/Materials/MaterialHLSLEmitter.cpp
@@ -870,7 +870,7 @@ bool MaterialEmitHLSL(const FMaterialCompileTargetParameters& InCompilerTarget,
 			EmitResultScope = EmitContext.PrepareScope(CachedTree->GetResultScope());
 
 			// Prepare code for just the normal
-			FRequestedType RequestedMaterialNormal;
+			FRequestedType RequestedMaterialNormal(CachedTree->GetMaterialAttributesType(), false);
 			RequestedMaterialNormal.SetFieldRequested(NormalField);
 			const FPreparedType PixelResultType1 = EmitContext.PrepareExpression(CachedTree->GetResultExpression(), *EmitResultScope, RequestedMaterialNormal);
 			if (PixelResultType1.IsVoid())
diff --git a/Engine/Source/Runtime/Engine/Private/Materials/MaterialHLSLTree.cpp b/Engine/Source/Runtime/Engine/Private/Materials/MaterialHLSLTree.cpp
--- a/Engine/Source/Runtime/Engine/Private/Materials/MaterialHLSLTree.cpp
+++ b/Engine/Source/Runtime/Engine/Private/Materials/MaterialHLSLTree.cpp
@@ -640,7 +640,7 @@ ETextureMipValueMode GetMipValueMode(FEmitContext& Context, const FExpressionTex
 bool FExpressionTextureSample::PrepareValue(FEmitContext& Context, FEmitScope& Scope, const FRequestedType& RequestedType, FPrepareValueResult& OutResult) const
 {
 	const FPreparedType& TextureType = Context.PrepareExpression(TextureExpression, Scope, FMaterialTextureValue::GetTypeName());
-	if (TextureType.ObjectType != FMaterialTextureValue::GetTypeName())
+	if (TextureType.Type.ObjectType != FMaterialTextureValue::GetTypeName())
 	{
 		return Context.Error(TEXT("Expected texture"));
 	}
@@ -1069,7 +1069,7 @@ void FExpressionTextureSample::EmitValueShader(FEmitContext& Context, FEmitScope
 bool FExpressionTextureSize::PrepareValue(FEmitContext& Context, FEmitScope& Scope, const FRequestedType& RequestedType, FPrepareValueResult& OutResult) const
 {
 	const FPreparedType& TextureType = Context.PrepareExpression(TextureExpression, Scope, FMaterialTextureValue::GetTypeName());
-	if (TextureType.ObjectType != FMaterialTextureValue::GetTypeName())
+	if (TextureType.Type.ObjectType != FMaterialTextureValue::GetTypeName())
 	{
 		return Context.Error(TEXT("Expected texture"));
 	}
@@ -1188,7 +1188,7 @@ bool FExpressionNoise::PrepareValue(FEmitContext& Context, FEmitScope& Scope, co
 void FExpressionNoise::EmitValueShader(FEmitContext& Context, FEmitScope& Scope, const FRequestedType& RequestedType, FEmitValueShaderResult& OutResult) const
 {
 	const FPreparedType& PreparedType = Context.GetPreparedType(PositionExpression);
-	bool bIsLWC = Shader::IsLWCType(PreparedType.ValueComponentType);
+	bool bIsLWC = Shader::IsLWCType(PreparedType.Type);
 	FEmitShaderExpression* EmitPosition = PositionExpression->GetValueShader(Context, Scope, bIsLWC ? Shader::EValueType::Double3 : Shader::EValueType::Float3);
 	FEmitShaderExpression* EmitFilterWidth = FilterWidthExpression->GetValueShader(Context, Scope, Shader::EValueType::Float1);
 
@@ -1263,21 +1263,20 @@ void FExpressionVertexInterpolator::EmitValueShader(FEmitContext& Context, FEmit
 	if (InterpolatorIndex >= 0)
 	{
 		const FVertexInterpolator& Interpolator = EmitMaterialData.VertexInterpolators[InterpolatorIndex];
+		const Shader::FType LocalType = Interpolator.PreparedType.GetType();
+		const int32 NumComponentsRequested = RequestedType.GetNumComponentsRequested(LocalType.GetNumComponents());
 
-		FRequestedType RequestedPreshaderType;
-		RequestedPreshaderType.ValueComponentType = RequestedType.ValueComponentType;
-		RequestedPreshaderType.StructType = RequestedType.StructType;
-		for (int32 Index = 0; Index < RequestedType.GetNumComponents(); ++Index)
+		// Request all non-shader components in preshader, only shader components will be packed in interpolator
+		FRequestedType RequestedPreshaderType(RequestedType, false);
+		for (int32 Index = 0; Index < NumComponentsRequested; ++Index)
 		{
-			// Requested components that are *not* requested by the interpolator will be either constant or preshader evaluation
-			if (RequestedType.IsComponentRequested(Index) &&
-				!Interpolator.RequestedType.IsComponentRequested(Index))
+			const EExpressionEvaluation ComponentEvaluation = Interpolator.PreparedType.GetComponent(Index).Evaluation;
+			if (RequestedType.IsComponentRequested(Index) && ComponentEvaluation != EExpressionEvaluation::Shader)
 			{
 				RequestedPreshaderType.SetComponentRequest(Index);
 			}
 		}
 
-		const Shader::FType LocalType = Interpolator.PreparedType.GetType();
 		FEmitShaderExpression* EmitPreshader = Context.EmitPreshaderOrConstant(Scope, RequestedPreshaderType, LocalType, VertexExpression);
 		OutResult.Code = Context.EmitExpression(Scope, LocalType, TEXT("MaterialVertexInterpolator%(Parameters, %)"), InterpolatorIndex, EmitPreshader);
 	}
@@ -1314,9 +1313,11 @@ void FEmitData::AddInterpolator(const FExpression* Expression, const FRequestedT
 		InterpolatorIndex = VertexInterpolators.Emplace(Expression);
 	}
 
+	const Shader::FType LocalType = PreparedType.GetType();
 	FVertexInterpolator& Interpolator = VertexInterpolators[InterpolatorIndex];
+	Interpolator.RequestedType = FRequestedType(RequestedType, false);
 	Interpolator.PreparedType = PreparedType;
-	for (int32 ComponentIndex = 0; ComponentIndex < RequestedType.GetNumComponents(); ++ComponentIndex)
+	for (int32 ComponentIndex = 0; ComponentIndex < RequestedType.GetNumComponentsRequested(LocalType.GetNumComponents()); ++ComponentIndex)
 	{
 		if (RequestedType.IsComponentRequested(ComponentIndex))
 		{
diff --git a/Engine/Source/Runtime/Engine/Private/Shader/ShaderValue.cpp b/Engine/Source/Runtime/Engine/Private/Shader/ShaderValue.cpp
--- a/Engine/Source/Runtime/Engine/Private/Shader/ShaderValue.cpp
+++ b/Engine/Source/Runtime/Engine/Private/Shader/ShaderValue.cpp
@@ -72,16 +72,30 @@ int32 FType::GetNumFlatFields() const
 
 EValueComponentType FType::GetComponentType(int32 Index) const
 {
+	if (Index < 0)
+	{
+		return EValueComponentType::Void;
+	}
+
 	if (IsStruct())
 	{
 		check(ValueType == EValueType::Struct);
-		return StructType->ComponentTypes.IsValidIndex(Index) ? StructType->ComponentTypes[Index] : EValueComponentType::Void;
+		if (Index < StructType->ComponentTypes.Num())
+		{
+			return StructType->ComponentTypes[Index];
+		}
 	}
-	else
+	else if(IsNumeric())
 	{
 		const FValueTypeDescription TypeDesc = GetValueTypeDescription(ValueType);
-		return (Index >= 0 && Index < TypeDesc.NumComponents) ? TypeDesc.ComponentType : EValueComponentType::Void;
+		// Scalar types replicate xyzw
+		if ((TypeDesc.NumComponents == 1 && Index < 4) || Index < TypeDesc.NumComponents)
+		{
+			return TypeDesc.ComponentType;
+		}
 	}
+
+	return EValueComponentType::Void;
 }
 
 EValueType FType::GetFlatFieldType(int32 Index) const
diff --git a/Engine/Source/Runtime/Engine/Public/HLSLTree/HLSLTree.h b/Engine/Source/Runtime/Engine/Public/HLSLTree/HLSLTree.h
--- a/Engine/Source/Runtime/Engine/Public/HLSLTree/HLSLTree.h
+++ b/Engine/Source/Runtime/Engine/Public/HLSLTree/HLSLTree.h
@@ -141,29 +141,36 @@ enum class EDefaultRequestedType : uint8
 };
 
 /**
- * Represents a request for an arbitrary set of untyped components
+ * FRequestedType is used when accessing various FExpression methods.  A requested type is represented by a concrete Shader::FType, along with a bit vector of components.  A bit set in this vector means the given component was requested.
+ * This allows FExpressions to track access of particular components, which can enable various optimizations.  For example, a 'Swizzle' expression may only access the 'Y' component of its input.  So it might request a float2 value from
+ * its input, but with only the second bit set.  This lets the input know that only the 'Y' component is needed for this request.  This is turn can be passed down the chain, which can allow certain expressions to be skipped or constant-folded.
+ * This is especially useful when dealing with struct types, which are internally represented as a flattened list of components of all fields.  In many cases only certain fields of a struct type will be relevant, which can be tracked by
+ * setting the bits associated with those fields.
  */
 class ENGINE_API FRequestedType
 {
 public:
 	FRequestedType() = default;
-	FRequestedType(Shader::EValueType InType, bool bDefaultRequest = true);
 	FRequestedType(const Shader::FType& InType, bool bDefaultRequest = true);
-	FRequestedType(FName InObjectName);
+	FRequestedType(const FRequestedType& InType, bool bDefaultRequest = true);
+	FRequestedType(Shader::EValueType InType, bool bDefaultRequest = true) : FRequestedType(Shader::FType(InType), bDefaultRequest) {}
+	FRequestedType(const Shader::FStructType* InType, bool bDefaultRequest = true) : FRequestedType(Shader::FType(InType), bDefaultRequest) {}
+	FRequestedType(const FName& InType, bool bDefaultRequest = true) : FRequestedType(Shader::FType(InType), bDefaultRequest) {}
+
 	FRequestedType(EDefaultRequestedType InType) : DefaultType(InType) { check(InType != EDefaultRequestedType::None); }
 
-	Shader::FType GetType() const;
-	const TCHAR* GetName() const { return GetType().GetName(); }
-	int32 GetNumComponents() const;
+	Shader::FType GetType(const Shader::FType& InDefaultType) const;
+	int32 GetNumComponentsRequested(int32 InDefaultNumComponents) const;
+	Shader::EValueComponentType GetValueComponentType() const;
 
 	bool IsComponentRequested(int32 Index) const
 	{
 		return (DefaultType == EDefaultRequestedType::Any) || (RequestedComponents.IsValidIndex(Index) ? (bool)RequestedComponents[Index] : false);
 	}
 
-	bool IsStruct() const { return !IsVoid() && StructType != nullptr; }
-	bool IsObject() const { return !IsVoid() && !ObjectType.IsNone(); }
-	bool IsNumeric() const { return !IsVoid() && Shader::IsNumericType(ValueComponentType); }
+	bool IsStruct() const { return !IsVoid() && Type.IsStruct(); }
+	bool IsObject() const { return !IsVoid() && Type.IsObject(); }
+	bool IsNumeric() const { return !IsVoid() && Type.IsNumeric(); }
 	bool IsVoid() const { return (DefaultType != EDefaultRequestedType::Any) && RequestedComponents.Find(true) == INDEX_NONE; }
 
 	void SetComponentRequest(int32 Index, bool bRequest = true);
@@ -177,14 +184,17 @@ public:
 	}
 
 	/** Marks the given field as requested, based on the input request type (which should match the field type) */
-	void SetField(const Shader::FStructField* Field, const FRequestedType& InRequest);
+	void SetFieldRequested(const Shader::FStructField* Field, const FRequestedType& InRequest);
 
 	/** Returns the requested type of the given field */
 	FRequestedType GetField(const Shader::FStructField* Field) const;
 
-	const Shader::FStructType* StructType = nullptr;
-	FName ObjectType;
-	Shader::EValueComponentType ValueComponentType = Shader::EValueComponentType::Void;
+	/**
+	 * The type used to initialize the request. The actual 'Type' represented by the FRequestedType may be different, as vector types will truncate to the number of requested components
+	 */
+	Shader::FType Type;
+
+	/** Used to signify 'special' requested types */
 	EDefaultRequestedType DefaultType = EDefaultRequestedType::None;
 
 	/** 1 bit per component, a value of 'true' means the specified component is requsted */
@@ -192,10 +202,7 @@ public:
 };
 inline bool operator==(const FRequestedType& Lhs, const FRequestedType& Rhs)
 {
-	return Lhs.ValueComponentType == Rhs.ValueComponentType &&
-		Lhs.ObjectType == Rhs.ObjectType &&
-		Lhs.StructType == Rhs.StructType &&
-		Lhs.RequestedComponents == Rhs.RequestedComponents;
+	return Lhs.Type == Rhs.Type && Lhs.RequestedComponents == Rhs.RequestedComponents;
 }
 inline bool operator!=(const FRequestedType& Lhs, const FRequestedType& Rhs)
 {
@@ -204,12 +211,13 @@ inline bool operator!=(const FRequestedType& Lhs, const FRequestedType& Rhs)
 
 inline void AppendHash(FHasher& Hasher, const FRequestedType& Value)
 {
-	AppendHash(Hasher, Value.StructType);
-	AppendHash(Hasher, Value.ObjectType);
-	AppendHash(Hasher, Value.ValueComponentType);
+	AppendHash(Hasher, Value.Type);
 	AppendHash(Hasher, Value.RequestedComponents);
 }
 
+/**
+ * Prepared state for a single component.  Used by FPreparedType, which includes a list of these
+ */
 struct FPreparedComponent
 {
 	FPreparedComponent() = default;
@@ -218,12 +226,31 @@ struct FPreparedComponent
 	inline bool IsNone() const { return Evaluation == EExpressionEvaluation::None; }
 	inline bool IsRequested() const { return IsRequestedEvaluation(Evaluation); }
 
+	/**
+	 * Get the evaluation of this component within the given scope
+	 * See comment of 'LoopScope' for explanation of why evaluation may change due to scope
+	 */
 	EExpressionEvaluation GetEvaluation(const FEmitScope& Scope) const;
 
 	void SetLoopEvaluation(FEmitScope& Scope);
 
+	/**
+	 * If 'Evaluation' is a loop evaluation, the loop status is only applied inside this scope.
+	 * This is important for constant-folding and preshaders working with loops.  Consider something like this:
+	 * int Value = 0;
+	 * for(int Index = 0; Index < 10; ++Index)
+	 * {
+	 *     Value += Index;
+	 * }
+	 * In this example, 'Value' would have 'ConstantLoop' evaluation, with LoopScope set to the for-loop scope.  This means that *within* the scope of the loop,
+	 * 'Value' is not constant, but it is constant *outside* the loop.  Outside of LoopScope, 'ConstantLoop' evaluation will switch to 'Constant'
+	 */ 
 	FEmitScope* LoopScope = nullptr;
+
+	/** Numeric bounds of this component */
 	Shader::FComponentBounds Bounds;
+
+	/** Evaluation type of this component */
 	EExpressionEvaluation Evaluation = EExpressionEvaluation::None;
 };
 inline bool operator==(const FPreparedComponent& Lhs, const FPreparedComponent& Rhs)
@@ -240,16 +267,18 @@ inline bool operator!=(const FPreparedComponent& Lhs, const FPreparedComponent&
 FPreparedComponent CombineComponents(const FPreparedComponent& Lhs, const FPreparedComponent& Rhs);
 
 /**
- * Like FRequestedType, but tracks an EExpressionEvaluation per component, rather than a simple requested flag
+ * Represents the prepared type of an FExpression.  Each FExpression will prepare a type during FExpression::PrepareValue().
+ * The prepared type is represented by a concrete Shader::FType, along with a list of FPreparedComponents.
+ * This structure matches FRequestedType, and in fact FRequestedType is often used as a 'mask' to only consider certain prepared components
  */
 class FPreparedType
 {
 public:
 	FPreparedType() = default;
-	FPreparedType(Shader::EValueComponentType InComponentType) : ValueComponentType(InComponentType) {}
-	FPreparedType(const Shader::FStructType* InStructType) : StructType(InStructType) {}
-	FPreparedType(Shader::EValueType InType, const FPreparedComponent& InComponent = FPreparedComponent());
 	FPreparedType(const Shader::FType& InType, const FPreparedComponent& InComponent = FPreparedComponent());
+	FPreparedType(Shader::EValueType InType, const FPreparedComponent& InComponent = FPreparedComponent()) : FPreparedType(Shader::FType(InType), InComponent) {}
+	FPreparedType(const Shader::FStructType* InType, const FPreparedComponent& InComponent = FPreparedComponent()) : FPreparedType(Shader::FType(InType), InComponent) {}
+	FPreparedType(const FName& InType, const FPreparedComponent& InComponent = FPreparedComponent()) : FPreparedType(Shader::FType(InType), InComponent) {}
 
 	void SetEvaluation(EExpressionEvaluation Evaluation);
 	void MergeEvaluation(EExpressionEvaluation Evaluation);
@@ -261,11 +290,12 @@ public:
 	int32 GetNumComponents() const;
 	FRequestedType GetRequestedType() const;
 	Shader::FType GetType() const;
+	Shader::EValueComponentType GetValueComponentType() const;
 	const TCHAR* GetName() const { return GetType().GetName(); }
-	bool IsStruct() const { return !IsVoid() && StructType != nullptr; }
-	bool IsObject() const { return !IsVoid() && !ObjectType.IsNone(); }
-	bool IsNumeric() const { return !IsVoid() && Shader::IsNumericType(ValueComponentType); }
-	bool IsInitialized() const { return StructType != nullptr || !ObjectType.IsNone() || ValueComponentType != Shader::EValueComponentType::Void; }
+	bool IsStruct() const { return !IsVoid() && Type.IsStruct(); }
+	bool IsObject() const { return !IsVoid() && Type.IsObject(); }
+	bool IsNumeric() const { return !IsVoid() && Type.IsNumeric(); }
+	bool IsInitialized() const { return !Type.IsVoid(); }
 	bool IsVoid() const;
 
 	EExpressionEvaluation GetEvaluation(const FEmitScope& Scope) const;
@@ -282,21 +312,19 @@ public:
 	void SetComponentBounds(int32 Index, const Shader::FComponentBounds Bounds);
 	void MergeComponent(int32 Index, const FPreparedComponent& InComponent);
 
-	const Shader::FStructType* StructType = nullptr;
-	FName ObjectType;
-	Shader::EValueComponentType ValueComponentType = Shader::EValueComponentType::Void;
-
 	void EnsureNumComponents(int32 NumComponents);
 
+	/**
+	 * The type used to initialize the this. The actual 'Type' represented by the FPreparedType may be different, as vector types will truncate to the number of prepared components
+	 */
+	Shader::FType Type;
+
 	/** Evaluation type for each component, may be 'None' for components that are unused */
 	TArray<FPreparedComponent, TInlineAllocator<4>> PreparedComponents;
 };
 inline bool operator==(const FPreparedType& Lhs, const FPreparedType& Rhs)
 {
-	return Lhs.ValueComponentType == Rhs.ValueComponentType &&
-		Lhs.StructType == Rhs.StructType &&
-		Lhs.ObjectType == Rhs.ObjectType &&
-		Lhs.PreparedComponents == Rhs.PreparedComponents;
+	return Lhs.Type == Rhs.Type && Lhs.PreparedComponents == Rhs.PreparedComponents;
 }
 inline bool operator!=(const FPreparedType& Lhs, const FPreparedType& Rhs)
 {
@@ -317,11 +345,7 @@ public:
 	bool SetType(FEmitContext& Context, const FRequestedType& RequestedType, const FPreparedType& Type);
 
 private:
-	bool TryMergePreparedType(FEmitContext& Context,
-		const Shader::FStructType* StructType,
-		FName ObjectType,
-		Shader::EValueComponentType ComponentType,
-		int32 NumComponents);
+	bool TryMergePreparedType(FEmitContext& Context, const Shader::FType& Type);
 
 	FPreparedType PreparedType;
 	bool bPreparingValue = false;
@@ -367,7 +391,7 @@ struct FExpressionDerivatives
 
 /**
  * Represents an HLSL expression.  This is a piece of code that evaluates to a value, but has no side effects.
- * Unlike statements, expressions are not expected to execute in any particular order.  They may be cached (or not) in generated code, without the underlying implementation needing to care.
+ * Unlike statements, expressions are not expected to execute in any particular order.
  * Examples include constant literals, variable accessors, and various types of math operations
  * This is an abstract base class, with derived classes representing various types of expression
  */
@@ -434,12 +458,33 @@ public:
 		FStringBuilderBase& OutForwardCode) const;
 
 protected:
+	/** Create new expressions representing DDX/DDY of this expression */
 	virtual void ComputeAnalyticDerivatives(FTree& Tree, FExpressionDerivatives& OutResult) const;
+
+	/** Creates a new expression representing this expression on the previous frame.  By default returns nullptr, which means the previous frame is the same as the current frame */
 	virtual const FExpression* ComputePreviousFrame(FTree& Tree, const FRequestedType& RequestedType) const;
+
+	/**
+	 * Computes a FPreparedType for this expression, given a FRequestedType.  Will be called multiple times, with potentially different requested types, if the FExpression is used multiple times.
+	 * In this case, all the prepared types will be merged together if possible (or generate an error otherwise).
+	 * The EExpressionEvaluation set by this is important for determining which EmitValue*** methods may be called
+	 * A given FExpression implementation must support at least one of EmitValueShader/EmitValuePreshader, but doesn't need to support both
+	 */
 	virtual bool PrepareValue(FEmitContext& Context, FEmitScope& Scope, const FRequestedType& RequestedType, FPrepareValueResult& OutResult) const = 0;
+
+	/** Emit HLSL shader code representing this expression */
 	virtual void EmitValueShader(FEmitContext& Context, FEmitScope& Scope, const FRequestedType& RequestedType, FEmitValueShaderResult& OutResult) const;
+
+	/** Emit Preshader code representing this expression */
 	virtual void EmitValuePreshader(FEmitContext& Context, FEmitScope& Scope, const FRequestedType& RequestedType, FEmitValuePreshaderResult& OutResult) const;
+
+	/** Emit an object.  The given 'ObjectTypeName' determines the C++ type pointed to by OutObjectBase */
 	virtual bool EmitValueObject(FEmitContext& Context, FEmitScope& Scope, const FName& ObjectTypeName, void* OutObjectBase) const;
+
+	/**
+	 * Allows custom objects to be passed to custom HLSL functions.  This needs to initialize some HLSL code that facilitates this interface.
+	 * If this returns 'true', then EmitValueShader() will be called to generate the actual HLSL code for the FExpression
+	 */
 	virtual bool EmitCustomHLSLParameter(FEmitContext& Context, FEmitScope& Scope, const FName& ObjectTypeName, const TCHAR* ParameterName, FEmitCustomHLSLParameterResult& OutResult) const;
 
 private:
diff --git a/Engine/Source/Runtime/Engine/Public/Shader/ShaderTypes.h b/Engine/Source/Runtime/Engine/Public/Shader/ShaderTypes.h
--- a/Engine/Source/Runtime/Engine/Public/Shader/ShaderTypes.h
+++ b/Engine/Source/Runtime/Engine/Public/Shader/ShaderTypes.h
@@ -149,6 +149,17 @@ ENGINE_API EValueType MakeComparisonResultType(EValueType Lhs, EValueType Rhs, F
 
 inline bool IsNumericType(EValueType Type) { return IsNumericType(GetValueTypeDescription(Type).ComponentType); }
 
+inline bool IsNumericVectorType(EValueType Type)
+{
+	const FValueTypeDescription TypeDesc = GetValueTypeDescription(Type);
+	return IsNumericType(TypeDesc.ComponentType) && TypeDesc.NumComponents <= 4;
+}
+inline bool IsNumericMatrixType(EValueType Type)
+{
+	const FValueTypeDescription TypeDesc = GetValueTypeDescription(Type);
+	return IsNumericType(TypeDesc.ComponentType) && TypeDesc.NumComponents == 16;
+}
+
 struct FType
 {
 	FType() : ValueType(EValueType::Void) {}
@@ -163,6 +174,8 @@ struct FType
 	bool IsStruct() const { return ValueType == EValueType::Struct; }
 	bool IsObject() const { return ValueType == EValueType::Object; }
 	bool IsNumeric() const { return !IsStruct() && !IsObject() && IsNumericType(ValueType); }
+	bool IsNumericVector() const { return !IsStruct() && !IsObject() && IsNumericVectorType(ValueType); }
+	bool IsNumericMatrix() const { return !IsStruct() && !IsObject() && IsNumericMatrixType(ValueType); }
 	bool IsNumericLWC() const { return IsNumeric() && IsLWCType(ValueType); }
 	int32 GetNumComponents() const;
 	int32 GetNumFlatFields() const;
