Fix a bug of HISM not rendered correctly on mobile preview.

diff --git a/Engine/Source/Runtime/Engine/Private/InstancedStaticMesh.cpp b/Engine/Source/Runtime/Engine/Private/InstancedStaticMesh.cpp
--- a/Engine/Source/Runtime/Engine/Private/InstancedStaticMesh.cpp
+++ b/Engine/Source/Runtime/Engine/Private/InstancedStaticMesh.cpp
@@ -4796,7 +4796,7 @@ void FInstancedStaticMeshVertexFactoryShaderParameters::GetElementShaderBindings
 
 	ShaderBindings.Add(InstanceOffset, InstanceOffsetValue);
 	
-	if (!UseGPUScene(GMaxRHIShaderPlatform, GMaxRHIFeatureLevel))
+	if (!UseGPUScene(Scene ? Scene->GetShaderPlatform() : GMaxRHIShaderPlatform))
 	{
 		ShaderBindings.Add(Shader->GetUniformBufferParameter<FInstancedStaticMeshVertexFactoryUniformShaderParameters>(), InstancedVertexFactory->GetUniformBuffer());
 		if (InstancedVertexFactory->SupportsManualVertexFetch(FeatureLevel))
