Allow audio thread access to FSOundWaveCompiler in standalone game

diff --git a/Engine/Source/Runtime/Engine/Private/AudioDerivedData.cpp b/Engine/Source/Runtime/Engine/Private/AudioDerivedData.cpp
--- a/Engine/Source/Runtime/Engine/Private/AudioDerivedData.cpp
+++ b/Engine/Source/Runtime/Engine/Private/AudioDerivedData.cpp
@@ -756,7 +756,21 @@ void FStreamedAudioPlatformData::Cache(USoundWave& InSoundWave, const FPlatformA
 			AsyncTask->StartBackgroundTask(SoundWaveThreadPool, BasePriority, EQueuedWorkFlags::DoNotRunInsideBusyWait);
 		}
 
-		FSoundWaveCompilingManager::Get().AddSoundWaves({&InSoundWave});
+		if (IsInAudioThread())
+		{
+			TWeakObjectPtr<USoundWave> WeakSoundWavePtr(&InSoundWave);
+			FAudioThread::RunCommandOnGameThread([WeakSoundWavePtr]()
+			{
+				if (USoundWave* SoundWave = WeakSoundWavePtr.Get())
+				{
+					FSoundWaveCompilingManager::Get().AddSoundWaves({SoundWave});
+				}
+			});
+		}
+		else
+		{
+			FSoundWaveCompilingManager::Get().AddSoundWaves({&InSoundWave});
+		}
 	}
 	else
 	{
@@ -1260,7 +1274,7 @@ struct FAudioCookInputs
 #endif
 	{
 #if WITH_EDITORONLY_DATA
-		checkf(IsInGameThread(), TEXT("FAudioCookInputs creation must happen on the game-thread as it reads from many non-thread safe properties of USoundWave"));
+		checkf(IsInGameThread() || IsInAudioThread(), TEXT("FAudioCookInputs creation must happen on the game-thread or audio-thread as it reads from many non-thread safe properties of USoundWave"));
 
 #if FORCE_RESAMPLE
 		FPlatformAudioCookOverrides NewCompressionOverrides = FPlatformAudioCookOverrides();
