Do not serialize the PaintedVertices when collecting object references. This is causing useless operation when launching PIE because of FPIEFixupSerializer.

diff --git a/Engine/Source/Runtime/Engine/Private/Components/StaticMeshComponent.cpp b/Engine/Source/Runtime/Engine/Private/Components/StaticMeshComponent.cpp
--- a/Engine/Source/Runtime/Engine/Private/Components/StaticMeshComponent.cpp
+++ b/Engine/Source/Runtime/Engine/Private/Components/StaticMeshComponent.cpp
@@ -2963,15 +2963,15 @@ bool UStaticMeshComponent::ComponentIsTouchingSelectionFrustum(const FConvexVolu
 					if (!bMustEncompassEntireComponent && bIntersect)
 					{
 						// Needn't encompass entire component: any intersection, we consider as touching
-						return true;
-					}
+					return true;
+				}
 					else if (bMustEncompassEntireComponent && !bFullyContained)
-					{
+				{
 						// Must encompass entire component: any non intersection, we consider as not touching
-						return false;
-					}
+					return false;
 				}
 			}
+			}
 
 			// Either:
 			// a) It must encompass the entire component and all points were intersected (return true), or;
@@ -3252,7 +3252,7 @@ FArchive& operator<<(FArchive& Ar,FStaticMeshComponentLODInfo& I)
 	}
 
 	// Serialize out cached vertex information if necessary.
-	if (!StripFlags.IsEditorDataStripped() && !(Ar.IsFilterEditorOnly() && Ar.IsCountingMemory()))
+	if (!StripFlags.IsEditorDataStripped() && !(Ar.IsFilterEditorOnly() && Ar.IsCountingMemory()) && !Ar.IsObjectReferenceCollector())
 	{
 		Ar << I.PaintedVertices;
 	}
