World Partition: add support for child actors not using an actor template.

#rnx

diff --git a/Engine/Source/Runtime/Engine/Private/Level.cpp b/Engine/Source/Runtime/Engine/Private/Level.cpp
--- a/Engine/Source/Runtime/Engine/Private/Level.cpp
+++ b/Engine/Source/Runtime/Engine/Private/Level.cpp
@@ -806,6 +806,15 @@ void ULevel::AddLoadedActor(AActor* Actor, const FTransform* TransformToApply)
 		}
 
 		OnLoadedActorAddedToLevelEvent.Broadcast(*Actor);
+
+		// Handle child actors
+		Actor->ForEachComponent<UChildActorComponent>(false, [this](UChildActorComponent* ChildActorComponent)
+		{
+			if (AActor* ChildActor = ChildActorComponent->GetChildActor())
+			{
+				AddLoadedActor(ChildActor);
+			}
+		});
 	}
 }
 
@@ -815,6 +824,15 @@ void ULevel::RemoveLoadedActor(AActor* Actor, const FTransform* TransformToRemov
 	check(Actor->GetLevel() == this);
 	check(IsValidChecked(Actor));
 
+	// Handle child actors
+	Actor->ForEachComponent<UChildActorComponent>(false, [this](UChildActorComponent* ChildActorComponent)
+	{
+		if (AActor* ChildActor = ChildActorComponent->GetChildActor())
+		{
+			RemoveLoadedActor(ChildActor);
+		}
+	});
+
 	Actor->UnregisterAllComponents();
 	Actor->RegisterAllActorTickFunctions(false, true);	
 
diff --git a/Engine/Source/Runtime/Engine/Private/WorldPartition/WorldPartitionLevelHelper.cpp b/Engine/Source/Runtime/Engine/Private/WorldPartition/WorldPartitionLevelHelper.cpp
--- a/Engine/Source/Runtime/Engine/Private/WorldPartition/WorldPartitionLevelHelper.cpp
+++ b/Engine/Source/Runtime/Engine/Private/WorldPartition/WorldPartitionLevelHelper.cpp
@@ -402,6 +402,16 @@ bool FWorldPartitionLevelHelper::LoadActors(UWorld* InOwningWorld, ULevel* InDes
 					check(Actor->IsPackageExternal());
 					InDestLevel->Actors.Add(Actor);
 					checkf(Actor->GetLevel() == InDestLevel, TEXT("Levels mismatch, got : %s, expected: %s\nActor: %s\nActorFullName: %s\nActorPackage: %s"), *InDestLevel->GetFullName(), *Actor->GetLevel()->GetFullName(), *Actor->GetActorNameOrLabel(), *Actor->GetFullName(), *Actor->GetPackage()->GetFullName());
+
+					// Handle child actors
+					Actor->ForEachComponent<UChildActorComponent>(true, [InDestLevel](UChildActorComponent* ChildActorComponent)
+					{
+						if (AActor* ChildActor = ChildActorComponent->GetChildActor())
+						{
+							InDestLevel->Actors.Add(ChildActor);
+							check(ChildActor->GetLevel() == InDestLevel);
+						}
+					});
 				}
 
 				UE_LOG(LogEngine, Verbose, TEXT(" ==> Loaded %s (remaining: %d)"), *Actor->GetFullName(), LoadProgress->NumPendingLoadRequests);
diff --git a/Engine/Source/Runtime/Engine/Private/WorldPartition/WorldPartitionRuntimeSpatialHash.cpp b/Engine/Source/Runtime/Engine/Private/WorldPartition/WorldPartitionRuntimeSpatialHash.cpp
--- a/Engine/Source/Runtime/Engine/Private/WorldPartition/WorldPartitionRuntimeSpatialHash.cpp
+++ b/Engine/Source/Runtime/Engine/Private/WorldPartition/WorldPartitionRuntimeSpatialHash.cpp
@@ -907,6 +907,16 @@ bool UWorldPartitionRuntimeSpatialHash::CreateStreamingGrid(const FSpatialHashRu
 									FWorldPartitionReference Reference(WorldPartition, ActorInstance.Actor);
 									AActor* AlwaysLoadedActor = FindObject<AActor>(nullptr, *ActorDescView.GetActorPath().ToString());
 									AlwaysLoadedActorsForPIE.Emplace(Reference, AlwaysLoadedActor);
+
+									// Handle child actors
+									AlwaysLoadedActor->ForEachComponent<UChildActorComponent>(true, [this, &Reference](UChildActorComponent* ChildActorComponent)
+									{
+										if (AActor* ChildActor = ChildActorComponent->GetChildActor())
+										{
+											AlwaysLoadedActorsForPIE.Emplace(Reference, ChildActor);
+										}
+									});
+
 									continue;
 								}
 							}
@@ -980,6 +990,15 @@ bool UWorldPartitionRuntimeSpatialHash::CreateStreamingGrid(const FSpatialHashRu
 						if (AActor* Actor = FindObject<AActor>(nullptr, *ActorDescView.GetActorPath().ToString()))
 						{
 							StreamingCell->UnsavedActorsContainer->Actors.Add(Actor->GetFName(), Actor);
+
+							// Handle child actors
+							Actor->ForEachComponent<UChildActorComponent>(true, [StreamingCell](UChildActorComponent* ChildActorComponent)
+							{
+								if (AActor* ChildActor = ChildActorComponent->GetChildActor())
+								{
+									StreamingCell->UnsavedActorsContainer->Actors.Add(ChildActor->GetFName(), ChildActor);
+								}
+							});
 						}
 					}
 					UE_LOG(LogWorldPartition, Verbose, TEXT("  Actor : %s (%s) (Container %s)"), *(ActorDescView.GetActorPath().ToString()), *ActorDescView.GetGuid().ToString(EGuidFormats::UniqueObjectGuid), *ActorInstance.ContainerInstance->ID.ToString());
