Explicitly release the pending shadermap.

- PendingShaderMap and the ShaderCommonCompileJob both reference each other and the PendingShaderMap is not released when the job is released
- This resulted in leaking maps, which can be quite large when running with unoptimized, uncompressed, debug shaders.

diff --git a/Engine/Source/Runtime/Engine/Private/ShaderCompiler/ShaderCompiler.cpp b/Engine/Source/Runtime/Engine/Private/ShaderCompiler/ShaderCompiler.cpp
--- a/Engine/Source/Runtime/Engine/Private/ShaderCompiler/ShaderCompiler.cpp
+++ b/Engine/Source/Runtime/Engine/Private/ShaderCompiler/ShaderCompiler.cpp
@@ -3787,6 +3787,7 @@ void FShaderCompilingManager::ReleaseJob(FShaderCommonCompileJobPtr& Job)
 
 void FShaderCompilingManager::ReleaseJob(FShaderCommonCompileJob* Job)
 {
+	Job->PendingShaderMap.SafeRelease();
 	Job->bReleased = true;
 	AllJobs.RemoveJob(Job);
 }
