Log each device profile separately to avoid creating one giant log string.

#rnx

diff --git a/Engine/Source/Runtime/Engine/Private/DeviceProfiles/DeviceProfileManager.cpp b/Engine/Source/Runtime/Engine/Private/DeviceProfiles/DeviceProfileManager.cpp
--- a/Engine/Source/Runtime/Engine/Private/DeviceProfiles/DeviceProfileManager.cpp
+++ b/Engine/Source/Runtime/Engine/Private/DeviceProfiles/DeviceProfileManager.cpp
@@ -1218,19 +1218,18 @@ void UDeviceProfileManager::SetActiveDeviceProfile( UDeviceProfile* DeviceProfil
 {
 	ActiveDeviceProfile = DeviceProfile;
 
-	FString ProfileNames;
+	UE_LOG(LogDeviceProfileManager, Log, TEXT("Available device profiles:"));
 	for (int32 Idx = 0; Idx < Profiles.Num(); ++Idx)
 	{
 		UDeviceProfile* Profile = Cast<UDeviceProfile>(Profiles[Idx]);
 		const void* TextureLODGroupsAddr = Profile ? Profile->TextureLODGroups.GetData() : nullptr;
 		const int32 NumTextureLODGroups = Profile ? Profile->TextureLODGroups.Num() : 0;
-		ProfileNames += FString::Printf(TEXT("[%p][%p %d] %s, "), Profile, TextureLODGroupsAddr, NumTextureLODGroups, Profile ? *Profile->GetName() : TEXT("None"));
+		UE_LOG(LogDeviceProfileManager, Log, TEXT("\t[%p][%p %d] %s, "), Profile, TextureLODGroupsAddr, NumTextureLODGroups, Profile ? *Profile->GetName() : TEXT("None"));
 	}
 
 	const void* TextureLODGroupsAddr = ActiveDeviceProfile ? ActiveDeviceProfile->TextureLODGroups.GetData() : nullptr;
 	const int32 NumTextureLODGroups = ActiveDeviceProfile ? ActiveDeviceProfile->TextureLODGroups.Num() : 0;
 	UE_LOG(LogDeviceProfileManager, Log, TEXT("Active device profile: [%p][%p %d] %s"), ActiveDeviceProfile, TextureLODGroupsAddr, NumTextureLODGroups, ActiveDeviceProfile ? *ActiveDeviceProfile->GetName() : TEXT("None"));
-	UE_LOG(LogDeviceProfileManager, Log, TEXT("Profiles: %s"), *ProfileNames);
 
 	ActiveDeviceProfileChangedDelegate.Broadcast();
 
