Fixed WPO was disabled for static meshes as a fallout of 19566743

diff --git a/Engine/Source/Runtime/Engine/Private/StaticMeshRender.cpp b/Engine/Source/Runtime/Engine/Private/StaticMeshRender.cpp
--- a/Engine/Source/Runtime/Engine/Private/StaticMeshRender.cpp
+++ b/Engine/Source/Runtime/Engine/Private/StaticMeshRender.cpp
@@ -305,7 +305,7 @@ FStaticMeshSceneProxy::FStaticMeshSceneProxy(UStaticMeshComponent* InComponent,
 
 			if (CVarRayTracingStaticMeshesWPO.GetValueOnAnyThread() == 1)
 			{
-				bDynamicRayTracingGeometry &= InComponent->bEvaluateWorldPositionOffset;
+				bDynamicRayTracingGeometry &= InComponent->bEvaluateWorldPositionOffsetInRayTracing;
 			}
 		}
 
