Fixup material feature level nodes on PostLoad to deal with unconnected new SM6 pin.
If the node has an SM5 setting but no SM6 setting then copy from SM5 to SM6.
This should help with cases of existing materials not working with SM6.

diff --git a/Engine/Source/Runtime/Core/Public/UObject/UE5MainStreamObjectVersion.h b/Engine/Source/Runtime/Core/Public/UObject/UE5MainStreamObjectVersion.h
--- a/Engine/Source/Runtime/Core/Public/UObject/UE5MainStreamObjectVersion.h
+++ b/Engine/Source/Runtime/Core/Public/UObject/UE5MainStreamObjectVersion.h
@@ -212,6 +212,9 @@ struct CORE_API FUE5MainStreamObjectVersion
 		// SoundWave Update to use EditorBuildData for it's RawData
 		SoundWaveVirtualizationUpdate,
 
+		// Fix material feature level nodes to account for new SM6 input pin.
+		MaterialFeatureLevelNodeFixForSM6,
+
 		// -----<new versions can be added above this line>-------------------------------------------------
 		VersionPlusOne,
 		LatestVersion = VersionPlusOne - 1
diff --git a/Engine/Source/Runtime/Engine/Classes/Materials/Material.h b/Engine/Source/Runtime/Engine/Classes/Materials/Material.h
--- a/Engine/Source/Runtime/Engine/Classes/Materials/Material.h
+++ b/Engine/Source/Runtime/Engine/Classes/Materials/Material.h
@@ -1773,6 +1773,10 @@ public:
 	 * Checks whether a Material is arranged in the old style, with inputs flowing from right to left
 	 */
 	bool HasFlippedCoordinates();
+
+	/** Fixup feature level nodes that haven't set the SM6 pin. Returns true if a fixup was needed. */
+	static bool FixFeatureLevelNodesForSM6(TArray<UMaterialExpression*> const& InExpressions);
+
 #endif //WITH_EDITORONLY_DATA
 
 private:
diff --git a/Engine/Source/Runtime/Engine/Private/Materials/Material.cpp b/Engine/Source/Runtime/Engine/Private/Materials/Material.cpp
--- a/Engine/Source/Runtime/Engine/Private/Materials/Material.cpp
+++ b/Engine/Source/Runtime/Engine/Private/Materials/Material.cpp
@@ -153,6 +153,7 @@ FUObjectAnnotationSparseBool GMaterialsThatNeedExpressionsFlipped;
 FUObjectAnnotationSparseBool GMaterialsThatNeedCoordinateCheck;
 FUObjectAnnotationSparseBool GMaterialsThatNeedCommentFix;
 FUObjectAnnotationSparseBool GMaterialsThatNeedDecalFix;
+FUObjectAnnotationSparseBool GMaterialsThatNeedFeatureLevelSM6Fix;
 
 #endif // #if WITH_EDITOR
 
@@ -2510,6 +2511,11 @@ void UMaterial::Serialize(FArchive& Ar)
 			GMaterialsThatNeedDecalFix.Set(this);
 		}
 	}
+
+	if (Ar.IsLoading() && Ar.CustomVer(FUE5MainStreamObjectVersion::GUID) < FUE5MainStreamObjectVersion::MaterialFeatureLevelNodeFixForSM6)
+	{
+		GMaterialsThatNeedFeatureLevelSM6Fix.Set(this);
+	}
 #endif
 
 #if WITH_EDITOR
@@ -3445,6 +3451,17 @@ void UMaterial::PostLoad()
 		GMaterialsThatNeedCommentFix.Clear(this);
 		FixCommentPositions(EditorComments);
 	}
+
+	if (GMaterialsThatNeedFeatureLevelSM6Fix.Get(this))
+	{
+		GMaterialsThatNeedFeatureLevelSM6Fix.Clear(this);
+		if (FixFeatureLevelNodesForSM6(Expressions))
+		{
+			// Change this guid if you change the conversion logic.
+			static FGuid BackwardsCompatibilityFeatureLevelSM6ConversionGuid(TEXT("FC75DED7-2FB3-463B-B56C-8295871A340C"));
+			ReleaseResourcesAndMutateDDCKey(BackwardsCompatibilityFeatureLevelSM6ConversionGuid);
+		}
+	}
 #endif // #if WITH_EDITOR
 
 	//DumpDebugInfo();
@@ -5927,6 +5944,29 @@ bool UMaterial::HasFlippedCoordinates()
 	// Can't be sure coords are flipped if most are set out correctly
 	return ReversedInputCount > StandardInputCount;
 }
+
+bool UMaterial::FixFeatureLevelNodesForSM6(TArray<UMaterialExpression*> const& InExpressions)
+{
+	bool bRequiredChange = false;
+	for (int32 ExpressionIndex = 0; ExpressionIndex < InExpressions.Num(); ExpressionIndex++)
+	{
+		UMaterialExpressionFeatureLevelSwitch* Expression = Cast<UMaterialExpressionFeatureLevelSwitch>(InExpressions[ExpressionIndex]);
+		if (Expression != nullptr)
+		{
+			FExpressionInput const& ExpressionSM5 = Expression->Inputs[ERHIFeatureLevel::SM5];
+			FExpressionInput& ExpressionSM6 = Expression->Inputs[ERHIFeatureLevel::SM6];
+
+			if (ExpressionSM5.IsConnected() && !ExpressionSM6.IsConnected())
+			{
+				ExpressionSM6.Connect(ExpressionSM5.OutputIndex, ExpressionSM5.Expression);
+				bRequiredChange = true;
+			}
+		}
+	}
+
+	return bRequiredChange;
+}
+
 #endif //WITH_EDITORONLY_DATA
 
 void UMaterial::GetLightingGuidChain(bool bIncludeTextures, TArray<FGuid>& OutGuids) const
diff --git a/Engine/Source/Runtime/Engine/Private/Materials/MaterialExpressions.cpp b/Engine/Source/Runtime/Engine/Private/Materials/MaterialExpressions.cpp
--- a/Engine/Source/Runtime/Engine/Private/Materials/MaterialExpressions.cpp
+++ b/Engine/Source/Runtime/Engine/Private/Materials/MaterialExpressions.cpp
@@ -302,6 +302,7 @@ FUObjectAnnotationSparseBool GMaterialFunctionsThatNeedExpressionsFlipped;
 FUObjectAnnotationSparseBool GMaterialFunctionsThatNeedCoordinateCheck;
 FUObjectAnnotationSparseBool GMaterialFunctionsThatNeedCommentFix;
 FUObjectAnnotationSparseBool GMaterialFunctionsThatNeedSamplerFixup;
+FUObjectAnnotationSparseBool GMaterialFunctionsThatNeedFeatureLevelSM6Fix;
 #endif // #if WITH_EDITOR
 
 bool Engine_IsStrataEnabled();
@@ -12963,6 +12964,11 @@ void UMaterialFunction::Serialize(FArchive& Ar)
 			LibraryCategoriesText.Add(FText::FromString(Category));
 		}
 	}
+
+	if (Ar.IsLoading() && Ar.CustomVer(FUE5MainStreamObjectVersion::GUID) < FUE5MainStreamObjectVersion::MaterialFeatureLevelNodeFixForSM6)
+	{
+		GMaterialFunctionsThatNeedFeatureLevelSM6Fix.Set(this);
+	}
 #endif // #if WITH_EDITOR
 }
 
@@ -13022,6 +13028,12 @@ void UMaterialFunction::PostLoad()
 		UMaterial::FixCommentPositions(FunctionEditorComments);
 	}
 
+	if (GMaterialFunctionsThatNeedFeatureLevelSM6Fix.Get(this))
+	{
+		GMaterialFunctionsThatNeedFeatureLevelSM6Fix.Clear(this);
+		UMaterial::FixFeatureLevelNodesForSM6(FunctionExpressions);
+	}
+
 	if (GMaterialFunctionsThatNeedSamplerFixup.Get(this))
 	{
 		GMaterialFunctionsThatNeedSamplerFixup.Clear(this);
