Made WorldPartition's RuntimeHash transactional.

diff --git a/Engine/Source/Runtime/Engine/Private/WorldPartition/WorldPartition.cpp b/Engine/Source/Runtime/Engine/Private/WorldPartition/WorldPartition.cpp
--- a/Engine/Source/Runtime/Engine/Private/WorldPartition/WorldPartition.cpp
+++ b/Engine/Source/Runtime/Engine/Private/WorldPartition/WorldPartition.cpp
@@ -430,6 +430,7 @@ void UWorldPartition::Initialize(UWorld* InWorld, const FTransform& InTransform)
 	}
 
 	check(RuntimeHash);
+	RuntimeHash->SetFlags(RF_Transactional);
 
 	if (bIsEditor || bIsGame || bPIEWorldTravel || bIsDedicatedServer)
 	{
@@ -825,7 +826,7 @@ UWorldPartition* UWorldPartition::CreateOrRepairWorldPartition(AWorldSettings* W
 			RuntimeHashClass = FindObject<UClass>(ANY_PACKAGE, TEXT("WorldPartitionRuntimeSpatialHash"));
 		}
 
-		WorldPartition->RuntimeHash = NewObject<UWorldPartitionRuntimeHash>(WorldPartition, RuntimeHashClass);
+		WorldPartition->RuntimeHash = NewObject<UWorldPartitionRuntimeHash>(WorldPartition, RuntimeHashClass, NAME_None, RF_Transactional);
 		WorldPartition->RuntimeHash->SetDefaultValues();
 	}
 
