Fix for blend profile entries not being cleaned up when unused bones are removed from a skeleton

diff --git a/Engine/Source/Runtime/Engine/Private/Animation/Skeleton.cpp b/Engine/Source/Runtime/Engine/Private/Animation/Skeleton.cpp
--- a/Engine/Source/Runtime/Engine/Private/Animation/Skeleton.cpp
+++ b/Engine/Source/Runtime/Engine/Private/Animation/Skeleton.cpp
@@ -1698,6 +1698,19 @@ void USkeleton::HandleSkeletonHierarchyChange()
 		CurveMappingTable->InitializeCurveMetaData(this);
 	}
 
+	// Remove entries from Blend Profiles for bones that no longer exists
+	for (UBlendProfile* Profile : BlendProfiles)
+	{
+		for (int32 EntryIndex = Profile->ProfileEntries.Num() - 1; EntryIndex >= 0; --EntryIndex)
+		{
+			const FBlendProfileBoneEntry& Entry = Profile->ProfileEntries[EntryIndex];
+			if (ReferenceSkeleton.FindBoneIndex(Entry.BoneReference.BoneName) == INDEX_NONE)
+			{
+				Profile->RemoveEntry(Entry.BoneReference.BoneIndex);
+			}
+		}
+	}
+
 	OnSkeletonHierarchyChanged.Broadcast();
 }
 
