Blueprint Debugger: Look in WorldPartition sublevels when transfering debug object references to PIE World

diff --git a/Engine/Source/Runtime/Engine/Private/World.cpp b/Engine/Source/Runtime/Engine/Private/World.cpp
--- a/Engine/Source/Runtime/Engine/Private/World.cpp
+++ b/Engine/Source/Runtime/Engine/Private/World.cpp
@@ -2520,6 +2520,23 @@ void UWorld::TransferBlueprintDebugReferences(UWorld* NewWorld)
 					if (GetNetMode() != NM_Client)
 					{
 						NewTargetObject = FindObject<UObject>(NewWorld, *OldTargetObject->GetPathName(this));
+
+						// if we didn't find the object in the Persistent level, we may need to look in WorldPartition sublevels
+						if (!NewTargetObject && NewWorld->GetWorldSettings()->IsPartitionedWorld())
+						{
+							const FString OldTargetPathName = OldTargetObject->GetPathName(PersistentLevel);
+							for (TObjectPtr<ULevelStreaming> StreamingLevel : NewWorld->StreamingLevels)
+							{
+								if (StreamingLevel && StreamingLevel->GetLoadedLevel())
+								{
+									NewTargetObject = FindObject<UObject>(StreamingLevel->GetLoadedLevel(), *OldTargetPathName);
+									if (NewTargetObject)
+									{
+										break;
+									}
+								}
+							}
+						}
 					}
 				}
 
