Tweaks to the new HLSL translator type system
- Add generic types to Shader::EValueType, replace EDefaultRequestedType::Any
- Update code to mostly respect FRequestedType::Type, FPreparedType::Type.  Generally not a good idea to clamp intermediate types to the number of prepared/requested components, as that often ends up creating duplicate code with different components masked out
Tracking indivual components is still useful for selecting evaluation modes that only apply to certain components

diff --git a/Engine/Source/Editor/UnrealEd/Private/MaterialGraphSchema.cpp b/Engine/Source/Editor/UnrealEd/Private/MaterialGraphSchema.cpp
--- a/Engine/Source/Editor/UnrealEd/Private/MaterialGraphSchema.cpp
+++ b/Engine/Source/Editor/UnrealEd/Private/MaterialGraphSchema.cpp
@@ -628,18 +628,24 @@ namespace Private
 {
 FLinearColor GetColorForConnectionType(const UGraphEditorSettings* Settings, UE::Shader::EValueType ConnectionType)
 {
-	if (ConnectionType == UE::Shader::EValueType::Struct)
+	using namespace UE::Shader;
+	if (ConnectionType == EValueType::Any)
+	{
+		return Settings->WildcardPinTypeColor;
+	}
+	else if (ConnectionType == EValueType::Struct)
 	{
 		return Settings->StructPinTypeColor;
 	}
-	else if (ConnectionType == UE::Shader::EValueType::Object)
+	else if (ConnectionType == EValueType::Object)
 	{
 		return Settings->ObjectPinTypeColor;
 	}
 	else
 	{
-		const UE::Shader::FValueTypeDescription TypeDesc = UE::Shader::GetValueTypeDescription(ConnectionType);
-		if (TypeDesc.ComponentType == UE::Shader::EValueComponentType::Float)
+		const FValueTypeDescription TypeDesc = GetValueTypeDescription(ConnectionType);
+		if (TypeDesc.ComponentType == EValueComponentType::Float ||
+			TypeDesc.ComponentType == EValueComponentType::Numeric)
 		{
 			if (TypeDesc.NumComponents == 1)
 			{
@@ -650,15 +656,15 @@ FLinearColor GetColorForConnectionType(const UGraphEditorSettings* Settings, UE:
 				return Settings->VectorPinTypeColor;
 			}
 		}
-		else if (TypeDesc.ComponentType == UE::Shader::EValueComponentType::Double)
+		else if (TypeDesc.ComponentType == EValueComponentType::Double)
 		{
 			return Settings->DoublePinTypeColor;
 		}
-		else if (TypeDesc.ComponentType == UE::Shader::EValueComponentType::Bool)
+		else if (TypeDesc.ComponentType == EValueComponentType::Bool)
 		{
 			return Settings->BooleanPinTypeColor;
 		}
-		else if (TypeDesc.ComponentType == UE::Shader::EValueComponentType::Int)
+		else if (TypeDesc.ComponentType == EValueComponentType::Int)
 		{
 			return Settings->IntPinTypeColor;
 		}
diff --git a/Engine/Source/Runtime/Engine/Private/HLSLTree/HLSLTree.cpp b/Engine/Source/Runtime/Engine/Private/HLSLTree/HLSLTree.cpp
--- a/Engine/Source/Runtime/Engine/Private/HLSLTree/HLSLTree.cpp
+++ b/Engine/Source/Runtime/Engine/Private/HLSLTree/HLSLTree.cpp
@@ -534,7 +534,6 @@ FRequestedType::FRequestedType(const Shader::FType& InType, bool bDefaultRequest
 
 FRequestedType::FRequestedType(const FRequestedType& InType, bool bDefaultRequest)
 	: Type(InType.Type)
-	, DefaultType(InType.DefaultType)
 {
 	if (bDefaultRequest)
 	{
@@ -542,58 +541,13 @@ FRequestedType::FRequestedType(const FRequestedType& InType, bool bDefaultReques
 	}
 }
 
-int32 FRequestedType::GetNumComponentsRequested(int32 InDefaultNumComponents) const
-{
-	if (DefaultType == EDefaultRequestedType::Any)
-	{
-		return InDefaultNumComponents;
-	}
-
-	if (Type.IsVoid())
-	{
-		check(!RequestedComponents.Contains(true));
-		return 0;
-	}
-
-	const int32 MaxComponentIndex = RequestedComponents.FindLast(true);
-	if (MaxComponentIndex == INDEX_NONE)
-	{
-		return 0;
-	}
-
-	if (!Type.IsNumericVector())
-	{
-		return Type.GetNumComponents();
-	}
-
-	return MaxComponentIndex + 1;
-}
-
-Shader::FType FRequestedType::GetType(const Shader::FType& InDefaultType) const
+Shader::EValueComponentType FRequestedType::GetValueComponentType() const
 {
-	if (DefaultType == EDefaultRequestedType::Any)
+	if (Type.IsAny())
 	{
-		return InDefaultType;
+		return Shader::EValueComponentType::Numeric;
 	}
-
-	const int32 NumComponents = GetNumComponentsRequested(0);
-	if (NumComponents == 0)
-	{
-		return Shader::EValueType::Void;
-	}
-
-	if (!Type.IsNumericVector())
-	{
-		return Type;
-	}
-
-	const Shader::FValueTypeDescription TypeDesc = Shader::GetValueTypeDescription(Type);
-	return Shader::MakeValueType(TypeDesc.ComponentType, NumComponents);
-}
-
-Shader::EValueComponentType FRequestedType::GetValueComponentType() const
-{
-	if (Type.IsNumeric())
+	else if (Type.IsNumeric())
 	{
 		return Shader::GetValueTypeDescription(Type).ComponentType;
 	}
@@ -602,7 +556,7 @@ Shader::EValueComponentType FRequestedType::GetValueComponentType() const
 
 void FRequestedType::SetComponentRequest(int32 Index, bool bRequested)
 {
-	if (DefaultType != EDefaultRequestedType::Any)
+	if (!Type.IsAny())
 	{
 		if (bRequested)
 		{
@@ -626,21 +580,17 @@ void FRequestedType::SetFieldRequested(const Shader::FStructField* Field, bool b
 
 void FRequestedType::SetFieldRequested(const Shader::FStructField* Field, const FRequestedType& InRequest)
 {
-	const int32 NumFieldComponents = Field->GetNumComponents();
-	if (NumFieldComponents == 1)
+	if (Field->Type.IsNumericScalar())
 	{
-		for (int32 Index = 0; Index < 4; ++Index)
+		// Scalar fields can replicate, so field is requested if we request any of xyzw
+		if (InRequest.IsNumericVectorRequested())
 		{
-			// Scalar fields can replicate, so field is requested if we request any of xyzw
-			if (InRequest.IsComponentRequested(Index))
-			{
-				SetComponentRequest(Field->ComponentIndex, true);
-				break;
-			}
+			SetComponentRequest(Field->ComponentIndex, true);
 		}
 	}
 	else
 	{
+		const int32 NumFieldComponents = Field->GetNumComponents();
 		for (int32 Index = 0; Index < NumFieldComponents; ++Index)
 		{
 			SetComponentRequest(Field->ComponentIndex + Index, InRequest.IsComponentRequested(Index));
@@ -730,7 +680,7 @@ FPreparedType::FPreparedType(const Shader::FType& InType, const FPreparedCompone
 	}
 }
 
-int32 FPreparedType::GetNumComponents() const
+int32 FPreparedType::GetNumPreparedComponents() const
 {
 	auto Predicate = [](const FPreparedComponent& InComponent) { return !InComponent.IsNone(); };
 
@@ -741,43 +691,29 @@ int32 FPreparedType::GetNumComponents() const
 	}
 
 	const int32 MaxComponentIndex = PreparedComponents.FindLastByPredicate(Predicate);
-	if (MaxComponentIndex == INDEX_NONE)
-	{
-		return 0;
-	}
-
-	if (!Type.IsNumericVector())
-	{
-		return Type.GetNumComponents();
-	}
-
-	return MaxComponentIndex + 1;
+	return (MaxComponentIndex != INDEX_NONE) ? (MaxComponentIndex + 1) : 0;
 }
 
-bool FPreparedType::IsVoid() const
+bool FPreparedType::IsEmpty() const
 {
-	return GetNumComponents() == 0;
+	return GetNumPreparedComponents() == 0;
 }
 
-Shader::FType FPreparedType::GetType() const
+Shader::FType FPreparedType::GetPreparedType() const
 {
-	const int32 NumComponents = GetNumComponents();
-	if (NumComponents == 0)
-	{
-		return Shader::EValueType::Void;
-	}
-
 	if (!Type.IsNumericVector())
 	{
 		return Type;
 	}
 
+	const int32 NumPreparedComponents = GetNumPreparedComponents();
 	const Shader::FValueTypeDescription TypeDesc = Shader::GetValueTypeDescription(Type);
-	return Shader::MakeValueType(TypeDesc.ComponentType, NumComponents);
+	return Shader::MakeValueType(TypeDesc.ComponentType, FMath::Max(NumPreparedComponents, 1));
 }
 
 Shader::EValueComponentType FPreparedType::GetValueComponentType() const
 {
+	check(!Type.IsAny());
 	if (Type.IsNumeric())
 	{
 		return Shader::GetValueTypeDescription(Type).ComponentType;
@@ -787,7 +723,7 @@ Shader::EValueComponentType FPreparedType::GetValueComponentType() const
 
 FRequestedType FPreparedType::GetRequestedType() const
 {
-	const int32 NumComponents = GetNumComponents();
+	const int32 NumComponents = GetNumPreparedComponents();
 	FRequestedType Result;
 	Result.Type = Type;
 	Result.RequestedComponents.Init(false, NumComponents);
@@ -815,11 +751,21 @@ EExpressionEvaluation FPreparedType::GetEvaluation(const FEmitScope& Scope) cons
 EExpressionEvaluation FPreparedType::GetEvaluation(const FEmitScope& Scope, const FRequestedType& RequestedType) const
 {
 	EExpressionEvaluation Result = EExpressionEvaluation::ConstantZero;
-	for (int32 Index = 0; Index < RequestedType.GetNumComponentsRequested(PreparedComponents.Num()); ++Index)
+	if (IsNumericScalar())
+	{
+		if (RequestedType.IsNumericVectorRequested())
+		{
+			Result = GetComponent(0).GetEvaluation(Scope);
+		}
+	}
+	else
 	{
-		if (RequestedType.IsComponentRequested(Index))
+		for (int32 Index = 0; Index < Type.GetNumComponents(); ++Index)
 		{
-			Result = CombineEvaluations(Result, GetComponent(Index).GetEvaluation(Scope));
+			if (RequestedType.IsComponentRequested(Index))
+			{
+				Result = CombineEvaluations(Result, GetComponent(Index).GetEvaluation(Scope));
+			}
 		}
 	}
 	return Result;
@@ -851,11 +797,18 @@ FPreparedComponent FPreparedType::GetMergedComponent() const
 
 FPreparedComponent FPreparedType::GetComponent(int32 Index) const
 {
-	const int32 NumComponents = GetNumComponents();
-	const int32 ComponentIndex = (NumComponents == 1 && Index >= 0 && Index < 4) ? 0 : Index;
-	if (PreparedComponents.IsValidIndex(ComponentIndex))
+	if (Index >= 0)
 	{
-		return PreparedComponents[ComponentIndex];
+		const int32 ComponentIndex = (Type.IsNumericScalar() && Index < 4) ? 0 : Index;
+		if (PreparedComponents.IsValidIndex(ComponentIndex))
+		{
+			return PreparedComponents[ComponentIndex];
+		}
+
+		if (Index < Type.GetNumComponents())
+		{
+			return FPreparedComponent(EExpressionEvaluation::ConstantZero);
+		}
 	}
 
 	return FPreparedComponent(EExpressionEvaluation::None);
@@ -933,16 +886,17 @@ void FPreparedType::MergeEvaluation(EExpressionEvaluation Evaluation)
 
 void FPreparedType::SetLoopEvaluation(FEmitScope& Scope, const FRequestedType& RequestedType)
 {
-	const int32 NumComponents = GetNumComponents();
-	const int32 NumRequestedComponents = RequestedType.GetNumComponentsRequested(NumComponents);
-
-	if (NumComponents == 1 && FMath::IsWithinInclusive(NumRequestedComponents, 1, 4))
+	if (Type.IsNumericScalar())
 	{
 		// If we're a scalar type, set loop evaluation if any of xyzw are requested
-		PreparedComponents[0].SetLoopEvaluation(Scope);
+		if (RequestedType.IsNumericVectorRequested())
+		{
+			PreparedComponents[0].SetLoopEvaluation(Scope);
+		}
 	}
 	else
 	{
+		const int32 NumComponents = GetNumPreparedComponents();
 		for (int32 Index = 0; Index < NumComponents; ++Index)
 		{
 			if (RequestedType.IsComponentRequested(Index))
@@ -965,7 +919,7 @@ Shader::FComponentBounds FPreparedType::GetComponentBounds(int32 Index) const
 		return Shader::FComponentBounds(Shader::EComponentBound::Zero, Shader::EComponentBound::Zero);
 	}
 
-	const Shader::EValueComponentType ComponentType = GetType().GetComponentType(Index);
+	const Shader::EValueComponentType ComponentType = GetPreparedType().GetComponentType(Index);
 	check(ComponentType != Shader::EValueComponentType::Void);
 
 	const Shader::FValueComponentTypeDescription ComponentTypeDesc = Shader::GetValueComponentTypeDescription(ComponentType);
@@ -976,17 +930,26 @@ Shader::FComponentBounds FPreparedType::GetComponentBounds(int32 Index) const
 
 Shader::FComponentBounds FPreparedType::GetBounds(const FRequestedType& RequestedType) const
 {
-	const int32 NumComponents = GetNumComponents();
-	const int32 NumRequestedComponents = RequestedType.GetNumComponentsRequested(NumComponents);
+	const int32 NumComponents = GetNumPreparedComponents();
 
 	Shader::FComponentBounds Result(Shader::EComponentBound::DoubleMax, Shader::EComponentBound::NegDoubleMax);
-	for (int32 Index = 0; Index < NumRequestedComponents; ++Index)
+	if (IsNumericScalar())
 	{
-		if (RequestedType.IsComponentRequested(Index))
+		if (RequestedType.IsNumericVectorRequested())
 		{
-			const Shader::FComponentBounds ComponentBounds = GetComponentBounds(Index);
-			Result.Min = Shader::MinBound(Result.Min, ComponentBounds.Min);
-			Result.Max = Shader::MaxBound(Result.Max, ComponentBounds.Max);
+			Result = GetComponentBounds(0);
+		}
+	}
+	else
+	{
+		for (int32 Index = 0; Index < NumComponents; ++Index)
+		{
+			if (RequestedType.IsComponentRequested(Index))
+			{
+				const Shader::FComponentBounds ComponentBounds = GetComponentBounds(Index);
+				Result.Min = Shader::MinBound(Result.Min, ComponentBounds.Min);
+				Result.Max = Shader::MaxBound(Result.Max, ComponentBounds.Max);
+			}
 		}
 	}
 	return Result;
@@ -1013,37 +976,17 @@ FPreparedType FPreparedType::GetFieldType(const Shader::FStructField* Field) con
 
 FPreparedType MergePreparedTypes(const FPreparedType& Lhs, const FPreparedType& Rhs)
 {
-	// If one type is not initialized yet, just use the other type
-	if (!Lhs.IsInitialized())
+	const Shader::FType UpdatedType = Shader::CombineTypes(Lhs.Type, Rhs.Type);
+	if (UpdatedType.IsVoid())
 	{
-		return Rhs;
-	}
-	else if (!Rhs.IsInitialized())
-	{
-		return Lhs;
+		// Mismatched types
+		return FPreparedType();
 	}
 
-	int32 NumComponents = 0;
-	FPreparedType Result;
-	if (Lhs.IsNumeric() && Rhs.IsNumeric())
-	{
-		const Shader::FValueTypeDescription LhsDesc = Shader::GetValueTypeDescription(Lhs.Type);
-		const Shader::FValueTypeDescription RhsDesc = Shader::GetValueTypeDescription(Rhs.Type);
-		const Shader::EValueComponentType ComponentType = Shader::CombineComponentTypes(LhsDesc.ComponentType, RhsDesc.ComponentType);
-		NumComponents = FMath::Max(Lhs.GetNumComponents(), Rhs.GetNumComponents());
-		Result.Type = Shader::MakeValueType(ComponentType, NumComponents);
-	}
-	else
-	{
-		if (Lhs.Type != Rhs.Type)
-		{
-			// Mismatched non-numeric
-			return Result;
-		}
-		Result.Type = Lhs.Type;
-		NumComponents = Result.Type.GetNumComponents();
-	}
+	const int32 NumComponents = UpdatedType.GetNumComponents();
 
+	FPreparedType Result;
+	Result.Type = UpdatedType;
 	Result.PreparedComponents.Reset(NumComponents);
 	for (int32 Index = 0; Index < NumComponents; ++Index)
 	{
@@ -1072,37 +1015,31 @@ bool FPrepareValueResult::TryMergePreparedType(FEmitContext& Context, const Shad
 		return false;
 	}
 
-	if (!PreparedType.IsInitialized())
+	check(!Type.IsGeneric());
+
+	if (PreparedType.IsVoid())
 	{
 		PreparedType.PreparedComponents.Reset();
 		PreparedType.Type = Type;
 		return true;
 	}
 
-	if (!Type.IsNumericVector())
+	const Shader::FType UpdatedType = Shader::CombineTypes(PreparedType.Type, Type);
+	if (UpdatedType.IsVoid())
 	{
-		if (Type != PreparedType.Type)
-		{
-			return Context.Error(TEXT("Invalid type"));
-		}
+		return Context.Errorf(TEXT("Expression was previously prepared as %s, now as %s, types are not compatible"),
+			PreparedType.Type.GetName(), Type.GetName());
 	}
-	else
-	{
-		const Shader::FValueTypeDescription TypeDesc = Shader::GetValueTypeDescription(Type);
-		const int32 PrevNumComponents = PreparedType.GetNumComponents();
-		if (PrevNumComponents == 1 && TypeDesc.NumComponents > 1)
-		{
-			// Once an expression is prepared as a scalar type, it's not valid to change later
-			// Other expressions may have already assumed this will be a scalar type, and its components will replicate
-			return Context.Error(TEXT("Type was already prepared as scalar"));
-		}
 
-		const Shader::FValueTypeDescription PrevTypeDesc = Shader::GetValueTypeDescription(PreparedType.Type);
-		const Shader::EValueComponentType ComponentType = Shader::CombineComponentTypes(PrevTypeDesc.ComponentType, TypeDesc.ComponentType);
-		const int32 NumComponents = FMath::Max<int32>(PrevNumComponents, TypeDesc.NumComponents);
-		PreparedType.Type = Shader::MakeValueType(ComponentType, NumComponents);
+	if (PreparedType.Type.GetNumComponents() != UpdatedType.GetNumComponents())
+	{
+		// Don't allow the number of components to change
+		// Other expressions may depend on the previous number of components (especially scalar replication)
+		return Context.Errorf(TEXT("Expression was previously prepared as %s, now as %s, number of components don't match"),
+			PreparedType.Type.GetName(), Type.GetName());
 	}
 
+	PreparedType.Type = UpdatedType;
 	return true;
 }
 
@@ -1113,21 +1050,23 @@ bool FPrepareValueResult::SetTypeVoid()
 	return false;
 }
 
-bool FPrepareValueResult::SetType(FEmitContext& Context, const FRequestedType& RequestedType, EExpressionEvaluation Evaluation, const Shader::FType& Type)
+bool FPrepareValueResult::SetType(FEmitContext& Context, const FRequestedType& RequestedType, EExpressionEvaluation Evaluation, const Shader::FType& InType)
 {
+	Shader::FType Type = InType.GetConcreteType();
 	if (TryMergePreparedType(Context, Type))
 	{
 		if (Evaluation != EExpressionEvaluation::None)
 		{
-			const int32 NumComponents = Type.GetNumComponents();
-			const int32 NumRequestedComponents = RequestedType.GetNumComponentsRequested(NumComponents);
-			
-			if (NumComponents == 1 && NumRequestedComponents > 0 && NumRequestedComponents <= 4)
+			if (Type.IsNumericScalar())
 			{
-				PreparedType.MergeComponent(0, Evaluation);
+				if (RequestedType.IsNumericVectorRequested())
+				{
+					PreparedType.MergeComponent(0, Evaluation);
+				}
 			}
 			else
 			{
+				const int32 NumComponents = Type.GetNumComponents();
 				for (int32 Index = 0; Index < NumComponents; ++Index)
 				{
 					if (RequestedType.IsComponentRequested(Index))
@@ -1149,7 +1088,7 @@ bool FPrepareValueResult::SetType(FEmitContext& Context, const FRequestedType& R
 
 bool FPrepareValueResult::SetType(FEmitContext& Context, const FRequestedType& RequestedType, const FPreparedType& InPreparedType)
 {
-	Shader::FType Type = InPreparedType.GetType();
+	Shader::FType Type = InPreparedType.Type.GetConcreteType();
 	if (Type.IsNumericLWC())
 	{
 		const Shader::FComponentBounds Bounds = InPreparedType.GetBounds(RequestedType);
@@ -1161,15 +1100,16 @@ bool FPrepareValueResult::SetType(FEmitContext& Context, const FRequestedType& R
 
 	if (TryMergePreparedType(Context, Type))
 	{
-		const int32 NumComponents = Type.GetNumComponents();
-		const int32 NumRequestedComponents = RequestedType.GetNumComponentsRequested(NumComponents);
-		
-		if (NumComponents == 1 && NumRequestedComponents > 0 && NumRequestedComponents <= 4)
+		if (Type.IsNumericScalar())
 		{
-			PreparedType.MergeComponent(0, InPreparedType.GetComponent(0));
+			if (RequestedType.IsNumericVectorRequested())
+			{
+				PreparedType.MergeComponent(0, InPreparedType.GetComponent(0));
+			}
 		}
 		else
 		{
+			const int32 NumComponents = Type.GetNumComponents();
 			for (int32 Index = 0; Index < NumComponents; ++Index)
 			{
 				const FPreparedComponent& Component = InPreparedType.GetComponent(Index);
@@ -1264,7 +1204,7 @@ FEmitShaderExpression* FExpression::GetValueShader(FEmitContext& Context, FEmitS
 FEmitShaderExpression* FExpression::GetValueShader(FEmitContext& Context, FEmitScope& Scope, const FRequestedType& RequestedType) const
 {
 	const FPreparedType& PreparedType = Context.GetPreparedType(this);
-	return GetValueShader(Context, Scope, RequestedType, PreparedType, PreparedType.GetType());
+	return GetValueShader(Context, Scope, RequestedType, PreparedType, PreparedType.GetPreparedType());
 }
 
 FEmitShaderExpression* FExpression::GetValueShader(FEmitContext& Context, FEmitScope& Scope, const Shader::FType& ResultType) const
@@ -1280,7 +1220,7 @@ FEmitShaderExpression* FExpression::GetValueShader(FEmitContext& Context, FEmitS
 FEmitShaderExpression* FExpression::GetValueShader(FEmitContext& Context, FEmitScope& Scope) const
 {
 	const FPreparedType& PreparedType = Context.GetPreparedType(this);
-	return GetValueShader(Context, Scope, PreparedType.GetRequestedType(), PreparedType, PreparedType.GetType());
+	return GetValueShader(Context, Scope, PreparedType.GetRequestedType(), PreparedType, PreparedType.GetPreparedType());
 }
 
 Shader::FType FExpression::GetValuePreshader(FEmitContext& Context, FEmitScope& Scope, const FRequestedType& RequestedType, const FPreparedType& PreparedType, const Shader::FType& ResultType, Shader::FPreshaderData& OutPreshader) const
@@ -1324,7 +1264,7 @@ Shader::FType FExpression::GetValuePreshader(FEmitContext& Context, FEmitScope&
 Shader::FType FExpression::GetValuePreshader(FEmitContext& Context, FEmitScope& Scope, const FRequestedType& RequestedType, Shader::FPreshaderData& OutPreshader) const
 {
 	const FPreparedType& PreparedType = Context.GetPreparedType(this);
-	return GetValuePreshader(Context, Scope, RequestedType, PreparedType, PreparedType.GetType(), OutPreshader);
+	return GetValuePreshader(Context, Scope, RequestedType, PreparedType, PreparedType.GetPreparedType(), OutPreshader);
 }
 
 Shader::FType FExpression::GetValuePreshader(FEmitContext& Context, FEmitScope& Scope, const Shader::FType& ResultType, Shader::FPreshaderData& OutPreshader) const
@@ -1368,7 +1308,7 @@ Shader::FValue FExpression::GetValueConstant(FEmitContext& Context, FEmitScope&
 
 Shader::FValue FExpression::GetValueConstant(FEmitContext& Context, FEmitScope& Scope, const FRequestedType& RequestedType, const FPreparedType& PreparedType) const
 {
-	return GetValueConstant(Context, Scope, RequestedType, PreparedType, PreparedType.GetType());
+	return GetValueConstant(Context, Scope, RequestedType, PreparedType, PreparedType.GetPreparedType());
 }
 
 Shader::FValue FExpression::GetValueConstant(FEmitContext& Context, FEmitScope& Scope, const FRequestedType& RequestedType, const Shader::FType& ResultType) const
@@ -1380,7 +1320,7 @@ Shader::FValue FExpression::GetValueConstant(FEmitContext& Context, FEmitScope&
 Shader::FValue FExpression::GetValueConstant(FEmitContext& Context, FEmitScope& Scope, const FRequestedType& RequestedType) const
 {
 	const FPreparedType PreparedType = Context.GetPreparedType(this);
-	return GetValueConstant(Context, Scope, RequestedType, PreparedType, PreparedType.GetType());
+	return GetValueConstant(Context, Scope, RequestedType, PreparedType, PreparedType.GetPreparedType());
 }
 
 Shader::FValue FExpression::GetValueConstant(FEmitContext& Context, FEmitScope& Scope, const FPreparedType& PreparedType, const Shader::FType& ResultType) const
diff --git a/Engine/Source/Runtime/Engine/Private/HLSLTree/HLSLTreeCommon.cpp b/Engine/Source/Runtime/Engine/Private/HLSLTree/HLSLTreeCommon.cpp
--- a/Engine/Source/Runtime/Engine/Private/HLSLTree/HLSLTreeCommon.cpp
+++ b/Engine/Source/Runtime/Engine/Private/HLSLTree/HLSLTreeCommon.cpp
@@ -127,16 +127,8 @@ bool FExpressionConstant::PrepareValue(FEmitContext& Context, FEmitScope& Scope,
 	ResultType.PreparedComponents.Reserve(NumComponents);
 	for (int32 Index = 0; Index < NumComponents; ++Index)
 	{
-		EExpressionEvaluation Evaluation = EExpressionEvaluation::ConstantZero;
-		if (RequestedType.IsComponentRequested(Index))
-		{
-			const Shader::FValueComponent Component = Value.GetComponent(Index);
-			if (Component.Packed != 0u)
-			{
-				Evaluation = EExpressionEvaluation::Constant;
-			}
-		}
-		ResultType.PreparedComponents.Add(Evaluation);
+		const Shader::FValueComponent Component = Value.GetComponent(Index);
+		ResultType.PreparedComponents.Add(Component.Packed ? EExpressionEvaluation::Constant : EExpressionEvaluation::ConstantZero);
 	}
 	return OutResult.SetType(Context, RequestedType, ResultType);
 }
@@ -443,15 +435,29 @@ namespace Private
 {
 FRequestedType GetRequestedSwizzleType(const FSwizzleParameters& Params, const FRequestedType& RequestedType)
 {
-	FRequestedType RequestedInputType(RequestedType, false);
-	for (int32 Index = 0; Index < RequestedType.GetNumComponentsRequested(Params.NumComponents); ++Index)
+	int32 MaxSwizzledComponentIndex = 0;
+	FRequestedType RequestedInputType;
+	if (Params.NumComponents == 1)
 	{
-		if (RequestedType.IsComponentRequested(Index))
+		MaxSwizzledComponentIndex = Params.GetSwizzleComponentIndex(0);
+		if (RequestedType.IsNumericVectorRequested())
+		{
+			RequestedInputType.SetComponentRequest(MaxSwizzledComponentIndex);
+		}
+	}
+	else
+	{
+		for (int32 Index = 0; Index < Params.NumComponents; ++Index)
 		{
 			const int32 SwizzledComponentIndex = Params.GetSwizzleComponentIndex(Index);
-			RequestedInputType.SetComponentRequest(SwizzledComponentIndex);
+			MaxSwizzledComponentIndex = FMath::Max(MaxSwizzledComponentIndex, SwizzledComponentIndex);
+			if (RequestedType.IsComponentRequested(Index))
+			{
+				RequestedInputType.SetComponentRequest(SwizzledComponentIndex);
+			}
 		}
 	}
+	RequestedInputType.Type = Shader::MakeValueType(RequestedType.GetValueComponentType(), MaxSwizzledComponentIndex + 1);
 	return RequestedInputType;
 }
 } // namespace Private
@@ -491,36 +497,15 @@ bool SwizzlePrepareValue(FEmitContext& Context,
 	FPrepareValueResult& OutResult)
 {
 	const FRequestedType RequestedInputType = Private::GetRequestedSwizzleType(Parameters, RequestedType);
+	const FPreparedType& InputType = Context.PrepareExpression(Input, Scope, RequestedInputType);
 
 	FPreparedType ResultType;
-	if (RequestedInputType.IsVoid())
+	ResultType.Type = Shader::MakeValueType(InputType.GetValueComponentType(), Parameters.NumComponents);
+	for (int32 ComponentIndex = 0; ComponentIndex < Parameters.NumComponents; ++ComponentIndex)
 	{
-		// All the requested components are outside the swizzle, so just return 0
-		ResultType.Type = RequestedType.Type;
-		for (int32 ComponentIndex = 0; ComponentIndex < Parameters.NumComponents; ++ComponentIndex)
-		{
-			ResultType.SetComponent(ComponentIndex, EExpressionEvaluation::ConstantZero);
-		}
+		const int32 SwizzledComponentIndex = Parameters.GetSwizzleComponentIndex(ComponentIndex);
+		ResultType.SetComponent(ComponentIndex, InputType.GetComponent(SwizzledComponentIndex));
 	}
-	else
-	{
-		const FPreparedType& InputType = Context.PrepareExpression(Input, Scope, RequestedInputType);
-
-		ResultType.Type = InputType.Type;
-		for (int32 ComponentIndex = 0; ComponentIndex < Parameters.NumComponents; ++ComponentIndex)
-		{
-			if (RequestedType.IsComponentRequested(ComponentIndex))
-			{
-				const int32 SwizzledComponentIndex = Parameters.GetSwizzleComponentIndex(ComponentIndex);
-				ResultType.SetComponent(ComponentIndex, InputType.GetComponent(SwizzledComponentIndex));
-			}
-			else
-			{
-				ResultType.SetComponent(ComponentIndex, EExpressionEvaluation::ConstantZero);
-			}
-		}
-	}
-
 	return OutResult.SetType(Context, RequestedType, ResultType);
 }
 
@@ -700,8 +685,8 @@ FAppendTypes GetAppendTypes(const FRequestedType& RequestedType, const FPrepared
 	const Shader::FValueTypeDescription LhsTypeDesc = Shader::GetValueTypeDescription(LhsType.Type);
 	const Shader::FValueTypeDescription RhsTypeDesc = Shader::GetValueTypeDescription(RhsType.Type);
 	const Shader::EValueComponentType ComponentType = Shader::CombineComponentTypes(LhsTypeDesc.ComponentType, RhsTypeDesc.ComponentType);
-	const int32 NumLhsComponents = LhsType.GetNumComponents();
-	const int32 NumRhsComponents = RhsType.GetNumComponents();
+	const int32 NumLhsComponents = LhsType.Type.GetNumComponents();
+	const int32 NumRhsComponents = RhsType.Type.GetNumComponents();
 	const int32 NumComponents = FMath::Min(NumLhsComponents + NumRhsComponents, 4);
 
 	FAppendTypes Types;
@@ -753,42 +738,42 @@ const FExpression* FExpressionAppend::ComputePreviousFrame(FTree& Tree, const FR
 bool FExpressionAppend::PrepareValue(FEmitContext& Context, FEmitScope& Scope, const FRequestedType& RequestedType, FPrepareValueResult& OutResult) const
 {
 	const FPreparedType& LhsType = Context.PrepareExpression(Lhs, Scope, RequestedType);
-	const int32 NumLhsComponents = LhsType.GetNumComponents();
-
-	FRequestedType RhsRequestedType(RequestedType, false);
-	for (int32 Index = NumLhsComponents; Index < RequestedType.GetNumComponentsRequested(4 - NumLhsComponents); ++Index)
-	{
-		RhsRequestedType.SetComponentRequest(Index - NumLhsComponents, RequestedType.IsComponentRequested(Index));
-	}
+	const int32 NumRequestedComponents = RequestedType.Type.GetNumComponents();
+	const int32 NumLhsComponents = LhsType.Type.GetNumComponents();
+	const int32 NumRhsComponents = FMath::Max(NumRequestedComponents - NumLhsComponents, 1);
 
-	FPreparedType RhsType;
-	if (!RhsRequestedType.IsVoid())
+	FRequestedType RhsRequestedType;
+	RhsRequestedType.Type = Shader::MakeValueType(RequestedType.GetValueComponentType(), NumRhsComponents);
+	for (int32 Index = NumLhsComponents; Index < NumRequestedComponents; ++Index)
 	{
-		RhsType = Context.PrepareExpression(Rhs, Scope, RhsRequestedType);
-		if (LhsType.GetValueComponentType() != RhsType.GetValueComponentType())
+		if (RequestedType.IsComponentRequested(Index))
 		{
-			return Context.Error(TEXT("Type mismatch"));
+			RhsRequestedType.SetComponentRequest(Index - NumLhsComponents);
 		}
 	}
-	else
+
+	const FPreparedType& RhsType = Context.PrepareExpression(Rhs, Scope, RhsRequestedType);
+	if (!RhsType.IsVoid() && LhsType.GetValueComponentType() != RhsType.GetValueComponentType())
 	{
-		// We never want the result of an 'Append' operation to be a scalar value
-		// Ensure that we include at least a single component for the rhs, even if it's not currently requested
-		RhsType.Type = RequestedType.Type;
-		RhsType.SetComponent(0, EExpressionEvaluation::ConstantZero);
+		return Context.Error(TEXT("Type mismatch"));
 	}
-
+	
 	const Private::FAppendTypes Types = Private::GetAppendTypes(RequestedType, LhsType, RhsType);
 	Context.MarkInputType(Lhs, Types.LhsType);
 	Context.MarkInputType(Rhs, Types.RhsType);
 
+	if (Types.ResultType.IsNumericScalar())
+	{
+		int a = 0;
+	}
+
 	return OutResult.SetType(Context, RequestedType, Types.ResultType);
 }
 
 void FExpressionAppend::EmitValueShader(FEmitContext& Context, FEmitScope& Scope, const FRequestedType& RequestedType, FEmitValueShaderResult& OutResult) const
 {
 	const Private::FAppendTypes Types = Private::GetAppendTypes(RequestedType, Context.GetPreparedType(Lhs), Context.GetPreparedType(Rhs));
-	const Shader::FType ResultType = Types.ResultType.GetType();
+	const Shader::FType ResultType = Types.ResultType.GetPreparedType();
 	FEmitShaderExpression* LhsValue = Lhs->GetValueShader(Context, Scope, Types.LhsRequestedType, Types.LhsType);
 
 	if (Types.RhsType == Shader::EValueType::Void)
@@ -827,7 +812,7 @@ void FExpressionAppend::EmitValuePreshader(FEmitContext& Context, FEmitScope& Sc
 
 		OutResult.Preshader.WriteOpcode(Shader::EPreshaderOpcode::AppendVector);
 	}
-	OutResult.Type = Types.ResultType.GetType();
+	OutResult.Type = Types.ResultType.GetPreparedType();
 }
 
 void FExpressionSwitchBase::ComputeAnalyticDerivatives(FTree& Tree, FExpressionDerivatives& OutResult) const
@@ -937,7 +922,7 @@ bool FExpressionCustomHLSL::PrepareValue(FEmitContext& Context, FEmitScope& Scop
 {
 	for (const FCustomHLSLInput& Input : Inputs)
 	{
-		const FPreparedType& InputType = Context.PrepareExpression(Input.Expression, Scope, EDefaultRequestedType::Any);
+		const FPreparedType& InputType = Context.PrepareExpression(Input.Expression, Scope, Shader::EValueType::Any);
 		if (InputType.IsVoid())
 		{
 			return false;
@@ -951,7 +936,7 @@ bool FExpressionCustomHLSL::PrepareValue(FEmitContext& Context, FEmitScope& Scop
 			}
 		}
 
-		Context.MarkInputType(Input.Expression, InputType.GetType());
+		Context.MarkInputType(Input.Expression, InputType.Type);
 	}
 
 	return OutResult.SetType(Context, RequestedType, EExpressionEvaluation::Shader, Shader::FType(OutputStructType));
diff --git a/Engine/Source/Runtime/Engine/Private/HLSLTree/HLSLTreeEmit.cpp b/Engine/Source/Runtime/Engine/Private/HLSLTree/HLSLTreeEmit.cpp
--- a/Engine/Source/Runtime/Engine/Private/HLSLTree/HLSLTreeEmit.cpp
+++ b/Engine/Source/Runtime/Engine/Private/HLSLTree/HLSLTreeEmit.cpp
@@ -376,7 +376,7 @@ FRequestedType FEmitContext::GetRequestedType(const FExpression* Expression) con
 
 Shader::FType FEmitContext::GetType(const FExpression* Expression) const
 {
-	return GetPreparedType(Expression).GetType();
+	return GetPreparedType(Expression).GetPreparedType();
 }
 
 EExpressionEvaluation FEmitContext::GetEvaluation(const FExpression* Expression, const FEmitScope& Scope, const FRequestedType& RequestedType) const
@@ -400,6 +400,12 @@ FPreparedType FEmitContext::PrepareExpression(const FExpression* InExpression, F
 		PrepareValueMap.Add(InExpression, Result);
 	}
 
+	if (RequestedType.IsEmpty() && !Result->PreparedType.IsVoid())
+	{
+		// If RequestedType is 'Empty' (nothing requested), we can skip all but the first call to Prepare
+		return Result->PreparedType;
+	}
+
 	if (Result->bPreparingValue)
 	{
 		// Valid for this to be called reentrantly
@@ -428,7 +434,7 @@ FPreparedType FEmitContext::PrepareExpression(const FExpression* InExpression, F
 	{
 		ResultType = Result->PreparedType;
 		check(!ResultType.IsVoid());
-		MarkInputType(InExpression, RequestedType.GetType(ResultType.GetType()));
+		MarkInputType(InExpression, RequestedType.Type.GetConcreteType());
 	}
 	return ResultType;
 }
@@ -768,6 +774,11 @@ FEmitShaderExpression* FEmitContext::InternalEmitExpression(FEmitScope& Scope, T
 {
 	FEmitShaderExpression* ShaderValue = nullptr;
 
+	if (Code == TEXT("(((float3)0).z * Material.PreshaderBuffer[6].x)"))
+	{
+		int a = 0;
+	}
+
 	FXxHash64Builder Hasher;
 	Hasher.Update(Code.GetData(), Code.Len() * sizeof(TCHAR));
 	if (bInline)
diff --git a/Engine/Source/Runtime/Engine/Private/HLSLTree/HLSLTreeOperation.cpp b/Engine/Source/Runtime/Engine/Private/HLSLTree/HLSLTreeOperation.cpp
--- a/Engine/Source/Runtime/Engine/Private/HLSLTree/HLSLTreeOperation.cpp
+++ b/Engine/Source/Runtime/Engine/Private/HLSLTree/HLSLTreeOperation.cpp
@@ -211,7 +211,7 @@ FOperationTypes GetOperationTypes(EOperation Op, TConstArrayView<FPreparedType>
 			Types.ResultType = FPreparedType(Shader::EValueType::Float3, IntermediateComponent);
 			break;
 		case EOperation::VecMulMatrix4:
-			switch (InputPreparedType[1].GetType().ValueType)
+			switch (InputPreparedType[1].Type.ValueType)
 			{
 			case Shader::EValueType::Double4x4:
 				// float3 * FLWCMatrix -> FLWCVector3
@@ -251,7 +251,7 @@ FOperationTypes GetOperationTypes(EOperation Op, TConstArrayView<FPreparedType>
 		}
 
 		const Shader::EValueComponentType IntermediateComponentType = IntermediateType.GetValueComponentType();
-		const int32 NumIntermediateComponents = IntermediateType.GetNumComponents();
+		const int32 NumIntermediateComponents = IntermediateType.Type.GetNumComponents();
 		for (int32 Index = 0; Index < InputPreparedType.Num(); ++Index)
 		{
 			Shader::EValueComponentType InputComponentType = InputPreparedType[Index].GetValueComponentType();
@@ -271,8 +271,9 @@ FOperationTypes GetOperationTypes(EOperation Op, TConstArrayView<FPreparedType>
 					break;
 				}
 			}
-			const int32 NumInputComponents = InputPreparedType[Index].GetNumComponents();
-			Types.InputType[Index] = Shader::MakeValueType(InputComponentType, (NumInputComponents == 1) ? 1 : NumIntermediateComponents);
+
+			const int32 NumInputComponents = InputPreparedType[Index].IsNumericScalar() ? 1 : NumIntermediateComponents;
+			Types.InputType[Index] = Shader::MakeValueType(InputComponentType, NumInputComponents);
 		}
 		Types.ResultType = IntermediateType;
 		Types.bIsLWC = (IntermediateComponentType == Shader::EValueComponentType::Double);
@@ -663,7 +664,7 @@ void FExpressionOperation::EmitValueShader(FEmitContext& Context, FEmitScope& Sc
 		InputValue[Index] = Inputs[Index]->GetValueShader(Context, Scope, RequestedTypes.InputType[Index], Types.InputType[Index]);
 	}
 
-	const Shader::EValueType ResultType = Types.ResultType.GetType();
+	const Shader::EValueType ResultType = Types.ResultType.GetPreparedType();
 	check(Shader::IsNumericType(ResultType));
 
 	switch (Op)
@@ -778,7 +779,7 @@ void FExpressionOperation::EmitValuePreshader(FEmitContext& Context, FEmitScope&
 	}
 
 	OutResult.Preshader.WriteOpcode(OpDesc.PreshaderOpcode);
-	OutResult.Type = Types.ResultType.GetType();
+	OutResult.Type = Types.ResultType.GetPreparedType();
 }
 
 } // namespace UE::HLSLTree
diff --git a/Engine/Source/Runtime/Engine/Private/Materials/MaterialHLSLTree.cpp b/Engine/Source/Runtime/Engine/Private/Materials/MaterialHLSLTree.cpp
--- a/Engine/Source/Runtime/Engine/Private/Materials/MaterialHLSLTree.cpp
+++ b/Engine/Source/Runtime/Engine/Private/Materials/MaterialHLSLTree.cpp
@@ -1263,12 +1263,11 @@ void FExpressionVertexInterpolator::EmitValueShader(FEmitContext& Context, FEmit
 	if (InterpolatorIndex >= 0)
 	{
 		const FVertexInterpolator& Interpolator = EmitMaterialData.VertexInterpolators[InterpolatorIndex];
-		const Shader::FType LocalType = Interpolator.PreparedType.GetType();
-		const int32 NumComponentsRequested = RequestedType.GetNumComponentsRequested(LocalType.GetNumComponents());
+		const Shader::FType LocalType = Interpolator.PreparedType.GetPreparedType();
 
 		// Request all non-shader components in preshader, only shader components will be packed in interpolator
 		FRequestedType RequestedPreshaderType(RequestedType, false);
-		for (int32 Index = 0; Index < NumComponentsRequested; ++Index)
+		for (int32 Index = 0; Index < Interpolator.PreparedType.GetNumPreparedComponents(); ++Index)
 		{
 			const EExpressionEvaluation ComponentEvaluation = Interpolator.PreparedType.GetComponent(Index).Evaluation;
 			if (RequestedType.IsComponentRequested(Index) && ComponentEvaluation != EExpressionEvaluation::Shader)
@@ -1313,11 +1312,11 @@ void FEmitData::AddInterpolator(const FExpression* Expression, const FRequestedT
 		InterpolatorIndex = VertexInterpolators.Emplace(Expression);
 	}
 
-	const Shader::FType LocalType = PreparedType.GetType();
+	const Shader::FType LocalType = PreparedType.GetPreparedType();
 	FVertexInterpolator& Interpolator = VertexInterpolators[InterpolatorIndex];
 	Interpolator.RequestedType = FRequestedType(RequestedType, false);
 	Interpolator.PreparedType = PreparedType;
-	for (int32 ComponentIndex = 0; ComponentIndex < RequestedType.GetNumComponentsRequested(LocalType.GetNumComponents()); ++ComponentIndex)
+	for (int32 ComponentIndex = 0; ComponentIndex < PreparedType.GetNumPreparedComponents(); ++ComponentIndex)
 	{
 		if (RequestedType.IsComponentRequested(ComponentIndex))
 		{
@@ -1345,7 +1344,7 @@ void FEmitData::EmitInterpolatorStatements(FEmitContext& Context, FEmitScope& Sc
 	for (int32 Index = 0; Index < VertexInterpolators.Num(); ++Index)
 	{
 		const FVertexInterpolator& Interpolator = VertexInterpolators[Index];
-		FEmitShaderExpression* EmitExpression = Interpolator.Expression->GetValueShader(Context, Scope, Interpolator.RequestedType, Interpolator.PreparedType, Interpolator.PreparedType.GetType());
+		FEmitShaderExpression* EmitExpression = Interpolator.Expression->GetValueShader(Context, Scope, Interpolator.RequestedType, Interpolator.PreparedType, Interpolator.PreparedType.GetPreparedType());
 		Context.EmitStatement(Scope, TEXT("MaterialPackVertexInterpolator%(Parameters, %);"), Index, EmitExpression);
 	}
 }
@@ -1414,7 +1413,7 @@ void FEmitData::EmitInterpolatorShader(FEmitContext& Context, FStringBuilderBase
 	for (int32 Index = 0; Index < VertexInterpolators.Num(); ++Index)
 	{
 		const FVertexInterpolator& Interpolator = VertexInterpolators[Index];
-		const Shader::FType Type = Interpolator.PreparedType.GetType();
+		const Shader::FType Type = Interpolator.PreparedType.GetPreparedType();
 		const TCHAR* TypeName = Type.GetName();
 
 		TStringBuilder<1024> ReadCode;
diff --git a/Engine/Source/Runtime/Engine/Private/Shader/ShaderValue.cpp b/Engine/Source/Runtime/Engine/Private/Shader/ShaderValue.cpp
--- a/Engine/Source/Runtime/Engine/Private/Shader/ShaderValue.cpp
+++ b/Engine/Source/Runtime/Engine/Private/Shader/ShaderValue.cpp
@@ -111,40 +111,38 @@ EValueType FType::GetFlatFieldType(int32 Index) const
 	}
 }
 
-bool FType::Merge(const FType& OtherType)
+FType CombineTypes(const FType& Lhs, const FType& Rhs)
 {
-	if (IsVoid())
+	if (Lhs.IsVoid() || Lhs.IsAny())
 	{
-		ValueType = OtherType.ValueType;
-		ObjectType = OtherType.ObjectType;
-		StructType = OtherType.StructType;
-		return true;
+		return Rhs;
 	}
-
-	if (IsStruct() || OtherType.IsStruct())
+	if (Rhs.IsVoid() || Rhs.IsAny())
 	{
-		return StructType == OtherType.StructType;
-	}
-
-	if (IsObject() || OtherType.IsObject())
-	{
-		return ObjectType == OtherType.ObjectType;
+		return Lhs;
 	}
 
-	if (ValueType != OtherType.ValueType)
+	if (Lhs.IsNumericVector() && Rhs.IsNumericVector())
 	{
-		const FValueTypeDescription TypeDesc = GetValueTypeDescription(ValueType);
-		const FValueTypeDescription OtherTypeDesc = GetValueTypeDescription(OtherType);
-		const EValueComponentType ComponentType = CombineComponentTypes(TypeDesc.ComponentType, OtherTypeDesc.ComponentType);
+		const FValueTypeDescription LhsDesc = GetValueTypeDescription(Lhs);
+		const FValueTypeDescription RhsDesc = GetValueTypeDescription(Rhs);
+		const EValueComponentType ComponentType = CombineComponentTypes(LhsDesc.ComponentType, RhsDesc.ComponentType);
 		if (ComponentType == EValueComponentType::Void)
 		{
-			return false;
+			return Shader::EValueType::Void;
 		}
 
-		const int8 NumComponents = FMath::Max(TypeDesc.NumComponents, OtherTypeDesc.NumComponents);
-		ValueType = MakeValueType(ComponentType, NumComponents);
+		const int8 NumComponents = FMath::Max(LhsDesc.NumComponents, RhsDesc.NumComponents);
+		return MakeValueType(ComponentType, NumComponents);
 	}
-	return true;
+
+	// Non-numeric types can only combine if they are the same
+	if (Lhs == Rhs)
+	{
+		return Lhs;
+	}
+
+	return Shader::EValueType::Void;
 }
 
 const FStructField* FStructType::FindFieldByName(const TCHAR* InName) const
@@ -436,6 +434,7 @@ FValueComponentTypeDescription GetValueComponentTypeDescription(EValueComponentT
 	case EValueComponentType::Double: return FValueComponentTypeDescription(TEXT("double"), sizeof(double), EComponentBound::NegDoubleMax, EComponentBound::DoubleMax);
 	case EValueComponentType::Int: return FValueComponentTypeDescription(TEXT("int"), sizeof(int32), EComponentBound::IntMin, EComponentBound::IntMax);
 	case EValueComponentType::Bool: return FValueComponentTypeDescription(TEXT("bool"), 1u, EComponentBound::Zero, EComponentBound::One);
+	case EValueComponentType::Numeric: return FValueComponentTypeDescription(TEXT("Numeric"), sizeof(double), EComponentBound::NegDoubleMax, EComponentBound::DoubleMax);
 	default: checkNoEntry() return FValueComponentTypeDescription();
 	}
 }
@@ -454,6 +453,10 @@ EValueComponentType CombineComponentTypes(EValueComponentType Lhs, EValueCompone
 	{
 		return Lhs;
 	}
+	else if (Lhs == EValueComponentType::Numeric && Rhs == EValueComponentType::Numeric)
+	{
+		return EValueComponentType::Numeric;
+	}
 	else if (Lhs == EValueComponentType::Double || Rhs == EValueComponentType::Double)
 	{
 		return EValueComponentType::Double;
@@ -479,7 +482,7 @@ const TCHAR* FValueComponent::ToString(EValueComponentType Type, FStringBuilderB
 	case EValueComponentType::Int: OutString.Appendf(TEXT("%d"), Int); break;
 	case EValueComponentType::Bool: OutString.Append(AsBool() ? TEXT("true") : TEXT("false")); break;
 	case EValueComponentType::Float: OutString.Appendf(TEXT("%#.9gf"), Float); break;
-	default: checkNoEntry(); break; // TODO - double
+	default: checkNoEntry(); break; // TODO - double, Numeric
 	}
 	return OutString.ToString();
 }
@@ -500,6 +503,7 @@ const TCHAR* FValue::ToString(EValueStringFormat Format, FStringBuilderBase& Out
 		else
 		{
 			const FValueTypeDescription TypeDesc = GetValueTypeDescription(Type.ValueType);
+			check(TypeDesc.ComponentType != EValueComponentType::Numeric);
 			if (TypeDesc.ComponentType != EValueComponentType::Double)
 			{
 				OutString.Appendf(TEXT("%s("), TypeDesc.Name);
@@ -520,7 +524,9 @@ const TCHAR* FValue::ToString(EValueStringFormat Format, FStringBuilderBase& Out
 		case EValueComponentType::Int: OutString.Appendf(TEXT("%d"), Component[Index].Int); break;
 		case EValueComponentType::Bool: OutString.Append(Component[Index].Bool ? TEXT("true") : TEXT("false")); break;
 		case EValueComponentType::Float: Private::FormatComponent_Double((double)Component[Index].Float, NumComponents, Format, OutString); break;
-		case EValueComponentType::Double: Private::FormatComponent_Double(Component[Index].Double, NumComponents, Format, OutString); break;
+		case EValueComponentType::Double:
+		case EValueComponentType::Numeric:
+			Private::FormatComponent_Double(Component[Index].Double, NumComponents, Format, OutString); break;
 		default: checkNoEntry(); break;
 		}
 	}
@@ -530,81 +536,6 @@ const TCHAR* FValue::ToString(EValueStringFormat Format, FStringBuilderBase& Out
 		OutString.Append(ClosingSuffix);
 	}
 
-
-
-
-	/*if (Format == EValueStringFormat::HLSL && ComponentType == EValueComponentType::Double)
-	{
-		// Construct an HLSL LWC Vector
-		TStringBuilder<256> TileValue;
-		TStringBuilder<256> OffsetValue;
-		for (int32 i = 0; i < NumComponents; ++i)
-		{
-			if (i > 0)
-			{
-				TileValue.Append(TEXT(", "));
-				OffsetValue.Append(TEXT(", "));
-			}
-
-			const FLargeWorldRenderScalar Value(Component[i].Double);
-			Private::FormatComponent_Double(Value.GetTileAsDouble(), NumComponents, Format, TileValue);
-			Private::FormatComponent_Double(Value.GetOffsetAsDouble(), NumComponents, Format, OffsetValue);
-		}
-
-		if (NumComponents > 1)
-		{
-			OutString.Appendf(TEXT("MakeLWCVector%d(float%d(%s), float%d(%s))"), NumComponents, NumComponents, TileValue.ToString(), NumComponents, OffsetValue.ToString());
-		}
-		else
-		{
-			OutString.Appendf(TEXT("MakeLWCScalar(%s, %s)"), TileValue.ToString(), OffsetValue.ToString());
-		}
-	}
-	else
-	{
-		if (Format == EValueStringFormat::HLSL)
-		{
-			const TCHAR* ComponentName = nullptr;
-			switch (ComponentType)
-			{
-			case EValueComponentType::Int: ComponentName = TEXT("int"); break;
-			case EValueComponentType::Bool: ComponentName = TEXT("bool"); break;
-			case EValueComponentType::Float: ComponentName = TEXT("float"); break;
-			default: checkNoEntry(); break; // Double is handled by above case
-			}
-			if (NumComponents > 1)
-			{
-				OutString.Appendf(TEXT("%s%d("), ComponentName, NumComponents);
-			}
-			else
-			{
-				OutString.Appendf(TEXT("%s("), ComponentName);
-			}
-		}
-
-		for (int32 i = 0; i < NumComponents; ++i)
-		{
-			if (i > 0)
-			{
-				OutString.Append(TEXT(", "));
-			}
-
-			switch (ComponentType)
-			{
-			case EValueComponentType::Int: OutString.Appendf(TEXT("%d"), Component[i].Int); break;
-			case EValueComponentType::Bool: OutString.Append(Component[i].Bool ? TEXT("true") : TEXT("false")); break;
-			case EValueComponentType::Float: Private::FormatComponent_Double((double)Component[i].Float, NumComponents, Format, OutString); break;
-			case EValueComponentType::Double: Private::FormatComponent_Double(Component[i].Double, NumComponents, Format, OutString); break;
-			default: checkNoEntry(); break;
-			}
-		}
-
-		if (Format == EValueStringFormat::HLSL)
-		{
-			OutString.Append(TEXT(")"));
-		}
-	}*/
-
 	return OutString.ToString();
 }
 
@@ -629,11 +560,17 @@ FValueTypeDescription GetValueTypeDescription(EValueType Type)
 	case EValueType::Bool2: return FValueTypeDescription(TEXT("bool2"), EValueComponentType::Bool, 2);
 	case EValueType::Bool3: return FValueTypeDescription(TEXT("bool3"), EValueComponentType::Bool, 3);
 	case EValueType::Bool4: return FValueTypeDescription(TEXT("bool4"), EValueComponentType::Bool, 4);
+	case EValueType::Numeric1: return FValueTypeDescription(TEXT("Numeric1"), EValueComponentType::Numeric, 1);
+	case EValueType::Numeric2: return FValueTypeDescription(TEXT("Numeric2"), EValueComponentType::Numeric, 2);
+	case EValueType::Numeric3: return FValueTypeDescription(TEXT("Numeric3"), EValueComponentType::Numeric, 3);
+	case EValueType::Numeric4: return FValueTypeDescription(TEXT("Numeric4"), EValueComponentType::Numeric, 4);
 	case EValueType::Float4x4: return FValueTypeDescription(TEXT("float4x4"), EValueComponentType::Float, 16);
 	case EValueType::Double4x4: return FValueTypeDescription(TEXT("FLWCMatrix"), EValueComponentType::Double, 16);
 	case EValueType::DoubleInverse4x4: return FValueTypeDescription(TEXT("FLWCInverseMatrix"), EValueComponentType::Double, 16);
+	case EValueType::Numeric4x4: return FValueTypeDescription(TEXT("Numeric4x4"), EValueComponentType::Numeric, 16);
 	case EValueType::Struct: return FValueTypeDescription(TEXT("struct"), EValueComponentType::Void, 0);
 	case EValueType::Object: return FValueTypeDescription(TEXT("object"), EValueComponentType::Void, 0);
+	case EValueType::Any: return FValueTypeDescription(TEXT("Any"), EValueComponentType::Void, 0);
 	default: checkNoEntry(); return FValueTypeDescription(TEXT("<INVALID>"), EValueComponentType::Void, 0);
 	}
 }
@@ -683,6 +620,17 @@ EValueType MakeValueType(EValueComponentType ComponentType, int32 NumComponents)
 		default: break;
 		}
 		break;
+	case EValueComponentType::Numeric:
+		switch (NumComponents)
+		{
+		case 1: return EValueType::Numeric1;
+		case 2: return EValueType::Numeric2;
+		case 3: return EValueType::Numeric3;
+		case 4: return EValueType::Numeric4;
+		case 16: return EValueType::Numeric4x4;
+		default: break;
+		}
+		break;
 	case EValueComponentType::Int:
 		switch (NumComponents)
 		{
@@ -727,93 +675,32 @@ EValueType MakeValueTypeWithRequestedNumComponents(EValueType BaseType, int8 Req
 EValueType MakeNonLWCType(EValueType Type)
 {
 	const FValueTypeDescription TypeDesc = GetValueTypeDescription(Type);
-	return MakeValueType(MakeNonLWCType(TypeDesc.ComponentType), TypeDesc.NumComponents);
-}
-
-EValueType MakeDerivativeType(EValueType Type)
-{
-	const FValueTypeDescription TypeDesc = GetValueTypeDescription(Type);
-	if (IsNumericType(TypeDesc.ComponentType))
+	check(IsNumericType(TypeDesc.ComponentType));
+	if (TypeDesc.ComponentType == EValueComponentType::Double)
 	{
-		return MakeValueType(EValueComponentType::Float, TypeDesc.NumComponents);
+		return MakeValueType(MakeNonLWCType(TypeDesc.ComponentType), TypeDesc.NumComponents);
 	}
-	return EValueType::Void;
+	return Type;
 }
 
-EValueType MakeArithmeticResultType(EValueType Lhs, EValueType Rhs, FString& OutErrorMessage)
+EValueType MakeConcreteType(EValueType Type)
 {
-	const FValueTypeDescription LhsDesc = GetValueTypeDescription(Lhs);
-	const FValueTypeDescription RhsDesc = GetValueTypeDescription(Rhs);
-	// Types with 0 components are non-arithmetic
-	if (LhsDesc.NumComponents > 0 && RhsDesc.NumComponents > 0)
-	{
-		if (Lhs == Rhs)
-		{
-			return Lhs;
-		}
-
-		EValueComponentType ComponentType = EValueComponentType::Void;
-		if (LhsDesc.ComponentType == RhsDesc.ComponentType)
-		{
-			ComponentType = LhsDesc.ComponentType;
-		}
-		else if (LhsDesc.ComponentType == EValueComponentType::Double || RhsDesc.ComponentType == EValueComponentType::Double)
-		{
-			ComponentType = EValueComponentType::Double;
-		}
-		else if (LhsDesc.ComponentType == EValueComponentType::Float || RhsDesc.ComponentType == EValueComponentType::Float)
-		{
-			ComponentType = EValueComponentType::Float;
-		}
-		else
-		{
-			ComponentType = EValueComponentType::Int;
-		}
-
-		if (ComponentType != EValueComponentType::Void)
-		{
-			if (LhsDesc.NumComponents == 1 || RhsDesc.NumComponents == 1)
-			{
-				// single component type is valid to combine with other type
-				return MakeValueType(ComponentType, FMath::Max(LhsDesc.NumComponents, RhsDesc.NumComponents));
-			}
-			else if (LhsDesc.NumComponents == RhsDesc.NumComponents)
-			{
-				return MakeValueType(ComponentType, LhsDesc.NumComponents);
-			}
-		}
-
-		OutErrorMessage = FString::Printf(TEXT("Arithmetic between types %s and %s are undefined"), LhsDesc.Name, RhsDesc.Name);
-	}
-	else
+	const FValueTypeDescription TypeDesc = GetValueTypeDescription(Type);
+	check(IsNumericType(TypeDesc.ComponentType));
+	if (TypeDesc.ComponentType == EValueComponentType::Numeric)
 	{
-		OutErrorMessage = FString::Printf(TEXT("Attempting to perform arithmetic on non-numeric types: %s %s"), LhsDesc.Name, RhsDesc.Name);
+		return MakeValueType(MakeConcreteType(TypeDesc.ComponentType), TypeDesc.NumComponents);
 	}
-
-	return EValueType::Void;
+	return Type;
 }
 
-EValueType MakeComparisonResultType(EValueType Lhs, EValueType Rhs, FString& OutErrorMessage)
+EValueType MakeDerivativeType(EValueType Type)
 {
-	const FValueTypeDescription LhsDesc = GetValueTypeDescription(Lhs);
-	const FValueTypeDescription RhsDesc = GetValueTypeDescription(Rhs);
-
-	if (Lhs == Rhs)
-	{
-		if (LhsDesc.NumComponents > 0)
-		{
-			return MakeValueType(EValueComponentType::Bool, LhsDesc.NumComponents);
-		}
-		else
-		{
-			OutErrorMessage = FString::Printf(TEXT("Attempting to perform comparison on non-numeric types: %s %s"), LhsDesc.Name, RhsDesc.Name);
-		}
-	}
-	else
+	const FValueTypeDescription TypeDesc = GetValueTypeDescription(Type);
+	if (IsNumericType(TypeDesc.ComponentType))
 	{
-		OutErrorMessage = FString::Printf(TEXT("Comparison between types %s and %s are undefined"), LhsDesc.Name, RhsDesc.Name);
+		return MakeValueType(EValueComponentType::Float, TypeDesc.NumComponents);
 	}
-
 	return EValueType::Void;
 }
 
@@ -1071,6 +958,7 @@ inline FValue UnaryOp(const Operation& Op, const FValue& Value)
 		return FValue();
 	}
 	const FValueTypeDescription TypeDesc = GetValueTypeDescription(Value.Type);
+	check(TypeDesc.ComponentType != EValueComponentType::Numeric);
 	const int8 NumComponents = TypeDesc.NumComponents;
 	
 	FValue Result;
@@ -1127,6 +1015,8 @@ inline FValue BinaryOp(const Operation& Op, const FValue& Lhs, const FValue& Rhs
 	}
 	const FValueTypeDescription LhsDesc = GetValueTypeDescription(Lhs.Type);
 	const FValueTypeDescription RhsDesc = GetValueTypeDescription(Rhs.Type);
+	check(LhsDesc.ComponentType != EValueComponentType::Numeric);
+	check(RhsDesc.ComponentType != EValueComponentType::Numeric);
 	const int8 NumComponents = GetNumComponentsResult(LhsDesc.NumComponents, RhsDesc.NumComponents);
 	
 	FValue Result;
@@ -1179,6 +1069,8 @@ inline FValue CompareOp(const Operation& Op, const FValue& Lhs, const FValue& Rh
 	}
 	const FValueTypeDescription LhsDesc = GetValueTypeDescription(Lhs.Type);
 	const FValueTypeDescription RhsDesc = GetValueTypeDescription(Rhs.Type);
+	check(LhsDesc.ComponentType != EValueComponentType::Numeric);
+	check(RhsDesc.ComponentType != EValueComponentType::Numeric);
 	const int8 NumComponents = GetNumComponentsResult(LhsDesc.NumComponents, RhsDesc.NumComponents);
 
 	FValue Result;
@@ -1263,6 +1155,7 @@ uint32 GetTypeHash(const FValue& Value)
 		case EValueComponentType::Double: ComponentHash = ::GetTypeHash(Value.Component[Index].Double); break;
 		case EValueComponentType::Int: ComponentHash = ::GetTypeHash(Value.Component[Index].Int); break;
 		case EValueComponentType::Bool: ComponentHash = ::GetTypeHash(Value.Component[Index].Bool); break;
+		case EValueComponentType::Numeric: ComponentHash = ::GetTypeHash(Value.Component[Index].Double); break;
 		default: checkNoEntry(); break;
 		}
 		Result = HashCombine(Result, ComponentHash);
diff --git a/Engine/Source/Runtime/Engine/Public/HLSLTree/HLSLTree.h b/Engine/Source/Runtime/Engine/Public/HLSLTree/HLSLTree.h
--- a/Engine/Source/Runtime/Engine/Public/HLSLTree/HLSLTree.h
+++ b/Engine/Source/Runtime/Engine/Public/HLSLTree/HLSLTree.h
@@ -134,12 +134,6 @@ private:
 	friend class FExpressionLocalPHI;
 };
 
-enum class EDefaultRequestedType : uint8
-{
-	None,
-	Any,
-};
-
 /**
  * FRequestedType is used when accessing various FExpression methods.  A requested type is represented by a concrete Shader::FType, along with a bit vector of components.  A bit set in this vector means the given component was requested.
  * This allows FExpressions to track access of particular components, which can enable various optimizations.  For example, a 'Swizzle' expression may only access the 'Y' component of its input.  So it might request a float2 value from
@@ -157,21 +151,24 @@ public:
 	FRequestedType(const Shader::FStructType* InType, bool bDefaultRequest = true) : FRequestedType(Shader::FType(InType), bDefaultRequest) {}
 	FRequestedType(const FName& InType, bool bDefaultRequest = true) : FRequestedType(Shader::FType(InType), bDefaultRequest) {}
 
-	FRequestedType(EDefaultRequestedType InType) : DefaultType(InType) { check(InType != EDefaultRequestedType::None); }
-
-	Shader::FType GetType(const Shader::FType& InDefaultType) const;
-	int32 GetNumComponentsRequested(int32 InDefaultNumComponents) const;
 	Shader::EValueComponentType GetValueComponentType() const;
 
 	bool IsComponentRequested(int32 Index) const
 	{
-		return (DefaultType == EDefaultRequestedType::Any) || (RequestedComponents.IsValidIndex(Index) ? (bool)RequestedComponents[Index] : false);
+		return Type.IsAny() || (RequestedComponents.IsValidIndex(Index) ? (bool)RequestedComponents[Index] : false);
 	}
 
-	bool IsStruct() const { return !IsVoid() && Type.IsStruct(); }
-	bool IsObject() const { return !IsVoid() && Type.IsObject(); }
-	bool IsNumeric() const { return !IsVoid() && Type.IsNumeric(); }
-	bool IsVoid() const { return (DefaultType != EDefaultRequestedType::Any) && RequestedComponents.Find(true) == INDEX_NONE; }
+	/** Check if any of xyzw are requested */
+	bool IsNumericVectorRequested() const
+	{
+		return Type.IsAny() || FMath::IsWithin(RequestedComponents.Find(true), 0, 4);
+	}
+
+	bool IsVoid() const { return Type.IsVoid(); }
+	bool IsStruct() const { return Type.IsStruct(); }
+	bool IsObject() const { return Type.IsObject(); }
+	bool IsNumeric() const { return Type.IsNumeric(); }
+	bool IsEmpty() const { return !Type.IsAny() && !RequestedComponents.Contains(true); }
 
 	void SetComponentRequest(int32 Index, bool bRequest = true);
 
@@ -194,9 +191,6 @@ public:
 	 */
 	Shader::FType Type;
 
-	/** Used to signify 'special' requested types */
-	EDefaultRequestedType DefaultType = EDefaultRequestedType::None;
-
 	/** 1 bit per component, a value of 'true' means the specified component is requsted */
 	TBitArray<> RequestedComponents;
 };
@@ -287,16 +281,21 @@ public:
 	void SetField(const Shader::FStructField* Field, const FPreparedType& FieldType);
 	FPreparedType GetFieldType(const Shader::FStructField* Field) const;
 
-	int32 GetNumComponents() const;
+	/** For numeric vectors, the number of components in the prepared type will only include components that have valid evaluation */
+	int32 GetNumPreparedComponents() const;
+	Shader::FType GetPreparedType() const;
+
+	/** Converts to a requested type, based on IsRequestedEvaluation() */
 	FRequestedType GetRequestedType() const;
-	Shader::FType GetType() const;
+	
 	Shader::EValueComponentType GetValueComponentType() const;
-	const TCHAR* GetName() const { return GetType().GetName(); }
-	bool IsStruct() const { return !IsVoid() && Type.IsStruct(); }
-	bool IsObject() const { return !IsVoid() && Type.IsObject(); }
-	bool IsNumeric() const { return !IsVoid() && Type.IsNumeric(); }
-	bool IsInitialized() const { return !Type.IsVoid(); }
-	bool IsVoid() const;
+	const TCHAR* GetName() const { return Type.GetName(); }
+	bool IsVoid() const { return Type.IsVoid(); }
+	bool IsStruct() const { return Type.IsStruct(); }
+	bool IsObject() const { return Type.IsObject(); }
+	bool IsNumeric() const { return Type.IsNumeric(); }
+	bool IsNumericScalar() const { return Type.IsNumericScalar(); }
+	bool IsEmpty() const;
 
 	EExpressionEvaluation GetEvaluation(const FEmitScope& Scope) const;
 	EExpressionEvaluation GetEvaluation(const FEmitScope& Scope, const FRequestedType& RequestedType) const;
diff --git a/Engine/Source/Runtime/Engine/Public/Shader/ShaderTypes.h b/Engine/Source/Runtime/Engine/Public/Shader/ShaderTypes.h
--- a/Engine/Source/Runtime/Engine/Public/Shader/ShaderTypes.h
+++ b/Engine/Source/Runtime/Engine/Public/Shader/ShaderTypes.h
@@ -62,6 +62,9 @@ enum class EValueComponentType : uint8
 	Int,
 	Bool,
 
+	// May be any numeric type, stored internally as 'double' within FValue
+	Numeric,
+
 	Num,
 };
 static constexpr int32 NumValueComponentTypes = (int32)EValueComponentType::Num;
@@ -85,6 +88,8 @@ inline bool IsComponentTypeWithinBounds(EValueComponentType Type, FComponentBoun
 ENGINE_API EValueComponentType CombineComponentTypes(EValueComponentType Lhs, EValueComponentType Rhs);
 
 inline EValueComponentType MakeNonLWCType(EValueComponentType Type) { return Type == EValueComponentType::Double ? EValueComponentType::Float : Type; }
+inline EValueComponentType MakeConcreteType(EValueComponentType Type) { return Type == EValueComponentType::Numeric ? EValueComponentType::Float : Type; }
+inline bool IsGenericType(EValueComponentType Type) { return Type == EValueComponentType::Numeric; }
 inline bool IsLWCType(EValueComponentType Type) { return Type == EValueComponentType::Double; }
 
 enum class EValueType : uint8
@@ -111,6 +116,12 @@ enum class EValueType : uint8
 	Bool3,
 	Bool4,
 
+	// Any scalar/vector type
+	Numeric1,
+	Numeric2,
+	Numeric3,
+	Numeric4,
+
 	// float4x4
 	Float4x4,
 
@@ -119,8 +130,12 @@ enum class EValueType : uint8
 	Double4x4,
 	DoubleInverse4x4,
 
+	// Any matrix type
+	Numeric4x4,
+
 	Struct,
 	Object,
+	Any,
 
 	Num,
 };
@@ -139,16 +154,21 @@ struct FValueTypeDescription
 ENGINE_API FValueTypeDescription GetValueTypeDescription(EValueType Type);
 ENGINE_API EValueType FindValueType(FName Name);
 inline bool IsLWCType(EValueType Type) { return IsLWCType(GetValueTypeDescription(Type).ComponentType); }
+inline bool IsGenericType(EValueType Type) { return Type == EValueType::Any || IsGenericType(GetValueTypeDescription(Type).ComponentType); }
 ENGINE_API EValueType MakeValueType(EValueComponentType ComponentType, int32 NumComponents);
 ENGINE_API EValueType MakeValueType(EValueType BaseType, int32 NumComponents);
 ENGINE_API EValueType MakeValueTypeWithRequestedNumComponents(EValueType BaseType, int8 RequestedNumComponents);
 ENGINE_API EValueType MakeNonLWCType(EValueType Type);
+ENGINE_API EValueType MakeConcreteType(EValueType Type);
 ENGINE_API EValueType MakeDerivativeType(EValueType Type);
-ENGINE_API EValueType MakeArithmeticResultType(EValueType Lhs, EValueType Rhs, FString& OutErrorMessage);
-ENGINE_API EValueType MakeComparisonResultType(EValueType Lhs, EValueType Rhs, FString& OutErrorMessage);
 
 inline bool IsNumericType(EValueType Type) { return IsNumericType(GetValueTypeDescription(Type).ComponentType); }
 
+inline bool IsNumericScalarType(EValueType Type)
+{
+	const FValueTypeDescription TypeDesc = GetValueTypeDescription(Type);
+	return IsNumericType(TypeDesc.ComponentType) && TypeDesc.NumComponents == 1;
+}
 inline bool IsNumericVectorType(EValueType Type)
 {
 	const FValueTypeDescription TypeDesc = GetValueTypeDescription(Type);
@@ -170,10 +190,14 @@ struct FType
 	const TCHAR* GetName() const;
 	FType GetDerivativeType() const;
 	FType GetNonLWCType() const { return IsNumericLWC() ? FType(MakeNonLWCType(ValueType)) : *this; }
+	FType GetConcreteType() const { return IsNumeric() ? FType(MakeConcreteType(ValueType)) : *this; }
 	bool IsVoid() const { return ValueType == EValueType::Void; }
 	bool IsStruct() const { return ValueType == EValueType::Struct; }
 	bool IsObject() const { return ValueType == EValueType::Object; }
+	bool IsAny() const { return ValueType == EValueType::Any; }
+	bool IsGeneric() const { return !IsStruct() && !IsObject() && IsGenericType(ValueType); }
 	bool IsNumeric() const { return !IsStruct() && !IsObject() && IsNumericType(ValueType); }
+	bool IsNumericScalar() const { return !IsStruct() && !IsObject() && IsNumericScalarType(ValueType); }
 	bool IsNumericVector() const { return !IsStruct() && !IsObject() && IsNumericVectorType(ValueType); }
 	bool IsNumericMatrix() const { return !IsStruct() && !IsObject() && IsNumericMatrixType(ValueType); }
 	bool IsNumericLWC() const { return IsNumeric() && IsLWCType(ValueType); }
@@ -181,7 +205,6 @@ struct FType
 	int32 GetNumFlatFields() const;
 	EValueComponentType GetComponentType(int32 Index) const;
 	EValueType GetFlatFieldType(int32 Index) const;
-	bool Merge(const FType& OtherType);
 
 	inline operator EValueType() const { return ValueType; }
 	inline operator bool() const { return !IsVoid(); }
@@ -222,6 +245,8 @@ inline bool operator!=(const EValueType& Lhs, const FType& Rhs)
 	return !operator==(Lhs, Rhs);
 }
 
+ENGINE_API FType CombineTypes(const FType& Lhs, const FType& Rhs);
+
 struct FStructField
 {
 	const TCHAR* Name;
