Actor::ActorGuid must not be reset in AActor::Serialize() when duplicating for PIE

diff --git a/Engine/Source/Runtime/Engine/Private/Actor.cpp b/Engine/Source/Runtime/Engine/Private/Actor.cpp
--- a/Engine/Source/Runtime/Engine/Private/Actor.cpp
+++ b/Engine/Source/Runtime/Engine/Private/Actor.cpp
@@ -870,7 +870,7 @@ void AActor::Serialize(FArchive& Ar)
 				ActorGuid.Invalidate();
 			}
 		}
-		else if ((Ar.GetPortFlags() & PPF_Duplicate) || (Ar.IsPersistent() && !ActorGuid.IsValid()))
+		else if (Ar.IsPersistent() && !ActorGuid.IsValid())
 		{
 			ActorGuid = FGuid::NewGuid();
 		}
