Fixed crash when changing the skeletal mesh on a skeletal mesh actor that has an anim class from a template AnimBP

Added nullptr check when calling GetTargetSkeleton as templated anim BPs have no skeleton

diff --git a/Engine/Source/Runtime/Engine/Private/Components/SkeletalMeshComponent.cpp b/Engine/Source/Runtime/Engine/Private/Components/SkeletalMeshComponent.cpp
--- a/Engine/Source/Runtime/Engine/Private/Components/SkeletalMeshComponent.cpp
+++ b/Engine/Source/Runtime/Engine/Private/Components/SkeletalMeshComponent.cpp
@@ -3795,18 +3795,23 @@ void USkeletalMeshComponent::ValidateAnimation()
 	else if (AnimationMode == EAnimationMode::AnimationBlueprint)
 	{
 		IAnimClassInterface* AnimClassInterface = IAnimClassInterface::GetFromClass(AnimClass);
-		if (AnimClassInterface && SkeletalMesh && !AnimClassInterface->GetTargetSkeleton()->IsCompatible(SkeletalMesh->GetSkeleton()))
+		if (AnimClassInterface && SkeletalMesh)
 		{
-			if(SkeletalMesh->GetSkeleton())
+			if (AnimClassInterface->GetTargetSkeleton() == nullptr)
 			{
-				UE_LOG(LogAnimation, Warning, TEXT("AnimBP %s is incompatible with skeleton %s, removing AnimBP from actor."), *AnimClass->GetName(), *SkeletalMesh->GetSkeleton()->GetName());
+				UE_LOG(LogAnimation, Warning, TEXT("AnimBP %s is incompatible because it has no target skeleton, removing AnimBP from actor."), *AnimClass->GetName());
+				AnimClass = nullptr;
 			}
-			else
+			else if (SkeletalMesh->GetSkeleton() == nullptr)
 			{
 				UE_LOG(LogAnimation, Warning, TEXT("AnimBP %s is incompatible because mesh %s has no skeleton, removing AnimBP from actor."), *AnimClass->GetName(), *SkeletalMesh->GetName());
+				AnimClass = nullptr;
+			}
+			else if (!AnimClassInterface->GetTargetSkeleton()->IsCompatible(SkeletalMesh->GetSkeleton()))
+			{
+				UE_LOG(LogAnimation, Warning, TEXT("AnimBP %s is incompatible with skeleton %s, removing AnimBP from actor."), *AnimClass->GetName(), *SkeletalMesh->GetSkeleton()->GetName());
+				AnimClass = nullptr;
 			}
-
-			AnimClass = nullptr;
 		}
 	}
 }
