PR #8725:  Remove unused field from FNetworkPredictionData_Client_Character. This is likely a duplication of a field from FSavedMove_Character.

diff --git a/Engine/Source/Runtime/Engine/Classes/GameFramework/CharacterMovementComponent.h b/Engine/Source/Runtime/Engine/Classes/GameFramework/CharacterMovementComponent.h
--- a/Engine/Source/Runtime/Engine/Classes/GameFramework/CharacterMovementComponent.h
+++ b/Engine/Source/Runtime/Engine/Classes/GameFramework/CharacterMovementComponent.h
@@ -2966,9 +2966,6 @@ public:
 	int32 MaxFreeMoveCount;					// Limit on size of free list
 	int32 MaxSavedMoveCount;				// Limit on the size of the saved move buffer
 
-	/** RootMotion saved while animation is updated, so we can store it and replay if needed in case of a position correction. */
-	FRootMotionMovementParams RootMotionMovement;
-
 	uint32 bUpdatePosition:1; // when true, update the position (via ClientUpdatePosition)
 
 	// Mesh smoothing variables (for network smoothing)
