Fixed deprecation error in CitySample. In general, we shouldn't be calling the Conv_ functions directly from C++. They're intended to be BP wrappers.

Additionally, the deprecation version needs to be 5.1.

diff --git a/Engine/Source/Runtime/Engine/Classes/Kismet/KismetTextLibrary.h b/Engine/Source/Runtime/Engine/Classes/Kismet/KismetTextLibrary.h
--- a/Engine/Source/Runtime/Engine/Classes/Kismet/KismetTextLibrary.h
+++ b/Engine/Source/Runtime/Engine/Classes/Kismet/KismetTextLibrary.h
@@ -210,7 +210,7 @@ class ENGINE_API UKismetTextLibrary : public UBlueprintFunctionLibrary
 	
 	// Default values are duplicated from FNumberFormattingOptions and should be replicated in all functions and in the struct when changed!
 	/** Converts a passed in float to text based on formatting options */
-	UE_DEPRECATED(5.0, "This method has been deprecated and will be removed.")
+	UE_DEPRECATED(5.1, "This method has been deprecated and will be removed.")
 	static FText Conv_FloatToText(float Value, TEnumAsByte<ERoundingMode> RoundingMode, bool bAlwaysSign = false, bool bUseGrouping = true, int32 MinimumIntegralDigits = 1, int32 MaximumIntegralDigits = 324, int32 MinimumFractionalDigits = 0, int32 MaximumFractionalDigits = 3);
 	
 	/**
