Backout CL19387283: problems has been found with old pending kill behavior by licensees and this needs further investigation to cover all cases.

#rnx

diff --git a/Engine/Source/Runtime/Engine/Classes/Engine/World.h b/Engine/Source/Runtime/Engine/Classes/Engine/World.h
--- a/Engine/Source/Runtime/Engine/Classes/Engine/World.h
+++ b/Engine/Source/Runtime/Engine/Classes/Engine/World.h
@@ -421,8 +421,6 @@ public:
  */
 struct ENGINE_API FLevelStreamingGCHelper
 {
-	friend UWorld;
-
 	/** Called when streamed out levels are going to be garbage collected  */
 	DECLARE_MULTICAST_DELEGATE(FOnGCStreamedOutLevelsEvent);
 	static FOnGCStreamedOutLevelsEvent OnGCStreamedOutLevels;
diff --git a/Engine/Source/Runtime/Engine/Private/World.cpp b/Engine/Source/Runtime/Engine/Private/World.cpp
--- a/Engine/Source/Runtime/Engine/Private/World.cpp
+++ b/Engine/Source/Runtime/Engine/Private/World.cpp
@@ -5052,42 +5052,26 @@ void UWorld::CleanupWorldInternal(bool bSessionEnded, bool bCleanupResources, bo
 				});
 			}
 		}
-	}
 
-	auto CleanupLevelResourcesAndReferences = [](ULevel* Level, bool bCleanupResources, bool bUnloadFromEditor)
-	{
-		if (Level)
+		// Cleanup Persistent level outside of following loop because uninitialized worlds don't have a valid Levels array
+		// StreamingLevels are not initialized.
+		if (PersistentLevel)
 		{
-			Level->CleanupLevel(bCleanupResources, bUnloadFromEditor);
-			Level->CleanupReferences();
+			PersistentLevel->CleanupLevel(bCleanupResources);
+			PersistentLevel->CleanupReferences();
 		}
-	};
-
-	// Cleanup Persistent level outside of following loop because uninitialized worlds don't have a valid Levels array
-	// StreamingLevels are not initialized.
-	CleanupLevelResourcesAndReferences(PersistentLevel, bCleanupResources, bUnloadFromEditor);
 
-	if (GetNumLevels() > 1)
-	{
-		check(GetLevel(0) == PersistentLevel);
-		for (int32 LevelIndex = 1; LevelIndex < GetNumLevels(); ++LevelIndex)
+		if (GetNumLevels() > 1)
 		{
-			CleanupLevelResourcesAndReferences(GetLevel(LevelIndex), bCleanupResources, bUnloadFromEditor);
+			check(GetLevel(0) == PersistentLevel);
+			for (int32 LevelIndex = 1; LevelIndex < GetNumLevels(); ++LevelIndex)
+			{
+				ULevel* Level = GetLevel(LevelIndex);
+				Level->CleanupLevel(bCleanupResources);
+				Level->CleanupReferences();
+			}
 		}
 	}
-
-	// Also cleanup levels pending a GC purge
-	for (int32 LevelIndex = 0; LevelIndex < FLevelStreamingGCHelper::LevelsPendingUnload.Num(); ++LevelIndex)
-	{
-		CleanupLevelResourcesAndReferences(FLevelStreamingGCHelper::LevelsPendingUnload[LevelIndex].Get(), bCleanupResources, bUnloadFromEditor);
-	}
-	FLevelStreamingGCHelper::LevelsPendingUnload.Empty();
-#else
-	if (PersistentLevel)
-	{
-		check(bCleanupResources);
-		PersistentLevel->CleanupLevel();
-	}
 #endif //WITH_EDITOR
 
 	for (int32 LevelIndex = 0; LevelIndex < GetNumLevels(); ++LevelIndex)
