Async Flag - Fix deadlock issue when doing a garbage collection of skeletalmesh that is currently building asynchronously.
#rnx

diff --git a/Engine/Source/Runtime/Engine/Private/SkeletalMesh.cpp b/Engine/Source/Runtime/Engine/Private/SkeletalMesh.cpp
--- a/Engine/Source/Runtime/Engine/Private/SkeletalMesh.cpp
+++ b/Engine/Source/Runtime/Engine/Private/SkeletalMesh.cpp
@@ -2151,6 +2151,8 @@ void USkeletalMesh::BeginBuildInternal(FSkeletalMeshBuildContext& Context)
 {
 	TRACE_CPUPROFILER_EVENT_SCOPE(USkeletalMesh::BeginBuildInternal);
 
+	SetInternalFlags(EInternalObjectFlags::Async);
+	
 	// Unregister all instances of this component
 	Context.RecreateRenderStateContext = MakeUnique<FSkinnedMeshComponentRecreateRenderStateContext>(this, false);
 
@@ -2204,6 +2206,8 @@ void USkeletalMesh::FinishBuildInternal(FSkeletalMeshBuildContext& Context)
 {
 	TRACE_CPUPROFILER_EVENT_SCOPE(USkeletalMesh::FinishBuildInternal);
 
+	ClearInternalFlags(EInternalObjectFlags::Async);
+
 	ReleaseAsyncProperty();
 
 	ApplyFinishBuildInternalData(&Context);
@@ -3276,6 +3280,8 @@ void USkeletalMesh::BeginPostLoadInternal(FSkeletalMeshPostLoadContext& Context)
 #if WITH_EDITOR
 	TRACE_CPUPROFILER_EVENT_SCOPE(UStaticMesh::BeginPostLoadInternal);
 
+	SetInternalFlags(EInternalObjectFlags::Async);
+
 	// Lock all properties that should not be modified/accessed during async post-load
 	AcquireAsyncProperty();
 
@@ -3524,6 +3530,9 @@ void USkeletalMesh::FinishPostLoadInternal(FSkeletalMeshPostLoadContext& Context
 	TRACE_CPUPROFILER_EVENT_SCOPE(USkeletalMesh::FinishPostLoadInternal);
 	
 #if WITH_EDITOR
+
+	ClearInternalFlags(EInternalObjectFlags::Async);
+
 	// This scope allows us to use any locked properties without causing stalls
 	FSkeletalMeshAsyncBuildScope AsyncBuildScope(this);
 
