Fix a flicking bug of mobile subsurface profile when updating the subsurface profile asset.

diff --git a/Engine/Source/Runtime/Engine/Private/Rendering/SubsurfaceProfile.cpp b/Engine/Source/Runtime/Engine/Private/Rendering/SubsurfaceProfile.cpp
--- a/Engine/Source/Runtime/Engine/Private/Rendering/SubsurfaceProfile.cpp
+++ b/Engine/Source/Runtime/Engine/Private/Rendering/SubsurfaceProfile.cpp
@@ -294,7 +294,7 @@ IPooledRenderTarget* FSubsurfaceProfileTexture::GetSSProfilesPreIntegratedTextur
 			FPixelShaderUtils::AddFullscreenPass(GraphBuilder, GlobalShaderMap, RDG_EVENT_NAME("SSS::SSProfilePreIntegrated"), PixelShader, PassParameters, ViewRect);
 		}
 
-		GraphBuilder.QueueTextureExtraction(ProfileTexture, &GSSProfilesPreIntegratedTexture);
+		GSSProfilesPreIntegratedTexture = ConvertToFinalizedExternalTexture(GraphBuilder, ProfileTexture);
 	}
 
 	return GSSProfilesPreIntegratedTexture;
