WorldPartition: Fix failure to load embedded Level Instances with multiple PIE worlds (broken in 19740564)

#rnx

diff --git a/Engine/Source/Runtime/Engine/Private/WorldPartition/WorldPartitionLevelHelper.cpp b/Engine/Source/Runtime/Engine/Private/WorldPartition/WorldPartitionLevelHelper.cpp
--- a/Engine/Source/Runtime/Engine/Private/WorldPartition/WorldPartitionLevelHelper.cpp
+++ b/Engine/Source/Runtime/Engine/Private/WorldPartition/WorldPartitionLevelHelper.cpp
@@ -311,7 +311,14 @@ bool FWorldPartitionLevelHelper::LoadActors(UWorld* InOwningWorld, ULevel* InDes
 		if (!Context)
 		{
 			check(!PackageObjectMapping.ContainerID.IsMainContainer());
-			const FName ContainerPackageInstanceName(*FString::Printf(TEXT("/Temp%s_%s"), *PackageObjectMapping.ContainerPackage.ToString(), *PackageObjectMapping.ContainerID.ToString()));
+		
+			FString ContainerPackageName = PackageObjectMapping.ContainerPackage.ToString();
+			if (InDestLevel && InDestLevel->GetPackage()->GetPIEInstanceID() != INDEX_NONE)
+			{
+				ContainerPackageName = UWorld::ConvertToPIEPackageName(ContainerPackageName, InDestLevel->GetPackage()->GetPIEInstanceID());
+			}
+			
+			const FName ContainerPackageInstanceName(*FString::Printf(TEXT("/Temp%s_%s"), *ContainerPackageName, *PackageObjectMapping.ContainerID.ToString()));
 
 			FLinkerInstancingContext& NewContext = LinkerInstancingContexts.Add(PackageObjectMapping.ContainerID);
 			NewContext.AddTag(ULevel::DontLoadExternalObjectsTag);
