Fix world unloading during async skeletalmesh build
#rnx

diff --git a/Engine/Source/Runtime/Engine/Private/SkeletalMesh.cpp b/Engine/Source/Runtime/Engine/Private/SkeletalMesh.cpp
--- a/Engine/Source/Runtime/Engine/Private/SkeletalMesh.cpp
+++ b/Engine/Source/Runtime/Engine/Private/SkeletalMesh.cpp
@@ -7274,16 +7274,17 @@ FSkinnedMeshComponentRecreateRenderStateContext::~FSkinnedMeshComponentRecreateR
 	const int32 ComponentCount = MeshComponents.Num();
 	for (int32 ComponentIndex = 0; ComponentIndex < ComponentCount; ++ComponentIndex)
 	{
-		USkinnedMeshComponent* Component = MeshComponents[ComponentIndex];
-
-		if (bRefreshBounds)
+		if(USkinnedMeshComponent* Component = MeshComponents[ComponentIndex].Get())
 		{
-			Component->UpdateBounds();
-		}
+			if (bRefreshBounds)
+			{
+				Component->UpdateBounds();
+			}
 
-		if (Component->IsRegistered() && !Component->IsRenderStateCreated() && Component->ShouldCreateRenderState())
-		{
-			Component->CreateRenderState_Concurrent(nullptr);
+			if (Component->IsRegistered() && !Component->IsRenderStateCreated() && Component->ShouldCreateRenderState())
+			{
+				Component->CreateRenderState_Concurrent(nullptr);
+			}
 		}
 	}
 }
diff --git a/Engine/Source/Runtime/Engine/Public/SkeletalMeshTypes.h b/Engine/Source/Runtime/Engine/Public/SkeletalMeshTypes.h
--- a/Engine/Source/Runtime/Engine/Public/SkeletalMeshTypes.h
+++ b/Engine/Source/Runtime/Engine/Public/SkeletalMeshTypes.h
@@ -7,6 +7,7 @@
 #include "Materials/MaterialInterface.h"
 #include "ComponentReregisterContext.h"
 #include "SkeletalMeshLegacyCustomVersions.h"
+#include "UObject/WeakObjectPtrTemplates.h"
 
 class FMaterialRenderProxy;
 class FMeshElementCollector;
@@ -350,7 +351,7 @@ public:
 private:
 
 	/** List of components to reset */
-	TArray< class USkinnedMeshComponent*> MeshComponents;
+	TArray< TWeakObjectPtr<USkinnedMeshComponent>> MeshComponents;
 
 	/** Whether we'll refresh the component bounds as we reset */
 	bool bRefreshBounds;
