Fix 'Enclosing block should never be called' from UDeviceProfileManager::FindMatchingFragments

diff --git a/Engine/Source/Runtime/Engine/Classes/DeviceProfiles/DeviceProfileManager.h b/Engine/Source/Runtime/Engine/Classes/DeviceProfiles/DeviceProfileManager.h
--- a/Engine/Source/Runtime/Engine/Classes/DeviceProfiles/DeviceProfileManager.h
+++ b/Engine/Source/Runtime/Engine/Classes/DeviceProfiles/DeviceProfileManager.h
@@ -258,8 +258,11 @@ private:
 
 	/** Get the current platform's the selector module. Can return null */
 	static class IDeviceProfileSelectorModule* GetDeviceProfileSelectorModule();
+
+#if ALLOW_OTHER_PLATFORM_CONFIG && WITH_EDITOR
 	/** Get another platform's selector module. Can return null */
 	static class IDeviceProfileSelectorModule* GetPreviewDeviceProfileSelectorModule(class FConfigCacheIni* PreviewConfigSystem);
+#endif
 public:
 
 	static class UDeviceProfileManager* DeviceProfileManagerSingleton;
diff --git a/Engine/Source/Runtime/Engine/Private/DeviceProfiles/DeviceProfileManager.cpp b/Engine/Source/Runtime/Engine/Private/DeviceProfiles/DeviceProfileManager.cpp
--- a/Engine/Source/Runtime/Engine/Private/DeviceProfiles/DeviceProfileManager.cpp
+++ b/Engine/Source/Runtime/Engine/Private/DeviceProfiles/DeviceProfileManager.cpp
@@ -1102,9 +1102,9 @@ bool UDeviceProfileManager::AreTextureGroupsTheSame(UDeviceProfile* Profile1, UD
 	return false;
 }
 
+#if ALLOW_OTHER_PLATFORM_CONFIG && WITH_EDITOR
 IDeviceProfileSelectorModule* UDeviceProfileManager::GetPreviewDeviceProfileSelectorModule(FConfigCacheIni* PreviewConfigSystemIn)
 {
-#if ALLOW_OTHER_PLATFORM_CONFIG && WITH_EDITOR
 	// If we're getting the selector for previewing the PreviewDeviceProfileSelectionModule, this can be separately configured to preview a target device.
 	FString PreviewDeviceProfileSelectionModuleName;
 	if (PreviewConfigSystemIn->GetString(TEXT("DeviceProfileManager"), TEXT("PreviewDeviceProfileSelectionModule"), PreviewDeviceProfileSelectionModuleName, GEngineIni))
@@ -1115,11 +1115,9 @@ IDeviceProfileSelectorModule* UDeviceProfileManager::GetPreviewDeviceProfileSele
 			return DPSelectorModule;
 		}
 	}
-#else
-	checkNoEntry();
-#endif
 	return nullptr;
 }
+#endif
 
 IDeviceProfileSelectorModule* UDeviceProfileManager::GetDeviceProfileSelectorModule()
 {
diff --git a/Engine/Source/Runtime/Engine/Private/DeviceProfiles/DeviceProfileMatching.cpp b/Engine/Source/Runtime/Engine/Private/DeviceProfiles/DeviceProfileMatching.cpp
--- a/Engine/Source/Runtime/Engine/Private/DeviceProfiles/DeviceProfileMatching.cpp
+++ b/Engine/Source/Runtime/Engine/Private/DeviceProfiles/DeviceProfileMatching.cpp
@@ -627,15 +627,18 @@ TArray<FSelectedFragmentProperties> UDeviceProfileManager::FindMatchingFragments
 	else
 #endif
 	{
-		bool bIsPreview = false;
-#if ALLOW_OTHER_PLATFORM_CONFIG
-		bIsPreview = ConfigSystem != GConfig;
+		IDeviceProfileSelectorModule* PreviewDPSelector = nullptr;
+#if ALLOW_OTHER_PLATFORM_CONFIG && WITH_EDITOR
+		if (ConfigSystem != GConfig)
+		{
+			PreviewDPSelector = GetPreviewDeviceProfileSelectorModule(ConfigSystem);
+			// previewing a DP with matching rules will run if conditions with the host device's data sources. It will likely not match the preview device's behavior.
+			UE_CLOG(!PreviewDPSelector && DPHasMatchingRules(ParentDP, ConfigSystem), LogInit, Warning, TEXT("Preview DP %s contains fragment matching rules, but no preview profile selector was found. The selected fragments for %s will likely not match the behavior of the intended preview device."), *ParentDP, *ParentDP);
+		}
 #endif
-		IDeviceProfileSelectorModule* DPSelector = bIsPreview ? GetPreviewDeviceProfileSelectorModule(ConfigSystem) : GetDeviceProfileSelectorModule();
+		IDeviceProfileSelectorModule* DPSelector = PreviewDPSelector ? PreviewDPSelector : GetDeviceProfileSelectorModule();
 		SelectedFragments = LoadAndProcessMatchingRulesFromConfig(ParentDP, DPSelector, ConfigSystem);
-		
-		// previewing a DP with matching rules will run if conditions with the host device's data sources. It will likely not match the preview device's behavior.
-		UE_CLOG(bIsPreview && !DPSelector && DPHasMatchingRules(ParentDP, ConfigSystem), LogInit, Warning, TEXT("Preview DP %s contains fragment matching rules, but no preview profile selector was found. The selected fragments for %s will likely not match the behavior of the intended preview device."), *ParentDP, *ParentDP);
+	
 	}
 	SelectedFragments = RemoveAllWhiteSpace(SelectedFragments);
 	if(!SelectedFragments.IsEmpty())
