Allowed the OnAnimInitialize event on Skeletal Mesh Components to work with Blueprints in Editor

diff --git a/Engine/Source/Runtime/Engine/Private/Components/SkeletalMeshComponent.cpp b/Engine/Source/Runtime/Engine/Private/Components/SkeletalMeshComponent.cpp
--- a/Engine/Source/Runtime/Engine/Private/Components/SkeletalMeshComponent.cpp
+++ b/Engine/Source/Runtime/Engine/Private/Components/SkeletalMeshComponent.cpp
@@ -786,6 +786,8 @@ void USkeletalMeshComponent::InitAnim(bool bForceReinit)
 
 				if (bInitializedAnimInstance)
 				{
+					// Allow blueprints to respond to the event in editor
+					FEditorScriptExecutionGuard ScriptGuard;
 					OnAnimInitialized.Broadcast();
 				}
 			}
