Fix crash when building HLOD in a commandlet.

- Remove shadermap access that can pose threading hazards.

diff --git a/Engine/Source/Runtime/Engine/Private/DebugViewModeHelpers.cpp b/Engine/Source/Runtime/Engine/Private/DebugViewModeHelpers.cpp
--- a/Engine/Source/Runtime/Engine/Private/DebugViewModeHelpers.cpp
+++ b/Engine/Source/Runtime/Engine/Private/DebugViewModeHelpers.cpp
@@ -333,7 +333,7 @@ bool CompileDebugViewModeShaders(EDebugViewShaderMode ShaderMode, EMaterialQuali
 			{
 				FMaterialShaderTypes ShaderTypes;
 				DebugViewModeInterface->AddShaderTypes(FeatureLevel, LocalVertexFactory, ShaderTypes);
-				if (Material->ShouldCacheShaders(ShaderTypes, LocalVertexFactory) && !Material->HasShaders(ShaderTypes, LocalVertexFactory))
+				if (Material->ShouldCacheShaders(GetFeatureLevelShaderPlatform(FeatureLevel), ShaderTypes, LocalVertexFactory) && !Material->HasShaders(ShaderTypes, LocalVertexFactory))
 				{
 					bMaterialFinished = false;
 				}
diff --git a/Engine/Source/Runtime/Engine/Private/Materials/MaterialShared.cpp b/Engine/Source/Runtime/Engine/Private/Materials/MaterialShared.cpp
--- a/Engine/Source/Runtime/Engine/Private/Materials/MaterialShared.cpp
+++ b/Engine/Source/Runtime/Engine/Private/Materials/MaterialShared.cpp
@@ -3048,9 +3048,8 @@ bool FMaterial::HasShaders(const FMaterialShaderTypes& InTypes, const FVertexFac
 	return TryGetShaders(InTypes, InVertexFactoryType, UnusedShaders);
 }
 
-bool FMaterial::ShouldCacheShaders(const FMaterialShaderTypes& InTypes, const FVertexFactoryType* InVertexFactoryType) const
+bool FMaterial::ShouldCacheShaders(const EShaderPlatform ShaderPlatform, const FMaterialShaderTypes& InTypes, const FVertexFactoryType* InVertexFactoryType) const
 {
-	const EShaderPlatform ShaderPlatform = RenderingThreadShaderMap->GetShaderPlatform();
 	for (int32 FrequencyIndex = 0; FrequencyIndex < SF_NumGraphicsFrequencies; ++FrequencyIndex)
 	{
 		const FShaderType* ShaderType = InTypes.ShaderType[FrequencyIndex];
diff --git a/Engine/Source/Runtime/Engine/Public/MaterialShared.h b/Engine/Source/Runtime/Engine/Public/MaterialShared.h
--- a/Engine/Source/Runtime/Engine/Public/MaterialShared.h
+++ b/Engine/Source/Runtime/Engine/Public/MaterialShared.h
@@ -2059,7 +2059,7 @@ public:
 
 	ENGINE_API bool HasShaders(const FMaterialShaderTypes& InTypes, const FVertexFactoryType* InVertexFactoryType) const;
 
-	ENGINE_API bool ShouldCacheShaders(const FMaterialShaderTypes& InTypes, const FVertexFactoryType* InVertexFactoryType) const;
+	ENGINE_API bool ShouldCacheShaders(const EShaderPlatform ShaderPlatform, const FMaterialShaderTypes& InTypes, const FVertexFactoryType* InVertexFactoryType) const;
 
 	ENGINE_API void SubmitCompileJobs_GameThread(EShaderCompileJobPriority Priority);
 	ENGINE_API void SubmitCompileJobs_RenderThread(EShaderCompileJobPriority Priority) const;
diff --git a/Engine/Source/Runtime/Renderer/Private/DebugViewModeRendering.cpp b/Engine/Source/Runtime/Renderer/Private/DebugViewModeRendering.cpp
--- a/Engine/Source/Runtime/Renderer/Private/DebugViewModeRendering.cpp
+++ b/Engine/Source/Runtime/Renderer/Private/DebugViewModeRendering.cpp
@@ -288,7 +288,7 @@ void FDebugViewModeMeshProcessor::AddMeshBatch(const FMeshBatch& RESTRICT MeshBa
 
 	FMaterialShaderTypes ShaderTypes;
 	DebugViewModeInterface->AddShaderTypes(FeatureLevel, VertexFactoryType, ShaderTypes);
-	if (!Material->ShouldCacheShaders(ShaderTypes, VertexFactoryType))
+	if (!Material->ShouldCacheShaders(GetFeatureLevelShaderPlatform(FeatureLevel), ShaderTypes, VertexFactoryType))
 	{
 		return;
 	}
