Cooked platform ODSC.  Only compile the shaders that are rendered with and not complete shader maps for every cooked material.

# Results
- In my small tests this resulted in compiling 8.9x less shaders (1558 vs. 174)
- Spent 107 seconds less compiling.

# Changes
- Adds CacheGivenTypesForCooking to UMaterialInterface and implemented for UMaterial and UMaterialInstance.  This is responsible for compiling just the given types of shaders which is fully qualified by VF Type, Shader Type (or Shader Pipeline Type).
- Fully qualify each shader compile request from the client by FeatureLevel and QualityLevel.  This way we don't compile multiple quality levels that aren't actively rendered with.
- CacheGivenTypesForCooking for UMaterial and UMaterialInterface explicity tries to re-use and find existing resources so we can progressively build up the shadermap while running.
- FMaterial::CacheGivenTypes now can compile shader pipelines.
- Adds CompileODSCMaterialsForRemoteRecompile which compiles just the requested shaders.  This was because I wanted to leave the functionality in CompileMaterialsForRemoteRecompile alone.
- FShaderRecompileData now has a serialize function so the Load/Save serialize code can be the same on host and client.
- Made CompileMaterialsForRemoteRecompile compiled only for the Editor.
- Don't assert on needing complete shader maps when using ODSC.
- Mass conversion from NULL to nullptr in ShaderCompiler.cpp
- Refactoring some common code from RecompileShadersForRemote into separate functions.

diff --git a/Engine/Source/Editor/MaterialEditor/Private/MaterialStats.cpp b/Engine/Source/Editor/MaterialEditor/Private/MaterialStats.cpp
--- a/Engine/Source/Editor/MaterialEditor/Private/MaterialStats.cpp
+++ b/Engine/Source/Editor/MaterialEditor/Private/MaterialStats.cpp
@@ -207,21 +207,23 @@ bool FShaderPlatformSettings::CheckShaders()
 				TMap<FName, TArray<FMaterialStatsUtils::FRepresentativeShaderInfo>> ShaderTypeNamesAndDescriptions;
 				FMaterialStatsUtils::GetRepresentativeShaderTypesAndDescriptions(ShaderTypeNamesAndDescriptions, Data.MaterialResourcesStats);
 
-				TArray<FVertexFactoryType*> VFTypes;
-				TArray<FShaderType*> ShaderTypes;
+				TArray<const FVertexFactoryType*> VFTypes;
+				TArray<const FShaderPipelineType*> PipelineTypes;
+				TArray<const FShaderType*> ShaderTypes;
 				for (auto& DescriptionPair : ShaderTypeNamesAndDescriptions)
 				{
-					FVertexFactoryType* VFType = FindVertexFactoryType(DescriptionPair.Key);
+					const FVertexFactoryType* VFType = FindVertexFactoryType(DescriptionPair.Key);
 					check(VFType);
 
 					auto& DescriptionArray = DescriptionPair.Value;
 					for (const FMaterialStatsUtils::FRepresentativeShaderInfo& ShaderInfo : DescriptionArray)
 					{
-						FShaderType* ShaderType = FindShaderTypeByName(ShaderInfo.ShaderName);
+						const FShaderType* ShaderType = FindShaderTypeByName(ShaderInfo.ShaderName);
 						if (ShaderType && VFType)
 						{
 							VFTypes.Add(VFType);
 							ShaderTypes.Add(ShaderType);
+							PipelineTypes.Add(nullptr);
 						}
 					}
 				}
@@ -232,7 +234,7 @@ bool FShaderPlatformSettings::CheckShaders()
 				if (bSuccess)
 				{
 					// Compile just the types we want.
-					Data.MaterialResourcesStats->CacheGivenTypes(PlatformShaderID, VFTypes, ShaderTypes);
+					Data.MaterialResourcesStats->CacheGivenTypes(PlatformShaderID, VFTypes, PipelineTypes, ShaderTypes);
 				}
 
 				Data.bCompilingShaders = true;
diff --git a/Engine/Source/Editor/UnrealEd/Private/Cooker/IoStoreCookOnTheFlyRequestManager.cpp b/Engine/Source/Editor/UnrealEd/Private/Cooker/IoStoreCookOnTheFlyRequestManager.cpp
--- a/Engine/Source/Editor/UnrealEd/Private/Cooker/IoStoreCookOnTheFlyRequestManager.cpp
+++ b/Engine/Source/Editor/UnrealEd/Private/Cooker/IoStoreCookOnTheFlyRequestManager.cpp
@@ -1064,14 +1064,8 @@ private:
 
 		FShaderRecompileData RecompileData(PlatformName.ToString(), &RecompileModifiedFiles, &MeshMaterialMaps, &GlobalShaderMap);
 		{
-			int32 iShaderPlatform = static_cast<int32>(RecompileData.ShaderPlatform);
 			TUniquePtr<FArchive> Ar = Request.ReadBody();
-			*Ar << RecompileData.MaterialsToLoad;
-			*Ar << RecompileData.ShaderTypesToLoad;
-			*Ar << iShaderPlatform;
-			RecompileData.ShaderPlatform = static_cast<EShaderPlatform>(iShaderPlatform);
-			*Ar << RecompileData.CommandType;
-			*Ar << RecompileData.ShadersToRecompile;
+			*Ar << RecompileData;
 		}
 
 		RecompileShaders(RecompileData);
diff --git a/Engine/Source/Runtime/Engine/Classes/Materials/Material.h b/Engine/Source/Runtime/Engine/Classes/Materials/Material.h
--- a/Engine/Source/Runtime/Engine/Classes/Materials/Material.h
+++ b/Engine/Source/Runtime/Engine/Classes/Materials/Material.h
@@ -28,6 +28,7 @@ class UPhysicalMaterial;
 class UPhysicalMaterialMask;
 class USubsurfaceProfile;
 class UTexture;
+struct FODSCRequestPayload;
 
 #if WITH_EDITOR
 
@@ -1254,6 +1255,9 @@ public:
 	ENGINE_API bool NeedsSetMaterialUsage_Concurrent(bool &bOutHasUsage, const EMaterialUsage Usage) const;
 
 	ENGINE_API virtual void CacheShaders(EMaterialShaderPrecompileMode CompileMode) override;
+#if WITH_EDITOR
+	ENGINE_API virtual void CacheGivenTypesForCooking(EShaderPlatform ShaderPlatform, ERHIFeatureLevel::Type FeatureLevel, EMaterialQualityLevel::Type QualityLevel, const TArray<const FVertexFactoryType*>& VFTypes, const TArray<const FShaderPipelineType*> PipelineTypes, const TArray<const FShaderType*>& ShaderTypes) override;
+#endif
 	ENGINE_API virtual bool IsComplete() const override;
 
 #if WITH_EDITORONLY_DATA
@@ -1528,6 +1532,7 @@ public:
 	 */
 	ENGINE_API static void RestoreMaterialShadersFromMemory(const TMap<class FMaterialShaderMap*, TUniquePtr<TArray<uint8> > >& ShaderMapToSerializedShaderData);
 
+#if WITH_EDITOR
 	/** Builds a map from UMaterialInterface name to the shader maps that are needed for rendering on the given platform. */
 	ENGINE_API static void CompileMaterialsForRemoteRecompile(
 		const TArray<UMaterialInterface*>& MaterialsToCompile,
@@ -1535,7 +1540,8 @@ public:
 		class ITargetPlatform* TargetPlatform,
 		TMap<FString, TArray<TRefCountPtr<class FMaterialShaderMap> > >& OutShaderMaps);
 
-#if WITH_EDITOR
+	ENGINE_API static void CompileODSCMaterialsForRemoteRecompile(TArray<FODSCRequestPayload> ShadersToRecompile, TMap<FString, TArray<TRefCountPtr<class FMaterialShaderMap>>>& OutShaderMaps);
+
 	/**
 	 * Add an expression node that represents a parameter to the list of material parameters.
 	 * @param	Expression	Pointer to the node that is going to be inserted if it's a parameter type.
diff --git a/Engine/Source/Runtime/Engine/Classes/Materials/MaterialInstance.h b/Engine/Source/Runtime/Engine/Classes/Materials/MaterialInstance.h
--- a/Engine/Source/Runtime/Engine/Classes/Materials/MaterialInstance.h
+++ b/Engine/Source/Runtime/Engine/Classes/Materials/MaterialInstance.h
@@ -871,6 +871,9 @@ public:
 #endif // WITH_EDITOR
 
 	ENGINE_API virtual void CacheShaders(EMaterialShaderPrecompileMode CompileMode) override;
+#if WITH_EDITOR
+	ENGINE_API virtual void CacheGivenTypesForCooking(EShaderPlatform Platform, ERHIFeatureLevel::Type InFeatureLevel, EMaterialQualityLevel::Type QualityLevel, const TArray<const FVertexFactoryType*>& VFTypes, const TArray<const FShaderPipelineType*> PipelineTypes, const TArray<const FShaderType*>& ShaderTypes) override;
+#endif
 	ENGINE_API virtual bool IsComplete() const override;
 
 protected:
diff --git a/Engine/Source/Runtime/Engine/Classes/Materials/MaterialInterface.h b/Engine/Source/Runtime/Engine/Classes/Materials/MaterialInterface.h
--- a/Engine/Source/Runtime/Engine/Classes/Materials/MaterialInterface.h
+++ b/Engine/Source/Runtime/Engine/Classes/Materials/MaterialInterface.h
@@ -822,6 +822,10 @@ public:
 	*/
 	ENGINE_API virtual void CacheShaders(EMaterialShaderPrecompileMode CompileMode = EMaterialShaderPrecompileMode::Default) {}
 
+#if WITH_EDITOR
+	ENGINE_API virtual void CacheGivenTypesForCooking(EShaderPlatform Platform, ERHIFeatureLevel::Type FeatureLevel, EMaterialQualityLevel::Type QualityLevel, const TArray<const FVertexFactoryType*>& VFTypes, const TArray<const FShaderPipelineType*> PipelineTypes, const TArray<const FShaderType*>& ShaderTypes) {}
+#endif
+
 	/** @brief Checks to see if this material has all its shaders cached.
 	*
 	* Materials are not guaranteed to have all their shaders compiled after loading.  It can be useful to
diff --git a/Engine/Source/Runtime/Engine/Private/Materials/Material.cpp b/Engine/Source/Runtime/Engine/Private/Materials/Material.cpp
--- a/Engine/Source/Runtime/Engine/Private/Materials/Material.cpp
+++ b/Engine/Source/Runtime/Engine/Private/Materials/Material.cpp
@@ -2189,7 +2189,7 @@ void UMaterial::CacheResourceShadersForRendering(bool bRegenerateId, EMaterialSh
 			if (AdditionalPlatform != SP_NumPlatforms)
 			{
 				ResourcesToCache.Reset();
-				CacheResourceShadersForCooking(AdditionalPlatform,ResourcesToCache);
+				CacheResourceShadersForCooking(AdditionalPlatform, ResourcesToCache);
 				for (int32 i = 0; i < ResourcesToCache.Num(); ++i)
 				{
 					FMaterialResource* Resource = ResourcesToCache[i];
@@ -2303,6 +2303,28 @@ void UMaterial::CacheShaders(EMaterialShaderPrecompileMode CompileMode)
 	CacheResourceShadersForRendering(false, CompileMode);
 }
 
+#if WITH_EDITOR
+void UMaterial::CacheGivenTypesForCooking(EShaderPlatform ShaderPlatform, ERHIFeatureLevel::Type FeatureLevel, EMaterialQualityLevel::Type QualityLevel, const TArray<const FVertexFactoryType*>& VFTypes, const TArray<const FShaderPipelineType*> PipelineTypes, const TArray<const FShaderType*>& ShaderTypes)
+{
+	TRACE_CPUPROFILER_EVENT_SCOPE(UMaterial::CacheGivenTypes);
+
+	if (QualityLevel == EMaterialQualityLevel::Num)
+	{
+		QualityLevel = GetCachedScalabilityCVars().MaterialQualityLevel;
+	}
+
+	FMaterialResource* CurrentResource = FindOrCreateMaterialResource(MaterialResources, this, nullptr, FeatureLevel, QualityLevel);
+	check(CurrentResource);
+
+	// Prepare the resource for compilation, but don't compile the completed shader map.
+	const bool bSuccess = CurrentResource->CacheShaders(ShaderPlatform, EMaterialShaderPrecompileMode::None);
+	if (bSuccess)
+	{
+		CurrentResource->CacheGivenTypes(ShaderPlatform, VFTypes, PipelineTypes, ShaderTypes);
+	}
+}
+#endif
+
 bool UMaterial::IsComplete() const
 {
 	bool bComplete = true;
@@ -4649,6 +4671,7 @@ void UMaterial::RestoreMaterialShadersFromMemory(const TMap<FMaterialShaderMap*,
 #endif // WITH_EDITOR
 }
 
+#if WITH_EDITOR
 void UMaterial::CompileMaterialsForRemoteRecompile(
 	const TArray<UMaterialInterface*>& MaterialsToCompile,
 	EShaderPlatform ShaderPlatform,
@@ -4701,6 +4724,111 @@ void UMaterial::CompileMaterialsForRemoteRecompile(
 	}
 }
 
+void UMaterial::CompileODSCMaterialsForRemoteRecompile(TArray<FODSCRequestPayload> ShadersToRecompile, TMap<FString, TArray<TRefCountPtr<class FMaterialShaderMap>>>& OutShaderMaps)
+{
+	TRACE_CPUPROFILER_EVENT_SCOPE(UMaterial::CompileODSCMaterialsForRemoteRecompile);
+
+	// Build a map from UMaterial / UMaterialInstance to the resources which are being compiled
+	TMap<FString, TArray<FMaterialResource*> > CompilingResources;
+
+	if (ShadersToRecompile.Num())
+	{
+		UE_LOG(LogShaders, Display, TEXT("Received %d shaders to compile."), ShadersToRecompile.Num());
+	}
+
+	// Group all shader types we need to compile by material so we can issue them all in a single call.
+	struct FShadersToCompile
+	{
+		EShaderPlatform ShaderPlatform;
+		ERHIFeatureLevel::Type FeatureLevel;
+		EMaterialQualityLevel::Type QualityLevel;
+		TArray<const FVertexFactoryType*> VFTypes;
+		TArray<const FShaderPipelineType*> PipelineTypes;
+		TArray<const FShaderType*> ShaderTypes;
+	};
+	TMap<UMaterialInterface*, FShadersToCompile> CoalescedShadersToCompile;
+
+	for (const FODSCRequestPayload& payload : ShadersToRecompile)
+	{
+		UE_LOG(LogShaders, Display, TEXT(""));
+		UE_LOG(LogShaders, Display, TEXT("Material:    %s "), *payload.MaterialName);
+
+		UMaterialInterface* MaterialInterface = LoadObject<UMaterialInterface>(nullptr, *payload.MaterialName);
+		if (MaterialInterface)
+		{
+			FShadersToCompile& Shaders = CoalescedShadersToCompile.FindOrAdd(MaterialInterface);
+			const FVertexFactoryType* VFType = FVertexFactoryType::GetVFByName(payload.VertexFactoryName);
+			const FShaderPipelineType* PipelineType = FShaderPipelineType::GetShaderPipelineTypeByName(payload.PipelineName);
+
+			if (VFType)
+			{
+				UE_LOG(LogShaders, Display, TEXT("VF Type:     %s "), *payload.VertexFactoryName);
+				Shaders.ShaderPlatform = payload.ShaderPlatform;
+				Shaders.FeatureLevel = payload.FeatureLevel;
+				Shaders.QualityLevel = payload.QualityLevel;
+				Shaders.VFTypes.Add(VFType);
+
+				if (PipelineType)
+				{
+					UE_LOG(LogShaders, Display, TEXT("Pipeline Type: %s"), *payload.PipelineName);
+
+					Shaders.PipelineTypes.Add(PipelineType);
+					Shaders.ShaderTypes.Add(nullptr);
+				}
+				else
+				{
+					for (const FString& ShaderTypeName : payload.ShaderTypeNames)
+					{
+						const FShaderType* ShaderType = FShaderType::GetShaderTypeByName(*ShaderTypeName);
+						if (ShaderType)
+						{
+							UE_LOG(LogShaders, Display, TEXT("\tShader Type: %s"), *ShaderTypeName);
+
+							Shaders.PipelineTypes.Add(nullptr);
+							Shaders.ShaderTypes.Add(ShaderType);
+						}
+					}
+				}
+			}
+		}
+	}
+
+	for (const auto& Entry : CoalescedShadersToCompile)
+	{
+		UMaterialInterface* MaterialInterface = Entry.Key;
+		const FShadersToCompile& Shaders = Entry.Value;
+
+		TArray<FMaterialResource*>& ResourceArray = CompilingResources.Add(MaterialInterface->GetPathName(), TArray<FMaterialResource*>());
+		MaterialInterface->CacheGivenTypesForCooking(Shaders.ShaderPlatform, Shaders.FeatureLevel, Shaders.QualityLevel, Shaders.VFTypes, Shaders.PipelineTypes, Shaders.ShaderTypes);
+
+		FMaterialResource* MaterialResource = MaterialInterface->GetMaterialResource(Shaders.FeatureLevel, Shaders.QualityLevel);
+		check(MaterialResource);
+		check(MaterialResource->GetFeatureLevel() == Shaders.FeatureLevel);
+		check(MaterialResource->GetQualityLevel() == Shaders.QualityLevel);
+		ResourceArray.Add(MaterialResource);
+	}
+
+	// Wait until all compilation is finished and all of the gathered FMaterialResources have their GameThreadShaderMap up to date
+	GShaderCompilingManager->FinishAllCompilation();
+
+	// This is heavy handed, but wait until we've set the render thread shader map before proceeding to delete the FMaterialResource below.
+	// This is code that should be run on the cooker so shouldn't be a big deal.
+	FlushRenderingCommands();
+
+	for (TMap<FString, TArray<FMaterialResource*> >::TIterator It(CompilingResources); It; ++It)
+	{
+		TArray<FMaterialResource*>& ResourceArray = It.Value();
+		TArray<TRefCountPtr<FMaterialShaderMap> >& OutShaderMapArray = OutShaderMaps.Add(It.Key(), TArray<TRefCountPtr<FMaterialShaderMap> >());
+
+		for (int32 Index = 0; Index < ResourceArray.Num(); Index++)
+		{
+			FMaterialResource* CurrentResource = ResourceArray[Index];
+			OutShaderMapArray.Add(CurrentResource->GetGameThreadShaderMap());
+		}
+	}
+}
+#endif // WITH_EDITOR
+
 bool UMaterial::UpdateLightmassTextureTracking()
 {
 	bool bTexturesHaveChanged = false;
diff --git a/Engine/Source/Runtime/Engine/Private/Materials/MaterialInstance.cpp b/Engine/Source/Runtime/Engine/Private/Materials/MaterialInstance.cpp
--- a/Engine/Source/Runtime/Engine/Private/Materials/MaterialInstance.cpp
+++ b/Engine/Source/Runtime/Engine/Private/Materials/MaterialInstance.cpp
@@ -2120,9 +2120,6 @@ void UMaterialInstance::CacheShadersForResources(EShaderPlatform ShaderPlatform,
 #if WITH_EDITOR
 	check(!HasAnyFlags(RF_NeedPostLoad));
 	check(BaseMaterial!=nullptr && !BaseMaterial->HasAnyFlags(RF_NeedPostLoad));
-#endif
-
-#if WITH_EDITOR
 	UpdateCachedData();
 #endif
 
@@ -2160,6 +2157,33 @@ void UMaterialInstance::CacheShaders(EMaterialShaderPrecompileMode CompileMode)
 	InitStaticPermutation(CompileMode);
 }
 
+#if WITH_EDITOR
+void UMaterialInstance::CacheGivenTypesForCooking(EShaderPlatform ShaderPlatform, ERHIFeatureLevel::Type FeatureLevel, EMaterialQualityLevel::Type QualityLevel, const TArray<const FVertexFactoryType*>& VFTypes, const TArray<const FShaderPipelineType*> PipelineTypes, const TArray<const FShaderType*>& ShaderTypes)
+{
+	TRACE_CPUPROFILER_EVENT_SCOPE(UMaterialInstance::CacheGivenTypes);
+
+	if (bHasStaticPermutationResource)
+	{
+		UMaterial* BaseMaterial = GetMaterial();
+
+		if (QualityLevel == EMaterialQualityLevel::Num)
+		{
+			QualityLevel = GetCachedScalabilityCVars().MaterialQualityLevel;
+		}
+
+		FMaterialResource* CurrentResource = FindOrCreateMaterialResource(StaticPermutationMaterialResources, BaseMaterial, nullptr, FeatureLevel, QualityLevel);
+		check(CurrentResource);
+
+		// Prepare the resource for compilation, but don't compile the completed shader map.
+		const bool bSuccess = CurrentResource->CacheShaders(ShaderPlatform, EMaterialShaderPrecompileMode::None);
+		if (bSuccess)
+		{
+			CurrentResource->CacheGivenTypes(ShaderPlatform, VFTypes, PipelineTypes, ShaderTypes);
+		}
+	}
+}
+#endif
+
 bool UMaterialInstance::GetMaterialLayers(FMaterialLayersFunctions& OutLayers, TMicRecursionGuard RecursionGuard) const
 {
 	if (StaticParameters.bHasMaterialLayers)
@@ -3343,7 +3367,7 @@ void UMaterialInstance::GetReferencedTexturesAndOverrides(TSet<const UTexture*>&
 
 void UMaterialInstance::UpdateCachedData()
 {
-	// Overriden for MIC/MID
+	// Overridden for MIC/MID
 }
 
 void UMaterialInstance::UpdateStaticPermutation(const FStaticParameterSet& NewParameters, FMaterialUpdateContext* MaterialUpdateContext)
diff --git a/Engine/Source/Runtime/Engine/Private/Materials/MaterialShader.cpp b/Engine/Source/Runtime/Engine/Private/Materials/MaterialShader.cpp
--- a/Engine/Source/Runtime/Engine/Private/Materials/MaterialShader.cpp
+++ b/Engine/Source/Runtime/Engine/Private/Materials/MaterialShader.cpp
@@ -1691,7 +1691,6 @@ void FMaterialShaderMap::LoadForRemoteRecompile(FArchive& Ar, EShaderPlatform Sh
 							|| (PassIndex == 1 && QualityLevelIndex == LoadedQualityLevel))
 						{
 							FMaterialResource* MaterialResource = MatchingMaterial->GetMaterialResource(GetMaxSupportedFeatureLevel(ShaderPlatform), (EMaterialQualityLevel::Type)QualityLevelIndex);
-							checkf(LoadedShaderMap->IsComplete(MaterialResource, false), TEXT("Shader map %s loaded for remote recompile is not complete"), LoadedShaderMap->GetFriendlyName());
 							MaterialResource->SetGameThreadShaderMap(LoadedShaderMap);
 						}
 					}
diff --git a/Engine/Source/Runtime/Engine/Private/Materials/MaterialShared.cpp b/Engine/Source/Runtime/Engine/Private/Materials/MaterialShared.cpp
--- a/Engine/Source/Runtime/Engine/Private/Materials/MaterialShared.cpp
+++ b/Engine/Source/Runtime/Engine/Private/Materials/MaterialShared.cpp
@@ -2527,7 +2527,7 @@ bool FMaterial::CacheShaders(const FMaterialShaderMapId& ShaderMapId, EShaderPla
 }
 
 #if WITH_EDITOR
-void FMaterial::CacheGivenTypes(EShaderPlatform Platform, const TArray<FVertexFactoryType*>& VFTypes, const TArray<FShaderType*>& ShaderTypes, const ITargetPlatform* TargetPlatform)
+void FMaterial::CacheGivenTypes(EShaderPlatform Platform, const TArray<const FVertexFactoryType*>& VFTypes, const TArray<const FShaderPipelineType*>& PipelineTypes, const TArray<const FShaderType*>& ShaderTypes)
 {
 	if (CompileErrors.Num())
 	{
@@ -2545,31 +2545,52 @@ void FMaterial::CacheGivenTypes(EShaderPlatform Platform, const TArray<FVertexFa
 		TRACE_CPUPROFILER_EVENT_SCOPE(FMaterial::CacheGivenTypes);
 		check(IsInGameThread());
 		checkf(ShaderTypes.Num() == VFTypes.Num(), TEXT("The size of the shader type array and vertex factory type array must match."));
+		checkf(PipelineTypes.Num() == ShaderTypes.Num(), TEXT("The size of the pipeline type array and shader type array must match.  Pass in null entries if pipelines are not used."));
 		checkf(GameThreadShaderMap, TEXT("Shader map is not initialized.  Please call CacheShaders first."));
 		checkf(GetGameThreadCompilingShaderMapId() != 0, TEXT("Material is not prepared to compile yet.  Please call CacheShaders first."));
 
 		TArray<FShaderCommonCompileJobPtr> CompileJobs;
 		for (int i = 0; i < VFTypes.Num(); ++i)
 		{
-			FVertexFactoryType* VFType = VFTypes[i];
+			const FVertexFactoryType* VFType = VFTypes[i];
 			check(VFType);
 
-			FShaderType* ShaderType = ShaderTypes[i];
-			check(ShaderType);
-			check(ShaderType->AsMeshMaterialShaderType());
+			const FShaderPipelineType* PipelineType = PipelineTypes[i];
+			const FShaderType* ShaderType = ShaderTypes[i];
 
-			ShaderType->AsMeshMaterialShaderType()->BeginCompileShader(
-				EShaderCompileJobPriority::ForceLocal,
-				GetGameThreadCompilingShaderMapId(),
-				0,
-				Platform,
-				GameThreadShaderMap->GetPermutationFlags(),
-				this,
-				GameThreadPendingCompilerEnvironment,
-				VFType,
-				CompileJobs,
-				nullptr,
-				nullptr);
+			if (PipelineType)
+			{
+				FMeshMaterialShaderType::BeginCompileShaderPipeline(
+					EShaderCompileJobPriority::ForceLocal,
+					GetGameThreadCompilingShaderMapId(),
+					0,
+					Platform,
+					GameThreadShaderMap->GetPermutationFlags(),
+					this,
+					GameThreadPendingCompilerEnvironment,
+					VFType,
+					PipelineType,
+					CompileJobs,
+					nullptr,
+					nullptr);
+			}
+			else if (ShaderType)
+			{
+				check(ShaderType->AsMeshMaterialShaderType());
+
+				ShaderType->AsMeshMaterialShaderType()->BeginCompileShader(
+					EShaderCompileJobPriority::ForceLocal,
+					GetGameThreadCompilingShaderMapId(),
+					0,
+					Platform,
+					GameThreadShaderMap->GetPermutationFlags(),
+					this,
+					GameThreadPendingCompilerEnvironment,
+					VFType,
+					CompileJobs,
+					nullptr,
+					nullptr);
+			}
 		}
 
 		GShaderCompilingManager->SubmitJobs(CompileJobs, GetBaseMaterialPathName(), GameThreadShaderMap->GetDebugDescription());
@@ -2962,7 +2983,7 @@ bool FMaterial::TryGetShaders(const FMaterialShaderTypes& InTypes, const FVertex
 						ShaderStageNamesToCompile.Add(ShaderType->GetName());
 					}
 
-					GODSCManager->AddThreadedShaderPipelineRequest(ShaderPlatform, MaterialName, VFTypeName, PipelineName, ShaderStageNamesToCompile);
+					GODSCManager->AddThreadedShaderPipelineRequest(ShaderPlatform, GetFeatureLevel(), GetQualityLevel(), MaterialName, VFTypeName, PipelineName, ShaderStageNamesToCompile);
 				}
 				else 
 #endif
@@ -3010,7 +3031,7 @@ bool FMaterial::TryGetShaders(const FMaterialShaderTypes& InTypes, const FVertex
 						TArray<FString> ShaderStageNamesToCompile;
 						ShaderStageNamesToCompile.Add(ShaderType->GetName());
 
-						GODSCManager->AddThreadedShaderPipelineRequest(ShaderPlatform, MaterialName, VFTypeName, PipelineName, ShaderStageNamesToCompile);
+						GODSCManager->AddThreadedShaderPipelineRequest(ShaderPlatform, GetFeatureLevel(), GetQualityLevel(), MaterialName, VFTypeName, PipelineName, ShaderStageNamesToCompile);
 					}
 					else
 #endif
diff --git a/Engine/Source/Runtime/Engine/Private/ODSC/ODSCManager.cpp b/Engine/Source/Runtime/Engine/Private/ODSC/ODSCManager.cpp
--- a/Engine/Source/Runtime/Engine/Private/ODSC/ODSCManager.cpp
+++ b/Engine/Source/Runtime/Engine/Private/ODSC/ODSCManager.cpp
@@ -70,18 +70,18 @@ bool FODSCManager::Tick(float DeltaSeconds)
 	return false;
 }
 
-void FODSCManager::AddThreadedRequest(const TArray<FString>& MaterialsToCompile, const FString& ShaderTypesToLoad, EShaderPlatform ShaderPlatform, ODSCRecompileCommand RecompileCommandType)
+void FODSCManager::AddThreadedRequest(const TArray<FString>& MaterialsToCompile, const FString& ShaderTypesToLoad, EShaderPlatform ShaderPlatform, ERHIFeatureLevel::Type FeatureLevel, EMaterialQualityLevel::Type QualityLevel, ODSCRecompileCommand RecompileCommandType)
 {
 	if (Thread)
 	{
-		Thread->AddRequest(MaterialsToCompile, ShaderTypesToLoad, ShaderPlatform, RecompileCommandType);
+		Thread->AddRequest(MaterialsToCompile, ShaderTypesToLoad, ShaderPlatform, FeatureLevel, QualityLevel, RecompileCommandType);
 	}
 }
 
-void FODSCManager::AddThreadedShaderPipelineRequest(EShaderPlatform ShaderPlatform, const FString& MaterialName, const FString& VertexFactoryName, const FString& PipelineName, const TArray<FString>& ShaderTypeNames)
+void FODSCManager::AddThreadedShaderPipelineRequest(EShaderPlatform ShaderPlatform, ERHIFeatureLevel::Type FeatureLevel, EMaterialQualityLevel::Type QualityLevel, const FString& MaterialName, const FString& VertexFactoryName, const FString& PipelineName, const TArray<FString>& ShaderTypeNames)
 {
 	if (Thread)
 	{
-		Thread->AddShaderPipelineRequest(ShaderPlatform, MaterialName, VertexFactoryName, PipelineName, ShaderTypeNames);
+		Thread->AddShaderPipelineRequest(ShaderPlatform, FeatureLevel, QualityLevel, MaterialName, VertexFactoryName, PipelineName, ShaderTypeNames);
 	}
 }
diff --git a/Engine/Source/Runtime/Engine/Private/ODSC/ODSCThread.cpp b/Engine/Source/Runtime/Engine/Private/ODSC/ODSCThread.cpp
--- a/Engine/Source/Runtime/Engine/Private/ODSC/ODSCThread.cpp
+++ b/Engine/Source/Runtime/Engine/Private/ODSC/ODSCThread.cpp
@@ -8,22 +8,26 @@
 #include "HAL/RunnableThread.h"
 #include "ShaderCompiler.h"
 
-FODSCRequestPayload::FODSCRequestPayload(EShaderPlatform InShaderPlatform, const FString& InMaterialName, const FString& InVertexFactoryName, const FString& InPipelineName, const TArray<FString>& InShaderTypeNames, const FString& InRequestHash)
-	: ShaderPlatform(InShaderPlatform), MaterialName(InMaterialName), VertexFactoryName(InVertexFactoryName), PipelineName(InPipelineName), ShaderTypeNames(std::move(InShaderTypeNames)), RequestHash(InRequestHash)
+FODSCRequestPayload::FODSCRequestPayload(EShaderPlatform InShaderPlatform, ERHIFeatureLevel::Type InFeatureLevel, EMaterialQualityLevel::Type InQualityLevel, const FString& InMaterialName, const FString& InVertexFactoryName, const FString& InPipelineName, const TArray<FString>& InShaderTypeNames, const FString& InRequestHash)
+	: ShaderPlatform(InShaderPlatform), FeatureLevel(InFeatureLevel), QualityLevel(InQualityLevel), MaterialName(InMaterialName), VertexFactoryName(InVertexFactoryName), PipelineName(InPipelineName), ShaderTypeNames(std::move(InShaderTypeNames)), RequestHash(InRequestHash)
 {
 
 }
 
-FODSCMessageHandler::FODSCMessageHandler(EShaderPlatform InShaderPlatform, ODSCRecompileCommand InRecompileCommandType) 
+FODSCMessageHandler::FODSCMessageHandler(EShaderPlatform InShaderPlatform, ERHIFeatureLevel::Type InFeatureLevel, EMaterialQualityLevel::Type InQualityLevel, ODSCRecompileCommand InRecompileCommandType)
 :	ShaderPlatform(InShaderPlatform),
+	FeatureLevel(InFeatureLevel),
+	QualityLevel(InQualityLevel),
 	RecompileCommandType(InRecompileCommandType)
 {
 }
 
-FODSCMessageHandler::FODSCMessageHandler(const TArray<FString>& InMaterials, const FString& ShaderTypesToLoad, EShaderPlatform InShaderPlatform, ODSCRecompileCommand InRecompileCommandType) :
+FODSCMessageHandler::FODSCMessageHandler(const TArray<FString>& InMaterials, const FString& ShaderTypesToLoad, EShaderPlatform InShaderPlatform, ERHIFeatureLevel::Type InFeatureLevel, EMaterialQualityLevel::Type InQualityLevel, ODSCRecompileCommand InRecompileCommandType) :
 	MaterialsToLoad(std::move(InMaterials)),
 	ShaderTypesToLoad(ShaderTypesToLoad),
 	ShaderPlatform(InShaderPlatform),
+	FeatureLevel(InFeatureLevel),
+	QualityLevel(InQualityLevel),
 	RecompileCommandType(InRecompileCommandType)
 {
 }
@@ -33,10 +37,15 @@ void FODSCMessageHandler::FillPayload(FArchive& Payload)
 	// When did we start this request?
 	RequestStartTime = FPlatformTime::Seconds();
 
+	int32 ConvertedShaderPlatform = static_cast<int32>(ShaderPlatform);
+	int32 ConvertedFeatureLevel = static_cast<int32>(FeatureLevel);
+	int32 ConvertedQualityLevel = static_cast<int32>(QualityLevel);
+
 	Payload << MaterialsToLoad;
 	Payload << ShaderTypesToLoad;
-	uint32 ConvertedShaderPlatform = (uint32)ShaderPlatform;
 	Payload << ConvertedShaderPlatform;
+	Payload << ConvertedFeatureLevel;
+	Payload << ConvertedQualityLevel;
 	Payload << RecompileCommandType;
 	Payload << RequestBatch;
 }
@@ -110,12 +119,12 @@ void FODSCThread::Tick()
 	Process();
 }
 
-void FODSCThread::AddRequest(const TArray<FString>& MaterialsToCompile, const FString& ShaderTypesToLoad, EShaderPlatform ShaderPlatform, ODSCRecompileCommand RecompileCommandType)
+void FODSCThread::AddRequest(const TArray<FString>& MaterialsToCompile, const FString& ShaderTypesToLoad, EShaderPlatform ShaderPlatform, ERHIFeatureLevel::Type FeatureLevel, EMaterialQualityLevel::Type QualityLevel, ODSCRecompileCommand RecompileCommandType)
 {
-	PendingMaterialThreadedRequests.Enqueue(new FODSCMessageHandler(MaterialsToCompile, ShaderTypesToLoad, ShaderPlatform, RecompileCommandType));
+	PendingMaterialThreadedRequests.Enqueue(new FODSCMessageHandler(MaterialsToCompile, ShaderTypesToLoad, ShaderPlatform, FeatureLevel, QualityLevel, RecompileCommandType));
 }
 
-void FODSCThread::AddShaderPipelineRequest(EShaderPlatform ShaderPlatform, const FString& MaterialName, const FString& VertexFactoryName, const FString& PipelineName, const TArray<FString>& ShaderTypeNames)
+void FODSCThread::AddShaderPipelineRequest(EShaderPlatform ShaderPlatform, ERHIFeatureLevel::Type FeatureLevel, EMaterialQualityLevel::Type QualityLevel, const FString& MaterialName, const FString& VertexFactoryName, const FString& PipelineName, const TArray<FString>& ShaderTypeNames)
 {
 	FString RequestString = (MaterialName + VertexFactoryName + PipelineName);
 	for (const auto& ShaderTypeName : ShaderTypeNames)
@@ -127,7 +136,7 @@ void FODSCThread::AddShaderPipelineRequest(EShaderPlatform ShaderPlatform, const
 	FScopeLock Lock(&RequestHashCriticalSection);
 	if (!RequestHashes.Contains(RequestHash))
 	{
-		PendingMeshMaterialThreadedRequests.Enqueue(FODSCRequestPayload(ShaderPlatform, MaterialName, VertexFactoryName, PipelineName, ShaderTypeNames, RequestHash));
+		PendingMeshMaterialThreadedRequests.Enqueue(FODSCRequestPayload(ShaderPlatform, FeatureLevel, QualityLevel, MaterialName, VertexFactoryName, PipelineName, ShaderTypeNames, RequestHash));
 		RequestHashes.Add(RequestHash);
 	}
 }
@@ -213,7 +222,7 @@ void FODSCThread::Process()
 	// process any specific mesh material shader requests.
 	if (PayloadsToAggregate.Num())
 	{
-		FODSCMessageHandler* RequestHandler = new FODSCMessageHandler(PayloadsToAggregate[0].ShaderPlatform, ODSCRecompileCommand::Material);
+		FODSCMessageHandler* RequestHandler = new FODSCMessageHandler(PayloadsToAggregate[0].ShaderPlatform, PayloadsToAggregate[0].FeatureLevel, PayloadsToAggregate[0].QualityLevel, ODSCRecompileCommand::Material);
 		for (const FODSCRequestPayload& payload : PayloadsToAggregate)
 		{
 			RequestHandler->AddPayload(payload);
diff --git a/Engine/Source/Runtime/Engine/Private/ODSC/ODSCThread.h b/Engine/Source/Runtime/Engine/Private/ODSC/ODSCThread.h
--- a/Engine/Source/Runtime/Engine/Private/ODSC/ODSCThread.h
+++ b/Engine/Source/Runtime/Engine/Private/ODSC/ODSCThread.h
@@ -16,8 +16,8 @@ class FRunnableThread;
 class FODSCMessageHandler : public IPlatformFile::IFileServerMessageHandler
 {
 public:
-	FODSCMessageHandler(EShaderPlatform InShaderPlatform, ODSCRecompileCommand InRecompileCommandType);
-	FODSCMessageHandler(const TArray<FString>& InMaterials, const FString& ShaderTypesToLoad, EShaderPlatform InShaderPlatform, ODSCRecompileCommand InRecompileCommandType);
+	FODSCMessageHandler(EShaderPlatform InShaderPlatform, ERHIFeatureLevel::Type InFeatureLevel, EMaterialQualityLevel::Type InQualityLevel, ODSCRecompileCommand InRecompileCommandType);
+	FODSCMessageHandler(const TArray<FString>& InMaterials, const FString& ShaderTypesToLoad, EShaderPlatform InShaderPlatform, ERHIFeatureLevel::Type InFeatureLevel, EMaterialQualityLevel::Type InQualityLevel, ODSCRecompileCommand InRecompileCommandType);
 	/** Subclass fills out an archive to send to the server */
 	virtual void FillPayload(FArchive& Payload) override;
 
@@ -44,6 +44,12 @@ private:
 	/** Which shader platform we are compiling for */
 	EShaderPlatform ShaderPlatform;
 
+	/** Which feature level to compile for. */
+	ERHIFeatureLevel::Type FeatureLevel;
+
+	/** Which material quality level to compile for. */
+	EMaterialQualityLevel::Type QualityLevel;
+
 	/** Whether or not to recompile changed shaders */
 	ODSCRecompileCommand RecompileCommandType = ODSCRecompileCommand::None;
 
@@ -94,12 +100,14 @@ public:
 	 *
 	 * @return false if no longer needs ticking
 	 */
-	void AddRequest(const TArray<FString>& MaterialsToCompile, const FString& ShaderTypesToLoad, EShaderPlatform ShaderPlatform, ODSCRecompileCommand RecompileCommandType);
+	void AddRequest(const TArray<FString>& MaterialsToCompile, const FString& ShaderTypesToLoad, EShaderPlatform ShaderPlatform, ERHIFeatureLevel::Type FeatureLevel, EMaterialQualityLevel::Type QualityLevel, ODSCRecompileCommand RecompileCommandType);
 
 	/**
 	 * Add a request to compile a pipeline (VS/PS) of shaders.  The results are submitted and processed in an async manner.
 	 *
 	 * @param ShaderPlatform - Which shader platform to compile for.
+	 * @param FeatureLevel - Which feature level to compile for.
+	 * @param QualityLevel - Which material quality level to compile for.
 	 * @param MaterialName - The name of the material to compile.
 	 * @param VertexFactoryName - The name of the vertex factory type we should compile.
 	 * @param PipelineName - The name of the shader pipeline we should compile.
@@ -107,7 +115,7 @@ public:
 	 *
 	 * @return false if no longer needs ticking
 	 */
-	void AddShaderPipelineRequest(EShaderPlatform ShaderPlatform, const FString& MaterialName, const FString& VertexFactoryName, const FString& PipelineName, const TArray<FString>& ShaderTypeNames);
+	void AddShaderPipelineRequest(EShaderPlatform ShaderPlatform, ERHIFeatureLevel::Type FeatureLevel, EMaterialQualityLevel::Type QualityLevel, const FString& MaterialName, const FString& VertexFactoryName, const FString& PipelineName, const TArray<FString>& ShaderTypeNames);
 
 	/**
 	 * Get completed requests.  Clears internal arrays.  Called on Game thread.
diff --git a/Engine/Source/Runtime/Engine/Private/ShaderCompiler/ShaderCompiler.cpp b/Engine/Source/Runtime/Engine/Private/ShaderCompiler/ShaderCompiler.cpp
--- a/Engine/Source/Runtime/Engine/Private/ShaderCompiler/ShaderCompiler.cpp
+++ b/Engine/Source/Runtime/Engine/Private/ShaderCompiler/ShaderCompiler.cpp
@@ -2276,13 +2276,13 @@ void FShaderCompileThreadRunnable::WriteNewTasks()
 
 			// Write out the file that the worker app is waiting for, which has all the information needed to compile the shader.
 			// 'Only' indicates that the worker should keep checking for more tasks after this one
-			FArchive* TransferFile = NULL;
+			FArchive* TransferFile = nullptr;
 
 			int32 RetryCount = 0;
 			// Retry over the next two seconds if we can't write out the input file
 			// Anti-virus and indexing applications can interfere and cause this write to fail
 			//@todo - switch to shared memory or some other method without these unpredictable hazards
-			while (TransferFile == NULL && RetryCount < 2000)
+			while (TransferFile == nullptr && RetryCount < 2000)
 			{
 				if (RetryCount > 0)
 				{
@@ -2290,12 +2290,12 @@ void FShaderCompileThreadRunnable::WriteNewTasks()
 				}
 				TransferFile = IFileManager::Get().CreateFileWriter(*TransferFileName, FILEWRITE_EvenIfReadOnly);
 				RetryCount++;
-				if (TransferFile == NULL)
+				if (TransferFile == nullptr)
 				{
 					UE_LOG(LogShaderCompilers, Warning, TEXT("Could not create the shader compiler transfer file '%s', retrying..."), *TransferFileName);
 				}
 			}
-			if (TransferFile == NULL)
+			if (TransferFile == nullptr)
 			{
 				UE_LOG(LogShaderCompilers, Fatal, TEXT("Could not create the shader compiler transfer file '%s'."), *TransferFileName);
 			}
@@ -2585,7 +2585,7 @@ void FShaderCompileUtilities::ExecuteShaderCompileJob(FShaderCommonCompileJob& J
 		{
 			UE_LOG(LogShaderCompilers, Fatal, TEXT("Can't compile shaders for format %s, couldn't load compiler dll"), *Format.ToString());
 		}
-		CA_ASSUME(Compiler != NULL);
+		CA_ASSUME(Compiler != nullptr);
 
 		if (IsValidRef(SingleJob->Input.SharedEnvironment))
 		{
@@ -2623,7 +2623,7 @@ void FShaderCompileUtilities::ExecuteShaderCompileJob(FShaderCommonCompileJob& J
 		{
 			UE_LOG(LogShaderCompilers, Fatal, TEXT("Can't compile shaders for format %s, couldn't load compiler dll"), *Format.ToString());
 		}
-		CA_ASSUME(Compiler != NULL);
+		CA_ASSUME(Compiler != nullptr);
 
 		// Verify same platform on all stages
 		for (int32 Index = 1; Index < PipelineJob->StageJobs.Num(); ++Index)
@@ -2813,7 +2813,7 @@ int32 FShaderCompileThreadRunnable::CompilingLoop()
 	return NumActiveThreads;
 }
 
-FShaderCompilerStats* GShaderCompilerStats = NULL;
+FShaderCompilerStats* GShaderCompilerStats = nullptr;
 
 void FShaderCompilerStats::WriteStats(FOutputDevice* Ar)
 {
@@ -3326,7 +3326,7 @@ double FShaderCompilerStats::GetTimeShaderCompilationWasActive()
 	return Sum;
 }
 
-FShaderCompilingManager* GShaderCompilingManager = NULL;
+FShaderCompilingManager* GShaderCompilingManager = nullptr;
 
 bool FShaderCompilingManager::AllTargetPlatformSupportsRemoteShaderCompiling()
 {
@@ -3941,7 +3941,7 @@ FProcHandle FShaderCompilingManager::LaunchWorker(const FString& WorkingDirector
 		// Disambiguate between SCW.exe missing vs other errors.
 		static bool bFirstLaunch = true;
 		uint32 WorkerId = 0;
-		FProcHandle WorkerHandle = FPlatformProcess::CreateProc(*ShaderCompileWorkerName, *WorkerParameters, true, false, false, &WorkerId, PriorityModifier, NULL, NULL);
+		FProcHandle WorkerHandle = FPlatformProcess::CreateProc(*ShaderCompileWorkerName, *WorkerParameters, true, false, false, &WorkerId, PriorityModifier, nullptr, nullptr);
 		if (WorkerHandle.IsValid())
 		{
 			// Process launched at least once successfully
@@ -4958,7 +4958,7 @@ void FShaderCompilingManager::FinishCompilation(const TCHAR* MaterialName, const
 	const double StartTime = FPlatformTime::Seconds();
 
 	FText StatusUpdate;
-	if ( MaterialName != NULL )
+	if ( MaterialName != nullptr)
 	{
 		FFormatNamedArguments Args;
 		Args.Add( TEXT("MaterialName"), FText::FromString( MaterialName ) );
@@ -6233,7 +6233,10 @@ bool RecompileShaders(const TCHAR* Cmd, FOutputDevice& Ar)
 		TArray<FString> MaterialsToLoad;
 		FString ShaderTypesToLoad;
 		ODSCRecompileCommand CommandType = ParseRecompileCommandString(Cmd, MaterialsToLoad, ShaderTypesToLoad);
-		GODSCManager->AddThreadedRequest(MaterialsToLoad, ShaderTypesToLoad, GMaxRHIShaderPlatform, CommandType);
+
+		ERHIFeatureLevel::Type TargetFeatureLevel = GetMaxSupportedFeatureLevel(GMaxRHIShaderPlatform);
+		const EMaterialQualityLevel::Type ActiveQualityLevel = GetCachedScalabilityCVars().MaterialQualityLevel;
+		GODSCManager->AddThreadedRequest(MaterialsToLoad, ShaderTypesToLoad, GMaxRHIShaderPlatform, TargetFeatureLevel, ActiveQualityLevel, CommandType);
 #endif
 		return true;
 	}
@@ -6323,7 +6326,7 @@ bool RecompileShaders(const TCHAR* Cmd, FOutputDevice& Ar)
 					}
 					else
 					{
-						Material->PreEditChange(NULL);
+						Material->PreEditChange(nullptr);
 						Material->PostEditChange();
 					}
 #endif // WITH_EDITOR
@@ -6353,7 +6356,7 @@ bool RecompileShaders(const TCHAR* Cmd, FOutputDevice& Ar)
 #if WITH_EDITOR
 					// <Pre/Post>EditChange will force a re-creation of the resource,
 					// in turn recompiling the shader.
-					Material->PreEditChange(NULL);
+					Material->PreEditChange(nullptr);
 					Material->PostEditChange();
 #endif // WITH_EDITOR
 				}
@@ -7217,7 +7220,7 @@ bool RecompileChangedShadersForPlatform(const FString& PlatformName)
 	// figure out what shader platforms to recompile
 	ITargetPlatformManagerModule* TPM = GetTargetPlatformManager();
 	ITargetPlatform* TargetPlatform = TPM->FindTargetPlatform(PlatformName);
-	if (TargetPlatform == NULL)
+	if (TargetPlatform == nullptr)
 	{
 		UE_LOG(LogShaders, Display, TEXT("Failed to find target platform module for %s"), *PlatformName);
 		return false;
@@ -7267,15 +7270,27 @@ bool RecompileChangedShadersForPlatform(const FString& PlatformName)
 	return false;
 }
 
-FArchive& operator<<(FArchive& Ar, FODSCRequestPayload& Elem)
+FArchive& operator<<(FArchive& Ar, FODSCRequestPayload& Payload)
 {
-	uint32 ConvertedShaderPlatform = (uint32)Elem.ShaderPlatform;
-	Ar << ConvertedShaderPlatform;
-	Ar << Elem.MaterialName;
-	Ar << Elem.VertexFactoryName;
-	Ar << Elem.PipelineName;
-	Ar << Elem.ShaderTypeNames;
-	Ar << Elem.RequestHash;
+	int32 iShaderPlatform = static_cast<int32>(Payload.ShaderPlatform);
+	int32 iFeatureLevel = static_cast<int32>(Payload.FeatureLevel);
+	int32 iQualityLevel = static_cast<int32>(Payload.QualityLevel);
+
+	Ar << iShaderPlatform;
+	Ar << iFeatureLevel;
+	Ar << iQualityLevel;
+	Ar << Payload.MaterialName;
+	Ar << Payload.VertexFactoryName;
+	Ar << Payload.PipelineName;
+	Ar << Payload.ShaderTypeNames;
+	Ar << Payload.RequestHash;
+
+	if (Ar.IsLoading())
+	{
+		Payload.ShaderPlatform = static_cast<EShaderPlatform>(iShaderPlatform);
+		Payload.FeatureLevel = static_cast<ERHIFeatureLevel::Type>(iFeatureLevel);
+		Payload.QualityLevel = static_cast<EMaterialQualityLevel::Type>(iQualityLevel);
+	}
 
 	return Ar;
 }
@@ -7298,6 +7313,31 @@ FShaderRecompileData::FShaderRecompileData(const FString& InPlatformName, EShade
 {
 }
 
+FArchive& operator<<(FArchive& Ar, FShaderRecompileData& RecompileData)
+{
+
+	int32 iShaderPlatform = static_cast<int32>(RecompileData.ShaderPlatform);
+	int32 iFeatureLevel = static_cast<int32>(RecompileData.FeatureLevel);
+	int32 iQualityLevel = static_cast<int32>(RecompileData.QualityLevel);
+
+	Ar << RecompileData.MaterialsToLoad;
+	Ar << RecompileData.ShaderTypesToLoad;
+	Ar << iShaderPlatform;
+	Ar << iFeatureLevel;
+	Ar << iQualityLevel;
+	Ar << RecompileData.CommandType;
+	Ar << RecompileData.ShadersToRecompile;
+
+	if (Ar.IsLoading())
+	{
+		RecompileData.ShaderPlatform = static_cast<EShaderPlatform>(iShaderPlatform);
+		RecompileData.FeatureLevel = static_cast<ERHIFeatureLevel::Type>(iFeatureLevel);
+		RecompileData.QualityLevel = static_cast<EMaterialQualityLevel::Type>(iQualityLevel);
+	}
+
+	return Ar;
+}
+
 extern ENGINE_API const TCHAR* ODSCCmdEnumToString(ODSCRecompileCommand Cmd)
 {
 	switch (Cmd)
@@ -7318,6 +7358,57 @@ extern ENGINE_API const TCHAR* ODSCCmdEnumToString(ODSCRecompileCommand Cmd)
 }
 
 #if WITH_EDITOR
+
+void CompileGlobalShaderMapForRemote(
+	const TArray<const FShaderType*>& OutdatedShaderTypes, 
+	const TArray<const FShaderPipelineType*>& OutdatedShaderPipelineTypes, 
+	const EShaderPlatform ShaderPlatform, 
+	const ITargetPlatform* TargetPlatform,
+	TArray<uint8>* OutArray)
+{
+	UE_LOG(LogShaders, Display, TEXT("Recompiling global shaders."));
+
+	// Kick off global shader recompiles
+	BeginRecompileGlobalShaders(OutdatedShaderTypes, OutdatedShaderPipelineTypes, ShaderPlatform, TargetPlatform);
+
+	// Block on global shaders
+	FinishRecompileGlobalShaders();
+
+	// Write the shader compilation info to memory, converting FName to strings
+	FMemoryWriter MemWriter(*OutArray, true);
+	FNameAsStringProxyArchive Ar(MemWriter);
+
+	TOptional<FArchiveCookData> CookData;
+	FArchiveCookContext CookContext(nullptr /*InPackage*/, FArchiveCookContext::ECookTypeUnknown);
+	if (TargetPlatform != nullptr)
+	{
+		CookData.Emplace(*TargetPlatform, CookContext);
+	}
+	Ar.SetCookData(CookData.GetPtrOrNull());
+
+	// save out the global shader map to the byte array
+	SaveGlobalShadersForRemoteRecompile(Ar, ShaderPlatform);
+}
+
+void SaveShaderMapsForRemote(ITargetPlatform* TargetPlatform, const TMap<FString, TArray<TRefCountPtr<FMaterialShaderMap>>>& CompiledShaderMaps, TArray<uint8>* OutArray)
+{
+	// write the shader compilation info to memory, converting fnames to strings
+	FMemoryWriter MemWriter(*OutArray, true);
+	FNameAsStringProxyArchive Ar(MemWriter);
+
+	TOptional<FArchiveCookData> CookData;
+	FArchiveCookContext CookContext(nullptr /*InPackage*/, FArchiveCookContext::ECookTypeUnknown);
+	if (TargetPlatform != nullptr)
+	{
+		CookData.Emplace(*TargetPlatform, CookContext);
+	}
+
+	Ar.SetCookData(CookData.GetPtrOrNull());
+
+	// save out the shader map to the byte array
+	FMaterialShaderMap::SaveForRemoteRecompile(Ar, CompiledShaderMaps);
+}
+
 void RecompileShadersForRemote(
 	FShaderRecompileData& Args,
 	const FString& OutputDirectory)
@@ -7327,7 +7418,7 @@ void RecompileShadersForRemote(
 	// figure out what shader platforms to recompile
 	ITargetPlatformManagerModule* TPM = GetTargetPlatformManager();
 	ITargetPlatform* TargetPlatform = TPM->FindTargetPlatform(Args.PlatformName);
-	if (TargetPlatform == NULL)
+	if (TargetPlatform == nullptr)
 	{
 		UE_LOG(LogShaders, Display, TEXT("Failed to find target platform module for %s"), *Args.PlatformName);
 		return;
@@ -7352,7 +7443,7 @@ void RecompileShadersForRemote(
 	for (int32 Index = 0; Index < Args.MaterialsToLoad.Num(); Index++)
 	{
 		UE_LOG(LogShaders, Verbose, TEXT("   --> %s"), *Args.MaterialsToLoad[Index]);
-		MaterialsToCompile.Add(LoadObject<UMaterialInterface>(NULL, *Args.MaterialsToLoad[Index]));
+		MaterialsToCompile.Add(LoadObject<UMaterialInterface>(nullptr, *Args.MaterialsToLoad[Index]));
 	}
 
 	UE_LOG(LogShaders, Verbose, TEXT("  Done!"));
@@ -7367,43 +7458,14 @@ void RecompileShadersForRemote(
 	// Pick up new changes to shader files
 	FlushShaderFileCache();
 
-	if (Args.ShadersToRecompile.Num())
+	// If we have an explicit list of shaders to compile from ODSC just compile those.
+	if (Args.ShadersToRecompile.Num() && (Args.MeshMaterialMaps != nullptr))
 	{
-		UE_LOG(LogShaders, Display, TEXT("Received %d shaders to compile."), Args.ShadersToRecompile.Num());
+		TMap<FString, TArray<TRefCountPtr<FMaterialShaderMap>>> CompiledShaderMaps;
+		UMaterial::CompileODSCMaterialsForRemoteRecompile(Args.ShadersToRecompile, CompiledShaderMaps);
+		SaveShaderMapsForRemote(TargetPlatform, CompiledShaderMaps, Args.MeshMaterialMaps);
 	}
-
-	for (const FODSCRequestPayload& payload: Args.ShadersToRecompile)
-	{
-		UE_LOG(LogShaders, Display, TEXT(""));
-		UE_LOG(LogShaders, Display, TEXT("\tMaterial:    %s "), *payload.MaterialName);
-		UE_LOG(LogShaders, Display, TEXT("\tVF Type:     %s "), *payload.VertexFactoryName);
-
-		MaterialsToCompile.Add(LoadObject<UMaterialInterface>(NULL, *payload.MaterialName));
-
-		const FVertexFactoryType* VFType = FVertexFactoryType::GetVFByName(payload.VertexFactoryName);
-		if (VFType)
-		{
-			OutdatedFactoryTypes.Add(VFType);
-		}
-
-		const FShaderPipelineType* PipelineType = FShaderPipelineType::GetShaderPipelineTypeByName(payload.PipelineName);
-		if (PipelineType)
-		{
-			OutdatedShaderPipelineTypes.Add(PipelineType);
-		}
-
-		for (const FString& ShaderTypeName : payload.ShaderTypeNames)
-		{
-			UE_LOG(LogShaders, Display, TEXT("\tShader Type: %s"), *ShaderTypeName);
-
-			const FShaderType* ShaderType = FShaderType::GetShaderTypeByName(*ShaderTypeName);
-			if (ShaderType)
-			{
-				OutdatedShaderTypes.Add(ShaderType);
-			}
-		}
-	}
-
+	else
 	{
 		for (int32 FormatIndex = 0; FormatIndex < DesiredShaderFormats.Num(); FormatIndex++)
 		{
@@ -7424,32 +7486,11 @@ void RecompileShadersForRemote(
 						UE_LOG(LogShaders, Display, TEXT("\t%s..."), ShaderType->GetName());
 					}
 
-					if (ShaderTypes.Num() > 0 || ShaderPipelineTypes.Num() > 0)
-					{
-						BeginRecompileGlobalShaders(ShaderTypes, ShaderPipelineTypes, ShaderPlatform, TargetPlatform);
-						FinishRecompileGlobalShaders();
-					}
-
-					// write the shader compilation info to memory, converting fnames to strings
-					FMemoryWriter MemWriter(*Args.GlobalShaderMap, true);
-					FNameAsStringProxyArchive Ar(MemWriter);
-
-					TOptional<FArchiveCookData> CookData;
-					FArchiveCookContext CookContext(nullptr /*InPackage*/, FArchiveCookContext::ECookTypeUnknown);
-					if (TargetPlatform != nullptr)
-					{
-						CookData.Emplace(*TargetPlatform, CookContext);
-					}
-					Ar.SetCookData(CookData.GetPtrOrNull());
-
-					// save out the global shader map to the byte array
-					SaveGlobalShadersForRemoteRecompile(Ar, ShaderPlatform);
+					CompileGlobalShaderMapForRemote(ShaderTypes, ShaderPipelineTypes, ShaderPlatform, TargetPlatform, Args.GlobalShaderMap);
 				}
 				else if (Args.CommandType == ODSCRecompileCommand::Global ||
 						 Args.CommandType == ODSCRecompileCommand::Changed)
 				{
-					UE_LOG(LogShaders, Display, TEXT("Recompiling global shaders."));
-
 					// Explicitly get outdated types for global shaders.
 					const FGlobalShaderMap* ShaderMap = GGlobalShaderMap[ShaderPlatform];
 					if (ShaderMap)
@@ -7459,50 +7500,16 @@ void RecompileShadersForRemote(
 
 					UE_LOG(LogShaders, Display, TEXT("\tFound %d outdated shader types."), OutdatedShaderTypes.Num() + OutdatedShaderPipelineTypes.Num());
 
-					// Kick off global shader recompiles
-					BeginRecompileGlobalShaders(OutdatedShaderTypes, OutdatedShaderPipelineTypes, ShaderPlatform, TargetPlatform);
-
-					// Block on global shaders
-					FinishRecompileGlobalShaders();
-
-					// write the shader compilation info to memory, converting fnames to strings
-					FMemoryWriter MemWriter(*Args.GlobalShaderMap, true);
-					FNameAsStringProxyArchive Ar(MemWriter);
-
-					TOptional<FArchiveCookData> CookData;
-					FArchiveCookContext CookContext(nullptr /*InPackage*/, FArchiveCookContext::ECookTypeUnknown);
-					if (TargetPlatform != nullptr)
-					{
-						CookData.Emplace(*TargetPlatform, CookContext);
-					}
-					Ar.SetCookData(CookData.GetPtrOrNull());
-
-					// save out the global shader map to the byte array
-					SaveGlobalShadersForRemoteRecompile(Ar, ShaderPlatform);
+					CompileGlobalShaderMapForRemote(OutdatedShaderTypes, OutdatedShaderPipelineTypes, ShaderPlatform, TargetPlatform, Args.GlobalShaderMap);
 				}
 
 				// we only want to actually compile mesh shaders if a client directly requested it
 				if ((Args.CommandType == ODSCRecompileCommand::Material || Args.CommandType == ODSCRecompileCommand::Changed) &&
 					Args.MeshMaterialMaps != nullptr)
 				{
-					TMap<FString, TArray<TRefCountPtr<FMaterialShaderMap> > > CompiledShaderMaps;
+					TMap<FString, TArray<TRefCountPtr<FMaterialShaderMap>>> CompiledShaderMaps;
 					UMaterial::CompileMaterialsForRemoteRecompile(MaterialsToCompile, ShaderPlatform, TargetPlatform, CompiledShaderMaps);
-
-					// write the shader compilation info to memory, converting fnames to strings
-					FMemoryWriter MemWriter(*Args.MeshMaterialMaps, true);
-					FNameAsStringProxyArchive Ar(MemWriter);
-
-					TOptional<FArchiveCookData> CookData;
-					FArchiveCookContext CookContext(nullptr /*InPackage*/, FArchiveCookContext::ECookTypeUnknown);
-					if (TargetPlatform != nullptr)
-					{
-						CookData.Emplace(*TargetPlatform, CookContext);
-					}
-
-					Ar.SetCookData(CookData.GetPtrOrNull());
-
-					// save out the shader map to the byte array
-					FMaterialShaderMap::SaveForRemoteRecompile(Ar, CompiledShaderMaps);
+					SaveShaderMapsForRemote(TargetPlatform, CompiledShaderMaps, Args.MeshMaterialMaps);
 				}
 
 				// save it out so the client can get it (and it's up to date next time)
diff --git a/Engine/Source/Runtime/Engine/Public/MaterialShared.h b/Engine/Source/Runtime/Engine/Public/MaterialShared.h
--- a/Engine/Source/Runtime/Engine/Public/MaterialShared.h
+++ b/Engine/Source/Runtime/Engine/Public/MaterialShared.h
@@ -1712,10 +1712,11 @@ public:
 	/**
 	 * Submits local compile jobs for the exact given shader types and vertex factory type combination.
 	 * @note CacheShaders() should be called first to prepare the resource for compilation.
-	 * @note The arrays of shader types and vertex factory types must match.
+	 * @note The arrays of shader types, pipeline types and vertex factory types must match.
+	 * @note Entries in the PipelineTypes and ShaderTypes arrays can contain null entries.
 	 * @note These compile jobs are submitted async and it is up to client code to block on results if needed.
 	 */
-	ENGINE_API void CacheGivenTypes(EShaderPlatform ShaderPlatform, const TArray<FVertexFactoryType*>& VFTypes, const TArray<FShaderType*>& ShaderTypes, const ITargetPlatform* TargetPlatform = nullptr);
+	ENGINE_API void CacheGivenTypes(EShaderPlatform ShaderPlatform, const TArray<const FVertexFactoryType*>& VFTypes, const TArray<const FShaderPipelineType*>& PipelineTypes, const TArray<const FShaderType*>& ShaderTypes);
 #endif
 
 	/**
diff --git a/Engine/Source/Runtime/Engine/Public/ODSC/ODSCManager.h b/Engine/Source/Runtime/Engine/Public/ODSC/ODSCManager.h
--- a/Engine/Source/Runtime/Engine/Public/ODSC/ODSCManager.h
+++ b/Engine/Source/Runtime/Engine/Public/ODSC/ODSCManager.h
@@ -51,12 +51,14 @@ public:
 	 *
 	 * @return false if no longer needs ticking
 	 */
-	void AddThreadedRequest(const TArray<FString>& MaterialsToCompile, const FString& ShaderTypesToLoad, EShaderPlatform ShaderPlatform, ODSCRecompileCommand RecompileCommandType);
+	void AddThreadedRequest(const TArray<FString>& MaterialsToCompile, const FString& ShaderTypesToLoad, EShaderPlatform ShaderPlatform, ERHIFeatureLevel::Type FeatureLevel, EMaterialQualityLevel::Type QualityLevel, ODSCRecompileCommand RecompileCommandType);
 
 	/**
 	 * Add a request to compile a pipeline (VS/PS) of shaders.  The results are submitted and processed in an async manner.
 	 *
 	 * @param ShaderPlatform - Which shader platform to compile for.
+	 * @param FeatureLevel - Which feature level to compile for.
+	 * @param QualityLevel - Which material quality level to compile for.
 	 * @param MaterialName - The name of the material to compile.
 	 * @param VertexFactoryName - The name of the vertex factory type we should compile.
 	 * @param PipelineName - The name of the shader pipeline we should compile.
@@ -64,7 +66,7 @@ public:
 	 *
 	 * @return false if no longer needs ticking
 	 */
-	void AddThreadedShaderPipelineRequest(EShaderPlatform ShaderPlatform, const FString& MaterialName, const FString& VertexFactoryName, const FString& PipelineName, const TArray<FString>& ShaderTypeNames);
+	void AddThreadedShaderPipelineRequest(EShaderPlatform ShaderPlatform, ERHIFeatureLevel::Type FeatureLevel, EMaterialQualityLevel::Type QualityLevel, const FString& MaterialName, const FString& VertexFactoryName, const FString& PipelineName, const TArray<FString>& ShaderTypeNames);
 private:
 
 	void OnEnginePreExit();
diff --git a/Engine/Source/Runtime/Engine/Public/ShaderCompiler.h b/Engine/Source/Runtime/Engine/Public/ShaderCompiler.h
--- a/Engine/Source/Runtime/Engine/Public/ShaderCompiler.h
+++ b/Engine/Source/Runtime/Engine/Public/ShaderCompiler.h
@@ -26,6 +26,7 @@
 #include "Containers/HashTable.h"
 #include "Containers/List.h"
 #include "Hash/Blake3.h"
+#include "SceneTypes.h"
 
 class FShaderCompileJob;
 class FShaderPipelineCompileJob;
@@ -1234,6 +1235,12 @@ struct FODSCRequestPayload
 	/** The shader platform to compile for. */
 	EShaderPlatform ShaderPlatform;
 
+	/** Which feature level to compile for. */
+	ERHIFeatureLevel::Type FeatureLevel;
+
+	/** Which material quality level to compile for. */
+	EMaterialQualityLevel::Type QualityLevel;
+
 	/** Which material do we compile for?. */
 	FString MaterialName;
 
@@ -1250,7 +1257,7 @@ struct FODSCRequestPayload
 	FString RequestHash;
 
 	ENGINE_API FODSCRequestPayload() {};
-	ENGINE_API FODSCRequestPayload(EShaderPlatform InShaderPlatform, const FString& InMaterialName, const FString& InVertexFactoryName, const FString& InPipelineName, const TArray<FString>& InShaderTypeNames, const FString& InRequestHash);
+	ENGINE_API FODSCRequestPayload(EShaderPlatform InShaderPlatform, ERHIFeatureLevel::Type InFeatureLevel, EMaterialQualityLevel::Type InQualityLevel, const FString& InMaterialName, const FString& InVertexFactoryName, const FString& InPipelineName, const TArray<FString>& InShaderTypeNames, const FString& InRequestHash);
 
 	/**
 	* Serializes FODSCRequestPayload value from or into this archive.
@@ -1281,6 +1288,10 @@ struct FShaderRecompileData
 	/** Shader platform */
 	EShaderPlatform ShaderPlatform = SP_NumPlatforms;
 
+	ERHIFeatureLevel::Type FeatureLevel = ERHIFeatureLevel::SM5;
+
+	EMaterialQualityLevel::Type QualityLevel = EMaterialQualityLevel::High;
+
 	/** All filenames that have been changed during the shader compilation. */
 	TArray<FString>* ModifiedFiles = nullptr;
 
@@ -1315,6 +1326,8 @@ struct FShaderRecompileData
 	{
 		PlatformName = Other.PlatformName;
 		ShaderPlatform = Other.ShaderPlatform;
+		FeatureLevel = Other.FeatureLevel;
+		QualityLevel = Other.QualityLevel;
 		ModifiedFiles = Other.ModifiedFiles;
 		MeshMaterialMaps = Other.MeshMaterialMaps;
 		MaterialsToLoad = Other.MaterialsToLoad;
@@ -1326,6 +1339,8 @@ struct FShaderRecompileData
 
 		return *this;
 	}
+
+	ENGINE_API friend FArchive& operator<<(FArchive& Ar, FShaderRecompileData& Elem);
 };
 
 #if WITH_EDITOR
diff --git a/Engine/Source/Runtime/NetworkFileSystem/Private/NetworkFileServerConnection.cpp b/Engine/Source/Runtime/NetworkFileSystem/Private/NetworkFileServerConnection.cpp
--- a/Engine/Source/Runtime/NetworkFileSystem/Private/NetworkFileServerConnection.cpp
+++ b/Engine/Source/Runtime/NetworkFileSystem/Private/NetworkFileServerConnection.cpp
@@ -1296,19 +1296,11 @@ void FNetworkFileServerClientConnection::ProcessRecompileShaders( FArchive& In,
 	TArray<uint8> GlobalShaderMap;
 	FShaderRecompileData RecompileData(ConnectedPlatformName, &RecompileModifiedFiles, &MeshMaterialMaps, &GlobalShaderMap);
 
-	// tell other side all the materials to load, by pathname
-	In << RecompileData.MaterialsToLoad;
-	In << RecompileData.ShaderTypesToLoad;
-	int32 iShaderPlatform = static_cast<int32>(RecompileData.ShaderPlatform);
-	In << iShaderPlatform;
-	RecompileData.ShaderPlatform = static_cast<EShaderPlatform>(iShaderPlatform);
-	In << RecompileData.CommandType;
-	In << RecompileData.ShadersToRecompile;
+	In << RecompileData;
 
 	NetworkFileDelegates->RecompileShadersDelegate.ExecuteIfBound(RecompileData);
 
 	// tell other side what to do!
-	Out << RecompileModifiedFiles;
 	Out << MeshMaterialMaps;
 	Out << GlobalShaderMap;
 }
