Fixed incorrect class filter for default destructible material to allow proper selection and avoid warnings for missing classes

diff --git a/Engine/Source/Runtime/Engine/Classes/Engine/Engine.h b/Engine/Source/Runtime/Engine/Classes/Engine/Engine.h
--- a/Engine/Source/Runtime/Engine/Classes/Engine/Engine.h
+++ b/Engine/Source/Runtime/Engine/Classes/Engine/Engine.h
@@ -1245,7 +1245,7 @@ public:
 	TObjectPtr<class UPhysicalMaterial> DefaultDestructiblePhysMaterial;
 
 	/** Path of the PhysicalMaterial to use if none is defined for a particular object. */
-	UPROPERTY(globalconfig, EditAnywhere, Category = DefaultMaterials, meta = (AllowedClasses = "Physics", DisplayName = "Destructible Physics Material"))
+	UPROPERTY(globalconfig, EditAnywhere, Category = DefaultMaterials, meta = (AllowedClasses = "PhysicalMaterial", DisplayName = "Destructible Physics Material"))
 	FSoftObjectPath DefaultDestructiblePhysMaterialName;
 
 	/** Deprecated rules for redirecting renamed objects, replaced by the CoreRedirects system*/
