UE-148463 - Update cumulative time tracking Engine variables to doubles

As a part of the LWC work, this change primarily focues on updating FGameTime and UWorld to use doubles for accumulated time. Note that time deltas will remain as floats.

Additionally:

* A few time-related BP functions were updated to use doubles where appropriate. as well as the FRenderThreadUpdateContext struct for the Slate renderer.
* Some compilation errors were also fixed due to the change.
* Minor changes of float literals to drop the 'f' suffix.

diff --git a/Engine/Source/Developer/FunctionalTesting/Public/AutomationScreenshotOptions.h b/Engine/Source/Developer/FunctionalTesting/Public/AutomationScreenshotOptions.h
--- a/Engine/Source/Developer/FunctionalTesting/Public/AutomationScreenshotOptions.h
+++ b/Engine/Source/Developer/FunctionalTesting/Public/AutomationScreenshotOptions.h
@@ -130,7 +130,7 @@ public:
 	 * on the game thread is unaffected.
 	 */
 	UPROPERTY(EditAnywhere, BlueprintReadWrite, Category="Screenshot", meta=(editcondition = "bOverride_OverrideTimeTo"))
-	float OverrideTimeTo;
+	double OverrideTimeTo;
 
 	/**
 	 * Disables Anti-Aliasing, Motion Blur, Screen Space Reflections, Eye Adaptation, Tonemapper and Contact
diff --git a/Engine/Source/Editor/GraphEditor/Private/MaterialNodes/SGraphNodeMaterialBase.cpp b/Engine/Source/Editor/GraphEditor/Private/MaterialNodes/SGraphNodeMaterialBase.cpp
--- a/Engine/Source/Editor/GraphEditor/Private/MaterialNodes/SGraphNodeMaterialBase.cpp
+++ b/Engine/Source/Editor/GraphEditor/Private/MaterialNodes/SGraphNodeMaterialBase.cpp
@@ -221,7 +221,7 @@ void FPreviewElement::DrawRenderThread(FRHICommandListImmediate& RHICmdList, con
 		RenderTarget->SetRenderTargetTexture(*(FTexture2DRHIRef*)InWindowBackBuffer);
 		{
 			// Check realtime mode for whether to pass current time to canvas
-			float CurrentTime = bIsRealtime ? (FApp::GetCurrentTime() - GStartTime) : 0.0f;
+			double CurrentTime = bIsRealtime ? (FApp::GetCurrentTime() - GStartTime) : 0.0;
 			float DeltaTime = bIsRealtime ? FApp::GetDeltaTime() : 0.0f;
 
 			FCanvas Canvas(RenderTarget, NULL, FGameTime::CreateUndilated(CurrentTime, DeltaTime), GMaxRHIFeatureLevel);
diff --git a/Engine/Source/Runtime/Engine/Classes/Engine/World.h b/Engine/Source/Runtime/Engine/Classes/Engine/World.h
--- a/Engine/Source/Runtime/Engine/Classes/Engine/World.h
+++ b/Engine/Source/Runtime/Engine/Classes/Engine/World.h
@@ -142,13 +142,11 @@ private:
 /** Contains all the timings of a gaming frame, to handle pause and time dilation (for instance bullet time) of the world. */
 struct ENGINE_API FGameTime
 {
-	using FTimeType = float;
-
 	FORCEINLINE_DEBUGGABLE FGameTime()
-		: RealTimeSeconds(FTimeType(0))
-		, WorldTimeSeconds(FTimeType(0))
-		, DeltaRealTimeSeconds(FTimeType(0))
-		, DeltaWorldTimeSeconds(FTimeType(0))
+		: RealTimeSeconds(0.0)
+		, WorldTimeSeconds(0.0)
+		, DeltaRealTimeSeconds(0.0f)
+		, DeltaWorldTimeSeconds(0.0f)
 	{ }
 
 	FGameTime(const FGameTime&) = default;
@@ -157,36 +155,36 @@ struct ENGINE_API FGameTime
 	// Returns the game time since GStartTime.
 	static FGameTime GetTimeSinceAppStart();
 
-	static FORCEINLINE_DEBUGGABLE FGameTime CreateUndilated(FTimeType InRealTimeSeconds, FTimeType InDeltaRealTimeSeconds)
+	static FORCEINLINE_DEBUGGABLE FGameTime CreateUndilated(double InRealTimeSeconds, float InDeltaRealTimeSeconds)
 	{
 		return FGameTime::CreateDilated(InRealTimeSeconds, InDeltaRealTimeSeconds, InRealTimeSeconds, InDeltaRealTimeSeconds);
 	}
 
-	static FORCEINLINE_DEBUGGABLE FGameTime CreateDilated(FTimeType InRealTimeSeconds, FTimeType InDeltaRealTimeSeconds, FTimeType InWorldTimeSeconds, FTimeType InDeltaWorldTimeSeconds)
+	static FORCEINLINE_DEBUGGABLE FGameTime CreateDilated(double InRealTimeSeconds, float InDeltaRealTimeSeconds, double InWorldTimeSeconds, float InDeltaWorldTimeSeconds)
 	{
 		return FGameTime(InRealTimeSeconds, InDeltaRealTimeSeconds, InWorldTimeSeconds, InDeltaWorldTimeSeconds);
 	}
 
 	/** Returns time in seconds since level began play, but IS NOT paused when the game is paused, and IS NOT dilated/clamped. */
-	FORCEINLINE_DEBUGGABLE FTimeType GetRealTimeSeconds() const
+	FORCEINLINE_DEBUGGABLE double GetRealTimeSeconds() const
 	{
 		return RealTimeSeconds;
 	}
 
 	/** Returns frame delta time in seconds with no adjustment for time dilation and pause. */
-	FORCEINLINE_DEBUGGABLE FTimeType GetDeltaRealTimeSeconds() const
+	FORCEINLINE_DEBUGGABLE float GetDeltaRealTimeSeconds() const
 	{
 		return DeltaRealTimeSeconds;
 	}
 
 	/** Returns time in seconds since level began play, but IS paused when the game is paused, and IS dilated/clamped. */
-	FORCEINLINE_DEBUGGABLE FTimeType GetWorldTimeSeconds() const
+	FORCEINLINE_DEBUGGABLE double GetWorldTimeSeconds() const
 	{
 		return WorldTimeSeconds;
 	}
 
 	/** Returns frame delta time in seconds adjusted by e.g. time dilation. */
-	FORCEINLINE_DEBUGGABLE FTimeType GetDeltaWorldTimeSeconds() const
+	FORCEINLINE_DEBUGGABLE float GetDeltaWorldTimeSeconds() const
 	{
 		return DeltaWorldTimeSeconds;
 	}
@@ -194,24 +192,24 @@ struct ENGINE_API FGameTime
 	/** Returns how much world time is slowed compared to real time. */
 	FORCEINLINE_DEBUGGABLE float GetTimeDilation() const
 	{
-		ensure(DeltaRealTimeSeconds > FTimeType(0));
-		return float(DeltaWorldTimeSeconds / DeltaRealTimeSeconds);
+		ensure(DeltaRealTimeSeconds > 0.0f);
+		return DeltaWorldTimeSeconds / DeltaRealTimeSeconds;
 	}
 
 	/** Returns whether the world time is paused. */
 	FORCEINLINE_DEBUGGABLE bool IsPaused() const
 	{
-		return DeltaWorldTimeSeconds == FTimeType(0);
+		return DeltaWorldTimeSeconds == 0.0f;
 	}
 
 private:
-	FTimeType RealTimeSeconds;
-	FTimeType WorldTimeSeconds;
+	double RealTimeSeconds;
+	double WorldTimeSeconds;
 
-	FTimeType DeltaRealTimeSeconds;
-	FTimeType DeltaWorldTimeSeconds;
+	float DeltaRealTimeSeconds;
+	float DeltaWorldTimeSeconds;
 
-	FORCEINLINE_DEBUGGABLE FGameTime(FTimeType InRealTimeSeconds, FTimeType InDeltaRealTimeSeconds, FTimeType InWorldTimeSeconds, FTimeType InDeltaWorldTimeSeconds)
+	FORCEINLINE_DEBUGGABLE FGameTime(double InRealTimeSeconds, float InDeltaRealTimeSeconds, double InWorldTimeSeconds, float InDeltaWorldTimeSeconds)
 		: RealTimeSeconds(InRealTimeSeconds)
 		, WorldTimeSeconds(InWorldTimeSeconds)
 		, DeltaRealTimeSeconds(InDeltaRealTimeSeconds)
@@ -337,7 +335,7 @@ private:
 	/** The context we are running in. Can be used to get the FWorldContext from Engine*/
 	FName WorldContextHandle;
 	/** Real time which we started traveling at  */
-	double SeamlessTravelStartTime = 0.f;
+	double SeamlessTravelStartTime = 0.0;
 
 	/** copy data between the old world and the new world */
 	void CopyWorldData();
@@ -1209,7 +1207,7 @@ public:
 	/** Cached view information from the last rendered frame. */
 	TArray<FWorldCachedViewInfo>				CachedViewInfoRenderedLastFrame;
 	/** WorldTimeSeconds when this world was last rendered. */
-	float										LastRenderTime = 0.0f;
+	double										LastRenderTime = 0.0;
 
 	/** The current renderer feature level of this world */
 	TEnumAsByte<ERHIFeatureLevel::Type> FeatureLevel;
@@ -1779,16 +1777,16 @@ public:
 	double LastTimeUnbuiltLightingWasEncountered;
 
 	/**  Time in seconds since level began play, but IS paused when the game is paused, and IS dilated/clamped. */
-	float TimeSeconds;
+	double TimeSeconds;
 
 	/**  Time in seconds since level began play, but IS NOT paused when the game is paused, and IS dilated/clamped. */
-	float UnpausedTimeSeconds;
+	double UnpausedTimeSeconds;
 
 	/** Time in seconds since level began play, but IS NOT paused when the game is paused, and IS NOT dilated/clamped. */
-	float RealTimeSeconds;
+	double RealTimeSeconds;
 
 	/** Time in seconds since level began play, but IS paused when the game is paused, and IS NOT dilated/clamped. */
-	float AudioTimeSeconds;
+	double AudioTimeSeconds;
 
 	/** Frame delta time in seconds with no adjustment for time dilation. */
 	float DeltaRealTimeSeconds;
@@ -1797,7 +1795,7 @@ public:
 	float DeltaTimeSeconds;
 
 	/** time at which to start pause **/
-	float PauseDelay;
+	double PauseDelay;
 
 	/** Current location of this world origin */
 	FIntVector OriginLocation;
@@ -2599,28 +2597,28 @@ public:
 	 *
 	 * @return time in seconds since world was brought up for play
 	 */
-	float GetTimeSeconds() const;
+	double GetTimeSeconds() const;
 
 	/**
 	* Returns time in seconds since world was brought up for play, IS NOT stopped when game pauses, IS dilated/clamped
 	*
 	* @return time in seconds since world was brought up for play
 	*/
-	float GetUnpausedTimeSeconds() const;
+	double GetUnpausedTimeSeconds() const;
 
 	/**
 	* Returns time in seconds since world was brought up for play, does NOT stop when game pauses, NOT dilated/clamped
 	*
 	* @return time in seconds since world was brought up for play
 	*/
-	float GetRealTimeSeconds() const;
+	double GetRealTimeSeconds() const;
 
 	/**
 	* Returns time in seconds since world was brought up for play, IS stopped when game pauses, NOT dilated/clamped
 	*
 	* @return time in seconds since world was brought up for play
 	*/
-	float GetAudioTimeSeconds() const;
+	double GetAudioTimeSeconds() const;
 
 	/**
 	 * Returns the frame delta time in seconds adjusted by e.g. time dilation.
@@ -2637,7 +2635,7 @@ public:
 	FGameTime GetTime() const;
 
 	/** Helper for getting the time since a certain time. */
-	float TimeSince( float Time ) const;
+	double TimeSince(double Time) const;
 
 	/** Creates a new physics scene for this world. */
 	void CreatePhysicsScene(const AWorldSettings* Settings = nullptr);
@@ -4254,23 +4252,23 @@ private:
 //////////////////////////////////////////////////////////////////////////
 // UWorld inlines:
 
-FORCEINLINE_DEBUGGABLE float UWorld::GetTimeSeconds() const
+FORCEINLINE_DEBUGGABLE double UWorld::GetTimeSeconds() const
 {
 	return TimeSeconds;
 }
 
-FORCEINLINE_DEBUGGABLE float UWorld::GetUnpausedTimeSeconds() const
+FORCEINLINE_DEBUGGABLE double UWorld::GetUnpausedTimeSeconds() const
 {
 	return UnpausedTimeSeconds;
 }
 
-FORCEINLINE_DEBUGGABLE float UWorld::GetRealTimeSeconds() const
+FORCEINLINE_DEBUGGABLE double UWorld::GetRealTimeSeconds() const
 {
 	checkSlow(!IsInActualRenderingThread());
 	return RealTimeSeconds;
 }
 
-FORCEINLINE_DEBUGGABLE float UWorld::GetAudioTimeSeconds() const
+FORCEINLINE_DEBUGGABLE double UWorld::GetAudioTimeSeconds() const
 {
 	return AudioTimeSeconds;
 }
@@ -4287,7 +4285,7 @@ FORCEINLINE_DEBUGGABLE FGameTime UWorld::GetTime() const
 		TimeSeconds, DeltaTimeSeconds);
 }
 
-FORCEINLINE_DEBUGGABLE float UWorld::TimeSince(float Time) const
+FORCEINLINE_DEBUGGABLE double UWorld::TimeSince(double Time) const
 {
 	return GetTimeSeconds() - Time;
 }
diff --git a/Engine/Source/Runtime/Engine/Classes/Kismet/GameplayStatics.h b/Engine/Source/Runtime/Engine/Classes/Kismet/GameplayStatics.h
--- a/Engine/Source/Runtime/Engine/Classes/Kismet/GameplayStatics.h
+++ b/Engine/Source/Runtime/Engine/Classes/Kismet/GameplayStatics.h
@@ -1155,27 +1155,30 @@ public:
 
 	/** Returns the frame delta time in seconds, adjusted by time dilation. */
 	UFUNCTION(BlueprintPure, Category = "Utilities|Time", meta = (WorldContext="WorldContextObject"))
-	static float GetWorldDeltaSeconds(const UObject* WorldContextObject);
+	static double GetWorldDeltaSeconds(const UObject* WorldContextObject);
 
 	/** Returns time in seconds since world was brought up for play, adjusted by time dilation and IS stopped when game pauses */
 	UFUNCTION(BlueprintPure, Category="Utilities|Time", meta=(WorldContext="WorldContextObject"))
-	static float GetTimeSeconds(const UObject* WorldContextObject);
+	static double GetTimeSeconds(const UObject* WorldContextObject);
 
 	/** Returns time in seconds since world was brought up for play, adjusted by time dilation and IS NOT stopped when game pauses */
 	UFUNCTION(BlueprintPure, Category="Utilities|Time", meta=(WorldContext="WorldContextObject"))
-	static float GetUnpausedTimeSeconds(const UObject* WorldContextObject);
+	static double GetUnpausedTimeSeconds(const UObject* WorldContextObject);
 
 	/** Returns time in seconds since world was brought up for play, does NOT stop when game pauses, NOT dilated/clamped */
 	UFUNCTION(BlueprintPure, Category="Utilities|Time", meta=(WorldContext="WorldContextObject"))
-	static float GetRealTimeSeconds(const UObject* WorldContextObject);
+	static double GetRealTimeSeconds(const UObject* WorldContextObject);
 	
 	/** Returns time in seconds since world was brought up for play, IS stopped when game pauses, NOT dilated/clamped. */
 	UFUNCTION(BlueprintPure, Category="Utilities|Time", meta=(WorldContext="WorldContextObject"))
-	static float GetAudioTimeSeconds(const UObject* WorldContextObject);
+	static double GetAudioTimeSeconds(const UObject* WorldContextObject);
+
+	UE_DEPRECATED(5.1, "This method has been deprecated and will be removed. Use the double version instead.")
+	static void GetAccurateRealTime(int32& Seconds, float& PartialSeconds);
 
 	/** Returns time in seconds since the application was started. Unlike the other time functions this is accurate to the exact time this function is called instead of set once per frame. */
 	UFUNCTION(BlueprintPure, Category="Utilities|Time")
-	static void GetAccurateRealTime(int32& Seconds, float& PartialSeconds);
+	static void GetAccurateRealTime(int32& Seconds, double& PartialSeconds);
 
 	/*~ DVRStreaming API */
 	
diff --git a/Engine/Source/Runtime/Engine/Classes/Kismet/KismetSystemLibrary.h b/Engine/Source/Runtime/Engine/Classes/Kismet/KismetSystemLibrary.h
--- a/Engine/Source/Runtime/Engine/Classes/Kismet/KismetSystemLibrary.h
+++ b/Engine/Source/Runtime/Engine/Classes/Kismet/KismetSystemLibrary.h
@@ -193,7 +193,7 @@ class ENGINE_API UKismetSystemLibrary : public UBlueprintFunctionLibrary
 	 * @param WorldContextObject	World context
 	 */
 	UFUNCTION(BlueprintPure, Category="Utilities|Time", meta=(WorldContext="WorldContextObject") )
-	static float GetGameTimeInSeconds(const UObject* WorldContextObject);
+	static double GetGameTimeInSeconds(const UObject* WorldContextObject);
 
 	/** Returns the value of GFrameCounter, a running count of the number of frames that have occurred. */
 	UFUNCTION(BlueprintPure, Category = "Utilities")
diff --git a/Engine/Source/Runtime/Engine/Private/GameplayStatics.cpp b/Engine/Source/Runtime/Engine/Private/GameplayStatics.cpp
--- a/Engine/Source/Runtime/Engine/Private/GameplayStatics.cpp
+++ b/Engine/Source/Runtime/Engine/Private/GameplayStatics.cpp
@@ -2406,34 +2406,34 @@ FMemoryReader UGameplayStatics::StripSaveGameHeader(const TArray<uint8>& SaveDat
 	return MemoryReader;
 }
 
-float UGameplayStatics::GetWorldDeltaSeconds(const UObject* WorldContextObject)
+double UGameplayStatics::GetWorldDeltaSeconds(const UObject* WorldContextObject)
 {
 	UWorld* World = GEngine->GetWorldFromContextObject(WorldContextObject, EGetWorldErrorMode::LogAndReturnNull);
-	return World ? World->GetDeltaSeconds() : 0.f;
+	return World ? World->GetDeltaSeconds() : 0.0;
 }
 
-float UGameplayStatics::GetTimeSeconds(const UObject* WorldContextObject)
+double UGameplayStatics::GetTimeSeconds(const UObject* WorldContextObject)
 {
 	UWorld* World = GEngine->GetWorldFromContextObject(WorldContextObject, EGetWorldErrorMode::LogAndReturnNull);
-	return World ? World->GetTimeSeconds() : 0.f;
+	return World ? World->GetTimeSeconds() : 0.0;
 }
 
-float UGameplayStatics::GetUnpausedTimeSeconds(const UObject* WorldContextObject)
+double UGameplayStatics::GetUnpausedTimeSeconds(const UObject* WorldContextObject)
 {
 	UWorld* World = GEngine->GetWorldFromContextObject(WorldContextObject, EGetWorldErrorMode::LogAndReturnNull);
-	return World ? World->GetUnpausedTimeSeconds() : 0.f;
+	return World ? World->GetUnpausedTimeSeconds() : 0.0;
 }
 
-float UGameplayStatics::GetRealTimeSeconds(const UObject* WorldContextObject)
+double UGameplayStatics::GetRealTimeSeconds(const UObject* WorldContextObject)
 {
 	UWorld* World = GEngine->GetWorldFromContextObject(WorldContextObject, EGetWorldErrorMode::LogAndReturnNull);
-	return World ? World->GetRealTimeSeconds() : 0.f;
+	return World ? World->GetRealTimeSeconds() : 0.0;
 }
 
-float UGameplayStatics::GetAudioTimeSeconds(const UObject* WorldContextObject)
+double UGameplayStatics::GetAudioTimeSeconds(const UObject* WorldContextObject)
 {
 	UWorld* World = GEngine->GetWorldFromContextObject(WorldContextObject, EGetWorldErrorMode::LogAndReturnNull);
-	return World ? World->GetAudioTimeSeconds() : 0.f;
+	return World ? World->GetAudioTimeSeconds() : 0.0;
 }
 
 void UGameplayStatics::GetAccurateRealTime(int32& Seconds, float& PartialSeconds)
@@ -2443,6 +2443,13 @@ void UGameplayStatics::GetAccurateRealTime(int32& Seconds, float& PartialSeconds
 	PartialSeconds = TimeSeconds - double(Seconds);
 }
 
+void UGameplayStatics::GetAccurateRealTime(int32& Seconds, double& PartialSeconds)
+{
+	double TimeSeconds = FPlatformTime::Seconds() - GStartTime;
+	Seconds = floor(TimeSeconds);
+	PartialSeconds = TimeSeconds - double(Seconds);
+}
+
 void UGameplayStatics::EnableLiveStreaming(bool Enable)
 {
 	if (IDVRStreamingSystem* StreamingSystem = IPlatformFeaturesModule::Get().GetStreamingSystem())
diff --git a/Engine/Source/Runtime/Engine/Private/KismetSystemLibrary.cpp b/Engine/Source/Runtime/Engine/Private/KismetSystemLibrary.cpp
--- a/Engine/Source/Runtime/Engine/Private/KismetSystemLibrary.cpp
+++ b/Engine/Source/Runtime/Engine/Private/KismetSystemLibrary.cpp
@@ -199,10 +199,10 @@ bool UKismetSystemLibrary::DoesImplementInterface(const UObject* TestObject, TSu
 	return false;
 }
 
-float UKismetSystemLibrary::GetGameTimeInSeconds(const UObject* WorldContextObject)
+double UKismetSystemLibrary::GetGameTimeInSeconds(const UObject* WorldContextObject)
 {
 	UWorld* World = GEngine->GetWorldFromContextObject(WorldContextObject, EGetWorldErrorMode::LogAndReturnNull);
-	return World ? World->GetTimeSeconds() : 0.f;
+	return World ? World->GetTimeSeconds() : 0.0;
 }
 
 int64 UKismetSystemLibrary::GetFrameCount()
diff --git a/Engine/Source/Runtime/Engine/Private/SceneView.cpp b/Engine/Source/Runtime/Engine/Private/SceneView.cpp
--- a/Engine/Source/Runtime/Engine/Private/SceneView.cpp
+++ b/Engine/Source/Runtime/Engine/Private/SceneView.cpp
@@ -2732,7 +2732,7 @@ FSceneViewFamily::FSceneViewFamily(const ConstructionValues& CVS)
 	int32 Value = CVarRenderTimeFrozen.GetValueOnAnyThread();
 	if(Value)
 	{
-		Time = FGameTime::CreateDilated(0.0f, Time.GetDeltaRealTimeSeconds(), 0.0f, Time.GetDeltaWorldTimeSeconds());
+		Time = FGameTime::CreateDilated(0.0, Time.GetDeltaRealTimeSeconds(), 0.0, Time.GetDeltaWorldTimeSeconds());
 	}
 
 	DebugViewShaderMode = ChooseDebugViewShaderMode();
diff --git a/Engine/Source/Runtime/Engine/Private/World.cpp b/Engine/Source/Runtime/Engine/Private/World.cpp
--- a/Engine/Source/Runtime/Engine/Private/World.cpp
+++ b/Engine/Source/Runtime/Engine/Private/World.cpp
@@ -4729,10 +4729,10 @@ void UWorld::InitializeActorsForPlay(const FURL& InURL, bool bResetTime, FRegist
 	// Don't reset time for seamless world transitions.
 	if (bResetTime)
 	{
-		TimeSeconds = 0.0f;
-		UnpausedTimeSeconds = 0.0f;
-		RealTimeSeconds = 0.0f;
-		AudioTimeSeconds = 0.0f;
+		TimeSeconds = 0.0;
+		UnpausedTimeSeconds = 0.0;
+		RealTimeSeconds = 0.0;
+		AudioTimeSeconds = 0.0;
 	}
 
 	// Get URL Options
diff --git a/Engine/Source/Runtime/Renderer/Private/SceneRendering.cpp b/Engine/Source/Runtime/Renderer/Private/SceneRendering.cpp
--- a/Engine/Source/Runtime/Renderer/Private/SceneRendering.cpp
+++ b/Engine/Source/Runtime/Renderer/Private/SceneRendering.cpp
@@ -2673,7 +2673,7 @@ FSceneRenderer::FSceneRenderer(TArrayView<const FSceneViewFamily*> InViewFamilie
 		// rendering differences in materials.
 		if (bAnyViewIsLocked)
 		{
-			ViewFamily.Time = FGameTime::CreateDilated(0.0f, ViewFamily.Time.GetDeltaRealTimeSeconds(), 0.0f, ViewFamily.Time.GetDeltaWorldTimeSeconds());
+			ViewFamily.Time = FGameTime::CreateDilated(0.0, ViewFamily.Time.GetDeltaRealTimeSeconds(), 0.0, ViewFamily.Time.GetDeltaWorldTimeSeconds());
 		}
 
 		if (HitProxyConsumer)
diff --git a/Engine/Source/Runtime/SlateCore/Public/Rendering/SlateRenderer.h b/Engine/Source/Runtime/SlateCore/Public/Rendering/SlateRenderer.h
--- a/Engine/Source/Runtime/SlateCore/Public/Rendering/SlateRenderer.h
+++ b/Engine/Source/Runtime/SlateCore/Public/Rendering/SlateRenderer.h
@@ -30,9 +30,9 @@ typedef TRefCountPtr<FRHITexture> FTexture2DRHIRef;
 struct FRenderThreadUpdateContext
 {
 	class FSlateDrawBuffer* WindowDrawBuffer;
-	float WorldTimeSeconds;
+	double WorldTimeSeconds;
 	float DeltaTimeSeconds;
-	float RealTimeSeconds;
+	double RealTimeSeconds;
 	float DeltaRealTimeSeconds;
 	FRenderTarget* RenderTarget;
 	ISlate3DRenderer* Renderer;
diff --git a/Engine/Source/Runtime/UMG/Private/Slate/WidgetRenderer.cpp b/Engine/Source/Runtime/UMG/Private/Slate/WidgetRenderer.cpp
--- a/Engine/Source/Runtime/UMG/Private/Slate/WidgetRenderer.cpp
+++ b/Engine/Source/Runtime/UMG/Private/Slate/WidgetRenderer.cpp
@@ -258,9 +258,9 @@ void FWidgetRenderer::DrawWindow(
 			FRenderThreadUpdateContext RenderThreadUpdateContext =
 			{
 				&(ScopedDrawBuffer.GetDrawBuffer()),
-				static_cast<float>(FApp::GetCurrentTime() - GStartTime),
+				(FApp::GetCurrentTime() - GStartTime),
 				static_cast<float>(FApp::GetDeltaTime()),
-				static_cast<float>(FPlatformTime::Seconds() - GStartTime),
+				(FPlatformTime::Seconds() - GStartTime),
 				static_cast<float>(FApp::GetDeltaTime()),
 				RenderTarget,
 				Renderer.Get(),
@@ -321,9 +321,9 @@ bool FWidgetRenderer::DrawInvalidationRoot(TSharedRef<SVirtualWindow>& VirtualWi
 				FRenderThreadUpdateContext RenderThreadUpdateContext =
 				{
 					&(ScopedDrawBuffer.GetDrawBuffer()),
-					static_cast<float>(FApp::GetCurrentTime() - GStartTime),
+					(FApp::GetCurrentTime() - GStartTime),
 					static_cast<float>(FApp::GetDeltaTime()),
-					static_cast<float>(FPlatformTime::Seconds() - GStartTime),
+					(FPlatformTime::Seconds() - GStartTime),
 					static_cast<float>(FApp::GetDeltaTime()),
 					static_cast<FRenderTarget*>(RenderTarget->GameThread_GetRenderTargetResource()),
 					Renderer.Get(),
