Camera: Update Draw Frustum on tick if changed

diff --git a/Engine/Source/Runtime/Engine/Classes/Camera/CameraComponent.h b/Engine/Source/Runtime/Engine/Classes/Camera/CameraComponent.h
--- a/Engine/Source/Runtime/Engine/Classes/Camera/CameraComponent.h
+++ b/Engine/Source/Runtime/Engine/Classes/Camera/CameraComponent.h
@@ -114,6 +114,7 @@ public:
 	// UActorComponent interface
 	virtual void OnRegister() override;
 	virtual void OnUpdateTransform(EUpdateTransformFlags UpdateTransformFlags, ETeleportType Teleport) override;
+	virtual void TickComponent(float DeltaTime, ELevelTick TickType, FActorComponentTickFunction* ThisTickFunction) override;
 #if WITH_EDITOR
 	virtual void OnComponentDestroyed(bool bDestroyingHierarchy) override;
 	virtual void CheckForErrors() override;
@@ -171,6 +172,9 @@ protected:
 	/** Ensure the proxy mesh is in the correct place */
 	void UpdateProxyMeshTransform();
 
+	/* Update draw frustum values */
+	void UpdateDrawFrustum();
+
 #endif	// WITH_EDITORONLY_DATA
 
 	/**
diff --git a/Engine/Source/Runtime/Engine/Private/Camera/CameraComponent.cpp b/Engine/Source/Runtime/Engine/Private/Camera/CameraComponent.cpp
--- a/Engine/Source/Runtime/Engine/Private/Camera/CameraComponent.cpp
+++ b/Engine/Source/Runtime/Engine/Private/Camera/CameraComponent.cpp
@@ -52,6 +52,12 @@ UCameraComponent::UCameraComponent(const FObjectInitializer& ObjectInitializer)
 	bUsePawnControlRotation = false;
 	bAutoActivate = true;
 	bLockToHmd = true;
+
+#if WITH_EDITORONLY_DATA
+	bTickInEditor = true;
+#endif
+	
+	PrimaryComponentTick.bCanEverTick = true;
 }
 
 void UCameraComponent::OnUpdateTransform(EUpdateTransformFlags UpdateTransformFlags, ETeleportType Teleport)
@@ -144,6 +150,15 @@ void UCameraComponent::OnRegister()
 	Super::OnRegister();
 }
 
+void UCameraComponent::TickComponent(float DeltaTime, ELevelTick TickType, FActorComponentTickFunction* ThisTickFunction)
+{
+#if WITH_EDITORONLY_DATA	
+	UpdateDrawFrustum();
+#endif
+
+	Super::TickComponent(DeltaTime, TickType, ThisTickFunction);
+}
+
 #if WITH_EDITORONLY_DATA
 
 void UCameraComponent::PostLoad()
@@ -179,27 +194,53 @@ void UCameraComponent::ResetProxyMeshTransform()
 	}
 }
 
-
-void UCameraComponent::RefreshVisualRepresentation()
+void UCameraComponent::UpdateDrawFrustum()
 {
 	if (DrawFrustum != nullptr)
 	{
+		bool bAnythingChanged = false;
 		const float FrustumDrawDistance = 1000.0f;
 		if (ProjectionMode == ECameraProjectionMode::Perspective)
 		{
-			DrawFrustum->FrustumAngle = FieldOfView;
-			DrawFrustum->FrustumStartDist = 10.f;
-			DrawFrustum->FrustumEndDist = DrawFrustum->FrustumStartDist + FrustumDrawDistance;
+			if (DrawFrustum->FrustumAngle != FieldOfView ||
+				DrawFrustum->FrustumStartDist != 10.f ||
+				DrawFrustum->FrustumEndDist != DrawFrustum->FrustumStartDist + FrustumDrawDistance)
+			{
+				DrawFrustum->FrustumAngle = FieldOfView;
+				DrawFrustum->FrustumStartDist = 10.f;
+				DrawFrustum->FrustumEndDist = DrawFrustum->FrustumStartDist + FrustumDrawDistance;
+				bAnythingChanged = true;
+			}
 		}
 		else
 		{
-			DrawFrustum->FrustumAngle = -OrthoWidth;
-			DrawFrustum->FrustumStartDist = OrthoNearClipPlane;
-			DrawFrustum->FrustumEndDist = FMath::Min(OrthoFarClipPlane - OrthoNearClipPlane, FrustumDrawDistance);
+			if (DrawFrustum->FrustumAngle != -OrthoWidth ||
+				DrawFrustum->FrustumStartDist != OrthoNearClipPlane ||
+				DrawFrustum->FrustumEndDist != FMath::Min(OrthoFarClipPlane - OrthoNearClipPlane, FrustumDrawDistance))
+			{
+				DrawFrustum->FrustumAngle = -OrthoWidth;
+				DrawFrustum->FrustumStartDist = OrthoNearClipPlane;
+				DrawFrustum->FrustumEndDist = FMath::Min(OrthoFarClipPlane - OrthoNearClipPlane, FrustumDrawDistance);
+				bAnythingChanged = true;
+			}
+		}
+
+		if (DrawFrustum->FrustumAspectRatio != AspectRatio)
+		{
+			DrawFrustum->FrustumAspectRatio = AspectRatio;
+			bAnythingChanged = true;
+		}	
+		
+		if (bAnythingChanged)
+		{
+			DrawFrustum->MarkRenderStateDirty();
 		}
-		DrawFrustum->FrustumAspectRatio = AspectRatio;
-		DrawFrustum->MarkRenderStateDirty();
 	}
+}
+
+void UCameraComponent::RefreshVisualRepresentation()
+{
+	UpdateDrawFrustum();
 
 	// Update the proxy camera mesh if necessary
 	if (ProxyMeshComponent && ProxyMeshComponent->GetStaticMesh() != CameraMesh)
