Fix RestartLevel command in Play in Editor, it needs to strip the PIE prefix from the last travel URL

diff --git a/Engine/Source/Runtime/Engine/Private/GameInstance.cpp b/Engine/Source/Runtime/Engine/Private/GameInstance.cpp
--- a/Engine/Source/Runtime/Engine/Private/GameInstance.cpp
+++ b/Engine/Source/Runtime/Engine/Private/GameInstance.cpp
@@ -1134,6 +1134,11 @@ void UGameInstance::SetPersistentTravelURL(FURL InURL)
 {
 	check(WorldContext);
 	WorldContext->LastURL = InURL;
+
+#if WITH_EDITOR
+	// Strip any PIE prefix from the map name
+	WorldContext->LastURL.Map = UWorld::StripPIEPrefixFromPackageName(WorldContext->LastURL.Map, WorldContext->PIEPrefix);
+#endif
 }
 
 bool UGameInstance::EnableListenServer(bool bEnable, int32 PortOverride /*= 0*/)
