Making BINK the default codec

- Bink supports seekable streaming across all platforms

diff --git a/Engine/Source/Runtime/Engine/Private/SoundWave.cpp b/Engine/Source/Runtime/Engine/Private/SoundWave.cpp
--- a/Engine/Source/Runtime/Engine/Private/SoundWave.cpp
+++ b/Engine/Source/Runtime/Engine/Private/SoundWave.cpp
@@ -681,6 +681,8 @@ USoundWave::USoundWave(const FObjectInitializer& ObjectInitializer)
 	bProcedural = false;
 	bRequiresStopFade = false;
 
+	SoundAssetCompressionType = ESoundAssetCompressionType::BinkAudio;
+
 #if WITH_EDITOR
 	bWasStreamCachingEnabledOnLastCook = FPlatformCompressionUtilities::IsCurrentPlatformUsingStreamCaching();
 	bLoadedFromCookedData = false;
