Fixes for tick.DoAsyncEndOfFrameTasks.ValidateReplicatedProperties
- DiffProperties skips empty layouts
- FGuidReferences copy constructor to deep copy container references, needed by FFastArraySerializer properties using delta structs
- Update the validation to occur around the slate tick and not during the async recording

diff --git a/Engine/Source/Runtime/Engine/Private/RepLayout.cpp b/Engine/Source/Runtime/Engine/Private/RepLayout.cpp
--- a/Engine/Source/Runtime/Engine/Private/RepLayout.cpp
+++ b/Engine/Source/Runtime/Engine/Private/RepLayout.cpp
@@ -3375,7 +3375,7 @@ static FGuidReferencesMap* PrepReceivedArray(
 		{
 			NewGuidReferencesArray = &ParentGuidReferences->FindOrAdd(AbsOffset);
 
-			NewGuidReferencesArray->Array = new FGuidReferencesMap;
+			NewGuidReferencesArray->Array.Reset(new FGuidReferencesMap());
 			NewGuidReferencesArray->ParentIndex = Cmd.ParentIndex;
 			NewGuidReferencesArray->CmdIndex = CmdIndex;
 		}
@@ -3412,7 +3412,7 @@ static FGuidReferencesMap* PrepReceivedArray(
 	// Re-compute the base data values since they could have changed after the resize above
 	*OutBaseData = DataArray->GetData();
 
-	return NewGuidReferencesArray ? NewGuidReferencesArray->Array : nullptr;
+	return NewGuidReferencesArray ? NewGuidReferencesArray->Array.Get() : nullptr;
 }
 
 /** Struct containing parameters that don't change or can be shared throughout recursion of ReceiveProperties_r */
@@ -4050,7 +4050,7 @@ void FRepLayout::GatherGuidReferences_r(
 		{
 			check(Cmds[GuidReferences.CmdIndex].Type == ERepLayoutCmdType::DynamicArray);
 
-			GatherGuidReferences_r(GuidReferences.Array, OutReferencedGuids, OutTrackedGuidMemoryBytes);
+			GatherGuidReferences_r(GuidReferences.Array.Get(), OutReferencedGuids, OutTrackedGuidMemoryBytes);
 			continue;
 		}
 
@@ -4109,7 +4109,7 @@ bool FRepLayout::MoveMappedObjectToUnmapped_r(FGuidReferencesMap* GuidReferences
 		{
 			check(Cmds[GuidReferences.CmdIndex].Type == ERepLayoutCmdType::DynamicArray);
 
-			if (MoveMappedObjectToUnmapped_r(GuidReferences.Array, GUID, OwningObject))
+			if (MoveMappedObjectToUnmapped_r(GuidReferences.Array.Get(), GUID, OwningObject))
 			{
 				bFoundGUID = true;
 			}
@@ -4230,7 +4230,7 @@ void FRepLayout::UpdateUnmappedObjects_r(
 
 				const int32 NewMaxOffset = FMath::Min(ShadowArray->Num() * Cmd.ElementSize, Array->Num() * Cmd.ElementSize);
 
-				UpdateUnmappedObjects_r(RepState, GuidReferences.Array, OriginalObject, Connection, ShadowArrayData, ArrayData, NewMaxOffset, bCalledPreNetReceive, bOutSomeObjectsWereMapped, bOutHasMoreUnmapped);
+				UpdateUnmappedObjects_r(RepState, GuidReferences.Array.Get(), OriginalObject, Connection, ShadowArrayData, ArrayData, NewMaxOffset, bCalledPreNetReceive, bOutSomeObjectsWereMapped, bOutHasMoreUnmapped);
 			}
 			else
 			{
@@ -4238,7 +4238,7 @@ void FRepLayout::UpdateUnmappedObjects_r(
 				FRepObjectDataBuffer ArrayData(Array->GetData());
 				const int32 NewMaxOffset = Array->Num() * Cmd.ElementSize;
 
-				UpdateUnmappedObjects_r(RepState, GuidReferences.Array, OriginalObject, Connection, nullptr, ArrayData, NewMaxOffset, bCalledPreNetReceive, bOutSomeObjectsWereMapped, bOutHasMoreUnmapped);
+				UpdateUnmappedObjects_r(RepState, GuidReferences.Array.Get(), OriginalObject, Connection, nullptr, ArrayData, NewMaxOffset, bCalledPreNetReceive, bOutSomeObjectsWereMapped, bOutHasMoreUnmapped);
 			}
 			continue;
 		}
@@ -5013,6 +5013,10 @@ bool FRepLayout::DiffProperties(
 	TConstRepDataBuffer<SourceType> Source,
 	const EDiffPropertiesFlags DiffFlags) const
 {
+	if (IsEmpty())
+	{
+		return false;
+	}
 
 	// Currently, only lifetime properties init from their defaults, so default to that,
 	// but also diff conditional properties if requested.
diff --git a/Engine/Source/Runtime/Launch/Private/LaunchEngineLoop.cpp b/Engine/Source/Runtime/Launch/Private/LaunchEngineLoop.cpp
--- a/Engine/Source/Runtime/Launch/Private/LaunchEngineLoop.cpp
+++ b/Engine/Source/Runtime/Launch/Private/LaunchEngineLoop.cpp
@@ -5319,13 +5319,7 @@ void FEngineLoop::Tick()
 		const UWorld* const GameViewportWorld = GameViewport ? GameViewport->GetWorld() : nullptr;
 		UDemoNetDriver* const CurrentDemoNetDriver = GameViewportWorld ? GameViewportWorld->GetDemoNetDriver() : nullptr;
 
-		// Optionally validate that Slate has not modified any replicated properties for client replay recording.
-		FDemoSavedPropertyState PreSlateObjectStates;
 		const bool bValidateReplicatedProperties = CurrentDemoNetDriver && CVarDoAsyncEndOfFrameTasksValidateReplicatedProperties.GetValueOnGameThread() != 0;
-		if (bValidateReplicatedProperties)
-		{
-			PreSlateObjectStates = CurrentDemoNetDriver->SavePropertyState();
-		}
 
 		if (bDoConcurrentSlateTick)
 		{
@@ -5351,8 +5345,28 @@ void FEngineLoop::Tick()
 				},
 				ConcurrentTaskCompleteEvent);
 				check(ConcurrentTask.IsValid());
+
+				// If we're validating, we want to only test the slate tick so wait for the task
+				if (bValidateReplicatedProperties)
+				{
+					CSV_SCOPED_SET_WAIT_STAT(ConcurrentWithSlate);
+
+					QUICK_SCOPE_CYCLE_COUNTER(STAT_ConcurrentWithSlateTickTasks_Wait);
+					check(ConcurrentTaskCompleteEvent);
+					ConcurrentTaskCompleteEvent->Wait();
+					FPlatformProcess::ReturnSynchEventToPool(ConcurrentTaskCompleteEvent);
+					ConcurrentTaskCompleteEvent = nullptr;
+					ConcurrentTask = nullptr;
+				}
 			}
 		}
+
+		// Optionally validate that Slate has not modified any replicated properties for client replay recording.
+		FDemoSavedPropertyState PreSlateObjectStates;
+		if (bValidateReplicatedProperties)
+		{
+			PreSlateObjectStates = CurrentDemoNetDriver->SavePropertyState();
+		}
 #endif
 
 		// Tick(Advance) Time for the application and then tick and paint slate application widgets.
diff --git a/Engine/Source/Runtime/Net/Core/Private/Net/Core/Misc/GuidReferences.cpp b/Engine/Source/Runtime/Net/Core/Private/Net/Core/Misc/GuidReferences.cpp
deleted file mode 100644
diff --git a/Engine/Source/Runtime/Net/Core/Public/Net/Core/Misc/GuidReferences.h b/Engine/Source/Runtime/Net/Core/Public/Net/Core/Misc/GuidReferences.h
--- a/Engine/Source/Runtime/Net/Core/Public/Net/Core/Misc/GuidReferences.h
+++ b/Engine/Source/Runtime/Net/Core/Public/Net/Core/Misc/GuidReferences.h
@@ -35,8 +35,7 @@ class NETCORE_API FGuidReferences
 {
 public:
 	FGuidReferences():
-		NumBufferBits(0),
-		Array(nullptr)
+		NumBufferBits(0)
 	{}
 
 	FGuidReferences(
@@ -50,8 +49,7 @@ public:
 		ParentIndex(InParentIndex),
 		CmdIndex(InCmdIndex),
 		UnmappedGUIDs(InUnmappedGUIDs),
-		MappedDynamicGUIDs(InMappedDynamicGUIDs),
-		Array(nullptr)
+		MappedDynamicGUIDs(InMappedDynamicGUIDs)
 	{
 		NumBufferBits = UE_PTRDIFF_TO_INT32(InReader.GetPosBits() - InMark.GetPos());
 		InMark.Copy(InReader, Buffer);
@@ -68,7 +66,22 @@ public:
 		Array(InArray)
 	{}
 
-	~FGuidReferences();
+	FGuidReferences& operator=(const FGuidReferences& Other) = delete;
+
+	FGuidReferences(const FGuidReferences& Other)
+		: ParentIndex(Other.ParentIndex)
+		, CmdIndex(Other.CmdIndex)
+		, NumBufferBits(Other.NumBufferBits)
+		, UnmappedGUIDs(Other.UnmappedGUIDs)
+		, MappedDynamicGUIDs(Other.MappedDynamicGUIDs)
+		, Buffer(Other.Buffer)
+		, Array(nullptr)
+	{
+		if (Other.Array.IsValid())
+		{
+			Array.Reset(new FGuidReferencesMap(*Other.Array));
+		}
+	}
 
 	void CountBytes(FArchive& Ar) const
 	{
@@ -107,5 +120,5 @@ public:
 	 * If this FGuidReferences instance is owned by an Array Property that contains nested GUID References,
 	 * then this will be a valid FGuidReferencesMap containing the nested FGuidReferences.
 	 */
-	FGuidReferencesMap* Array;
+	TUniquePtr<FGuidReferencesMap> Array;
 };
\ No newline at end of file
