Texture: Added FDerivedData alongside FByteBulkData in the platform data

diff --git a/Engine/Source/Runtime/Engine/Private/Texture2D.cpp b/Engine/Source/Runtime/Engine/Private/Texture2D.cpp
--- a/Engine/Source/Runtime/Engine/Private/Texture2D.cpp
+++ b/Engine/Source/Runtime/Engine/Private/Texture2D.cpp
@@ -155,29 +155,33 @@ void FTexture2DMipMap::Serialize(FArchive& Ar, UObject* Owner, int32 MipIdx)
 #if WITH_EDITORONLY_DATA
 	if (!Ar.IsFilterEditorOnly())
 	{
+		bool bLocalPagedToDerivedData = DerivedData.HasData();
 		Ar << FileRegionType;
-		Ar << bPagedToDerivedData;
+		Ar << bLocalPagedToDerivedData;
 	}
 #endif // #if WITH_EDITORONLY_DATA
 }
 
 #if WITH_EDITORONLY_DATA
-int64 FTexture2DMipMap::StoreInDerivedDataCache(const FString& InDerivedDataKey, const FStringView& TextureName, bool bReplaceExistingDDC)
+int64 FTexture2DMipMap::StoreInDerivedDataCache(const FStringView InKey, const FStringView InName, const bool bInReplaceExisting)
 {
+	using namespace UE;
 	using namespace UE::DerivedData;
 
-	int64 BulkDataSizeInBytes = BulkData.GetBulkDataSize();
+	const int64 BulkDataSizeInBytes = BulkData.GetBulkDataSize();
 	check(BulkDataSizeInBytes > 0);
 
-	FValue DerivedData;
-	DerivedData = FValue::Compress(FSharedBuffer::MakeView(BulkData.Lock(LOCK_READ_ONLY), BulkDataSizeInBytes));
+	const FSharedString Name = InName;
+	const FCacheKey Key = ConvertLegacyCacheKey(InKey);
+	FValue Value = FValue::Compress(FSharedBuffer::MakeView(BulkData.Lock(LOCK_READ_ONLY), BulkDataSizeInBytes));
 	BulkData.Unlock();
 
 	FRequestOwner AsyncOwner(EPriority::Normal);
-	GetCache().PutValue({{{TextureName}, ConvertLegacyCacheKey(InDerivedDataKey), MoveTemp(DerivedData)}}, AsyncOwner);
+	const ECachePolicy Policy = bInReplaceExisting ? ECachePolicy::Store : ECachePolicy::Default;
+	GetCache().PutValue({{Name, Key, MoveTemp(Value), Policy}}, AsyncOwner);
 	AsyncOwner.KeepAlive();
 
-	SetPagedToDerivedData(true);
+	DerivedData = FDerivedData(Name, Key);
 	BulkData.RemoveBulkData();
 	return BulkDataSizeInBytes;
 }
diff --git a/Engine/Source/Runtime/Engine/Private/TextureDerivedData.cpp b/Engine/Source/Runtime/Engine/Private/TextureDerivedData.cpp
--- a/Engine/Source/Runtime/Engine/Private/TextureDerivedData.cpp
+++ b/Engine/Source/Runtime/Engine/Private/TextureDerivedData.cpp
@@ -1359,7 +1359,7 @@ static bool LoadDerivedStreamingMips(FTexturePlatformData& PlatformData, int32 F
 		for (int32 MipIndex = FirstMipToLoad; MipIndex < Mips.Num(); ++MipIndex)
 		{
 			const FTexture2DMipMap& Mip = Mips[MipIndex];
-			if (Mip.IsPagedToDerivedData())
+			if (Mip.IsPagedToDerivedData() && !Mip.BulkData.IsBulkDataLoaded())
 			{
 				TStringBuilder<256> MipNameBuilder;
 				MipNameBuilder.Append(DebugContext).Appendf(TEXT(" [MIP %d]"), MipIndex);
@@ -1376,7 +1376,7 @@ static bool LoadDerivedStreamingMips(FTexturePlatformData& PlatformData, int32 F
 		for (int32 MipIndex = FirstMipToLoad; MipIndex < Mips.Num(); ++MipIndex)
 		{
 			const FTexture2DMipMap& Mip = Mips[MipIndex];
-			if (Mip.IsPagedToDerivedData())
+			if (Mip.IsPagedToDerivedData() && !Mip.BulkData.IsBulkDataLoaded())
 			{
 				TStringBuilder<256> MipNameBuilder;
 				MipNameBuilder.Append(DebugContext).Appendf(TEXT(" [MIP %d]"), MipIndex);
@@ -1427,12 +1427,10 @@ static bool LoadDerivedStreamingVTChunks(const TArray<FVirtualTextureDataChunk>&
 	for (int32 ChunkIndex = 0; ChunkIndex < Chunks.Num(); ++ChunkIndex)
 	{
 		const FVirtualTextureDataChunk& Chunk = Chunks[ChunkIndex];
-		if (!Chunk.DerivedDataKey.IsEmpty())
+		if (!Chunk.DerivedDataKey.IsEmpty() && !Chunk.BulkData.IsBulkDataLoaded())
 		{
-			TStringBuilder<256> ChunkNameBuilder;
-			ChunkNameBuilder.Append(DebugContext).Appendf(TEXT(" [Chunk %d]"), ChunkIndex);
 			FCacheGetChunkRequest& Request = Requests.AddDefaulted_GetRef();
-			Request.Name = ChunkNameBuilder;
+			Request.Name = FSharedString(WriteToString<256>(DebugContext, TEXT(" [Chunk ]"), ChunkIndex, TEXT("]")));
 			Request.Key = ConvertLegacyCacheKey(Chunk.DerivedDataKey);
 			Request.UserData = ChunkIndex;
 		}
@@ -1785,7 +1783,6 @@ bool FTexturePlatformData::TryInlineMipData(int32 FirstMipToLoad, FStringView De
 		void* MipData = Mip.BulkData.Realloc(int64(MipBuffer.GetSize()));
 		FMemory::Memcpy(MipData, MipBuffer.GetData(), MipBuffer.GetSize());
 		Mip.BulkData.Unlock();
-		Mip.SetPagedToDerivedData(false);
 	};
 
 	if (!LoadDerivedStreamingMips(*this, FirstMipToLoad, DebugContext, StoreMip))
@@ -1800,7 +1797,6 @@ bool FTexturePlatformData::TryInlineMipData(int32 FirstMipToLoad, FStringView De
 		void* ChunkData = Chunk.BulkData.Realloc(int64(ChunkBuffer.GetSize()));
 		FMemory::Memcpy(ChunkData, ChunkBuffer.GetData(), ChunkBuffer.GetSize());
 		Chunk.BulkData.Unlock();
-		Chunk.DerivedDataKey.Empty();
 	};
 
 	if (VTData && !LoadDerivedStreamingVTChunks(VTData->Chunks, DebugContext, StoreChunk))
@@ -2065,12 +2061,10 @@ bool FTexturePlatformData::CanBeLoaded() const
 {
 	for (const FTexture2DMipMap& Mip : Mips)
 	{
-#if WITH_EDITORONLY_DATA
-		if (Mip.IsPagedToDerivedData())
+		if (Mip.DerivedData.HasData())
 		{
 			return true;
 		}
-#endif 
 		if (Mip.BulkData.CanLoadFromDisk())
 		{
 			return true;
@@ -2984,7 +2978,7 @@ bool UTexture::IsAsyncCacheComplete() const
 	{
 		if (const FTexturePlatformData* PlatformData = *RunningPlatformDataPtr)
 		{
-			if (PlatformData->AsyncTask && !PlatformData->AsyncTask->Poll())
+			if (!PlatformData->IsAsyncWorkComplete())
 			{
 				return false;
 			}
@@ -2997,7 +2991,7 @@ bool UTexture::IsAsyncCacheComplete() const
 		{
 			if (const FTexturePlatformData* PlatformData = Kvp.Value)
 			{
-				if (PlatformData->AsyncTask && !PlatformData->AsyncTask->Poll())
+				if (!PlatformData->IsAsyncWorkComplete())
 				{
 					return false;
 				}
@@ -3007,7 +3001,7 @@ bool UTexture::IsAsyncCacheComplete() const
 
 	return true;
 }
-	
+
 bool UTexture::TryCancelCachePlatformData()
 {
 	TRACE_CPUPROFILER_EVENT_SCOPE(UTexture::TryCancelCachePlatformData);
diff --git a/Engine/Source/Runtime/Engine/Private/TextureDerivedDataTask.cpp b/Engine/Source/Runtime/Engine/Private/TextureDerivedDataTask.cpp
--- a/Engine/Source/Runtime/Engine/Private/TextureDerivedDataTask.cpp
+++ b/Engine/Source/Runtime/Engine/Private/TextureDerivedDataTask.cpp
@@ -47,13 +47,13 @@
 #include "Misc/FeedbackContext.h"
 #include "Misc/Optional.h"
 #include "Misc/ScopeRWLock.h"
+#include "Misc/ScopeExit.h"
 #include "ProfilingDebugging/CookStats.h"
 #include "Serialization/BulkDataRegistry.h"
 #include "TextureDerivedDataBuildUtils.h"
 #include "VT/VirtualTextureChunkDDCCache.h"
 #include "VT/VirtualTextureDataBuilder.h"
 #include "TextureEncodingSettings.h"
-#include <atomic>
 
 static TAutoConsoleVariable<int32> CVarVTValidateCompressionOnLoad(
 	TEXT("r.VT.ValidateCompressionOnLoad"),
@@ -725,6 +725,7 @@ FTextureCacheDerivedDataWorker::FTextureCacheDerivedDataWorker(
 // DDC1 primary fetch/build work function
 void FTextureCacheDerivedDataWorker::DoWork()
 {
+	using namespace UE;
 	using namespace UE::DerivedData;
 
 	TRACE_CPUPROFILER_EVENT_SCOPE(FTextureCacheDerivedDataWorker::DoWork);
@@ -808,6 +809,33 @@ void FTextureCacheDerivedDataWorker::DoWork()
 		FMemoryReaderView Ar(RawDerivedData.GetView(), /*bIsPersistent=*/ true);
 		DerivedData->Serialize(Ar, NULL);
 		bSucceeded = true;
+
+		if (bForVirtualTextureStreamingBuild)
+		{
+			if (DerivedData->VTData && DerivedData->VTData->IsInitialized())
+			{
+				const FSharedString Name(TexturePathName);
+				for (FVirtualTextureDataChunk& Chunk : DerivedData->VTData->Chunks)
+				{
+					if (!Chunk.DerivedDataKey.IsEmpty())
+					{
+						Chunk.DerivedData = FDerivedData(Name, ConvertLegacyCacheKey(Chunk.DerivedDataKey));
+					}
+				}
+			}
+		}
+		else
+		{
+			if (const int32 StreamingMipCount = FMath::Max(0, DerivedData->Mips.Num() - DerivedData->GetNumNonStreamingMips()))
+			{
+				const FSharedString Name(TexturePathName);
+				for (int32 MipIndex = 0; MipIndex < StreamingMipCount; ++MipIndex)
+				{
+					DerivedData->Mips[MipIndex].DerivedData = FDerivedData(Name, ConvertLegacyCacheKey(DerivedData->GetDerivedDataMipKeyString(MipIndex, DerivedData->Mips[MipIndex])));
+				}
+			}
+		}
+
 		// Load any streaming (not inline) mips that are necessary for our platform.
 		if (bForDDC)
 		{
@@ -1062,7 +1090,8 @@ public:
 		{
 			return;
 		}
-		
+
+		TOptional<FTextureStatusMessageContext> StatusMessage;
 		if (IsInGameThread() && OwnerPriority == EPriority::Blocking)
 		{
 			// this gets sent whether or not we are building the texture, and is a rare edge case for UI feedback.
@@ -1087,7 +1116,7 @@ public:
 
 		FBuildDefinition FetchOrBuildDefinition = CreateDefinition(Build, Texture, TexturePath, FunctionName, InSettingsFetchOrBuild, bUseCompositeTexture);
 		DerivedData.FetchOrBuildDerivedDataKey.Emplace<FTexturePlatformData::FStructuredDerivedDataKey>(GetKey(FetchOrBuildDefinition, Texture, bUseCompositeTexture));
-		
+
 		bool bBuildKicked = false;
 		if (InSettingsFetchFirst)
 		{
@@ -1104,13 +1133,13 @@ public:
 				FBuildPolicyBuilder BuildPolicyBuilder(bInlineMips ? EBuildPolicy::Cache : (EBuildPolicy::CacheQuery | EBuildPolicy::SkipData));
 				if (!bInlineMips)
 				{
-					BuildPolicyBuilder.AddValuePolicy(FValueId::FromName("Description"_ASV), EBuildPolicy::Cache);
-					BuildPolicyBuilder.AddValuePolicy(FValueId::FromName("MipTail"_ASV), EBuildPolicy::Cache);
+					BuildPolicyBuilder.AddValuePolicy(FValueId::FromName(ANSITEXTVIEW("Description")), EBuildPolicy::Cache);
+					BuildPolicyBuilder.AddValuePolicy(FValueId::FromName(ANSITEXTVIEW("MipTail")), EBuildPolicy::Cache);
 				}
 
 				BuildSession.Get().Build(FetchDefinition, {}, BuildPolicyBuilder.Build(), *Owner,
-					[this, 
-					 FetchOrBuildDefinition = MoveTemp(FetchOrBuildDefinition), 
+					[this,
+					 FetchOrBuildDefinition = MoveTemp(FetchOrBuildDefinition),
 					 Flags,
 					 FetchOrBuildMetadata = InFetchOrBuildMetadata ? *InFetchOrBuildMetadata : FTexturePlatformData::FTextureEncodeResultMetadata()
 					](FBuildCompleteParams&& Params)
@@ -1120,7 +1149,7 @@ public:
 						default:
 						case EStatus::Ok:
 							return EndBuild(MoveTemp(Params));
-						case EStatus::Error:							
+						case EStatus::Error:
 							this->DerivedData.ResultMetadata = FetchOrBuildMetadata;
 							return BeginBuild(FetchOrBuildDefinition, Flags);
 						}
@@ -1137,6 +1166,11 @@ public:
 			}
 			BeginBuild(FetchOrBuildDefinition, Flags);
 		}
+
+		if (StatusMessage.IsSet())
+		{
+			Owner->Wait();
+		}
 	}
 
 	static UE::DerivedData::FBuildDefinition CreateDefinition(
@@ -1175,8 +1209,8 @@ private:
 		else
 		{
 			FBuildPolicyBuilder BuildPolicyBuilder(EBuildPolicy::Build | EBuildPolicy::CacheQuery | EBuildPolicy::CacheStoreOnBuild | EBuildPolicy::SkipData);
-			BuildPolicyBuilder.AddValuePolicy(FValueId::FromName("Description"_ASV), EBuildPolicy::Default);
-			BuildPolicyBuilder.AddValuePolicy(FValueId::FromName("MipTail"_ASV), EBuildPolicy::Default);
+			BuildPolicyBuilder.AddValuePolicy(FValueId::FromName(ANSITEXTVIEW("Description")), EBuildPolicy::Default);
+			BuildPolicyBuilder.AddValuePolicy(FValueId::FromName(ANSITEXTVIEW("MipTail")), EBuildPolicy::Default);
 			BuildPolicy = BuildPolicyBuilder.Build();
 		}
 
@@ -1193,9 +1227,8 @@ private:
 			[](const FValue& Value) { return Value.GetData().GetRawSize(); }, uint64(0));
 		if (Params.Status != EStatus::Canceled)
 		{
-			WriteDerivedData(MoveTemp(Params.Output));
+			WriteDerivedData(MoveTemp(Params));
 		}
-		StatusMessage.Reset();
 	}
 
 	void Finalize(bool& bOutFoundInCache, uint64& OutProcessedByteCount) final
@@ -1374,46 +1407,49 @@ public:
 
 private:
 
-	static bool DeserializeTextureFromValues(FTexturePlatformData& DerivedData, const UE::DerivedData::FBuildOutput& Output, int32 FirstMipToLoad, bool bInlineMips)
+	static bool DeserializeTextureFromValues(FTexturePlatformData& DerivedData, UE::DerivedData::FBuildCompleteParams&& Params, int32 FirstMipToLoad, bool bInlineMips)
 	{
+		using namespace UE;
 		using namespace UE::DerivedData;
-		const FValueWithId& Value = Output.GetValue(FValueId::FromName("Description"_ASV));
+
+		const FBuildOutput& Output = Params.Output;
+		const FValueWithId& Value = Output.GetValue(FValueId::FromName(ANSITEXTVIEW("Description")));
 		if (!Value)
 		{
 			UE_LOG(LogTexture, Error, TEXT("Missing texture description for build of '%s' by %s."),
-				*WriteToString<128>(Output.GetName()), *WriteToString<32>(Output.GetFunction()));
+				*Output.GetName(), *WriteToString<32>(Output.GetFunction()));
 			return false;
 		}
 
 		FCbObject TextureDescription(Value.GetData().Decompress());
 
-		FCbFieldViewIterator SizeIt = TextureDescription["Size"_ASV].AsArrayView().CreateViewIterator();
+		FCbFieldViewIterator SizeIt = TextureDescription[ANSITEXTVIEW("Size")].AsArrayView().CreateViewIterator();
 		DerivedData.SizeX = SizeIt++->AsInt32();
 		DerivedData.SizeY = SizeIt++->AsInt32();
 		int32 NumSlices = SizeIt++->AsInt32();
 
 		UEnum* PixelFormatEnum = UTexture::GetPixelFormatEnum();
-		FUtf8StringView PixelFormatStringView = TextureDescription["PixelFormat"_ASV].AsString();
+		FUtf8StringView PixelFormatStringView = TextureDescription[ANSITEXTVIEW("PixelFormat")].AsString();
 		FName PixelFormatName(PixelFormatStringView.Len(), PixelFormatStringView.GetData());
 		DerivedData.PixelFormat = (EPixelFormat)PixelFormatEnum->GetValueByName(PixelFormatName);
 
-		const bool bCubeMap = TextureDescription["bCubeMap"_ASV].AsBool();
-		DerivedData.OptData.ExtData = TextureDescription["ExtData"_ASV].AsUInt32();
-		DerivedData.OptData.NumMipsInTail = TextureDescription["NumMipsInTail"_ASV].AsUInt32();
+		const bool bCubeMap = TextureDescription[ANSITEXTVIEW("bCubeMap")].AsBool();
+		DerivedData.OptData.ExtData = TextureDescription[ANSITEXTVIEW("ExtData")].AsUInt32();
+		DerivedData.OptData.NumMipsInTail = TextureDescription[ANSITEXTVIEW("NumMipsInTail")].AsUInt32();
 		const bool bHasOptData = (DerivedData.OptData.NumMipsInTail != 0) || (DerivedData.OptData.ExtData != 0);
 		static constexpr uint32 BitMask_CubeMap = 1u << 31u;
 		static constexpr uint32 BitMask_HasOptData = 1u << 30u;
 		static constexpr uint32 BitMask_NumSlices = BitMask_HasOptData - 1u;
 		DerivedData.PackedData = (NumSlices & BitMask_NumSlices) | (bCubeMap ? BitMask_CubeMap : 0) | (bHasOptData ? BitMask_HasOptData : 0);
 
-		int32 NumMips = TextureDescription["NumMips"_ASV].AsInt32();
-		int32 NumStreamingMips = TextureDescription["NumStreamingMips"_ASV].AsInt32();
+		int32 NumMips = TextureDescription[ANSITEXTVIEW("NumMips")].AsInt32();
+		int32 NumStreamingMips = TextureDescription[ANSITEXTVIEW("NumStreamingMips")].AsInt32();
 
-		FCbArrayView MipArrayView = TextureDescription["Mips"_ASV].AsArrayView();
+		FCbArrayView MipArrayView = TextureDescription[ANSITEXTVIEW("Mips")].AsArrayView();
 		if (NumMips != MipArrayView.Num())
 		{
 			UE_LOG(LogTexture, Error, TEXT("Mismatched mip quantity (%d and %d) for build of '%s' by %s."),
-				NumMips, MipArrayView.Num(), *WriteToString<128>(Output.GetName()), *WriteToString<32>(Output.GetFunction()));
+				NumMips, MipArrayView.Num(), *Output.GetName(), *WriteToString<32>(Output.GetFunction()));
 			return false;
 		}
 		check(NumMips >= (int32)DerivedData.OptData.NumMipsInTail);
@@ -1422,11 +1458,11 @@ private:
 		FSharedBuffer MipTailData;
 		if (NumMips > NumStreamingMips)
 		{
-			const FValueWithId& MipTailValue = Output.GetValue(FValueId::FromName("MipTail"_ASV));
+			const FValueWithId& MipTailValue = Output.GetValue(FValueId::FromName(ANSITEXTVIEW("MipTail")));
 			if (!MipTailValue)
 			{
 				UE_LOG(LogTexture, Error, TEXT("Missing texture mip tail for build of '%s' by %s."),
-					*WriteToString<128>(Output.GetName()), *WriteToString<32>(Output.GetFunction()));
+					*Output.GetName(), *WriteToString<32>(Output.GetFunction()));
 				return false;
 			}
 			MipTailData = MipTailValue.GetData().Decompress();
@@ -1437,60 +1473,65 @@ private:
 		for (FCbFieldView MipFieldView : MipArrayView)
 		{
 			FCbObjectView MipObjectView = MipFieldView.AsObjectView();
-			FTexture2DMipMap* NewMip = new FTexture2DMipMap();
+			TUniquePtr<FTexture2DMipMap> NewMip = MakeUnique<FTexture2DMipMap>();
 
-			FCbFieldViewIterator MipSizeIt = MipObjectView["Size"_ASV].AsArrayView().CreateViewIterator();
+			FCbFieldViewIterator MipSizeIt = MipObjectView[ANSITEXTVIEW("Size")].AsArrayView().CreateViewIterator();
 			NewMip->SizeX = MipSizeIt++->AsInt32();
 			NewMip->SizeY = MipSizeIt++->AsInt32();
 			NewMip->SizeZ = MipSizeIt++->AsInt32();
-			NewMip->FileRegionType = static_cast<EFileRegionType>(MipObjectView["FileRegion"_ASV].AsInt32());
+			NewMip->FileRegionType = static_cast<EFileRegionType>(MipObjectView[ANSITEXTVIEW("FileRegion")].AsInt32());
 			
 			if (MipIndex >= NumStreamingMips)
 			{
-				uint64 MipSize = MipObjectView["NumBytes"_ASV].AsUInt64();
-				FMemoryView MipView = MipTailData.GetView().Mid(MipObjectView["MipOffset"_ASV].AsUInt64(), MipSize);
+				uint64 MipSize = MipObjectView[ANSITEXTVIEW("NumBytes")].AsUInt64();
+				FMemoryView MipView = MipTailData.GetView().Mid(MipObjectView[ANSITEXTVIEW("MipOffset")].AsUInt64(), MipSize);
 
 				NewMip->BulkData.Lock(LOCK_READ_WRITE);
 				void* MipAllocData = NewMip->BulkData.Realloc(int64(MipSize));
 				MakeMemoryView(MipAllocData, MipSize).CopyFrom(MipView);
 				NewMip->BulkData.Unlock();
-				NewMip->SetPagedToDerivedData(false);
 			}
-			else if (bInlineMips && (MipIndex >= FirstMipToLoad))
+			else
 			{
-				const FValueWithId& StreamingMipValue = Output.GetValue(FTexturePlatformData::MakeMipId(MipIndex));
-				if (!StreamingMipValue)
+				const FValueId MipId = FTexturePlatformData::MakeMipId(MipIndex);
+				const FValueWithId& MipValue = Output.GetValue(MipId);
+				if (!MipValue)
 				{
-					delete NewMip;
-					UE_LOG(LogTexture, Error, TEXT("Missing texture streaming mip '%d' for build of '%s' by %s."),
-						MipIndex, *WriteToString<128>(Output.GetName()), *WriteToString<32>(Output.GetFunction()));
+					UE_LOG(LogTexture, Error, TEXT("Missing streaming texture mip %d for build of '%s' by %s."),
+						MipIndex, *Output.GetName(), *WriteToString<32>(Output.GetFunction()));
 					return false;
 				}
-				FSharedBuffer StreamingMipData = StreamingMipValue.GetData().Decompress();
-				uint64 MipSize = StreamingMipData.GetSize();
 
-				NewMip->BulkData.Lock(LOCK_READ_WRITE);
-				void* MipAllocData = NewMip->BulkData.Realloc(int64(MipSize));
-				MakeMemoryView(MipAllocData, MipSize).CopyFrom(StreamingMipData.GetView());
-				NewMip->BulkData.Unlock();
-				NewMip->SetPagedToDerivedData(false);
-			}
-			else
-			{
-				NewMip->SetPagedToDerivedData(true);
+				if (bInlineMips && (MipIndex >= FirstMipToLoad))
+				{
+					NewMip->BulkData.Lock(LOCK_READ_WRITE);
+					const uint64 MipSize = MipValue.GetRawSize();
+					void* MipData = NewMip->BulkData.Realloc(IntCastChecked<int64>(MipSize));
+					ON_SCOPE_EXIT { NewMip->BulkData.Unlock(); };
+					if (!MipValue.GetData().TryDecompressTo(MakeMemoryView(MipData, MipSize)))
+					{
+						UE_LOG(LogTexture, Error, TEXT("Failed to decompress streaming texture mip %d for build of '%s' by %s."),
+							MipIndex, *Output.GetName(), *WriteToString<32>(Output.GetFunction()));
+						return false;
+					}
+				}
+
+				FSharedString MipName(WriteToString<256>(Output.GetName(), TEXT(" [MIP "), MipIndex, TEXT("]")));
+				NewMip->DerivedData = FDerivedData(MoveTemp(MipName), Params.CacheKey, MipId);
 			}
 
-			DerivedData.Mips.Add(NewMip);
+			DerivedData.Mips.Add(NewMip.Release());
 			++MipIndex;
 		}
 
 		return true;
 	}
 
-	void WriteDerivedData(UE::DerivedData::FBuildOutput&& Output)
+	void WriteDerivedData(UE::DerivedData::FBuildCompleteParams&& Params)
 	{
 		using namespace UE::DerivedData;
 
+		const FBuildOutput& Output = Params.Output;
 		const FSharedString& Name = Output.GetName();
 		const FUtf8SharedString& Function = Output.GetFunction();
 
@@ -1543,7 +1584,7 @@ private:
 			return;
 		}
 
-		DeserializeTextureFromValues(DerivedData, Output, FirstMipToLoad, bInlineMips);
+		DeserializeTextureFromValues(DerivedData, MoveTemp(Params), FirstMipToLoad, bInlineMips);
 	}
 
 	static UE::DerivedData::EPriority ConvertPriority(EQueuedWorkPriority SourcePriority)
@@ -1590,7 +1631,6 @@ private:
 	bool bInlineMips;
 	int32 FirstMipToLoad;
 	uint64 BuildOutputSize = 0;
-	TOptional<FTextureStatusMessageContext> StatusMessage;
 	UE::TextureDerivedData::FTextureBuildInputResolver InputResolver;
 	FRWLock Lock;
 }; // end DDC2 fetch/build task (FTextureBuildTask)
diff --git a/Engine/Source/Runtime/Engine/Private/VT/VirtualTextureBuiltData.cpp b/Engine/Source/Runtime/Engine/Private/VT/VirtualTextureBuiltData.cpp
--- a/Engine/Source/Runtime/Engine/Private/VT/VirtualTextureBuiltData.cpp
+++ b/Engine/Source/Runtime/Engine/Private/VT/VirtualTextureBuiltData.cpp
@@ -543,26 +543,31 @@ bool FVirtualTextureDataChunk::ShortenKey(const FString& CacheKey, FString& Resu
 	return true;
 }
 
-int64 FVirtualTextureDataChunk::StoreInDerivedDataCache(const FString& InDerivedDataKey, const FStringView& TextureName, bool bReplaceExistingDDC)
+int64 FVirtualTextureDataChunk::StoreInDerivedDataCache(const FStringView InKey, const FStringView InName, const bool bInReplaceExisting)
 {
+	using namespace UE;
+	using namespace UE::DerivedData;
+
 	TRACE_CPUPROFILER_EVENT_SCOPE(FVirtualTextureDataChunk::StoreInDerivedDataCache);
 
-	int64 BulkDataSizeInBytes = BulkData.GetBulkDataSize();
+	const int64 BulkDataSizeInBytes = BulkData.GetBulkDataSize();
 	check(BulkDataSizeInBytes > 0);
 
-	using namespace UE::DerivedData;
-
-	FValue DerivedData = FValue::Compress(FSharedBuffer::MakeView(BulkData.Lock(LOCK_READ_ONLY), BulkDataSizeInBytes));
+	const FSharedString Name = InName;
+	const FCacheKey Key = ConvertLegacyCacheKey(InKey);
+	FValue Value = FValue::Compress(FSharedBuffer::MakeView(BulkData.Lock(LOCK_READ_ONLY), BulkDataSizeInBytes));
 	BulkData.Unlock();
 
 	FRequestOwner AsyncOwner(EPriority::Normal);
-	const ECachePolicy Policy = bReplaceExistingDDC ? ECachePolicy::Store : ECachePolicy::Default;
-	GetCache().PutValue({{{TextureName}, ConvertLegacyCacheKey(InDerivedDataKey), MoveTemp(DerivedData), Policy}}, AsyncOwner);
+	const ECachePolicy Policy = bInReplaceExisting ? ECachePolicy::Store : ECachePolicy::Default;
+	GetCache().PutValue({{Name, Key, MoveTemp(Value), Policy}}, AsyncOwner);
 	AsyncOwner.KeepAlive();
 
-	DerivedDataKey = InDerivedDataKey;
+	DerivedDataKey = InKey;
 	ShortenKey(DerivedDataKey, ShortDerivedDataKey);
 
+	DerivedData = FDerivedData(Name, Key);
+
 	// remove the actual bulkdata so when we serialize the owning FVirtualTextureBuiltData, this is actually serializing only the meta data
 	BulkData.RemoveBulkData();
 	return BulkDataSizeInBytes;
diff --git a/Engine/Source/Runtime/Engine/Private/VT/VirtualTextureBuiltData.h b/Engine/Source/Runtime/Engine/Private/VT/VirtualTextureBuiltData.h
--- a/Engine/Source/Runtime/Engine/Private/VT/VirtualTextureBuiltData.h
+++ b/Engine/Source/Runtime/Engine/Private/VT/VirtualTextureBuiltData.h
@@ -4,6 +4,7 @@
 #include "CoreMinimal.h"
 #include "PixelFormat.h"
 #include "Serialization/BulkData.h"
+#include "Serialization/DerivedData.h"
 #include "Engine/Texture.h"
 #include "HAL/ThreadSafeBool.h"
 
@@ -35,6 +36,9 @@ struct FVirtualTextureChunkHeader
 
 struct FVirtualTextureDataChunk
 {
+	/** Reference to the data for the chunk if it can be streamed. */
+	UE::FDerivedData DerivedData;
+	/** Stores the data for the chunk when it is loaded. */
 	FByteBulkData BulkData;
 	FSHAHash BulkDataHash;
 	uint32 SizeInBytes;
@@ -66,7 +70,7 @@ struct FVirtualTextureDataChunk
 	bool ShortenKey(const FString& CacheKey, FString& Result);
 	FThreadSafeBool bFileAvailableInVTDDCDache;
 	bool bCorruptDataLoadedFromDDC = false;
-	int64 StoreInDerivedDataCache(const FString& InDerivedDataKey, const FStringView& TextureName, bool bReplaceExistingDDC);
+	int64 StoreInDerivedDataCache(FStringView InKey, FStringView InName, bool bReplaceExisting);
 #endif // WITH_EDITORONLY_DATA
 };
 
diff --git a/Engine/Source/Runtime/Engine/Public/TextureResource.h b/Engine/Source/Runtime/Engine/Public/TextureResource.h
--- a/Engine/Source/Runtime/Engine/Public/TextureResource.h
+++ b/Engine/Source/Runtime/Engine/Public/TextureResource.h
@@ -16,6 +16,7 @@
 #include "RHI.h"
 #include "RenderResource.h"
 #include "Serialization/BulkData.h"
+#include "Serialization/DerivedData.h"
 #include "Engine/TextureDefines.h"
 #include "UnrealClient.h"
 #include "Templates/UniquePtr.h"
@@ -49,7 +50,10 @@ struct FTexture2DMipMap
 	/** Depth of the mip-map. */
 	int32 SizeZ;
 
-	/** Bulk data if stored in the package. */
+	/** Reference to the data for the mip if it can be streamed. */
+	UE::FDerivedData DerivedData;
+
+	/** Stores the data for the mip when it is loaded. */
 	FByteBulkData BulkData;
 
 	/** Default constructor. */
@@ -64,21 +68,22 @@ struct FTexture2DMipMap
 	ENGINE_API void Serialize(FArchive& Ar, UObject* Owner, int32 MipIndex);
 
 #if WITH_EDITORONLY_DATA
-	/** The file region type appropriate for this mip's pixel format. */
+	/** The file region type appropriate for the pixel format of this mip-map. */
 	EFileRegionType FileRegionType = EFileRegionType::None;
 
-	/** Whether this mip is stored in the derived data cache. */
+	UE_DEPRECATED(5.1, "Use DerivedData.HasData().")
 	bool bPagedToDerivedData = false;
 
-	bool IsPagedToDerivedData() const { return bPagedToDerivedData; }
-	void SetPagedToDerivedData(bool InValue) { bPagedToDerivedData = InValue; }
+	/** Whether this mip-map is stored in the derived data cache. */
+	inline bool IsPagedToDerivedData() const { return DerivedData.HasData(); }
 
-	/**
-	 * Place mip-map data in the derived data cache associated with the provided
-	 * key.
-	 */
-	int64 StoreInDerivedDataCache(const FString& InDerivedDataKey, const FStringView& TextureName, bool bReplaceExistingDDC);
+	UE_DEPRECATED(5.1, "Setting DerivedData is sufficient to control this state.")
+	inline void SetPagedToDerivedData(bool InValue)
+	{
+	}
 
+	/** Place mip-map data in the derived data cache associated with the provided key. */
+	int64 StoreInDerivedDataCache(FStringView Key, FStringView Name, bool bReplaceExisting);
 #endif // #if WITH_EDITORONLY_DATA
 };
 
@@ -161,7 +166,7 @@ public:
 
 	ENGINE_API virtual void InitRHI() override;
 	ENGINE_API virtual void ReleaseRHI() override;
-	
+
 	// Dynamic cast methods.
 	ENGINE_API virtual FVirtualTexture2DResource* GetVirtualTexture2DResource() { return this; }
 	// Dynamic cast methods (const).
