Removed redundant copy of reference pose if no results exist (usually for pickaxes and other mostly-static meshes).

diff --git a/Engine/Source/Runtime/Engine/Private/Components/SkeletalMeshComponent.cpp b/Engine/Source/Runtime/Engine/Private/Components/SkeletalMeshComponent.cpp
--- a/Engine/Source/Runtime/Engine/Private/Components/SkeletalMeshComponent.cpp
+++ b/Engine/Source/Runtime/Engine/Private/Components/SkeletalMeshComponent.cpp
@@ -484,10 +484,6 @@ void USkeletalMeshComponent::FinalizePoseEvaluationResult(const USkeletalMesh* I
 			OutBoneSpaceTransforms[MeshPoseIndex.GetInt()] = InFinalPose[BoneIndex];
 		}
 	}
-	else
-	{
-		OutBoneSpaceTransforms = InMesh->GetRefSkeleton().GetRefBonePose();
-	}
 
 	OutRootBoneTranslation = OutBoneSpaceTransforms[0].GetTranslation() - InMesh->GetRefSkeleton().GetRefBonePose()[0].GetTranslation();
 }
