Texture: Removed deprecated member causing deprecation warnings in generated functions

#rnx

diff --git a/Engine/Source/Runtime/Engine/Public/TextureResource.h b/Engine/Source/Runtime/Engine/Public/TextureResource.h
--- a/Engine/Source/Runtime/Engine/Public/TextureResource.h
+++ b/Engine/Source/Runtime/Engine/Public/TextureResource.h
@@ -71,17 +71,9 @@ struct FTexture2DMipMap
 	/** The file region type appropriate for the pixel format of this mip-map. */
 	EFileRegionType FileRegionType = EFileRegionType::None;
 
-	UE_DEPRECATED(5.1, "Use DerivedData.HasData().")
-	bool bPagedToDerivedData = false;
-
 	/** Whether this mip-map is stored in the derived data cache. */
 	inline bool IsPagedToDerivedData() const { return DerivedData.HasData(); }
 
-	UE_DEPRECATED(5.1, "Setting DerivedData is sufficient to control this state.")
-	inline void SetPagedToDerivedData(bool InValue)
-	{
-	}
-
 	/** Place mip-map data in the derived data cache associated with the provided key. */
 	int64 StoreInDerivedDataCache(FStringView Key, FStringView Name, bool bReplaceExisting);
 #endif // #if WITH_EDITORONLY_DATA
