Reinitialize ActorPartitionSubsystem when new WP map is created from a non WP template

diff --git a/Engine/Source/Runtime/Engine/Private/ActorPartition/ActorPartitionSubsystem.cpp b/Engine/Source/Runtime/Engine/Private/ActorPartition/ActorPartitionSubsystem.cpp
--- a/Engine/Source/Runtime/Engine/Private/ActorPartition/ActorPartitionSubsystem.cpp
+++ b/Engine/Source/Runtime/Engine/Private/ActorPartition/ActorPartitionSubsystem.cpp
@@ -5,6 +5,7 @@
 #include "WorldPartition/ActorPartition/PartitionActorDesc.h"
 #include "WorldPartition/WorldPartitionHelpers.h"
 #include "WorldPartition/WorldPartitionSubsystem.h"
+#include "WorldPartition/WorldPartition.h"
 #include "WorldPartition/DataLayer/DataLayerSubsystem.h"
 #include "WorldPartition/DataLayer/WorldDataLayers.h"
 #include "Engine/World.h"
@@ -179,8 +180,7 @@ public:
 	FActorPartitionWorldPartition(UWorld* InWorld)
 		: FBaseActorPartition(InWorld)
 	{
-		WorldPartition = InWorld->GetWorldPartition();
-		check(WorldPartition);
+		check(InWorld->GetWorldPartition());
 	}
 
 	UActorPartitionSubsystem::FCellCoord GetActorPartitionHash(const FActorPartitionGetParams& GetParams) const override 
@@ -224,6 +224,8 @@ public:
 			return true;
 		};
 
+		UWorldPartition* WorldPartition = World->GetWorldPartition();
+
 		FBox CellBounds = UActorPartitionSubsystem::FCellCoord::GetCellBounds(InCellCoord, InGridSize);
 		if (bInBoundsSearch)
 		{
@@ -318,9 +320,6 @@ public:
 			return true;
 			});
 	}
-
-private:
-	UWorldPartition* WorldPartition;
 };
 
 #endif // WITH_EDITOR
@@ -345,9 +344,15 @@ void UActorPartitionSubsystem::Initialize(FSubsystemCollectionBase& Collection)
 /** Implement this for deinitialization of instances of the system */
 void UActorPartitionSubsystem::Deinitialize()
 {
-	if (ActorPartition)
+	UninitializeActorPartition();
+}
+
+void UActorPartitionSubsystem::OnWorldPartitionInitialized(UWorldPartition* InWorldPartition)
+{
+	if (InWorldPartition->IsMainWorldPartition())
 	{
-		ActorPartition->GetOnActorPartitionHashInvalidated().Remove(ActorPartitionHashInvalidatedHandle);
+		UninitializeActorPartition();
+		InitializeActorPartition();
 	}
 }
 
@@ -371,6 +376,12 @@ void UActorPartitionSubsystem::InitializeActorPartition()
 	if (IsLevelPartition())
 	{
 		ActorPartition.Reset(new FActorPartitionLevel(GetWorld()));
+
+		// Specific use case where map is Converted to World Partition from a non World Partition template
+		if (GetWorld()->GetPackage()->HasAnyPackageFlags(PKG_NewlyCreated))
+		{
+			GetWorld()->OnWorldPartitionInitialized().AddUObject(this, &UActorPartitionSubsystem::OnWorldPartitionInitialized);
+		}
 	}
 	else
 	{
@@ -379,6 +390,18 @@ void UActorPartitionSubsystem::InitializeActorPartition()
 	ActorPartitionHashInvalidatedHandle = ActorPartition->GetOnActorPartitionHashInvalidated().AddUObject(this, &UActorPartitionSubsystem::OnActorPartitionHashInvalidated);
 }
 
+void UActorPartitionSubsystem::UninitializeActorPartition()
+{
+	PartitionedActors.Empty();
+	if (ActorPartition)
+	{
+		ActorPartition->GetOnActorPartitionHashInvalidated().Remove(ActorPartitionHashInvalidatedHandle);
+	}
+
+	ActorPartition = nullptr;
+	GetWorld()->OnWorldPartitionInitialized().RemoveAll(this);
+}
+
 APartitionActor* UActorPartitionSubsystem::GetActor(const FActorPartitionGetParams& GetParams)
 {
 	FCellCoord CellCoord = ActorPartition->GetActorPartitionHash(GetParams);
diff --git a/Engine/Source/Runtime/Engine/Public/ActorPartition/ActorPartitionSubsystem.h b/Engine/Source/Runtime/Engine/Public/ActorPartition/ActorPartitionSubsystem.h
--- a/Engine/Source/Runtime/Engine/Public/ActorPartition/ActorPartitionSubsystem.h
+++ b/Engine/Source/Runtime/Engine/Public/ActorPartition/ActorPartitionSubsystem.h
@@ -10,6 +10,7 @@
 #include "ActorPartitionSubsystem.generated.h"
 
 class FBaseActorPartition;
+class UWorldPartition;
 
 #if WITH_EDITOR
 /**
@@ -178,8 +179,10 @@ private:
 
 #if WITH_EDITOR
 	void OnActorPartitionHashInvalidated(const FCellCoord& Hash);
+	void OnWorldPartitionInitialized(UWorldPartition* InWorldPartition);
 
 	void InitializeActorPartition();
+	void UninitializeActorPartition();
 
 	TMap<FCellCoord, TMap<FActorPartitionIdentifier, TWeakObjectPtr<APartitionActor>>> PartitionedActors;
 	TUniquePtr<FBaseActorPartition> ActorPartition;
