Fix UGameplayStatics::GetPlayerController to correctly deal with cases where the passed in world is not the same as the active world on GameInstance, it now falls back to the old behavior
This fixes cases with multiple active gameplay worlds or an active world with no game instance but multiple player controllers

diff --git a/Engine/Source/Runtime/Engine/Private/GameplayStatics.cpp b/Engine/Source/Runtime/Engine/Private/GameplayStatics.cpp
--- a/Engine/Source/Runtime/Engine/Private/GameplayStatics.cpp
+++ b/Engine/Source/Runtime/Engine/Private/GameplayStatics.cpp
@@ -295,9 +295,19 @@ class APlayerController* UGameplayStatics::GetPlayerController(const UObject* Wo
 	// 99% of the time people pass in index 0 and want the primary local player controller
 	// After we've finished iterating the local player controllers, iterate the GameState list to find remote ones in a consistent order
 
+	UWorld* World = GEngine->GetWorldFromContextObject(WorldContextObject, EGetWorldErrorMode::LogAndReturnNull);
+
+	if (!World)
+	{
+		return nullptr;
+	}
+
+	// Don't use the game instance if the passed in world isn't the primary active world
+	UGameInstance* GameInstance = World->GetGameInstance();
+	const bool bUseGameInstance = GameInstance && GameInstance->GetWorld() == World;
+
 	int32 Index = 0;
-	UGameInstance* GameInstance = GetGameInstance(WorldContextObject);
-	if (GameInstance)
+	if (bUseGameInstance)
 	{		
 		const TArray<ULocalPlayer*>& LocalPlayers = GameInstance->GetLocalPlayers();
 		for (ULocalPlayer* LocalPlayer : LocalPlayers)
@@ -314,14 +324,15 @@ class APlayerController* UGameplayStatics::GetPlayerController(const UObject* Wo
 		}
 	}
 
-	AGameStateBase* GameState = GetGameState(WorldContextObject);
+	// If we have a game state, use the consistent order there to pick up remote player controllers
+	AGameStateBase* GameState = World->GetGameState();
 	if (GameState)
 	{
 		for (APlayerState* PlayerState : GameState->PlayerArray)
 		{
-			// Only count valid player controllers that we skipped over in the last list
+			// Ignore local player controllers we would have found in the previous pass
 			APlayerController* PC = PlayerState ? PlayerState->GetPlayerController() : nullptr;
-			if (PC && !PC->GetLocalPlayer())
+			if (PC && !(bUseGameInstance && PC->GetLocalPlayer()))
 			{
 				if (Index == PlayerIndex)
 				{
@@ -332,12 +343,17 @@ class APlayerController* UGameplayStatics::GetPlayerController(const UObject* Wo
 		}
 	}
 
-	// Fallback case for weird world contexts with no game instance or state, only support index 0
-	if (PlayerIndex == 0)
+	// Fallback to the old behavior with a raw iterator, but only if we didn't find any potential player controllers with the other methods
+	if (Index == 0)
 	{
-		if (UWorld* World = GEngine->GetWorldFromContextObject(WorldContextObject, EGetWorldErrorMode::LogAndReturnNull))
+		for (FConstPlayerControllerIterator Iterator = World->GetPlayerControllerIterator(); Iterator; ++Iterator)
 		{
-			return World->GetFirstPlayerController();
+			APlayerController* PlayerController = Iterator->Get();
+			if (Index == PlayerIndex)
+			{
+				return PlayerController;
+			}
+			Index++;
 		}
 	}
 
