PR #8927: Fix SetupAttachment() false positive ensure when called on a component that already had an AttachParent (Contributed by [REDACTED_PERSON])

diff --git a/Engine/Source/Runtime/Engine/Private/Components/SceneComponent.cpp b/Engine/Source/Runtime/Engine/Private/Components/SceneComponent.cpp
--- a/Engine/Source/Runtime/Engine/Private/Components/SceneComponent.cpp
+++ b/Engine/Source/Runtime/Engine/Private/Components/SceneComponent.cpp
@@ -1792,17 +1792,20 @@ void USceneComponent::SetRelativeRotationCache(const FRotationConversionCache& I
 
 void USceneComponent::SetupAttachment(class USceneComponent* InParent, FName InSocketName)
 {
-	if (ensureMsgf(!bRegistered, TEXT("SetupAttachment should only be used to initialize AttachParent and AttachSocketName for a future AttachToComponent. Once a component is registered you must use AttachToComponent. Owner [%s], InParent [%s], InSocketName [%s]"), *GetPathNameSafe(GetOwner()), *GetNameSafe(InParent), *InSocketName.ToString()))
+	if (InParent != AttachParent || InSocketName != AttachSocketName)
 	{
-		if (ensureMsgf(InParent != this, TEXT("Cannot attach a component to itself.")))
+		if (ensureMsgf(!bRegistered, TEXT("SetupAttachment should only be used to initialize AttachParent and AttachSocketName for a future AttachToComponent. Once a component is registered you must use AttachToComponent. Owner [%s], InParent [%s], InSocketName [%s]"), *GetPathNameSafe(GetOwner()), *GetNameSafe(InParent), *InSocketName.ToString()))
 		{
-			if (ensureMsgf(InParent == nullptr || !InParent->IsAttachedTo(this), TEXT("Setting up attachment would create a cycle.")))
+			if (ensureMsgf(InParent != this, TEXT("Cannot attach a component to itself.")))
 			{
-				if (ensureMsgf(AttachParent == nullptr || !AttachParent->AttachChildren.Contains(this), TEXT("SetupAttachment cannot be used once a component has already had AttachTo used to connect it to a parent.")))
+				if (ensureMsgf(InParent == nullptr || !InParent->IsAttachedTo(this), TEXT("Setting up attachment would create a cycle.")))
 				{
-					SetAttachParent(InParent);
-					SetAttachSocketName(InSocketName);
-					SetShouldBeAttached(AttachParent != nullptr);
+					if (ensureMsgf(AttachParent == nullptr || !AttachParent->AttachChildren.Contains(this), TEXT("SetupAttachment cannot be used once a component has already had AttachTo used to connect it to a parent.")))
+					{
+						SetAttachParent(InParent);
+						SetAttachSocketName(InSocketName);
+						SetShouldBeAttached(AttachParent != nullptr);
+					}
 				}
 			}
 		}
