Move some PostEditChange steps to the async build process to avoid synchronous stalls as soon as the build is launched

#rnx

diff --git a/Engine/Source/Runtime/Engine/Classes/Engine/StaticMesh.h b/Engine/Source/Runtime/Engine/Classes/Engine/StaticMesh.h
--- a/Engine/Source/Runtime/Engine/Classes/Engine/StaticMesh.h
+++ b/Engine/Source/Runtime/Engine/Classes/Engine/StaticMesh.h
@@ -1594,12 +1594,41 @@ public:
 	*/
 	ENGINE_API static void CancelAllPendingStreamingActions();
 
+	/**
+	 * Contains all the parameters required to build the mesh.
+	 */
+	struct FBuildParameters
+	{
+		// Required to work around a Clang bug
+#if PLATFORM_COMPILER_CLANG
+		FBuildParameters()
+			: bInSilent(false)
+			, OutErrors(nullptr)
+			, bInRebuildUVChannelData(false)
+			, bInEnforceLightmapRestrictions(false)
+		{
+		}
+#endif
+
+		// If true will not popup a progress dialog.
+		bool bInSilent = false;
+
+		// If provided, will contain the errors that occurred during this process.This will prevent async static mesh compilation because OutErrors could get out of scope.
+		TArray<FText>* OutErrors = nullptr;
+
+		// Whether to completely rebuild the UV Channel data after the render data has been computed.
+		bool bInRebuildUVChannelData = false;
+
+		// Whether to call EnforceLightmapRestrictions as part of the build process.
+		bool bInEnforceLightmapRestrictions = false;
+	};
+
 	/**
 	 * Rebuilds renderable data for this static mesh, automatically made async if enabled.
 	 * @param		bInSilent	If true will not popup a progress dialog.
 	 * @param [out]	OutErrors	If provided, will contain the errors that occurred during this process. This will prevent async static mesh compilation because OutErrors could get out of scope.
 	 */
-	ENGINE_API void Build(bool bInSilent = false, TArray<FText>* OutErrors = nullptr);
+	ENGINE_API void Build(bool bInSilent, TArray<FText>* OutErrors = nullptr);
 
 	/**
 	 * Rebuilds renderable data for a batch of static meshes.
@@ -1608,7 +1637,21 @@ public:
 	 * @param		InProgressCallback	If provided, will be used to abort task and report progress to higher level functions (should return true to continue, false to abort).
 	 * @param [out]	OutErrors			If provided, will contain the errors that occurred during this process. This will prevent async static mesh compilation because OutErrors could get out of scope.
 	 */
-	ENGINE_API static void BatchBuild(const TArray<UStaticMesh*>& InStaticMeshes, bool bInSilent = false, TFunction<bool(UStaticMesh*)> InProgressCallback = nullptr, TArray<FText>* OutErrors = nullptr);
+	ENGINE_API static void BatchBuild(const TArray<UStaticMesh*>& InStaticMeshes, bool bInSilent, TFunction<bool(UStaticMesh*)> InProgressCallback = nullptr, TArray<FText>* OutErrors = nullptr);
+
+	/**
+	 * Rebuilds renderable data for this static mesh, automatically made async if enabled.
+	 * @param		BuildParameters	  Contains all the information required to build the mesh.
+	 */
+	ENGINE_API void Build(const FBuildParameters& BuildParameters = FBuildParameters());
+
+	/**
+	 * Rebuilds renderable data for a batch of static meshes.
+	 * @param		InStaticMeshes		The list of all static meshes to build.
+	 * @param		BuildParameters	    Contains all the parameters required to build the mesh.
+	 * @param		InProgressCallback	If provided, will be used to abort task and report progress to higher level functions (should return true to continue, false to abort).
+	 */
+	ENGINE_API static void BatchBuild(const TArray<UStaticMesh*>& InStaticMeshes, const FBuildParameters& BuildParameters = FBuildParameters(), TFunction<bool(UStaticMesh*)> InProgressCallback = nullptr);
 
 	/**
 	 * Initialize the static mesh's render resources.
@@ -1914,7 +1957,7 @@ private:
 	/**
 	 * Build the static mesh
 	 */
-	bool ExecuteBuildInternal(bool bSilent, TArray<FText>* OutErrors);
+	bool ExecuteBuildInternal(const FBuildParameters& BuildParameters);
 
 	/**
 	 * Complete the static mesh building process - Can't be done in parallel.
@@ -2003,5 +2046,11 @@ public:
 class FStaticMeshBuildContext : public FStaticMeshCompilationContext
 {
 public:
+	FStaticMeshBuildContext(const UStaticMesh::FBuildParameters& InBuildParameters)
+		: BuildParameters(InBuildParameters)
+	{
+	}
+
+	UStaticMesh::FBuildParameters BuildParameters;
 	bool bHasRenderDataChanged = false;
 };
\ No newline at end of file
diff --git a/Engine/Source/Runtime/Engine/Private/StaticMesh.cpp b/Engine/Source/Runtime/Engine/Private/StaticMesh.cpp
--- a/Engine/Source/Runtime/Engine/Private/StaticMesh.cpp
+++ b/Engine/Source/Runtime/Engine/Private/StaticMesh.cpp
@@ -207,7 +207,7 @@ void FStaticMeshAsyncBuildWorker::DoWork()
 
 	if (BuildContext.IsValid())
 	{
-		BuildContext->bHasRenderDataChanged = StaticMesh->ExecuteBuildInternal(true, nullptr);
+		BuildContext->bHasRenderDataChanged = StaticMesh->ExecuteBuildInternal(BuildContext->BuildParameters);
 	}
 }
 
@@ -3700,8 +3700,8 @@ void UStaticMesh::PostEditChangeProperty(FPropertyChangedEvent& PropertyChangedE
 	{
 		// Force an update of LOD group settings
 
-		// Dont rebuild inside here.  We're doing that below.
-		bool bRebuild = false;
+		// Don't rebuild inside here.  We're doing that below.
+		constexpr bool bRebuild = false;
 		SetLODGroup(LODGroup, bRebuild);
 	}
 #if WITH_EDITORONLY_DATA
@@ -3745,32 +3745,30 @@ void UStaticMesh::PostEditChangeProperty(FPropertyChangedEvent& PropertyChangedE
 		}
 	}
 
-	//Don't use the render data here because the property that just changed might be invalidating the current RenderData.
-	EnforceLightmapRestrictions(/*bUseRenderData=*/false);
-
-	Build(/*bSilent=*/ true);
-
-	if (PropertyName == GET_MEMBER_NAME_CHECKED(UStaticMesh, bHasNavigationData)
-		|| PropertyName == UStaticMesh::GetBodySetupName())
-	{
-		// Build called above will result in creation, update or destruction 
-		// of NavCollision. We need to let related StaticMeshComponents know
-		BroadcastNavCollisionChange();
-	}
-
 	if (PropertyChangedEvent.ChangeType != EPropertyChangeType::Interactive)
 	{
 		// Only unbuild lighting for properties which affect static lighting
 		if (PropertyName == UStaticMesh::GetLightMapResolutionName()
 			|| PropertyName == UStaticMesh::GetLightMapCoordinateIndexName())
 		{
-			FStaticMeshComponentRecreateRenderStateContext Context(this, true);
 			SetLightingGuid();
 		}
 	}
-	
-	UpdateUVChannelData(true);
 
+	UStaticMesh::FBuildParameters BuildParameters;
+	BuildParameters.bInSilent = true;
+	BuildParameters.bInRebuildUVChannelData = true;
+	BuildParameters.bInEnforceLightmapRestrictions = true;
+	Build(BuildParameters);
+
+	if (PropertyName == GET_MEMBER_NAME_CHECKED(UStaticMesh, bHasNavigationData)
+		|| PropertyName == UStaticMesh::GetBodySetupName())
+	{
+		// Build called above will result in creation, update or destruction 
+		// of NavCollision. We need to let related StaticMeshComponents know
+		BroadcastNavCollisionChange();
+	}
+	
 	for (UAssetUserData* Datum : AssetUserData)
 	{
 		if (Datum != nullptr)
diff --git a/Engine/Source/Runtime/Engine/Private/StaticMeshBuild.cpp b/Engine/Source/Runtime/Engine/Private/StaticMeshBuild.cpp
--- a/Engine/Source/Runtime/Engine/Private/StaticMeshBuild.cpp
+++ b/Engine/Source/Runtime/Engine/Private/StaticMeshBuild.cpp
@@ -113,24 +113,24 @@ static TAutoConsoleVariable<int32> CVarStaticMeshDisableThreadedBuild(
 
 #endif // #if WITH_EDITOR
 
-void UStaticMesh::Build(bool bInSilent, TArray<FText>* OutErrors)
+void UStaticMesh::Build(const FBuildParameters& BuildParameters)
 {
 #if WITH_EDITOR
 	FFormatNamedArguments Args;
-	Args.Add( TEXT("Path"), FText::FromString( GetPathName() ) );
-	const FText StatusUpdate = FText::Format( LOCTEXT("BeginStaticMeshBuildingTask", "({Path}) Building"), Args );
+	Args.Add(TEXT("Path"), FText::FromString(GetPathName()));
+	const FText StatusUpdate = FText::Format(LOCTEXT("BeginStaticMeshBuildingTask", "({Path}) Building"), Args);
 	FScopedSlowTask StaticMeshBuildingSlowTask(1, StatusUpdate);
-	if (!bInSilent)
+	if (!BuildParameters.bInSilent)
 	{
 		StaticMeshBuildingSlowTask.MakeDialogDelayed(1.0f);
 	}
 	StaticMeshBuildingSlowTask.EnterProgressFrame(1);
 #endif // #if WITH_EDITOR
 
-	BatchBuild({ this }, bInSilent, nullptr, OutErrors);
+	BatchBuild({ this }, BuildParameters, nullptr);
 }
 
-void UStaticMesh::BatchBuild(const TArray<UStaticMesh*>& InStaticMeshes, bool bInSilent, TFunction<bool(UStaticMesh*)> InProgressCallback, TArray<FText>* OutErrors)
+void UStaticMesh::BatchBuild(const TArray<UStaticMesh*>& InStaticMeshes, const FBuildParameters& BuildParameters, TFunction<bool(UStaticMesh*)> InProgressCallback)
 {
 #if WITH_EDITOR
 	check(IsInGameThread());
@@ -181,7 +181,7 @@ void UStaticMesh::BatchBuild(const TArray<UStaticMesh*>& InStaticMeshes, bool bI
 		{
 			GCardRepresentationAsyncQueue->CancelBuilds(StaticMeshesToProcess);
 		}
-		
+
 		TMap<UStaticMesh*, TArray<UStaticMeshComponent*>> StaticMeshComponents;
 		StaticMeshComponents.Reserve(InStaticMeshes.Num());
 
@@ -233,12 +233,12 @@ void UStaticMesh::BatchBuild(const TArray<UStaticMesh*>& InStaticMeshes, bool bI
 			};
 
 		auto LaunchAsyncBuild =
-			[OutErrors](UStaticMesh* StaticMesh)
+			[BuildParameters](UStaticMesh* StaticMesh)
 			{
 				// Only launch async compile if errors are not required
-				if (OutErrors == nullptr && FStaticMeshCompilingManager::Get().IsAsyncCompilationAllowed(StaticMesh))
+				if (BuildParameters.OutErrors == nullptr && FStaticMeshCompilingManager::Get().IsAsyncCompilationAllowed(StaticMesh))
 				{
-					TUniquePtr<FStaticMeshBuildContext> Context = MakeUnique<FStaticMeshBuildContext>();
+					TUniquePtr<FStaticMeshBuildContext> Context = MakeUnique<FStaticMeshBuildContext>(BuildParameters);
 					StaticMesh->BeginBuildInternal(Context.Get());
 
 					FQueuedThreadPool* StaticMeshThreadPool = FStaticMeshCompilingManager::Get().GetThreadPool();
@@ -286,7 +286,7 @@ void UStaticMesh::BatchBuild(const TArray<UStaticMesh*>& InStaticMeshes, bool bI
 						StaticMesh,
 						Async(
 							EAsyncExecution::LargeThreadPool,
-							[StaticMesh, bInSilent, OutErrors, &OutErrorsLock, &bCancelled]()
+							[StaticMesh, BuildParameters, &OutErrorsLock, &bCancelled]()
 							{
 								if (bCancelled.load(std::memory_order_relaxed))
 								{
@@ -294,11 +294,11 @@ void UStaticMesh::BatchBuild(const TArray<UStaticMesh*>& InStaticMeshes, bool bI
 								}
 
 								TArray<FText> Errors;
-								const bool bHasRenderDataChanged = StaticMesh->ExecuteBuildInternal(bInSilent, &Errors);
-								if (OutErrors)
+								const bool bHasRenderDataChanged = StaticMesh->ExecuteBuildInternal(BuildParameters);
+								if (BuildParameters.OutErrors)
 								{
 									FScopeLock ScopeLock(&OutErrorsLock);
-									OutErrors->Append(Errors);
+									BuildParameters.OutErrors->Append(Errors);
 								}
 
 								return bHasRenderDataChanged;
@@ -334,10 +334,10 @@ void UStaticMesh::BatchBuild(const TArray<UStaticMesh*>& InStaticMeshes, bool bI
 					{
 						break;
 					}
-				
+
 					StaticMesh->BeginBuildInternal();
 
-					const bool bHasRenderDataChanged = StaticMesh->ExecuteBuildInternal(bInSilent, OutErrors);
+					const bool bHasRenderDataChanged = StaticMesh->ExecuteBuildInternal(BuildParameters);
 
 					StaticMesh->FinishBuildInternal(StaticMeshComponents.FindChecked(StaticMesh), bHasRenderDataChanged);
 					FinalizeStaticMesh(StaticMesh);
@@ -355,6 +355,22 @@ void UStaticMesh::BatchBuild(const TArray<UStaticMesh*>& InStaticMeshes, bool bI
 #endif
 }
 
+void UStaticMesh::Build(bool bInSilent, TArray<FText>* OutErrors)
+{
+	FBuildParameters BuildParameters;
+	BuildParameters.bInSilent = bInSilent;
+	BuildParameters.OutErrors = OutErrors;
+	Build(BuildParameters);
+}
+
+void UStaticMesh::BatchBuild(const TArray<UStaticMesh*>& InStaticMeshes, bool bInSilent, TFunction<bool(UStaticMesh*)> InProgressCallback, TArray<FText>* OutErrors)
+{
+	FBuildParameters BuildParameters;
+	BuildParameters.bInSilent = bInSilent;
+	BuildParameters.OutErrors = OutErrors;
+	BatchBuild(InStaticMeshes, BuildParameters, InProgressCallback);
+}
+
 #if WITH_EDITOR
 
 extern const FString& GetStaticMeshDerivedDataVersion();
@@ -414,7 +430,7 @@ void UStaticMesh::BeginBuildInternal(FStaticMeshBuildContext* Context)
 	}
 }
 
-bool UStaticMesh::ExecuteBuildInternal(bool bInSilent, TArray<FText>* OutErrors)
+bool UStaticMesh::ExecuteBuildInternal(const FBuildParameters& BuildParameters)
 {
 	TRACE_CPUPROFILER_EVENT_SCOPE(UStaticMesh::ExecuteBuildInternal);
 
@@ -438,6 +454,16 @@ bool UStaticMesh::ExecuteBuildInternal(bool bInSilent, TArray<FText>* OutErrors)
 	CacheDerivedData();
 	PrepareDerivedDataForActiveTargetPlatforms();
 
+	if (BuildParameters.bInEnforceLightmapRestrictions)
+	{
+		EnforceLightmapRestrictions();
+	}
+
+	if (BuildParameters.bInRebuildUVChannelData)
+	{
+		UpdateUVChannelData(true);
+	}
+
 	// InitResources will send commands to other threads that will
 	// use our RenderData, we must mark it as ready to be used since
 	// we're not going to modify it anymore
@@ -485,9 +511,9 @@ bool UStaticMesh::ExecuteBuildInternal(bool bInSilent, TArray<FText>* OutErrors)
 					UE_LOG(LogStaticMesh, Warning, TEXT("%s"), *WarningMsg.ToString());
 				}
 
-				if (!bInSilent && OutErrors)
+				if (!BuildParameters.bInSilent && BuildParameters.OutErrors)
 				{
-					OutErrors->Add(WarningMsg);
+					BuildParameters.OutErrors->Add(WarningMsg);
 				}
 			}
 		}
@@ -508,9 +534,9 @@ bool UStaticMesh::ExecuteBuildInternal(bool bInSilent, TArray<FText>* OutErrors)
 			{
 				UE_LOG(LogStaticMesh, Warning, TEXT("%s"), *WarningMsg.ToString());
 			}
-			if (!bInSilent && OutErrors)
+			if (!BuildParameters.bInSilent && BuildParameters.OutErrors)
 			{
-				OutErrors->Add(WarningMsg);
+				BuildParameters.OutErrors->Add(WarningMsg);
 			}
 		}
 
@@ -530,9 +556,9 @@ bool UStaticMesh::ExecuteBuildInternal(bool bInSilent, TArray<FText>* OutErrors)
 				UE_LOG(LogStaticMesh, Warning, TEXT("%s"), *WarningMsg.ToString());
 			}
 
-			if (!bInSilent && OutErrors)
+			if (!BuildParameters.bInSilent && BuildParameters.OutErrors)
 			{
-				OutErrors->Add(WarningMsg);
+				BuildParameters.OutErrors->Add(WarningMsg);
 			}
 		}
 
@@ -552,9 +578,9 @@ bool UStaticMesh::ExecuteBuildInternal(bool bInSilent, TArray<FText>* OutErrors)
 				UE_LOG(LogStaticMesh, Warning, TEXT("%s"), *WarningMsg.ToString());
 			}
 
-			if (!bInSilent && OutErrors)
+			if (!BuildParameters.bInSilent && BuildParameters.OutErrors)
 			{
-				OutErrors->Add(WarningMsg);
+				BuildParameters.OutErrors->Add(WarningMsg);
 			}
 		}
 
@@ -565,7 +591,7 @@ bool UStaticMesh::ExecuteBuildInternal(bool bInSilent, TArray<FText>* OutErrors)
 	return bHasRenderDataChanged;
 }
 
-void UStaticMesh::FinishBuildInternal(const TArray<UStaticMeshComponent*> & InAffectedComponents, bool bHasRenderDataChanged, bool bShouldComputeExtendedBounds)
+void UStaticMesh::FinishBuildInternal(const TArray<UStaticMeshComponent*>& InAffectedComponents, bool bHasRenderDataChanged, bool bShouldComputeExtendedBounds)
 {
 	TRACE_CPUPROFILER_EVENT_SCOPE(UStaticMesh::PostBuildInternal);
 
