Fixing non-unity build error

diff --git a/Engine/Source/Runtime/Engine/Private/Level.cpp b/Engine/Source/Runtime/Engine/Private/Level.cpp
--- a/Engine/Source/Runtime/Engine/Private/Level.cpp
+++ b/Engine/Source/Runtime/Engine/Private/Level.cpp
@@ -91,6 +91,7 @@ Level.cpp: Level-related functions
 #include "HAL/LowLevelMemTracker.h"
 #include "ObjectTrace.h"
 #include "ProfilingDebugging/TagTrace.h"
+#include "UObject/MetaData.h"
 
 #define LOCTEXT_NAMESPACE "ULevel"
 DEFINE_LOG_CATEGORY(LogLevel);
