Expose UStaticMesh::GetNumTriangles() as Blueprint function

diff --git a/Engine/Source/Runtime/Engine/Classes/Engine/StaticMesh.h b/Engine/Source/Runtime/Engine/Classes/Engine/StaticMesh.h
--- a/Engine/Source/Runtime/Engine/Classes/Engine/StaticMesh.h
+++ b/Engine/Source/Runtime/Engine/Classes/Engine/StaticMesh.h
@@ -1695,8 +1695,9 @@ public:
 	ENGINE_API int32 GetNumVertices(int32 LODIndex) const;
 
 	/**
-	 * Returns the number of triangles for the specified LOD.
+	 * Returns the number of triangles in the render data for the specified LOD.
 	 */
+	UFUNCTION(BlueprintPure, Category = StaticMesh)
 	ENGINE_API int32 GetNumTriangles(int32 LODIndex) const;
 
 	/**
