Add missing PersistentLevel->CleanupLevel from bad merge

#rnx

diff --git a/Engine/Source/Runtime/Engine/Private/World.cpp b/Engine/Source/Runtime/Engine/Private/World.cpp
--- a/Engine/Source/Runtime/Engine/Private/World.cpp
+++ b/Engine/Source/Runtime/Engine/Private/World.cpp
@@ -5072,6 +5072,12 @@ void UWorld::CleanupWorldInternal(bool bSessionEnded, bool bCleanupResources, bo
 			}
 		}
 	}
+#else
+	if (PersistentLevel)
+	{
+		check(bCleanupResources);
+		PersistentLevel->CleanupLevel();
+	}
 #endif //WITH_EDITOR
 
 	for (int32 LevelIndex = 0; LevelIndex < GetNumLevels(); ++LevelIndex)
