Fixed copy-pasted comment for FEdGraphSchemaAction::IsA()

diff --git a/Engine/Source/Runtime/Engine/Classes/EdGraph/EdGraphSchema.h b/Engine/Source/Runtime/Engine/Classes/EdGraph/EdGraphSchema.h
--- a/Engine/Source/Runtime/Engine/Classes/EdGraph/EdGraphSchema.h
+++ b/Engine/Source/Runtime/Engine/Classes/EdGraph/EdGraphSchema.h
@@ -292,7 +292,7 @@ public:
 	// (e.g., both are variables in the same Blueprint)
 	virtual FEdGraphSchemaActionDefiningObject GetPersistentItemDefiningObject() const { return FEdGraphSchemaActionDefiningObject(nullptr); }
 
-	// Returns true if the action refers to a external or local variable
+	// Returns true if the action is of the given type.
 	virtual bool IsA(const FName& InType) const
 	{
 		return InType == GetTypeId();
