Various cleanup/fixes for new HLSL translator
- Add documentation
- Add new 'DefaultValue' expression, support existing material behavior where properties will use default values if connected expressions fail
- Misc fixes, mostly related to component access

diff --git a/Engine/Source/Runtime/Engine/Private/HLSLTree/HLSLTree.cpp b/Engine/Source/Runtime/Engine/Private/HLSLTree/HLSLTree.cpp
--- a/Engine/Source/Runtime/Engine/Private/HLSLTree/HLSLTree.cpp
+++ b/Engine/Source/Runtime/Engine/Private/HLSLTree/HLSLTree.cpp
@@ -558,6 +558,7 @@ void FRequestedType::SetComponentRequest(int32 Index, bool bRequested)
 {
 	if (!Type.IsAny())
 	{
+		check(FMath::IsWithin(Index, 0, Type.GetNumComponents()));
 		if (bRequested)
 		{
 			RequestedComponents.PadToNum(Index + 1, false);
@@ -699,8 +700,9 @@ bool FPreparedType::IsEmpty() const
 	return GetNumPreparedComponents() == 0;
 }
 
-Shader::FType FPreparedType::GetPreparedType() const
+Shader::FType FPreparedType::GetResultType() const
 {
+	// Only numeric vectors will adjust the number of components
 	if (!Type.IsNumericVector())
 	{
 		return Type;
@@ -723,7 +725,7 @@ Shader::EValueComponentType FPreparedType::GetValueComponentType() const
 
 FRequestedType FPreparedType::GetRequestedType() const
 {
-	const int32 NumComponents = GetNumPreparedComponents();
+	const int32 NumComponents = PreparedComponents.Num();
 	FRequestedType Result;
 	Result.Type = Type;
 	Result.RequestedComponents.Init(false, NumComponents);
@@ -799,13 +801,14 @@ FPreparedComponent FPreparedType::GetComponent(int32 Index) const
 {
 	if (Index >= 0)
 	{
+		// Scalar replicated into xyzw
 		const int32 ComponentIndex = (Type.IsNumericScalar() && Index < 4) ? 0 : Index;
-		if (PreparedComponents.IsValidIndex(ComponentIndex))
+		if (ComponentIndex < PreparedComponents.Num())
 		{
 			return PreparedComponents[ComponentIndex];
 		}
 
-		if (Index < Type.GetNumComponents())
+		if (ComponentIndex < Type.GetNumComponents())
 		{
 			return FPreparedComponent(EExpressionEvaluation::ConstantZero);
 		}
@@ -825,6 +828,7 @@ void FPreparedType::EnsureNumComponents(int32 NumComponents)
 
 void FPreparedType::SetComponent(int32 Index, const FPreparedComponent& InComponent)
 {
+	check(FMath::IsWithin(Index, 0, Type.GetNumComponents()));
 	if (InComponent.Evaluation != EExpressionEvaluation::None)
 	{
 		EnsureNumComponents(Index + 1);
@@ -837,6 +841,7 @@ void FPreparedType::SetComponent(int32 Index, const FPreparedComponent& InCompon
 
 void FPreparedType::SetComponentBounds(int32 Index, const Shader::FComponentBounds Bounds)
 {
+	check(FMath::IsWithin(Index, 0, Type.GetNumComponents()));
 	if (PreparedComponents.IsValidIndex(Index))
 	{
 		FPreparedComponent& Component = PreparedComponents[Index];
@@ -849,6 +854,7 @@ void FPreparedType::SetComponentBounds(int32 Index, const Shader::FComponentBoun
 
 void FPreparedType::MergeComponent(int32 Index, const FPreparedComponent& InComponent)
 {
+	check(FMath::IsWithin(Index, 0, Type.GetNumComponents()));
 	if (InComponent.Evaluation != EExpressionEvaluation::None)
 	{
 		EnsureNumComponents(Index + 1);
@@ -896,8 +902,7 @@ void FPreparedType::SetLoopEvaluation(FEmitScope& Scope, const FRequestedType& R
 	}
 	else
 	{
-		const int32 NumComponents = GetNumPreparedComponents();
-		for (int32 Index = 0; Index < NumComponents; ++Index)
+		for (int32 Index = 0; Index < PreparedComponents.Num(); ++Index)
 		{
 			if (RequestedType.IsComponentRequested(Index))
 			{
@@ -919,7 +924,7 @@ Shader::FComponentBounds FPreparedType::GetComponentBounds(int32 Index) const
 		return Shader::FComponentBounds(Shader::EComponentBound::Zero, Shader::EComponentBound::Zero);
 	}
 
-	const Shader::EValueComponentType ComponentType = GetPreparedType().GetComponentType(Index);
+	const Shader::EValueComponentType ComponentType = Type.GetComponentType(Index);
 	check(ComponentType != Shader::EValueComponentType::Void);
 
 	const Shader::FValueComponentTypeDescription ComponentTypeDesc = Shader::GetValueComponentTypeDescription(ComponentType);
@@ -930,8 +935,6 @@ Shader::FComponentBounds FPreparedType::GetComponentBounds(int32 Index) const
 
 Shader::FComponentBounds FPreparedType::GetBounds(const FRequestedType& RequestedType) const
 {
-	const int32 NumComponents = GetNumPreparedComponents();
-
 	Shader::FComponentBounds Result(Shader::EComponentBound::DoubleMax, Shader::EComponentBound::NegDoubleMax);
 	if (IsNumericScalar())
 	{
@@ -942,7 +945,7 @@ Shader::FComponentBounds FPreparedType::GetBounds(const FRequestedType& Requeste
 	}
 	else
 	{
-		for (int32 Index = 0; Index < NumComponents; ++Index)
+		for (int32 Index = 0; Index < PreparedComponents.Num(); ++Index)
 		{
 			if (RequestedType.IsComponentRequested(Index))
 			{
@@ -1073,11 +1076,6 @@ bool FPrepareValueResult::SetType(FEmitContext& Context, const FRequestedType& R
 					{
 						PreparedType.MergeComponent(Index, Evaluation);
 					}
-					else
-					{
-						// If component wasn't requested, it can be replaced with constant-0
-						PreparedType.MergeComponent(Index, EExpressionEvaluation::ConstantZero);
-					}
 				}
 			}
 		}
@@ -1112,15 +1110,9 @@ bool FPrepareValueResult::SetType(FEmitContext& Context, const FRequestedType& R
 			const int32 NumComponents = Type.GetNumComponents();
 			for (int32 Index = 0; Index < NumComponents; ++Index)
 			{
-				const FPreparedComponent& Component = InPreparedType.GetComponent(Index);
 				if (RequestedType.IsComponentRequested(Index))
 				{
-					PreparedType.MergeComponent(Index, Component);
-				}
-				else if (!Component.IsNone())
-				{
-					// If component wasn't requested, it can be replaced with constant-0
-					PreparedType.MergeComponent(Index, EExpressionEvaluation::ConstantZero);
+					PreparedType.MergeComponent(Index, InPreparedType.GetComponent(Index));
 				}
 			}
 		}
@@ -1204,7 +1196,7 @@ FEmitShaderExpression* FExpression::GetValueShader(FEmitContext& Context, FEmitS
 FEmitShaderExpression* FExpression::GetValueShader(FEmitContext& Context, FEmitScope& Scope, const FRequestedType& RequestedType) const
 {
 	const FPreparedType& PreparedType = Context.GetPreparedType(this);
-	return GetValueShader(Context, Scope, RequestedType, PreparedType, PreparedType.GetPreparedType());
+	return GetValueShader(Context, Scope, RequestedType, PreparedType, PreparedType.GetResultType());
 }
 
 FEmitShaderExpression* FExpression::GetValueShader(FEmitContext& Context, FEmitScope& Scope, const Shader::FType& ResultType) const
@@ -1220,7 +1212,7 @@ FEmitShaderExpression* FExpression::GetValueShader(FEmitContext& Context, FEmitS
 FEmitShaderExpression* FExpression::GetValueShader(FEmitContext& Context, FEmitScope& Scope) const
 {
 	const FPreparedType& PreparedType = Context.GetPreparedType(this);
-	return GetValueShader(Context, Scope, PreparedType.GetRequestedType(), PreparedType, PreparedType.GetPreparedType());
+	return GetValueShader(Context, Scope, PreparedType.GetRequestedType(), PreparedType, PreparedType.GetResultType());
 }
 
 Shader::FType FExpression::GetValuePreshader(FEmitContext& Context, FEmitScope& Scope, const FRequestedType& RequestedType, const FPreparedType& PreparedType, const Shader::FType& ResultType, Shader::FPreshaderData& OutPreshader) const
@@ -1264,7 +1256,7 @@ Shader::FType FExpression::GetValuePreshader(FEmitContext& Context, FEmitScope&
 Shader::FType FExpression::GetValuePreshader(FEmitContext& Context, FEmitScope& Scope, const FRequestedType& RequestedType, Shader::FPreshaderData& OutPreshader) const
 {
 	const FPreparedType& PreparedType = Context.GetPreparedType(this);
-	return GetValuePreshader(Context, Scope, RequestedType, PreparedType, PreparedType.GetPreparedType(), OutPreshader);
+	return GetValuePreshader(Context, Scope, RequestedType, PreparedType, PreparedType.GetResultType(), OutPreshader);
 }
 
 Shader::FType FExpression::GetValuePreshader(FEmitContext& Context, FEmitScope& Scope, const Shader::FType& ResultType, Shader::FPreshaderData& OutPreshader) const
@@ -1308,7 +1300,7 @@ Shader::FValue FExpression::GetValueConstant(FEmitContext& Context, FEmitScope&
 
 Shader::FValue FExpression::GetValueConstant(FEmitContext& Context, FEmitScope& Scope, const FRequestedType& RequestedType, const FPreparedType& PreparedType) const
 {
-	return GetValueConstant(Context, Scope, RequestedType, PreparedType, PreparedType.GetPreparedType());
+	return GetValueConstant(Context, Scope, RequestedType, PreparedType, PreparedType.GetResultType());
 }
 
 Shader::FValue FExpression::GetValueConstant(FEmitContext& Context, FEmitScope& Scope, const FRequestedType& RequestedType, const Shader::FType& ResultType) const
@@ -1320,7 +1312,7 @@ Shader::FValue FExpression::GetValueConstant(FEmitContext& Context, FEmitScope&
 Shader::FValue FExpression::GetValueConstant(FEmitContext& Context, FEmitScope& Scope, const FRequestedType& RequestedType) const
 {
 	const FPreparedType PreparedType = Context.GetPreparedType(this);
-	return GetValueConstant(Context, Scope, RequestedType, PreparedType, PreparedType.GetPreparedType());
+	return GetValueConstant(Context, Scope, RequestedType, PreparedType, PreparedType.GetResultType());
 }
 
 Shader::FValue FExpression::GetValueConstant(FEmitContext& Context, FEmitScope& Scope, const FPreparedType& PreparedType, const Shader::FType& ResultType) const
@@ -1536,6 +1528,7 @@ void FTree::RegisterStatement(FScope& Scope, FStatement* Statement)
 
 void FTree::AssignLocal(FScope& Scope, const FName& LocalName, const FExpression* Value)
 {
+	check(Value);
 	Scope.LocalMap.Add(LocalName, Value);
 }
 
@@ -1544,6 +1537,7 @@ const FExpression* FTree::AcquireLocal(FScope& Scope, const FName& LocalName)
 	FExpression const* const* FoundExpression = Scope.LocalMap.Find(LocalName);
 	if (FoundExpression)
 	{
+		check(*FoundExpression);
 		return *FoundExpression;
 	}
 
diff --git a/Engine/Source/Runtime/Engine/Private/HLSLTree/HLSLTreeCommon.cpp b/Engine/Source/Runtime/Engine/Private/HLSLTree/HLSLTreeCommon.cpp
--- a/Engine/Source/Runtime/Engine/Private/HLSLTree/HLSLTreeCommon.cpp
+++ b/Engine/Source/Runtime/Engine/Private/HLSLTree/HLSLTreeCommon.cpp
@@ -120,7 +120,9 @@ void FExpressionConstant::ComputeAnalyticDerivatives(FTree& Tree, FExpressionDer
 	}
 }
 
-bool FExpressionConstant::PrepareValue(FEmitContext& Context, FEmitScope& Scope, const FRequestedType& RequestedType, FPrepareValueResult& OutResult) const
+namespace Private
+{
+FPreparedType PrepareConstant(const Shader::FValue& Value)
 {
 	const int32 NumComponents = Value.Type.GetNumComponents();
 	FPreparedType ResultType(Value.Type);
@@ -130,7 +132,13 @@ bool FExpressionConstant::PrepareValue(FEmitContext& Context, FEmitScope& Scope,
 		const Shader::FValueComponent Component = Value.GetComponent(Index);
 		ResultType.PreparedComponents.Add(Component.Packed ? EExpressionEvaluation::Constant : EExpressionEvaluation::ConstantZero);
 	}
-	return OutResult.SetType(Context, RequestedType, ResultType);
+	return ResultType;
+}
+} // namespace Private
+
+bool FExpressionConstant::PrepareValue(FEmitContext& Context, FEmitScope& Scope, const FRequestedType& RequestedType, FPrepareValueResult& OutResult) const
+{
+	return OutResult.SetType(Context, RequestedType, Private::PrepareConstant(Value));
 }
 
 void FExpressionConstant::EmitValuePreshader(FEmitContext& Context, FEmitScope& Scope, const FRequestedType& RequestedType, FEmitValuePreshaderResult& OutResult) const
@@ -140,6 +148,50 @@ void FExpressionConstant::EmitValuePreshader(FEmitContext& Context, FEmitScope&
 	OutResult.Preshader.WriteOpcode(Shader::EPreshaderOpcode::Constant).Write(Value);
 }
 
+void FExpressionDefaultValue::ComputeAnalyticDerivatives(FTree& Tree, FExpressionDerivatives& OutResult) const
+{
+	const Shader::FType DerivativeType = DefaultValue.Type.GetDerivativeType();
+	if (!DerivativeType.IsVoid())
+	{
+		const FExpressionDerivatives ExpressionDerivative = Tree.GetAnalyticDerivatives(Expression);
+		if (ExpressionDerivative.IsValid())
+		{
+			OutResult.ExpressionDdx = Tree.NewExpression<FExpressionDefaultValue>(ExpressionDerivative.ExpressionDdx, Shader::FValue(DerivativeType));
+			OutResult.ExpressionDdy = Tree.NewExpression<FExpressionDefaultValue>(ExpressionDerivative.ExpressionDdy, Shader::FValue(DerivativeType));
+		}
+	}
+}
+
+const FExpression* FExpressionDefaultValue::ComputePreviousFrame(FTree& Tree, const FRequestedType& RequestedType) const
+{
+	return Tree.NewExpression<FExpressionDefaultValue>(Tree.GetPreviousFrame(Expression, RequestedType), DefaultValue);
+}
+
+bool FExpressionDefaultValue::PrepareValue(FEmitContext& Context, FEmitScope& Scope, const FRequestedType& RequestedType, FPrepareValueResult& OutResult) const
+{
+	FPreparedType ResultType = Context.PrepareExpression(Expression, Scope, RequestedType);
+	if (ResultType.IsVoid())
+	{
+		ResultType = Private::PrepareConstant(DefaultValue);
+	}
+	return OutResult.SetType(Context, RequestedType, ResultType);
+}
+
+void FExpressionDefaultValue::EmitValuePreshader(FEmitContext& Context, FEmitScope& Scope, const FRequestedType& RequestedType, FEmitValuePreshaderResult& OutResult) const
+{
+	const FPreparedType& PreparedType = Context.GetPreparedType(Expression);
+	if (!PreparedType.IsVoid())
+	{
+		// If 'Exrpession' is valid, use that
+		return FExpressionForward::EmitValuePreshader(Context, Scope, RequestedType, OutResult);
+	}
+
+	// Otherwise emit the default value
+	Context.PreshaderStackPosition++;
+	OutResult.Type = DefaultValue.Type;
+	OutResult.Preshader.WriteOpcode(Shader::EPreshaderOpcode::Constant).Write(DefaultValue);
+}
+
 void FExpressionGetStructField::ComputeAnalyticDerivatives(FTree& Tree, FExpressionDerivatives& OutResult) const
 {
 	const FExpressionDerivatives StructDerivatives = Tree.GetAnalyticDerivatives(StructExpression);
@@ -308,8 +360,11 @@ void FExpressionSelect::ComputeAnalyticDerivatives(FTree& Tree, FExpressionDeriv
 {
 	const FExpressionDerivatives TrueDerivatives = Tree.GetAnalyticDerivatives(TrueExpression);
 	const FExpressionDerivatives FalseDerivatives = Tree.GetAnalyticDerivatives(FalseExpression);
-	OutResult.ExpressionDdx = Tree.NewExpression<FExpressionSelect>(ConditionExpression, TrueDerivatives.ExpressionDdx, FalseDerivatives.ExpressionDdx);
-	OutResult.ExpressionDdy = Tree.NewExpression<FExpressionSelect>(ConditionExpression, TrueDerivatives.ExpressionDdy, FalseDerivatives.ExpressionDdy);
+	if (TrueDerivatives.IsValid() && FalseDerivatives.IsValid())
+	{
+		OutResult.ExpressionDdx = Tree.NewExpression<FExpressionSelect>(ConditionExpression, TrueDerivatives.ExpressionDdx, FalseDerivatives.ExpressionDdx);
+		OutResult.ExpressionDdy = Tree.NewExpression<FExpressionSelect>(ConditionExpression, TrueDerivatives.ExpressionDdy, FalseDerivatives.ExpressionDdy);
+	}
 }
 
 const FExpression* FExpressionSelect::ComputePreviousFrame(FTree& Tree, const FRequestedType& RequestedType) const
@@ -435,29 +490,26 @@ namespace Private
 {
 FRequestedType GetRequestedSwizzleType(const FSwizzleParameters& Params, const FRequestedType& RequestedType)
 {
-	int32 MaxSwizzledComponentIndex = 0;
-	FRequestedType RequestedInputType;
+	const Shader::EValueType InputType = Shader::MakeValueType(RequestedType.GetValueComponentType(), Params.GetNumInputComponents());
+	FRequestedType RequestedInputType(InputType, false);
 	if (Params.NumComponents == 1)
 	{
-		MaxSwizzledComponentIndex = Params.GetSwizzleComponentIndex(0);
 		if (RequestedType.IsNumericVectorRequested())
 		{
-			RequestedInputType.SetComponentRequest(MaxSwizzledComponentIndex);
+			RequestedInputType.SetComponentRequest(Params.GetSwizzleComponentIndex(0));
 		}
 	}
 	else
 	{
 		for (int32 Index = 0; Index < Params.NumComponents; ++Index)
 		{
-			const int32 SwizzledComponentIndex = Params.GetSwizzleComponentIndex(Index);
-			MaxSwizzledComponentIndex = FMath::Max(MaxSwizzledComponentIndex, SwizzledComponentIndex);
 			if (RequestedType.IsComponentRequested(Index))
 			{
+				const int32 SwizzledComponentIndex = Params.GetSwizzleComponentIndex(Index);
 				RequestedInputType.SetComponentRequest(SwizzledComponentIndex);
 			}
 		}
 	}
-	RequestedInputType.Type = Shader::MakeValueType(RequestedType.GetValueComponentType(), MaxSwizzledComponentIndex + 1);
 	return RequestedInputType;
 }
 } // namespace Private
@@ -498,6 +550,10 @@ bool SwizzlePrepareValue(FEmitContext& Context,
 {
 	const FRequestedType RequestedInputType = Private::GetRequestedSwizzleType(Parameters, RequestedType);
 	const FPreparedType& InputType = Context.PrepareExpression(Input, Scope, RequestedInputType);
+	if (InputType.IsVoid())
+	{
+		return false;
+	}
 
 	FPreparedType ResultType;
 	ResultType.Type = Shader::MakeValueType(InputType.GetValueComponentType(), Parameters.NumComponents);
@@ -517,9 +573,12 @@ void SwizzleEmitValueShader(FEmitContext& Context,
 	FEmitValueShaderResult& OutResult)
 {
 	const FRequestedType RequestedInputType = Private::GetRequestedSwizzleType(Parameters, RequestedType);
-	FEmitShaderExpression* EmitInput = Input->GetValueShader(Context, Scope, RequestedInputType);
-	const Shader::FValueTypeDescription InputTypeDesc = Shader::GetValueTypeDescription(EmitInput->Type);
 
+	// Make sure the input is cast to the explicit requested type, this ensures it will have enough components for the swizzle
+	// Alternately, we could avoid the cast, and update the logic to insert 0s for swizzle access to invalid components
+	FEmitShaderExpression* EmitInput = Input->GetValueShader(Context, Scope, RequestedInputType, RequestedInputType.Type);
+
+	const Shader::FValueTypeDescription InputTypeDesc = Shader::GetValueTypeDescription(EmitInput->Type);
 	bool bSkipSwizzle = (InputTypeDesc.NumComponents == Parameters.NumComponents);
 	if (bSkipSwizzle)
 	{
@@ -642,7 +701,7 @@ FSwizzleParameters GetComponentMaskSwizzle(FEmitContext& Context, FEmitScope& Sc
 	const Shader::FBoolValue MaskValue = Mask->GetValueConstant(Context, Scope, Shader::EValueType::Bool4).AsBool();
 	return MakeSwizzleMask(MaskValue[0], MaskValue[1], MaskValue[2], MaskValue[3]);
 }
-}
+} // namespace Private
 
 bool FExpressionComponentMask::PrepareValue(FEmitContext& Context, FEmitScope& Scope, const FRequestedType& RequestedType, FPrepareValueResult& OutResult) const
 {
@@ -738,6 +797,11 @@ const FExpression* FExpressionAppend::ComputePreviousFrame(FTree& Tree, const FR
 bool FExpressionAppend::PrepareValue(FEmitContext& Context, FEmitScope& Scope, const FRequestedType& RequestedType, FPrepareValueResult& OutResult) const
 {
 	const FPreparedType& LhsType = Context.PrepareExpression(Lhs, Scope, RequestedType);
+	if (LhsType.IsVoid())
+	{
+		return false;
+	}
+
 	const int32 NumRequestedComponents = RequestedType.Type.GetNumComponents();
 	const int32 NumLhsComponents = LhsType.Type.GetNumComponents();
 	const int32 NumRhsComponents = FMath::Max(NumRequestedComponents - NumLhsComponents, 1);
@@ -753,27 +817,23 @@ bool FExpressionAppend::PrepareValue(FEmitContext& Context, FEmitScope& Scope, c
 	}
 
 	const FPreparedType& RhsType = Context.PrepareExpression(Rhs, Scope, RhsRequestedType);
-	if (!RhsType.IsVoid() && LhsType.GetValueComponentType() != RhsType.GetValueComponentType())
+	const Private::FAppendTypes Types = Private::GetAppendTypes(RequestedType, LhsType, RhsType);
+	check(!Types.ResultType.IsNumericScalar()); // Appending 2 values should never result in a scalar
+	if (Types.ResultType.IsVoid())
 	{
-		return Context.Error(TEXT("Type mismatch"));
+		return false;
 	}
-	
-	const Private::FAppendTypes Types = Private::GetAppendTypes(RequestedType, LhsType, RhsType);
+
 	Context.MarkInputType(Lhs, Types.LhsType);
 	Context.MarkInputType(Rhs, Types.RhsType);
 
-	if (Types.ResultType.IsNumericScalar())
-	{
-		int a = 0;
-	}
-
 	return OutResult.SetType(Context, RequestedType, Types.ResultType);
 }
 
 void FExpressionAppend::EmitValueShader(FEmitContext& Context, FEmitScope& Scope, const FRequestedType& RequestedType, FEmitValueShaderResult& OutResult) const
 {
 	const Private::FAppendTypes Types = Private::GetAppendTypes(RequestedType, Context.GetPreparedType(Lhs), Context.GetPreparedType(Rhs));
-	const Shader::FType ResultType = Types.ResultType.GetPreparedType();
+	const Shader::FType ResultType = Types.ResultType.GetResultType();
 	FEmitShaderExpression* LhsValue = Lhs->GetValueShader(Context, Scope, Types.LhsRequestedType, Types.LhsType);
 
 	if (Types.RhsType == Shader::EValueType::Void)
@@ -812,7 +872,7 @@ void FExpressionAppend::EmitValuePreshader(FEmitContext& Context, FEmitScope& Sc
 
 		OutResult.Preshader.WriteOpcode(Shader::EPreshaderOpcode::AppendVector);
 	}
-	OutResult.Type = Types.ResultType.GetPreparedType();
+	OutResult.Type = Types.ResultType.GetResultType();
 }
 
 void FExpressionSwitchBase::ComputeAnalyticDerivatives(FTree& Tree, FExpressionDerivatives& OutResult) const
diff --git a/Engine/Source/Runtime/Engine/Private/HLSLTree/HLSLTreeEmit.cpp b/Engine/Source/Runtime/Engine/Private/HLSLTree/HLSLTreeEmit.cpp
--- a/Engine/Source/Runtime/Engine/Private/HLSLTree/HLSLTreeEmit.cpp
+++ b/Engine/Source/Runtime/Engine/Private/HLSLTree/HLSLTreeEmit.cpp
@@ -222,7 +222,7 @@ bool FEmitContext::InternalError(FStringView ErrorMessage)
 		{
 			CurrentOwners = OwnerStack.Last()->GetOwners();
 		}
-		Errors->AddError(CurrentOwners, ErrorMessage);
+		Errors->AddErrorInternal(CurrentOwners, ErrorMessage);
 	}
 	NumErrors++;
 	return false;
@@ -376,7 +376,7 @@ FRequestedType FEmitContext::GetRequestedType(const FExpression* Expression) con
 
 Shader::FType FEmitContext::GetType(const FExpression* Expression) const
 {
-	return GetPreparedType(Expression).GetPreparedType();
+	return GetPreparedType(Expression).GetResultType();
 }
 
 EExpressionEvaluation FEmitContext::GetEvaluation(const FExpression* Expression, const FEmitScope& Scope, const FRequestedType& RequestedType) const
@@ -436,6 +436,10 @@ FPreparedType FEmitContext::PrepareExpression(const FExpression* InExpression, F
 		check(!ResultType.IsVoid());
 		MarkInputType(InExpression, RequestedType.Type.GetConcreteType());
 	}
+	else
+	{
+		int a = 0;
+	}
 	return ResultType;
 }
 
@@ -774,7 +778,7 @@ FEmitShaderExpression* FEmitContext::InternalEmitExpression(FEmitScope& Scope, T
 {
 	FEmitShaderExpression* ShaderValue = nullptr;
 
-	if (Code == TEXT("(((float3)0).z * Material.PreshaderBuffer[6].x)"))
+	if (Code.Contains(TEXT("((float2)Local21).xyz")))
 	{
 		int a = 0;
 	}
@@ -1151,6 +1155,10 @@ FEmitShaderExpression* FEmitContext::EmitCast(FEmitScope& Scope, FEmitShaderExpr
 	check(ShaderValue);
 	check(!DestType.IsVoid());
 
+	// TODO - could handle generic types here, cast to 'Any' should just return ShaderValue, cast to 'Numeric' could check number of components
+	// Not sure if this is needed in practice
+	check(!DestType.IsGeneric());
+
 	if (ShaderValue->Type == DestType)
 	{
 		return ShaderValue;
diff --git a/Engine/Source/Runtime/Engine/Private/HLSLTree/HLSLTreeOperation.cpp b/Engine/Source/Runtime/Engine/Private/HLSLTree/HLSLTreeOperation.cpp
--- a/Engine/Source/Runtime/Engine/Private/HLSLTree/HLSLTreeOperation.cpp
+++ b/Engine/Source/Runtime/Engine/Private/HLSLTree/HLSLTreeOperation.cpp
@@ -272,8 +272,10 @@ FOperationTypes GetOperationTypes(EOperation Op, TConstArrayView<FPreparedType>
 				}
 			}
 
-			const int32 NumInputComponents = InputPreparedType[Index].IsNumericScalar() ? 1 : NumIntermediateComponents;
-			Types.InputType[Index] = Shader::MakeValueType(InputComponentType, NumInputComponents);
+			// No implicit conversion from FLWCScalar to FLWCVector types
+			// Could potentially add overloads to LWCOperations.ush to allow this
+			const bool bInputScalar = (InputPreparedType[Index].IsNumericScalar() && InputComponentType != Shader::EValueComponentType::Double);
+			Types.InputType[Index] = Shader::MakeValueType(InputComponentType, bInputScalar ? 1 : NumIntermediateComponents);
 		}
 		Types.ResultType = IntermediateType;
 		Types.bIsLWC = (IntermediateComponentType == Shader::EValueComponentType::Double);
@@ -664,7 +666,7 @@ void FExpressionOperation::EmitValueShader(FEmitContext& Context, FEmitScope& Sc
 		InputValue[Index] = Inputs[Index]->GetValueShader(Context, Scope, RequestedTypes.InputType[Index], Types.InputType[Index]);
 	}
 
-	const Shader::EValueType ResultType = Types.ResultType.GetPreparedType();
+	const Shader::EValueType ResultType = Types.ResultType.GetResultType();
 	check(Shader::IsNumericType(ResultType));
 
 	switch (Op)
@@ -779,7 +781,7 @@ void FExpressionOperation::EmitValuePreshader(FEmitContext& Context, FEmitScope&
 	}
 
 	OutResult.Preshader.WriteOpcode(OpDesc.PreshaderOpcode);
-	OutResult.Type = Types.ResultType.GetPreparedType();
+	OutResult.Type = Types.ResultType.GetResultType();
 }
 
 } // namespace UE::HLSLTree
diff --git a/Engine/Source/Runtime/Engine/Private/Materials/Material.cpp b/Engine/Source/Runtime/Engine/Private/Materials/Material.cpp
--- a/Engine/Source/Runtime/Engine/Private/Materials/Material.cpp
+++ b/Engine/Source/Runtime/Engine/Private/Materials/Material.cpp
@@ -4924,7 +4924,7 @@ static void SetMaterialInputDescription(FScalarMaterialInput& Input, bool bHidde
 
 static void SetMaterialInputDescription(FShadingModelMaterialInput& Input, bool bHidden, FMaterialInputDescription& OutDescription)
 {
-	OutDescription.Type = UE::Shader::EValueType::Void;
+	OutDescription.Type = UE::Shader::EValueType::Int1;
 	OutDescription.Input = &Input;
 	OutDescription.bUseConstant = false;
 	OutDescription.bHidden = bHidden;
diff --git a/Engine/Source/Runtime/Engine/Private/Materials/MaterialExpressionHLSL.cpp b/Engine/Source/Runtime/Engine/Private/Materials/MaterialExpressionHLSL.cpp
--- a/Engine/Source/Runtime/Engine/Private/Materials/MaterialExpressionHLSL.cpp
+++ b/Engine/Source/Runtime/Engine/Private/Materials/MaterialExpressionHLSL.cpp
@@ -245,7 +245,16 @@ bool UMaterialExpressionStaticSwitch::GenerateHLSLExpression(FMaterialHLSLGenera
 		return false;
 	}
 
-	OutExpression = Generator.GenerateBranch(Scope, ConditionExpression, TrueExpression, FalseExpression);
+	if (TrueExpression && FalseExpression)
+	{
+		// Dynamic branch requires both true/false expression to be valid
+		OutExpression = Generator.GenerateBranch(Scope, ConditionExpression, TrueExpression, FalseExpression);
+	}
+	else
+	{
+		// Select can handle missing values
+		OutExpression = Generator.GetTree().NewExpression<FExpressionSelect>(ConditionExpression, TrueExpression, FalseExpression);
+	}
 	return true;
 }
 
diff --git a/Engine/Source/Runtime/Engine/Private/Materials/MaterialHLSLGenerator.cpp b/Engine/Source/Runtime/Engine/Private/Materials/MaterialHLSLGenerator.cpp
--- a/Engine/Source/Runtime/Engine/Private/Materials/MaterialHLSLGenerator.cpp
+++ b/Engine/Source/Runtime/Engine/Private/Materials/MaterialHLSLGenerator.cpp
@@ -140,9 +140,9 @@ static const UE::HLSLTree::FExpression* CompileMaterialInput(FMaterialHLSLGenera
 		FMaterialInputDescription InputDescription;
 		if (Material->GetExpressionInputDescription(InputProperty, InputDescription))
 		{
+			UE::Shader::FValue DefaultValue = UE::Shader::Cast(FMaterialAttributeDefinitionMap::GetDefaultValue(InputProperty), InputDescription.Type);
 			if (InputDescription.bUseConstant)
 			{
-				UE::Shader::FValue DefaultValue = UE::Shader::Cast(FMaterialAttributeDefinitionMap::GetDefaultValue(InputProperty), InputDescription.Type);
 				if (InputDescription.ConstantValue != DefaultValue)
 				{
 					Expression = Generator.NewConstant(InputDescription.ConstantValue);
@@ -152,6 +152,10 @@ static const UE::HLSLTree::FExpression* CompileMaterialInput(FMaterialHLSLGenera
 			{
 				check(InputDescription.Input);
 				Expression = InputDescription.Input->TryAcquireHLSLExpression(Generator, Scope, (int32)InputProperty);
+				if (Expression)
+				{
+					Expression = Generator.GetTree().NewExpression<FExpressionDefaultValue>(Expression, DefaultValue);
+				}
 			}
 		}
 	}
@@ -206,10 +210,12 @@ bool FMaterialHLSLGenerator::GenerateResult(UE::HLSLTree::FScope& Scope)
 
 					if (AttributesExpression)
 					{
+						AttributesExpression = GetTree().NewExpression<FExpressionDefaultValue>(AttributesExpression, CachedTree.GetMaterialAttributesDefaultValue());
+
 						const FString& WPOName = FMaterialAttributeDefinitionMap::GetAttributeName(MP_WorldPositionOffset);
 						const FStructField* WPOField = CachedTree.GetMaterialAttributesType()->FindFieldByName(*WPOName);
 
-						FRequestedType PrevRequestedType;
+						FRequestedType PrevRequestedType(CachedTree.GetMaterialAttributesType());
 						PrevRequestedType.SetFieldRequested(WPOField);
 
 						const FExpression* PrevAttributesExpression = GetTree().GetPreviousFrame(AttributesExpression, PrevRequestedType);
@@ -708,6 +714,10 @@ const UE::HLSLTree::FExpression* FMaterialHLSLGenerator::GenerateBranch(UE::HLSL
 {
 	using namespace UE::HLSLTree;
 
+	check(ConditionExpression);
+	check(TrueExpression);
+	check(FalseExpression);
+
 	FXxHash64 Hash;
 	{
 		FXxHash64Builder Hasher;
diff --git a/Engine/Source/Runtime/Engine/Private/Materials/MaterialHLSLTree.cpp b/Engine/Source/Runtime/Engine/Private/Materials/MaterialHLSLTree.cpp
--- a/Engine/Source/Runtime/Engine/Private/Materials/MaterialHLSLTree.cpp
+++ b/Engine/Source/Runtime/Engine/Private/Materials/MaterialHLSLTree.cpp
@@ -1263,11 +1263,11 @@ void FExpressionVertexInterpolator::EmitValueShader(FEmitContext& Context, FEmit
 	if (InterpolatorIndex >= 0)
 	{
 		const FVertexInterpolator& Interpolator = EmitMaterialData.VertexInterpolators[InterpolatorIndex];
-		const Shader::FType LocalType = Interpolator.PreparedType.GetPreparedType();
+		const Shader::FType LocalType = Interpolator.PreparedType.GetResultType();
 
 		// Request all non-shader components in preshader, only shader components will be packed in interpolator
 		FRequestedType RequestedPreshaderType(RequestedType, false);
-		for (int32 Index = 0; Index < Interpolator.PreparedType.GetNumPreparedComponents(); ++Index)
+		for (int32 Index = 0; Index < Interpolator.PreparedType.PreparedComponents.Num(); ++Index)
 		{
 			const EExpressionEvaluation ComponentEvaluation = Interpolator.PreparedType.GetComponent(Index).Evaluation;
 			if (RequestedType.IsComponentRequested(Index) && ComponentEvaluation != EExpressionEvaluation::Shader)
@@ -1312,11 +1312,11 @@ void FEmitData::AddInterpolator(const FExpression* Expression, const FRequestedT
 		InterpolatorIndex = VertexInterpolators.Emplace(Expression);
 	}
 
-	const Shader::FType LocalType = PreparedType.GetPreparedType();
+	const Shader::FType LocalType = PreparedType.GetResultType();
 	FVertexInterpolator& Interpolator = VertexInterpolators[InterpolatorIndex];
 	Interpolator.RequestedType = FRequestedType(RequestedType, false);
 	Interpolator.PreparedType = PreparedType;
-	for (int32 ComponentIndex = 0; ComponentIndex < PreparedType.GetNumPreparedComponents(); ++ComponentIndex)
+	for (int32 ComponentIndex = 0; ComponentIndex < PreparedType.PreparedComponents.Num(); ++ComponentIndex)
 	{
 		if (RequestedType.IsComponentRequested(ComponentIndex))
 		{
@@ -1344,7 +1344,7 @@ void FEmitData::EmitInterpolatorStatements(FEmitContext& Context, FEmitScope& Sc
 	for (int32 Index = 0; Index < VertexInterpolators.Num(); ++Index)
 	{
 		const FVertexInterpolator& Interpolator = VertexInterpolators[Index];
-		FEmitShaderExpression* EmitExpression = Interpolator.Expression->GetValueShader(Context, Scope, Interpolator.RequestedType, Interpolator.PreparedType, Interpolator.PreparedType.GetPreparedType());
+		FEmitShaderExpression* EmitExpression = Interpolator.Expression->GetValueShader(Context, Scope, Interpolator.RequestedType, Interpolator.PreparedType, Interpolator.PreparedType.GetResultType());
 		Context.EmitStatement(Scope, TEXT("MaterialPackVertexInterpolator%(Parameters, %);"), Index, EmitExpression);
 	}
 }
@@ -1413,7 +1413,7 @@ void FEmitData::EmitInterpolatorShader(FEmitContext& Context, FStringBuilderBase
 	for (int32 Index = 0; Index < VertexInterpolators.Num(); ++Index)
 	{
 		const FVertexInterpolator& Interpolator = VertexInterpolators[Index];
-		const Shader::FType Type = Interpolator.PreparedType.GetPreparedType();
+		const Shader::FType Type = Interpolator.PreparedType.GetResultType();
 		const TCHAR* TypeName = Type.GetName();
 
 		TStringBuilder<1024> ReadCode;
diff --git a/Engine/Source/Runtime/Engine/Public/HLSLTree/HLSLTree.h b/Engine/Source/Runtime/Engine/Public/HLSLTree/HLSLTree.h
--- a/Engine/Source/Runtime/Engine/Public/HLSLTree/HLSLTree.h
+++ b/Engine/Source/Runtime/Engine/Public/HLSLTree/HLSLTree.h
@@ -47,28 +47,15 @@ class FEmitShaderExpression;
 
 static constexpr int32 MaxNumPreviousScopes = 2;
 
+/** Used to report any errors generating when preparing the HLSLTree */
 class FErrorHandlerInterface
 {
 public:
-	virtual ~FErrorHandlerInterface() {}
-
-	bool AddError(UObject* InOwner, FStringView InError)
-	{
-		AddErrorInternal(MakeArrayView(&InOwner, 1), InError);
-		return false;
-	}
-
-	bool AddError(TConstArrayView<UObject*> InOwners, FStringView InError)
-	{
-		AddErrorInternal(InOwners, InError);
-		return false;
-	}
-
-protected:
+	/** Recieves a list of owners for the node that generated the error, along with the error message */
 	virtual void AddErrorInternal(TConstArrayView<UObject*> InOwners, FStringView InError) = 0;
 };
 
-class FNullErrorHandler : public FErrorHandlerInterface
+class FNullErrorHandler final : public FErrorHandlerInterface
 {
 public:
 	virtual void AddErrorInternal(TConstArrayView<UObject*> InOwners, FStringView InError) override
@@ -89,6 +76,14 @@ private:
 	friend class FTree;
 };
 
+/**
+ * OwnedNodes track 1 or more UObject 'owners'
+ * When generating HLSLTree for materials, the owner will typically be the UMaterialExpression that created the node, or the UMaterial itself
+ * Since certain nodes (FExpressions) are deduplicated and shared, it's possible to have multiple owners, if multiple owners attempt to create a node with the same parameters
+ * Tracking owners is important for the following reasons:
+ * - Attributing any generated errors to the owner(s)
+ * - Tracking input/output types (to color wires/pins in the material editor for example)
+ */
 class FOwnedNode : public FNode
 {
 public:
@@ -158,7 +153,7 @@ public:
 		return Type.IsAny() || (RequestedComponents.IsValidIndex(Index) ? (bool)RequestedComponents[Index] : false);
 	}
 
-	/** Check if any of xyzw are requested */
+	/** Check if any of xyzw are requested, any of these components will request a numeric scalar type */
 	bool IsNumericVectorRequested() const
 	{
 		return Type.IsAny() || FMath::IsWithin(RequestedComponents.Find(true), 0, 4);
@@ -168,11 +163,12 @@ public:
 	bool IsStruct() const { return Type.IsStruct(); }
 	bool IsObject() const { return Type.IsObject(); }
 	bool IsNumeric() const { return Type.IsNumeric(); }
+
+	/** No requested components */
 	bool IsEmpty() const { return !Type.IsAny() && !RequestedComponents.Contains(true); }
 
 	void SetComponentRequest(int32 Index, bool bRequest = true);
 
-
 	/** Marks the given field as requested (or not) */
 	void SetFieldRequested(const Shader::FStructField* Field, bool bRequest = true);
 	void ClearFieldRequested(const Shader::FStructField* Field)
@@ -187,7 +183,8 @@ public:
 	FRequestedType GetField(const Shader::FStructField* Field) const;
 
 	/**
-	 * The type used to initialize the request. The actual 'Type' represented by the FRequestedType may be different, as vector types will truncate to the number of requested components
+	 * The actual type that was requested.  Specific componets of this type are requested by setting RequestedComponents
+	 * Shader::EValueType::Any is a valid type here, which means that RequestedComponents is ignored, and IsComponentRequested() will be true for any component
 	 */
 	Shader::FType Type;
 
@@ -281,9 +278,11 @@ public:
 	void SetField(const Shader::FStructField* Field, const FPreparedType& FieldType);
 	FPreparedType GetFieldType(const Shader::FStructField* Field) const;
 
-	/** For numeric vectors, the number of components in the prepared type will only include components that have valid evaluation */
+	/** Returns the number of components, only including components with valid evaluation */
 	int32 GetNumPreparedComponents() const;
-	Shader::FType GetPreparedType() const;
+
+	/** For numeric vectors, the number of components in the result type will only include components that have valid evaluation */
+	Shader::FType GetResultType() const;
 
 	/** Converts to a requested type, based on IsRequestedEvaluation() */
 	FRequestedType GetRequestedType() const;
@@ -295,6 +294,8 @@ public:
 	bool IsObject() const { return Type.IsObject(); }
 	bool IsNumeric() const { return Type.IsNumeric(); }
 	bool IsNumericScalar() const { return Type.IsNumericScalar(); }
+
+	/** Either void, or no prepared components */
 	bool IsEmpty() const;
 
 	EExpressionEvaluation GetEvaluation(const FEmitScope& Scope) const;
@@ -314,7 +315,8 @@ public:
 	void EnsureNumComponents(int32 NumComponents);
 
 	/**
-	 * The type used to initialize the this. The actual 'Type' represented by the FPreparedType may be different, as vector types will truncate to the number of prepared components
+	 * The prepared type.  Status of components within this type can be configured by setting PreparedComponents
+	 * Components within Type that don't have a coorisponding entry in PreparedComponents are considered to have ConstantZero evaluation
 	 */
 	Shader::FType Type;
 
@@ -393,6 +395,8 @@ struct FExpressionDerivatives
  * Unlike statements, expressions are not expected to execute in any particular order.
  * Examples include constant literals, variable accessors, and various types of math operations
  * This is an abstract base class, with derived classes representing various types of expression
+ * Derived expression classes should provide a constructor to fully initialize all fields.  All types given to this constructor must be hashable via HLSLTree::AppendHash().
+ * A unique hash value will be generated for each expression, by hashing the constructor arguments, which will be used to deduplicate expressions
  */
 class FExpression : public FOwnedNode
 {
@@ -466,8 +470,14 @@ protected:
 	/**
 	 * Computes a FPreparedType for this expression, given a FRequestedType.  Will be called multiple times, with potentially different requested types, if the FExpression is used multiple times.
 	 * In this case, all the prepared types will be merged together if possible (or generate an error otherwise).
+	 * PrepareValue will be called in two phases, first with Context.bMarkLiveValues=false, then with bMarkLiveValues=true.
+	 * If Context.bMarkLiveValues is true, the expression should record any state to signal to the client that it's active.
+	 * For example, a material expression parameter might add the parameter name or texture value to the material data.  This is typically done through the Context.FindData() interface.
+	 * If Context.bMarkLiveValues is true, the expression should also take care *not* to call Context.PrepareExpression for any nested expressions that are not relevant to the final result.
+	 * For example, in a multiply expression where one input is 0, the other input is not relevant, and can be skipped.
 	 * The EExpressionEvaluation set by this is important for determining which EmitValue*** methods may be called
 	 * A given FExpression implementation must support at least one of EmitValueShader/EmitValuePreshader, but doesn't need to support both
+	 * TODO? - It feels messy to have PrepareValue() used for both of these phases, possibly better to have a separate MarkLive() virtual method...in practice that would often lead to duplicate code however
 	 */
 	virtual bool PrepareValue(FEmitContext& Context, FEmitScope& Scope, const FRequestedType& RequestedType, FPrepareValueResult& OutResult) const = 0;
 
@@ -494,13 +504,22 @@ private:
 	friend class FExpressionForward;
 };
 
+/**
+ * Functions are a way to dynamically inject new scopes into the tree
+ * They are used to represent calls to material functions using execution flow, and to inject dynamic branches without explicit control flow
+ */
 class FFunction final : public FNode
 {
 public:
 	FScope& GetRootScope() const { return *RootScope; }
 
+	/** The root scope of the function */
 	FScope* RootScope = nullptr;
+
+	/** The scope where this function is called, when HLSL is emitted, the RootScope will be injected here */
 	FScope* CalledScope = nullptr;
+
+	/** Function output expressions, should be expressions nested under RootScope */
 	TArray<const FExpression*, TInlineAllocator<8>> OutputExpressions;
 };
 
@@ -560,6 +579,10 @@ public:
 
 	FScope& GetRootScope() const { return *RootScope; }
 
+	/**
+	 * Creates a new FExpression-derived type by passing the given arguments to the constructor
+	 * Given arguments are hashed, and may result in returning an existing expression if the hash matches
+	 */
 	template<typename T, typename... ArgTypes>
 	inline const FExpression* NewExpression(ArgTypes&&... Args)
 	{
diff --git a/Engine/Source/Runtime/Engine/Public/HLSLTree/HLSLTreeCommon.h b/Engine/Source/Runtime/Engine/Public/HLSLTree/HLSLTreeCommon.h
--- a/Engine/Source/Runtime/Engine/Public/HLSLTree/HLSLTreeCommon.h
+++ b/Engine/Source/Runtime/Engine/Public/HLSLTree/HLSLTreeCommon.h
@@ -53,6 +53,23 @@ public:
 	virtual void EmitValuePreshader(FEmitContext& Context, FEmitScope& Scope, const FRequestedType& RequestedType, FEmitValuePreshaderResult& OutResult) const override;
 };
 
+/** Forwards to another expression, but will fallback to the 'DefaultValue' constant if the forwarded expression fails */
+class FExpressionDefaultValue : public FExpressionForward
+{
+public:
+	explicit FExpressionDefaultValue(const FExpression* InExpression, const Shader::FValue& InDefaultValue)
+		: FExpressionForward(InExpression)
+		, DefaultValue(InDefaultValue)
+	{}
+
+	Shader::FValue DefaultValue;
+
+	virtual void ComputeAnalyticDerivatives(FTree& Tree, FExpressionDerivatives& OutResult) const override;
+	virtual const FExpression* ComputePreviousFrame(FTree& Tree, const FRequestedType& RequestedType) const override;
+	virtual bool PrepareValue(FEmitContext& Context, FEmitScope& Scope, const FRequestedType& RequestedType, FPrepareValueResult& OutResult) const override;
+	virtual void EmitValuePreshader(FEmitContext& Context, FEmitScope& Scope, const FRequestedType& RequestedType, FEmitValuePreshaderResult& OutResult) const override;
+};
+
 class FExpressionGetStructField : public FExpression
 {
 public:
@@ -106,7 +123,9 @@ public:
 		: ConditionExpression(InCondition)
 		, TrueExpression(InTrue)
 		, FalseExpression(InFalse)
-	{}
+	{
+		check(ConditionExpression);
+	}
 
 	const FExpression* ConditionExpression;
 	const FExpression* TrueExpression;
@@ -154,6 +173,11 @@ public:
 	virtual void EmitValuePreshader(FEmitContext& Context, FEmitScope& Scope, const FRequestedType& RequestedType, FEmitValuePreshaderResult& OutResult) const override;
 };
 
+/**
+ * Similar to FExpressionSwizzle, except swizzle parameters are extracted from a 'Mask' expression, which is expected to generate a 'bool4' component mask (which is then used to initialize a FSwizzleParameters)
+ * The 'Mask' input is required to be constant (otherwise an error is generated)
+ * This is used to support material StaticParameterMask expressions, in most cases FExpressionSwizzle is a better choice
+ */
 class FExpressionComponentMask : public FExpression
 {
 public:
diff --git a/Engine/Source/Runtime/Engine/Public/HLSLTree/HLSLTreeEmit.h b/Engine/Source/Runtime/Engine/Public/HLSLTree/HLSLTreeEmit.h
--- a/Engine/Source/Runtime/Engine/Public/HLSLTree/HLSLTreeEmit.h
+++ b/Engine/Source/Runtime/Engine/Public/HLSLTree/HLSLTreeEmit.h
@@ -67,6 +67,9 @@ public:
 };
 using FEmitShaderDependencies = TArray<FEmitShaderNode*, TInlineAllocator<8>>;
 
+/**
+ * Represents an HLSL expression
+ */
 class FEmitShaderExpression final : public FEmitShaderNode
 {
 public:
@@ -80,9 +83,19 @@ public:
 	virtual FEmitShaderExpression* AsExpression() override { return this; }
 	inline bool IsInline() const { return Value == nullptr; }
 
+	/**
+	 * String used to reference this expression.
+	 * For inline epxressions, this will be the actual HLSL.  Otherwise it will be the name of a local variable
+	 */
 	const TCHAR* Reference = nullptr;
+
+	/** For non-inline expressions, this holds the actual HLSL.  Otherwise it will be nullptr */
 	const TCHAR* Value = nullptr;
+
+	/** Type of the HLSL */
 	Shader::FType Type;
+
+	/** Hash of the string, used to deduplicate identical expressions */
 	FXxHash64 Hash;
 };
 
@@ -162,6 +175,11 @@ void InternalFormatStrings(FStringBuilderBase* OutString0,
 	const FFormatArgList& ArgList);
 } // namespace Private
 
+/**
+ * FormatString() provides printf-like functionality, except dependencies are automatically extracted from the list of Args
+ * In addition, a single '%' is used for all format specifiers, since the Args list is strongly typed
+ * Not all printf arg types are currently supported, but more can be added by extending FFormatArgVariant to support more types
+ */
 template<typename FormatType, typename... Types>
 void FormatString(FStringBuilderBase& OutString, FEmitShaderDependencies& OutDependencies, const FormatType& Format, Types... Args)
 {
@@ -186,6 +204,10 @@ enum class EEmitScopeState : uint8
 	Dead,
 };
 
+/**
+ * FEmitScopes mostly track FScopes, except they hold transient state used while emitting HLSL
+ * The exception is FFunction, which allows FScopes to be dynamically injected into the hierarchy, which will result in FEmitScopes matching this dynamic hierarchy
+ */
 class FEmitScope
 {
 public:
@@ -261,6 +283,10 @@ struct FTargetParameters
 	const ITargetPlatform* TargetPlatform = nullptr;
 };
 
+/**
+ * Used to associate owner objects (such as UMaterialExpression) with the FExpressions used to generate their input
+ * This is currently only used to color pins/wires based on types in the material editor
+ */
 struct FConnectionKey
 {
 	FConnectionKey() = default;
@@ -308,6 +334,7 @@ public:
 		return InternalError(FStringView(String.ToString(), String.Len()));
 	}
 
+	/** Finds or creates arbitrary data that's associated with this context. This can be used to provide custom inputs/outputs to FExpression/FStatement implementations */
 	template<typename T>
 	T& AcquireData()
 	{
@@ -373,6 +400,14 @@ public:
 
 	FEmitShaderExpression* InternalEmitExpression(FEmitScope& Scope, TArrayView<FEmitShaderNode*> Dependencies, bool bInline, const Shader::FType& Type, FStringView Code);
 
+	/**
+	 * Generate a snippit of HLSL code, strings with identical contents will be deduplicated, and their dependencies will be merged
+	 * @param Scope the scope where the code should be included
+	 * @param Dependencies a list of explicit dependencies, dependencies are automatically extracted from the given vargs, so this is typically not needed
+	 * @param Type the type of the HLSL expression
+	 * @param Format printf-style format string, except all format specifiers use a single '%' character with no additional characters
+	 * @param Args variable list of args to fill in the format string, 'FEmitShaderExpression*' is directly supported here to inject the expression's value, as well as include it as a dependency
+	 */
 	template<typename FormatType, typename... Types>
 	FEmitShaderExpression* EmitExpressionWithDependencies(FEmitScope& Scope, TArrayView<FEmitShaderNode*> Dependencies, const Shader::FType& Type, const FormatType& Format, Types... Args)
 	{
@@ -560,7 +595,12 @@ public:
 	int32 NumExpressionLocals = 0;
 	int32 NumExpressionLocalPHIs = 0;
 
-	// TODO - Material values required for preshaders, need to decouple preshaders from material system
+	/**
+	 * TODO - Material values required for preshaders, need to decouple preshaders from material system
+	 * Current plan is to implement support for externally registered preshader opcodes, and move the preshader opcodes related to material parameters into a separate module.
+	 * This would require a runtime parallel to CustomDataMap, in order to pass the relevant material data to the preshader VM in a generic way.
+	 * Would also need some generic interface between the preshaders generated here, and the preshaders stored in FMaterialCompilationOutput
+	 */
 	const FMaterial* Material = nullptr;
 	FMaterialCompilationOutput* MaterialCompilationOutput = nullptr;
 	uint32 UniformPreshaderOffset = 0u;
diff --git a/Engine/Source/Runtime/Engine/Public/HLSLTree/HLSLTreeTypes.h b/Engine/Source/Runtime/Engine/Public/HLSLTree/HLSLTreeTypes.h
--- a/Engine/Source/Runtime/Engine/Public/HLSLTree/HLSLTreeTypes.h
+++ b/Engine/Source/Runtime/Engine/Public/HLSLTree/HLSLTreeTypes.h
@@ -142,12 +142,22 @@ struct FSwizzleParameters
 	FSwizzleParameters() : NumComponents(0), bHasSwizzle(false) { SwizzleComponentIndex[0] = SwizzleComponentIndex[1] = SwizzleComponentIndex[2] = SwizzleComponentIndex[3] = INDEX_NONE; }
 	explicit FSwizzleParameters(int8 IndexR, int8 IndexG = INDEX_NONE, int8 IndexB = INDEX_NONE, int8 IndexA = INDEX_NONE);
 
-	int32 GetSwizzleComponentIndex(int32 Index) const
+	inline int32 GetSwizzleComponentIndex(int32 Index) const
 	{
 		const int32 ComponentIndex = (NumComponents == 1) ? 0 : Index;
 		return SwizzleComponentIndex[ComponentIndex];
 	}
 
+	inline int32 GetNumInputComponents() const
+	{
+		int32 MaxComponentIndex = INDEX_NONE;
+		for (int32 i = 0; i < NumComponents; ++i)
+		{
+			MaxComponentIndex = FMath::Max<int32>(MaxComponentIndex, SwizzleComponentIndex[i]);
+		}
+		return (MaxComponentIndex != INDEX_NONE) ? (MaxComponentIndex + 1) : 0;
+	}
+
 	int8 SwizzleComponentIndex[4];
 	int8 NumComponents;
 	bool bHasSwizzle;
diff --git a/Engine/Source/Runtime/Engine/Public/Shader/ShaderTypes.h b/Engine/Source/Runtime/Engine/Public/Shader/ShaderTypes.h
--- a/Engine/Source/Runtime/Engine/Public/Shader/ShaderTypes.h
+++ b/Engine/Source/Runtime/Engine/Public/Shader/ShaderTypes.h
@@ -10,9 +10,7 @@
 class FMemStackBase;
 class UTexture;
 
-namespace UE
-{
-namespace Shader
+namespace UE::Shader
 {
 
 struct FStructType;
@@ -305,7 +303,10 @@ public:
 	void EmitDeclarationsCode(FStringBuilderBase& OutCode) const;
 
 	const FStructType* NewType(const FStructTypeInitializer& Initializer);
+
+	/** An external struct type is opaque is far as the system is concerned, it's expected to be defined by the engine's HLSL code */
 	const FStructType* NewExternalType(FStringView Name);
+	
 	const FStructType* FindType(uint64 Hash) const;
 
 private:
@@ -339,6 +340,7 @@ struct FMemoryImageValue
 	uint32 Size;
 };
 
+/** The value of a single component within FValue.  The actual type of the component is only known by FValue::Type */
 union FValueComponent
 {
 	FValueComponent() : Packed(0u) {}
@@ -360,6 +362,10 @@ union FValueComponent
 };
 static_assert(sizeof(FValueComponent) == sizeof(uint64), "bad packing");
 
+/**
+ * Can store a numeric/struct shader value of any type
+ * Values are represented as a flat list of components
+ */
 struct FValue
 {
 	FValue() = default;
@@ -494,7 +500,7 @@ struct FValue
 	inline FValueComponent GetComponent(int32 Index) const
 	{
 		// Return scalar value for any request to xyzw
-		const int32 ComponentIndex = (GetNumComponents() == 1 && Index >= 0 && Index < 4) ? 0 : Index;
+		const int32 ComponentIndex = (Type.IsNumericScalar() && Index >= 0 && Index < 4) ? 0 : Index;
 		if (Component.IsValidIndex(ComponentIndex))
 		{
 			return Component[ComponentIndex];
@@ -567,8 +573,8 @@ ENGINE_API FValue Cross(const FValue& Lhs, const FValue& Rhs);
 ENGINE_API FValue Append(const FValue& Lhs, const FValue& Rhs);
 
 ENGINE_API FValue Cast(const FValue& Value, EValueType Type);
-}
-}
+
+} // namespace UE::Shader
 
 DECLARE_INTRINSIC_TYPE_LAYOUT(UE::Shader::EValueType);
 DECLARE_INTRINSIC_TYPE_LAYOUT(UE::Shader::EValueComponentType);
