Texture: Fixed streaming mips to make IsBulkDataLoaded() return false when the mips are not loaded

#rnx

diff --git a/Engine/Source/Runtime/Engine/Private/Texture2D.cpp b/Engine/Source/Runtime/Engine/Private/Texture2D.cpp
--- a/Engine/Source/Runtime/Engine/Private/Texture2D.cpp
+++ b/Engine/Source/Runtime/Engine/Private/Texture2D.cpp
@@ -148,6 +148,14 @@ void FTexture2DMipMap::Serialize(FArchive& Ar, UObject* Owner, int32 MipIdx)
 	BulkData.Serialize(Ar, Owner, MipIdx, false);
 #endif
 
+	// Streaming mips are saved with a size of 0 because they are stored separately.
+	// IsBulkDataLoaded() returns true for this empty bulk data. Remove the empty
+	// bulk data to allow unloaded streaming mips to be detected.
+	if (Ar.IsLoading() && BulkData.GetBulkDataSize() == 0)
+	{
+		BulkData.RemoveBulkData();
+	}
+
 	Ar << SizeX;
 	Ar << SizeY;
 	Ar << SizeZ;
diff --git a/Engine/Source/Runtime/Engine/Private/VT/VirtualTextureBuiltData.cpp b/Engine/Source/Runtime/Engine/Private/VT/VirtualTextureBuiltData.cpp
--- a/Engine/Source/Runtime/Engine/Private/VT/VirtualTextureBuiltData.cpp
+++ b/Engine/Source/Runtime/Engine/Private/VT/VirtualTextureBuiltData.cpp
@@ -375,6 +375,14 @@ void FVirtualTextureBuiltData::Serialize(FArchive& Ar, UObject* Owner, int32 Fir
 
 		Chunk.BulkData.Serialize(Ar, Owner, SerialzeChunkId, false);
 
+		// Streaming chunks are saved with a size of 0 because they are stored separately.
+		// IsBulkDataLoaded() returns true for this empty bulk data. Remove the empty
+		// bulk data to allow unloaded streaming chunks to be detected.
+		if (Ar.IsLoading() && Chunk.BulkData.GetBulkDataSize() == 0)
+		{
+			Chunk.BulkData.RemoveBulkData();
+		}
+
 #if WITH_EDITORONLY_DATA
 		if (!bCooked)
 		{
