Replaced uses of _SV with TEXTVIEW

#rnx

diff --git a/Engine/Source/Developer/IoStoreUtilities/Private/ZenStoreHttpClient.cpp b/Engine/Source/Developer/IoStoreUtilities/Private/ZenStoreHttpClient.cpp
--- a/Engine/Source/Developer/IoStoreUtilities/Private/ZenStoreHttpClient.cpp
+++ b/Engine/Source/Developer/IoStoreUtilities/Private/ZenStoreHttpClient.cpp
@@ -451,11 +451,11 @@ TIoStatusOr<FIoBuffer> FZenStoreHttpClient::ReadOpLogUri(FStringBuilderBase& Chu
 	{
 		if (bHaveQuery)
 		{
-			ChunkUri.Append(TEXT("&"_WSV));
+			ChunkUri.AppendChar(TEXT('&'));
 		}
 		else
 		{
-			ChunkUri.Append(TEXT("?"_WSV));
+			ChunkUri.AppendChar(TEXT('?'));
 			bHaveQuery = true;
 		}
 	};
diff --git a/Engine/Source/Developer/TraceInsights/Private/Insights/MemoryProfiler/ViewModels/MemAllocTable.cpp b/Engine/Source/Developer/TraceInsights/Private/Insights/MemoryProfiler/ViewModels/MemAllocTable.cpp
--- a/Engine/Source/Developer/TraceInsights/Private/Insights/MemoryProfiler/ViewModels/MemAllocTable.cpp
+++ b/Engine/Source/Developer/TraceInsights/Private/Insights/MemoryProfiler/ViewModels/MemAllocTable.cpp
@@ -803,11 +803,11 @@ void FMemAllocTable::AddDefaultColumns()
 							// Ignore symbols by function prefix.
 							FStringView IgnoreSymbolsByFunctionName[] =
 							{
-								TEXT("FMemory::"_SV),
-								TEXT("FMallocWrapper::"_SV),
-								TEXT("FMallocPoisonProxy::"_SV),
-								TEXT("Malloc"_SV),
-								TEXT("Realloc"_SV),
+								TEXTVIEW("FMemory::"),
+								TEXTVIEW("FMallocWrapper::"),
+								TEXTVIEW("FMallocPoisonProxy::"),
+								TEXTVIEW("Malloc"),
+								TEXTVIEW("Realloc"),
 							};
 							for (uint32 StringIndex = 0; StringIndex < UE_ARRAY_COUNT(IgnoreSymbolsByFunctionName); ++StringIndex)
 							{
@@ -824,7 +824,7 @@ void FMemAllocTable::AddDefaultColumns()
 								// Ignore symbols by file, specified as RegexPattern strings.
 								FStringView IgnoreSymbolsByFilePath[] =
 								{
-									TEXT(".*/Containers/.*"_SV),
+									TEXTVIEW(".*/Containers/.*"),
 								};
 								for (uint32 StringIndex = 0; StringIndex < UE_ARRAY_COUNT(IgnoreSymbolsByFilePath); ++StringIndex)
 								{
diff --git a/Engine/Source/Developer/TraceServices/Private/Analyzers/GpuProfilerTraceAnalysis.cpp b/Engine/Source/Developer/TraceServices/Private/Analyzers/GpuProfilerTraceAnalysis.cpp
--- a/Engine/Source/Developer/TraceServices/Private/Analyzers/GpuProfilerTraceAnalysis.cpp
+++ b/Engine/Source/Developer/TraceServices/Private/Analyzers/GpuProfilerTraceAnalysis.cpp
@@ -93,7 +93,7 @@ bool FGpuProfilerAnalyzer::OnEvent(uint16 RouteId, EStyle Style, const FOnEventC
 				else
 				{
 					FTimingProfilerEvent Event;
-					Event.TimerIndex = TimingProfilerProvider.AddGpuTimer(TEXT("<unknown>"_SV));
+					Event.TimerIndex = TimingProfilerProvider.AddGpuTimer(TEXTVIEW("<unknown>"));
 					EventTypeMap.Add(EventType, Event.TimerIndex);
 					ThisTimeline.AppendBeginEvent(LastTime, Event);
 				}
diff --git a/Engine/Source/Developer/Zen/Internal/ZenServerInterface.cpp b/Engine/Source/Developer/Zen/Internal/ZenServerInterface.cpp
--- a/Engine/Source/Developer/Zen/Internal/ZenServerInterface.cpp
+++ b/Engine/Source/Developer/Zen/Internal/ZenServerInterface.cpp
@@ -1087,7 +1087,7 @@ FZenServiceInstance::GetStats(FZenStats& Stats)
 				Request.Reset();
 
 				TArray64<uint8> GetBuffer;
-				FZenHttpRequest::Result Result = Request.PerformBlockingDownload(TEXT("/stats/z$"_SV), &GetBuffer, Zen::EContentType::CbObject);
+				FZenHttpRequest::Result Result = Request.PerformBlockingDownload(TEXTVIEW("/stats/z$"), &GetBuffer, Zen::EContentType::CbObject);
 
 				FZenStats Stats;
 
diff --git a/Engine/Source/Editor/UnrealEd/Private/DerivedDataBuildLocalExecutor.cpp b/Engine/Source/Editor/UnrealEd/Private/DerivedDataBuildLocalExecutor.cpp
--- a/Engine/Source/Editor/UnrealEd/Private/DerivedDataBuildLocalExecutor.cpp
+++ b/Engine/Source/Editor/UnrealEd/Private/DerivedDataBuildLocalExecutor.cpp
@@ -269,7 +269,7 @@ public:
 
 	TConstArrayView<FStringView> GetHostPlatforms() const final
 	{
-		static constexpr FStringView HostPlatforms[]{TEXT("Win64"_SV)};
+		static constexpr FStringView HostPlatforms[]{TEXTVIEW("Win64")};
 		return HostPlatforms;
 	}
 
diff --git a/Engine/Source/Editor/UnrealEd/Private/EditorDomain/EditorDomainSave.cpp b/Engine/Source/Editor/UnrealEd/Private/EditorDomain/EditorDomainSave.cpp
--- a/Engine/Source/Editor/UnrealEd/Private/EditorDomain/EditorDomainSave.cpp
+++ b/Engine/Source/Editor/UnrealEd/Private/EditorDomain/EditorDomainSave.cpp
@@ -1391,7 +1391,7 @@ void FEditorDomainSaveClient::ConnectionWarning(FStringView Message, FStringView
 		{
 			LastConnectionWarningTimeSeconds = CurrentTime;
 		}
-		FStringView DefaultConsequence = TEXT("EditorDomainSave is currently unavailable"_SV);
+		FStringView DefaultConsequence = TEXTVIEW("EditorDomainSave is currently unavailable");
 		if (Consequence.Len() == 0)
 		{
 			Consequence = DefaultConsequence;
diff --git a/Engine/Source/Runtime/Engine/Private/AudioDerivedData.cpp b/Engine/Source/Runtime/Engine/Private/AudioDerivedData.cpp
--- a/Engine/Source/Runtime/Engine/Private/AudioDerivedData.cpp
+++ b/Engine/Source/Runtime/Engine/Private/AudioDerivedData.cpp
@@ -847,7 +847,7 @@ static void BeginLoadDerivedChunks(TIndirectArray<FStreamedAudioChunk>& Chunks,
 		const FStreamedAudioChunk& Chunk = Chunks[ChunkIndex];
 		if (!Chunk.DerivedDataKey.IsEmpty())
 		{
-			OutHandles[ChunkIndex] = DDC.GetAsynchronous(*Chunk.DerivedDataKey, TEXT("Unknown SoundWave"_SV));
+			OutHandles[ChunkIndex] = DDC.GetAsynchronous(*Chunk.DerivedDataKey, TEXTVIEW("Unknown SoundWave"));
 		}
 	}
 }
@@ -992,14 +992,14 @@ int32 FStreamedAudioPlatformData::GetChunkFromDDC(int32 ChunkIndex, uint8** OutC
 	{
 		if (bMakeSureChunkIsLoaded)
 		{
-			if (DDC.GetSynchronous(*Chunk.DerivedDataKey, TempData, TEXT("Unknown SoundWave"_SV)))
+			if (DDC.GetSynchronous(*Chunk.DerivedDataKey, TempData, TEXTVIEW("Unknown SoundWave")))
 			{
 				ChunkDataSize = DeserializeChunkFromDDC(TempData, Chunk, ChunkIndex, OutChunkData);
 			}
 		}
 		else
 		{
-			AsyncHandle = DDC.GetAsynchronous(*Chunk.DerivedDataKey, TEXT("Unknown SoundWave"_SV));
+			AsyncHandle = DDC.GetAsynchronous(*Chunk.DerivedDataKey, TEXTVIEW("Unknown SoundWave"));
 		}
 	}
 	else if (Chunk.bLoadedFromCookedPackage)
@@ -1088,7 +1088,7 @@ bool FStreamedAudioPlatformData::AreDerivedChunksAvailable(FStringView InContext
 
 bool FStreamedAudioPlatformData::AreDerivedChunksAvailable() const
 {
-	return AreDerivedChunksAvailable(TEXT("DerivedAudioChunks"_SV));
+	return AreDerivedChunksAvailable(TEXTVIEW("DerivedAudioChunks"));
 }
 
 #endif // #if WITH_EDITORONLY_DATA
diff --git a/Engine/Source/Runtime/Engine/Private/AudioStreaming.cpp b/Engine/Source/Runtime/Engine/Private/AudioStreaming.cpp
--- a/Engine/Source/Runtime/Engine/Private/AudioStreaming.cpp
+++ b/Engine/Source/Runtime/Engine/Private/AudioStreaming.cpp
@@ -68,7 +68,7 @@ void FAsyncStreamDerivedChunkWorker::DoWork()
 
 	TArray<uint8> DerivedChunkData;
 
-	if (GetDerivedDataCacheRef().GetSynchronous(*DerivedDataKey, DerivedChunkData, TEXT("Unknown Audio"_SV)))
+	if (GetDerivedDataCacheRef().GetSynchronous(*DerivedDataKey, DerivedChunkData, TEXTVIEW("Unknown Audio")))
 	{
 		FMemoryReader Ar(DerivedChunkData, true);
 		int32 ChunkSize = 0;
diff --git a/Engine/Source/Runtime/Engine/Private/TextureDerivedData.cpp b/Engine/Source/Runtime/Engine/Private/TextureDerivedData.cpp
--- a/Engine/Source/Runtime/Engine/Private/TextureDerivedData.cpp
+++ b/Engine/Source/Runtime/Engine/Private/TextureDerivedData.cpp
@@ -2493,7 +2493,7 @@ FString FTexturePlatformData::GetDerivedDataMipKeyString(int32 MipIndex, const F
 
 UE::DerivedData::FValueId FTexturePlatformData::MakeMipId(int32 MipIndex)
 {
-	return UE::DerivedData::FValueId::FromName(WriteToString<16>(TEXT("Mip"_SV), MipIndex));
+	return UE::DerivedData::FValueId::FromName(WriteToString<16>(TEXTVIEW("Mip"), MipIndex));
 }
 
 #endif // WITH_EDITORONLY_DATA
diff --git a/Engine/Source/Runtime/Engine/Private/TextureDerivedDataBuildUtils.cpp b/Engine/Source/Runtime/Engine/Private/TextureDerivedDataBuildUtils.cpp
--- a/Engine/Source/Runtime/Engine/Private/TextureDerivedDataBuildUtils.cpp
+++ b/Engine/Source/Runtime/Engine/Private/TextureDerivedDataBuildUtils.cpp
@@ -126,7 +126,7 @@ static void WriteBuildSettings(FCbWriter& Writer, const FTextureBuildSettings& B
 	WriteCbFieldWithDefault(Writer, "DiffuseConvolveMipLevel", BuildSettings.DiffuseConvolveMipLevel, DefaultSettings.DiffuseConvolveMipLevel);
 	WriteCbFieldWithDefault(Writer, "SharpenMipKernelSize", BuildSettings.SharpenMipKernelSize, DefaultSettings.SharpenMipKernelSize);
 	WriteCbFieldWithDefault(Writer, "MaxTextureResolution", BuildSettings.MaxTextureResolution, DefaultSettings.MaxTextureResolution);
-	WriteCbFieldWithDefault(Writer, "TextureFormatName", WriteToString<64>(BuildSettings.TextureFormatName).ToView(), TEXT(""_SV));
+	WriteCbFieldWithDefault(Writer, "TextureFormatName", WriteToString<64>(BuildSettings.TextureFormatName).ToView(), TEXTVIEW(""));
 	WriteCbFieldWithDefault(Writer, "bHDRSource", BuildSettings.bHDRSource, DefaultSettings.bHDRSource);
 	WriteCbFieldWithDefault(Writer, "MipGenSettings", BuildSettings.MipGenSettings, DefaultSettings.MipGenSettings);
 	WriteCbFieldWithDefault<bool>(Writer, "bCubemap", BuildSettings.bCubemap, DefaultSettings.bCubemap);
diff --git a/Engine/Source/Runtime/RenderCore/Private/ShaderCodeLibrary.cpp b/Engine/Source/Runtime/RenderCore/Private/ShaderCodeLibrary.cpp
--- a/Engine/Source/Runtime/RenderCore/Private/ShaderCodeLibrary.cpp
+++ b/Engine/Source/Runtime/RenderCore/Private/ShaderCodeLibrary.cpp
@@ -403,7 +403,7 @@ void FCompactFullName::ParseFromString(const FStringView& InSrc)
 	// do not split by '/' as this splits the original FName into per-path components
 	UE::String::ParseTokensMultiple(InSrc.TrimStartAndEnd(), {TEXT(' '), TEXT('.'), /*TEXT('/'),*/ TEXT('\t')},
 		[&Fields](FStringView Field) { if (!Field.IsEmpty()) { Fields.Add(Field); } });
-	if (Fields.Num() == 1 && Fields[0] == TEXT("empty"_SV))
+	if (Fields.Num() == 1 && Fields[0] == TEXTVIEW("empty"))
 	{
 		ObjectClassAndPath.Empty();
 	}
