Texture: Restored deprecated members with deprecation warnings disabled around defaulted special member functions

#rnx

diff --git a/Engine/Source/Runtime/Engine/Public/TextureResource.h b/Engine/Source/Runtime/Engine/Public/TextureResource.h
--- a/Engine/Source/Runtime/Engine/Public/TextureResource.h
+++ b/Engine/Source/Runtime/Engine/Public/TextureResource.h
@@ -44,11 +44,11 @@ struct IPooledRenderTarget;
 struct FTexture2DMipMap
 {
 	/** Width of the mip-map. */
-	int32 SizeX;
+	int32 SizeX = 0;
 	/** Height of the mip-map. */
-	int32 SizeY;
+	int32 SizeY = 0;
 	/** Depth of the mip-map. */
-	int32 SizeZ;
+	int32 SizeZ = 0;
 
 	/** Reference to the data for the mip if it can be streamed. */
 	UE::FDerivedData DerivedData;
@@ -56,13 +56,14 @@ struct FTexture2DMipMap
 	/** Stores the data for the mip when it is loaded. */
 	FByteBulkData BulkData;
 
-	/** Default constructor. */
-	FTexture2DMipMap()
-		: SizeX(0)
-		, SizeY(0)
-		, SizeZ(0)
-	{
-	}
+	PRAGMA_DISABLE_DEPRECATION_WARNINGS
+	FTexture2DMipMap() = default;
+	FTexture2DMipMap(FTexture2DMipMap&&) = default;
+	FTexture2DMipMap(const FTexture2DMipMap&) = default;
+	FTexture2DMipMap& operator=(FTexture2DMipMap&&) = default;
+	FTexture2DMipMap& operator=(const FTexture2DMipMap&) = default;
+	~FTexture2DMipMap() = default;
+	PRAGMA_ENABLE_DEPRECATION_WARNINGS
 
 	/** Serialization. */
 	ENGINE_API void Serialize(FArchive& Ar, UObject* Owner, int32 MipIndex);
@@ -71,9 +72,17 @@ struct FTexture2DMipMap
 	/** The file region type appropriate for the pixel format of this mip-map. */
 	EFileRegionType FileRegionType = EFileRegionType::None;
 
+	UE_DEPRECATED(5.1, "Use DerivedData.HasData().")
+	bool bPagedToDerivedData = false;
+
 	/** Whether this mip-map is stored in the derived data cache. */
 	inline bool IsPagedToDerivedData() const { return DerivedData.HasData(); }
 
+	UE_DEPRECATED(5.1, "Setting DerivedData is sufficient to control this state.")
+	inline void SetPagedToDerivedData(bool InValue)
+	{
+	}
+
 	/** Place mip-map data in the derived data cache associated with the provided key. */
 	int64 StoreInDerivedDataCache(FStringView Key, FStringView Name, bool bReplaceExisting);
 #endif // #if WITH_EDITORONLY_DATA
