Moving custom property conditions out of the net driver
Stored per object, and can be set from any game code, not just PreReplication

diff --git a/Engine/Source/Runtime/CoreUObject/Private/UObject/Obj.cpp b/Engine/Source/Runtime/CoreUObject/Private/UObject/Obj.cpp
--- a/Engine/Source/Runtime/CoreUObject/Private/UObject/Obj.cpp
+++ b/Engine/Source/Runtime/CoreUObject/Private/UObject/Obj.cpp
@@ -4855,6 +4855,12 @@ void UObject::GetLifetimeReplicatedProps( TArray< class FLifetimeProperty > & Ou
 
 }
 
+/** Called when this object begins replicating to initialize the state of custom property conditions */
+void UObject::GetReplicatedCustomConditionState(FCustomPropertyConditionState& OutActiveState) const
+{
+
+}
+
 /** Called right before receiving a bunch */
 void UObject::PreNetReceive()
 {
diff --git a/Engine/Source/Runtime/CoreUObject/Public/UObject/CoreNet.h b/Engine/Source/Runtime/CoreUObject/Public/UObject/CoreNet.h
--- a/Engine/Source/Runtime/CoreUObject/Public/UObject/CoreNet.h
+++ b/Engine/Source/Runtime/CoreUObject/Public/UObject/CoreNet.h
@@ -506,10 +506,7 @@ public:
 	IRepChangedPropertyTracker() { }
 	virtual ~IRepChangedPropertyTracker() { }
 
-	virtual void SetCustomIsActiveOverride(
-		UObject* OwningObject,
-		const uint16 RepIndex,
-		const bool bIsActive) = 0;
+	virtual void SetCustomIsActiveOverride(UObject* OwningObject, const uint16 RepIndex, const bool bIsActive) = 0;
 
 	UE_DEPRECATED(5.0, "Please use UReplaySubsystem::SetExternalDataForObject instead.")
 	virtual void SetExternalData(const uint8* Src, const int32 NumBits) = 0;
@@ -521,6 +518,38 @@ public:
 	virtual void CountBytes(FArchive& Ar) const {};
 };
 
+class FCustomPropertyConditionState
+{
+public:
+	FCustomPropertyConditionState() = delete;
+	FCustomPropertyConditionState(int32 NumProperties)
+	{
+		CurrentState.Init(true, NumProperties);
+	}
+
+	void SetActiveState(const uint16 RepIndex, const bool bIsActive)
+	{
+		CurrentState[RepIndex] = bIsActive;
+	}
+
+	bool GetActiveState(const uint16 RepIndex) const
+	{
+		return CurrentState[RepIndex];
+	}
+
+	int32 GetNumProperties() const
+	{
+		return CurrentState.Num();
+	}
+
+	void CountBytes(FArchive& Ar) const
+	{
+		CurrentState.CountBytes(Ar);
+	}
+
+private:
+	TBitArray<> CurrentState;
+};
 
 /**
  * FNetDeltaSerializeInfo
diff --git a/Engine/Source/Runtime/CoreUObject/Public/UObject/Object.h b/Engine/Source/Runtime/CoreUObject/Public/UObject/Object.h
--- a/Engine/Source/Runtime/CoreUObject/Public/UObject/Object.h
+++ b/Engine/Source/Runtime/CoreUObject/Public/UObject/Object.h
@@ -16,6 +16,7 @@
 
 struct FAssetData;
 class FConfigCacheIni;
+class FCustomPropertyConditionState;
 class FEditPropertyChain;
 class FObjectPostSaveContext;
 class FObjectPostSaveRootContext;
@@ -907,6 +908,9 @@ public:
 	/** Returns properties that are replicated for the lifetime of the actor channel */
 	virtual void GetLifetimeReplicatedProps( TArray< class FLifetimeProperty > & OutLifetimeProps ) const;
 
+	/** Called when this object begins replicating to initialize the state of custom property conditions */
+	virtual void GetReplicatedCustomConditionState(FCustomPropertyConditionState& OutActiveState) const;
+
 	/** IsNameStableForNetworking means an object can be referred to its path name (relative to outer) over the network */
 	virtual bool IsNameStableForNetworking() const;
 
diff --git a/Engine/Source/Runtime/Engine/Classes/Engine/EngineTypes.h b/Engine/Source/Runtime/Engine/Classes/Engine/EngineTypes.h
--- a/Engine/Source/Runtime/Engine/Classes/Engine/EngineTypes.h
+++ b/Engine/Source/Runtime/Engine/Classes/Engine/EngineTypes.h
@@ -3088,7 +3088,8 @@ struct FReplicationFlags
 			uint32 bSerializePropertyNames : 1;
 			/** True if a subclass of UActorChannel needs custom subobject replication */
 			uint32 bUseCustomSubobjectReplication : 1;
-
+			/** True if this actor is replicating on a replay connection on a game client. */
+			uint32 bClientReplay : 1;
 		};
 
 		uint32	Value;
diff --git a/Engine/Source/Runtime/Engine/Classes/Engine/NetDriver.h b/Engine/Source/Runtime/Engine/Classes/Engine/NetDriver.h
--- a/Engine/Source/Runtime/Engine/Classes/Engine/NetDriver.h
+++ b/Engine/Source/Runtime/Engine/Classes/Engine/NetDriver.h
@@ -701,7 +701,9 @@ private:
 public:
 
 	/** Destructor */
+	PRAGMA_DISABLE_DEPRECATION_WARNINGS
 	ENGINE_API virtual ~UNetDriver() {};
+	PRAGMA_ENABLE_DEPRECATION_WARNINGS
 
 	/** Used to specify the class to use for connections */
 	UPROPERTY(Config)
@@ -1065,7 +1067,7 @@ public:
 	 */
 	TMap<FName, FName>	RenamedStartupActors;
 
-	class FRepChangedPropertyTrackerWrapper
+	class UE_DEPRECATED(5.1, "No longer used.") FRepChangedPropertyTrackerWrapper
 	{
 	public:
 		FRepChangedPropertyTrackerWrapper(UObject* Obj, const TSharedPtr<FRepChangedPropertyTracker>& InRepChangedPropertyTracker) : RepChangedPropertyTracker(InRepChangedPropertyTracker), WeakObjectPtr(Obj) {}
@@ -1088,8 +1090,13 @@ public:
 	private:
 		TWeakObjectPtr<UObject> WeakObjectPtr;
 	};
+
+	PRAGMA_DISABLE_DEPRECATION_WARNINGS
 	/** Maps FRepChangedPropertyTracker to active objects that are replicating properties */
+	UE_DEPRECATED(5.1, "Property trackers have been moved to the NetCore module")
 	TMap<UObject*, FRepChangedPropertyTrackerWrapper>	RepChangedPropertyTrackerMap;
+	PRAGMA_ENABLE_DEPRECATION_WARNINGS
+
 	/** Used to invalidate properties marked "unchanged" in FRepChangedPropertyTracker's */
 	uint32																		ReplicationFrame;
 
@@ -1443,6 +1450,7 @@ public:
 	bool HandleDumpSubObjectsCommand(const TCHAR* Cmd, FOutputDevice& Ar);
 	bool HandleDumpRepLayoutFlagsCommand(const TCHAR* Cmd, FOutputDevice& Ar);
 	bool HandlePushModelMemCommand(const TCHAR* Cmd, FOutputDevice& Ar);
+	bool HandlePropertyConditionsMemCommand(const TCHAR* Cmd, FOutputDevice& Ar);
 #endif
 
 	void HandlePacketLossBurstCommand( int32 DurationInMilliseconds );
@@ -1478,6 +1486,8 @@ public:
 	/** Called when a spawned actor is destroyed. */
 	ENGINE_API virtual void NotifyActorDestroyed( AActor* Actor, bool IsSeamlessTravel=false );
 
+	void NotifySubObjectDestroyed(UObject* SubObject);
+
 	/** Called when an actor is renamed. */
 	ENGINE_API virtual void NotifyActorRenamed(AActor* Actor, FName PreviousName);
 
@@ -1555,6 +1565,9 @@ public:
 	*/
 	TSharedPtr<FRepChangedPropertyTracker> FindOrCreateRepChangedPropertyTracker(UObject *Obj);
 
+	/** Finds a FRepChangedPropertyTracker associated with an object. */
+	TSharedPtr<FRepChangedPropertyTracker> FindRepChangedPropertyTracker(UObject* Obj);
+
 	/** Returns true if the client should destroy immediately any actor that becomes torn-off */
 	ENGINE_API virtual bool ShouldClientDestroyTearOffActors() const { return false; }
 
diff --git a/Engine/Source/Runtime/Engine/Classes/GameFramework/Actor.h b/Engine/Source/Runtime/Engine/Classes/GameFramework/Actor.h
--- a/Engine/Source/Runtime/Engine/Classes/GameFramework/Actor.h
+++ b/Engine/Source/Runtime/Engine/Classes/GameFramework/Actor.h
@@ -182,7 +182,10 @@ private:
 
 public:
 	/** Returns the properties used for network replication, this needs to be overridden by all actor classes with native replicated properties */
-	void GetLifetimeReplicatedProps(TArray<FLifetimeProperty>& OutLifetimeProps) const override;
+	virtual void GetLifetimeReplicatedProps(TArray<FLifetimeProperty>& OutLifetimeProps) const override;
+
+	/** Called when this actor begins replicating to initialize the state of custom property conditions */
+	virtual void GetReplicatedCustomConditionState(FCustomPropertyConditionState& OutActiveState) const override;
 
 	/**
 	 * Primary Actor tick function, which calls TickActor().
diff --git a/Engine/Source/Runtime/Engine/Classes/GameFramework/Character.h b/Engine/Source/Runtime/Engine/Classes/GameFramework/Character.h
--- a/Engine/Source/Runtime/Engine/Classes/GameFramework/Character.h
+++ b/Engine/Source/Runtime/Engine/Classes/GameFramework/Character.h
@@ -220,7 +220,9 @@ public:
 	/** Default UObject constructor. */
 	ACharacter(const FObjectInitializer& ObjectInitializer = FObjectInitializer::Get());
 
-	void GetLifetimeReplicatedProps(TArray<FLifetimeProperty>& OutLifetimeProps) const override;
+	virtual void GetLifetimeReplicatedProps(TArray<FLifetimeProperty>& OutLifetimeProps) const override;
+
+	virtual void GetReplicatedCustomConditionState(FCustomPropertyConditionState& OutActiveState) const override;
 
 private:
 	/** The main skeletal mesh associated with this Character (optional sub-object). */
diff --git a/Engine/Source/Runtime/Engine/Private/Actor.cpp b/Engine/Source/Runtime/Engine/Private/Actor.cpp
--- a/Engine/Source/Runtime/Engine/Private/Actor.cpp
+++ b/Engine/Source/Runtime/Engine/Private/Actor.cpp
@@ -1431,10 +1431,10 @@ void AActor::PreReplication(IRepChangedPropertyTracker & ChangedPropertyTracker)
 
 	GatherCurrentMovement();
 
-	DOREPLIFETIME_ACTIVE_OVERRIDE(AActor, ReplicatedMovement, IsReplicatingMovement());
+	DOREPLIFETIME_ACTIVE_OVERRIDE_FAST(AActor, ReplicatedMovement, IsReplicatingMovement());
 
 	// Don't need to replicate AttachmentReplication if the root component replicates, because it already handles it.
-	DOREPLIFETIME_ACTIVE_OVERRIDE(AActor, AttachmentReplication, RootComponent && !RootComponent->GetIsReplicated());
+	DOREPLIFETIME_ACTIVE_OVERRIDE_FAST(AActor, AttachmentReplication, RootComponent && !RootComponent->GetIsReplicated());
 
 
 #if WITH_PUSH_MODEL
diff --git a/Engine/Source/Runtime/Engine/Private/ActorReplication.cpp b/Engine/Source/Runtime/Engine/Private/ActorReplication.cpp
--- a/Engine/Source/Runtime/Engine/Private/ActorReplication.cpp
+++ b/Engine/Source/Runtime/Engine/Private/ActorReplication.cpp
@@ -468,6 +468,14 @@ void AActor::GetLifetimeReplicatedProps( TArray< FLifetimeProperty > & OutLifeti
 	DOREPLIFETIME_WITH_PARAMS_FAST(AActor, ReplicatedMovement, ReplicatedMovementParams);
 }
 
+void AActor::GetReplicatedCustomConditionState(FCustomPropertyConditionState& OutActiveState) const
+{
+	Super::GetReplicatedCustomConditionState(OutActiveState);
+
+	DOREPCUSTOMCONDITION_ACTIVE_FAST(AActor, AttachmentReplication, IsReplicatingMovement());
+	DOREPCUSTOMCONDITION_ACTIVE_FAST(AActor, ReplicatedMovement, RootComponent && !RootComponent->GetIsReplicated());
+}
+
 bool AActor::ReplicateSubobjects(UActorChannel *Channel, FOutBunch *Bunch, FReplicationFlags *RepFlags)
 {
 	check(Channel);
diff --git a/Engine/Source/Runtime/Engine/Private/Character.cpp b/Engine/Source/Runtime/Engine/Private/Character.cpp
--- a/Engine/Source/Runtime/Engine/Private/Character.cpp
+++ b/Engine/Source/Runtime/Engine/Private/Character.cpp
@@ -1454,13 +1454,13 @@ void ACharacter::PreReplication( IRepChangedPropertyTracker & ChangedPropertyTra
 		RepRootMotion.Acceleration = CharacterMovement->GetCurrentAcceleration();
 		RepRootMotion.LinearVelocity = CharacterMovement->Velocity;
 
-		DOREPLIFETIME_ACTIVE_OVERRIDE( ACharacter, RepRootMotion, true );
+		DOREPLIFETIME_ACTIVE_OVERRIDE_FAST( ACharacter, RepRootMotion, true );
 	}
 	else
 	{
 		RepRootMotion.Clear();
 
-		DOREPLIFETIME_ACTIVE_OVERRIDE( ACharacter, RepRootMotion, false );
+		DOREPLIFETIME_ACTIVE_OVERRIDE_FAST( ACharacter, RepRootMotion, false );
 	}
 
 	bProxyIsJumpForceApplied = (JumpForceTimeRemaining > 0.0f);
@@ -1491,6 +1491,13 @@ void ACharacter::PreReplication( IRepChangedPropertyTracker & ChangedPropertyTra
 	}
 }
 
+void ACharacter::GetReplicatedCustomConditionState(FCustomPropertyConditionState& OutActiveState) const
+{
+	Super::GetReplicatedCustomConditionState(OutActiveState);
+
+	DOREPCUSTOMCONDITION_ACTIVE_FAST(ACharacter, RepRootMotion, CharacterMovement->CurrentRootMotion.HasActiveRootMotionSources() || IsPlayingNetworkedRootMotionMontage());
+}
+
 void ACharacter::PreReplicationForReplay(IRepChangedPropertyTracker& ChangedPropertyTracker)
 {
 	Super::PreReplicationForReplay(ChangedPropertyTracker);
diff --git a/Engine/Source/Runtime/Engine/Private/DataChannel.cpp b/Engine/Source/Runtime/Engine/Private/DataChannel.cpp
--- a/Engine/Source/Runtime/Engine/Private/DataChannel.cpp
+++ b/Engine/Source/Runtime/Engine/Private/DataChannel.cpp
@@ -2267,7 +2267,7 @@ void UActorChannel::DestroyActorAndComponents()
 
 			if (Connection != nullptr && Connection->Driver != nullptr)
 			{
-				Connection->Driver->RepChangedPropertyTrackerMap.Remove(SubObject);
+				Connection->Driver->NotifySubObjectDestroyed(SubObject);
 			}
 
 			Actor->OnSubobjectDestroyFromReplication(SubObject); //-V595
@@ -3369,6 +3369,7 @@ int64 UActorChannel::ReplicateActor()
 	RepFlags.bNetSimulated	= (Actor->GetRemoteRole() == ROLE_SimulatedProxy);
 	RepFlags.bRepPhysics	= Actor->GetReplicatedMovement().bRepPhysics;
 	RepFlags.bReplay		= bReplay;
+	RepFlags.bClientReplay	= ActorWorld->IsRecordingClientReplay();
 	RepFlags.bForceInitialDirty = Connection->IsForceInitialDirty();
 
 	UE_LOG(LogNetTraffic, Log, TEXT("Replicate %s, bNetInitial: %d, bNetOwner: %d"), *Actor->GetName(), RepFlags.bNetInitial, RepFlags.bNetOwner);
@@ -4216,7 +4217,7 @@ UObject* UActorChannel::ReadContentBlockHeader(FInBunch& Bunch, bool& bObjectDel
 
 			if (Connection != nullptr && Connection->Driver != nullptr)
 			{
-				Connection->Driver->RepChangedPropertyTrackerMap.Remove(SubObj);
+				Connection->Driver->NotifySubObjectDestroyed(SubObj);
 			}
 
 			Actor->OnSubobjectDestroyFromReplication(SubObj);
diff --git a/Engine/Source/Runtime/Engine/Private/DemoNetDriver.cpp b/Engine/Source/Runtime/Engine/Private/DemoNetDriver.cpp
--- a/Engine/Source/Runtime/Engine/Private/DemoNetDriver.cpp
+++ b/Engine/Source/Runtime/Engine/Private/DemoNetDriver.cpp
@@ -5002,8 +5002,6 @@ void UDemoNetDriver::NotifyActorDestroyed(AActor* Actor, bool IsSeamlessTravel)
 			// This is a stripped down version of UNetDriver::NotifyActorDestroy and UActorChannel::Close
 			// combined, and should be kept up to date with those methods.
 
-			// Remove the actor from the property tracker map
-			RepChangedPropertyTrackerMap.Remove(Actor);
 
 			if (UNetConnection* Connection = ClientConnections[0])
 			{
diff --git a/Engine/Source/Runtime/Engine/Private/LevelActor.cpp b/Engine/Source/Runtime/Engine/Private/LevelActor.cpp
--- a/Engine/Source/Runtime/Engine/Private/LevelActor.cpp
+++ b/Engine/Source/Runtime/Engine/Private/LevelActor.cpp
@@ -41,6 +41,7 @@
 #include "Engine/DemoNetDriver.h"
 #include "Engine/Player.h"
 #include "AssetRegistryModule.h"
+#include "Net/Core/PropertyConditions/PropertyConditions.h"
 
 #include "Components/BoxComponent.h"
 #include "GameFramework/MovementComponent.h"
@@ -942,6 +943,8 @@ bool UWorld::DestroyActor( AActor* ThisActor, bool bNetForce, bool bShouldModify
 		}
 	}
 
+	FNetPropertyConditionManager::Get().NotifyObjectDestroyed(ThisActor);
+
 	// Notify net drivers that this guy has been destroyed.
 	if (FWorldContext* Context = GEngine->GetWorldContextFromWorld(this))
 	{
diff --git a/Engine/Source/Runtime/Engine/Private/NetDriver.cpp b/Engine/Source/Runtime/Engine/Private/NetDriver.cpp
--- a/Engine/Source/Runtime/Engine/Private/NetDriver.cpp
+++ b/Engine/Source/Runtime/Engine/Private/NetDriver.cpp
@@ -58,6 +58,7 @@
 #include "NetworkingDistanceConstants.h"
 #include "Engine/ChildConnection.h"
 #include "Net/Core/Trace/NetTrace.h"
+#include "Net/Core/PropertyConditions/PropertyConditions.h"
 #include "Misc/ScopeExit.h"
 #include "Net/DataChannel.h"
 #include "GameFramework/PlayerState.h"
@@ -2636,7 +2637,6 @@ void UNetDriver::SetAnalyticsProvider(TSharedPtr<IAnalyticsProvider> InProvider)
 	}
 }
 
-
 void UNetDriver::FRepChangedPropertyTrackerWrapper::CountBytes(FArchive& Ar) const
 {
 	if (FRepChangedPropertyTracker const * const LocalTracker = RepChangedPropertyTracker.Get())
@@ -2710,6 +2710,7 @@ void UNetDriver::Serialize( FArchive& Ar )
 
 		GRANULAR_NETWORK_MEMORY_TRACKING_TRACK("RenamedStartupActors", RenamedStartupActors.CountBytes(Ar));
 
+		PRAGMA_DISABLE_DEPRECATION_WARNINGS
 		GRANULAR_NETWORK_MEMORY_TRACKING_TRACK("RepChangedPropertyTrackerMap",
 			RepChangedPropertyTrackerMap.CountBytes(Ar);
 
@@ -2718,6 +2719,7 @@ void UNetDriver::Serialize( FArchive& Ar )
 				RepChangedPropertyTrackerPair.Value.CountBytes(Ar);
 			}
 		);
+		PRAGMA_ENABLE_DEPRECATION_WARNINGS
 
 		GRANULAR_NETWORK_MEMORY_TRACKING_TRACK("RepLayoutMap",
 			RepLayoutMap.CountBytes(Ar);
@@ -3220,6 +3222,11 @@ bool UNetDriver::HandlePushModelMemCommand(const TCHAR* Cmd, FOutputDevice& Ar)
 	return true;
 }
 
+bool UNetDriver::HandlePropertyConditionsMemCommand(const TCHAR* Cmd, FOutputDevice& Ar)
+{
+	FNetPropertyConditionManager::Get().LogMemory(Ar);
+	return true;
+}
 #endif // !UE_BUILD_SHIPPING
 
 void UNetDriver::HandlePacketLossBurstCommand( int32 DurationInMilliseconds )
@@ -3350,6 +3357,11 @@ bool UNetDriver::Exec( UWorld* InWorld, const TCHAR* Cmd, FOutputDevice& Ar )
 		HandlePushModelMemCommand(Cmd, Ar);
 		return true;
 	}
+	else if (FParse::Command(&Cmd, TEXT("PROPERTYCONDITIONMEM")))
+	{
+		HandlePropertyConditionsMemCommand(Cmd, Ar);
+		return true;
+	}
 	else
 #endif // !UE_BUILD_SHIPPING
 	{
@@ -3407,9 +3419,6 @@ FActorDestructionInfo* UNetDriver::CreateDestructionInfo( UNetDriver* NetDriver,
 
 void UNetDriver::NotifyActorDestroyed( AActor* ThisActor, bool IsSeamlessTravel )
 {
-	// Remove the actor from the property tracker map
-	RepChangedPropertyTrackerMap.Remove(ThisActor);
-
 	const bool bIsServer = IsServer();
 	
 	if (bIsServer)
@@ -3469,6 +3478,10 @@ void UNetDriver::NotifyActorDestroyed( AActor* ThisActor, bool IsSeamlessTravel
 	RemoveNetworkActor( ThisActor );
 }
 
+void UNetDriver::NotifySubObjectDestroyed(UObject* SubObject)
+{
+	FNetPropertyConditionManager::Get().NotifyObjectDestroyed(SubObject);
+}
 
 void UNetDriver::RemoveNetworkActor(AActor* Actor)
 {
@@ -3769,6 +3782,7 @@ void UNetDriver::PostGarbageCollect()
 		}
 	}
 
+	PRAGMA_DISABLE_DEPRECATION_WARNINGS
 	for (auto It = RepChangedPropertyTrackerMap.CreateIterator(); It; ++It)
 	{
 		if (!It.Value().IsObjectValid())
@@ -3776,6 +3790,7 @@ void UNetDriver::PostGarbageCollect()
 			It.RemoveCurrent();
 		}
 	}
+	PRAGMA_ENABLE_DEPRECATION_WARNINGS
 }
 
 void UNetDriver::AddReferencedObjects(UObject* InThis, FReferenceCollector& Collector)
@@ -6045,29 +6060,12 @@ TSharedPtr<FRepChangedPropertyTracker> UNetDriver::FindOrCreateRepChangedPropert
 {
 	check(IsServer() || MaySendProperties());
 
-	FRepChangedPropertyTrackerWrapper * GlobalPropertyTrackerPtr = RepChangedPropertyTrackerMap.Find( Obj );
-
-	// Obj can be a new object with a pointer that matches an old, no longer valid, object
-	if ( GlobalPropertyTrackerPtr != nullptr && !GlobalPropertyTrackerPtr->IsObjectValid() )
-	{
-		RepChangedPropertyTrackerMap.Remove(Obj);
-		GlobalPropertyTrackerPtr = nullptr;
-	}
-
-	if ( !GlobalPropertyTrackerPtr ) 
-	{
-		const UWorld* const LocalWorld = GetWorld();
-		const bool bIsReplay = LocalWorld != nullptr && static_cast<void*>(LocalWorld->GetDemoNetDriver()) == static_cast<void*>(this);
-		const bool bIsClientReplayRecording = LocalWorld != nullptr ? LocalWorld->IsRecordingClientReplay() : false;
-
-		FRepChangedPropertyTracker * Tracker = new FRepChangedPropertyTracker(bIsReplay, bIsClientReplayRecording);
-
-		GetObjectClassRepLayout( Obj->GetClass() )->InitChangedTracker( Tracker );
-
-		GlobalPropertyTrackerPtr = &RepChangedPropertyTrackerMap.Add( Obj, FRepChangedPropertyTrackerWrapper( Obj, TSharedPtr<FRepChangedPropertyTracker>( Tracker ) ) );
-	}
+	return FNetPropertyConditionManager::Get().FindOrCreatePropertyTracker(Obj);
+}
 
-	return GlobalPropertyTrackerPtr->RepChangedPropertyTracker;
+TSharedPtr<FRepChangedPropertyTracker> UNetDriver::FindRepChangedPropertyTracker(UObject* Obj)
+{
+	return FNetPropertyConditionManager::Get().FindPropertyTracker(Obj);
 }
 
 TSharedPtr<FRepLayout> UNetDriver::GetObjectClassRepLayout( UClass * Class )
diff --git a/Engine/Source/Runtime/Engine/Private/RepLayout.cpp b/Engine/Source/Runtime/Engine/Private/RepLayout.cpp
--- a/Engine/Source/Runtime/Engine/Private/RepLayout.cpp
+++ b/Engine/Source/Runtime/Engine/Private/RepLayout.cpp
@@ -993,49 +993,6 @@ struct FNetPrivatePushIdHelper
 };
 #endif // (WITH_PUSH_MODEL)
 
-PRAGMA_DISABLE_DEPRECATION_WARNINGS
-FRepChangedPropertyTracker::FRepChangedPropertyTracker(const bool InbIsReplay, const bool InbIsClientReplayRecording) :
-	bIsClientReplayRecording(InbIsClientReplayRecording),
-	ExternalDataNumBits(0)
-{}
-PRAGMA_ENABLE_DEPRECATION_WARNINGS
-
-void FRepChangedPropertyTracker::SetCustomIsActiveOverride(UObject* OwningObject, const uint16 RepIndex, const bool bIsActive)
-{
-	FBitReference Active = ActiveParents[RepIndex];
-
-	const bool bOldActive = Active;
-	Active = (bIsActive || bIsClientReplayRecording);
-
-#if WITH_PUSH_MODEL
-	if (!bOldActive && Active)
-	{
-		FNetPrivatePushIdHelper::MarkPropertyDirty(OwningObject, RepIndex);
-	}
-#endif	
-}
-
-PRAGMA_DISABLE_DEPRECATION_WARNINGS
-void FRepChangedPropertyTracker::SetExternalData(const uint8* Src, const int32 NumBits)
-{
-	ExternalDataNumBits = NumBits;
-	const int32 NumBytes = (NumBits + 7) >> 3;
-	ExternalData.Reset(NumBytes);
-	ExternalData.AddUninitialized(NumBytes);
-	FMemory::Memcpy(ExternalData.GetData(), Src, NumBytes);
-}
-PRAGMA_ENABLE_DEPRECATION_WARNINGS
-
-void FRepChangedPropertyTracker::CountBytes(FArchive& Ar) const
-{
-	// Include our size here, because the caller won't know.
-	Ar.CountBytes(sizeof(FRepChangedPropertyTracker), sizeof(FRepChangedPropertyTracker));
-	ActiveParents.CountBytes(Ar);
-	PRAGMA_DISABLE_DEPRECATION_WARNINGS
-	ExternalData.CountBytes(Ar);
-	PRAGMA_ENABLE_DEPRECATION_WARNINGS
-}
-
 void FRepStateStaticBuffer::CountBytes(FArchive& Ar) const
 {
 	GRANULAR_NETWORK_MEMORY_TRACKING_INIT(Ar, "FRepStateStaticBuffer::CountBytes");
@@ -1331,6 +1288,7 @@ struct FComparePropertiesSharedParams
 	const bool bValidateProperties = false;
 	const bool bIsNetworkProfilerActive = false;
 	const bool bChangedNetOwner = false;
+	const bool bForceCustomPropsActive = false;
 #if (WITH_PUSH_VALIDATION_SUPPORT || USE_NETWORK_PROFILER)
 	TBitArray<> PropertiesCompared;
 	TBitArray<> PropertiesChanged;
@@ -1403,7 +1361,7 @@ static bool CompareParentProperty(
 	// If the property is inactive, we can skip comparing it because we know it won't be sent.
 	// Further, this will keep the last active state of the property in the shadow buffer,
 	// meaning the next time the property becomes active it will be sent to all connections.
-	const bool bIsActive = !SharedParams.RepChangedPropertyTracker || SharedParams.RepChangedPropertyTracker->IsParentActive(ParentIndex);
+	const bool bIsActive = SharedParams.bForceCustomPropsActive || !SharedParams.RepChangedPropertyTracker || SharedParams.RepChangedPropertyTracker->IsParentActive(ParentIndex);
 	const bool bShouldSkip = !bIsLifetime || !bIsActive || (Parent.Condition == COND_InitialOnly && !SharedParams.bIsInitial);
 
 	if (bShouldSkip)
@@ -1721,10 +1679,11 @@ ERepLayoutResult FRepLayout::CompareProperties(
 		(RepState ? RepState->RepChangedPropertyTracker.Get() : nullptr),
 		NetSerializeLayouts,
 		/*PushModelState=*/UE_RepLayout_Private::GetPerNetDriverState(RepChangelistState),
-		/*PushModelProperties=*/ LocalPushModelProperties,	
+		/*PushModelProperties=*/ LocalPushModelProperties,
 		/*bValidateProperties=*/GbPushModelValidateProperties,
 		/*bIsNetworkProfilerActive=*/UE_RepLayout_Private::IsNetworkProfilerComparisonTrackingEnabled(),
-		/*bChangedNetOwner=*/ RepState && RepState->RepFlags.bNetOwner != RepFlags.bNetOwner
+		/*bChangedNetOwner=*/ RepState && RepState->RepFlags.bNetOwner != RepFlags.bNetOwner,
+		/*bForceCustomPropsActive*/ !!RepFlags.bClientReplay
 	};
 
 	FComparePropertiesStackParams StackParams{
@@ -7054,10 +7013,12 @@ void FRepLayout::RebuildConditionalProperties(
 	RepState->RepFlags = RepFlags;
 }
 
+PRAGMA_DISABLE_DEPRECATION_WARNINGS
 void FRepLayout::InitChangedTracker(FRepChangedPropertyTracker* ChangedTracker) const
 {
-	ChangedTracker->InitActiveParents(Parents.Num());
+	checkf(ChangedTracker->GetParentCount() == Parents.Num(), TEXT("InitChangedTracker: Mismatched replicated parent properties for: %s"), *GetNameSafe(Owner));
 }
+PRAGMA_ENABLE_DEPRECATION_WARNINGS
 
 FRepStateStaticBuffer FRepLayout::CreateShadowBuffer(const FConstRepObjectDataBuffer Source) const
 {
diff --git a/Engine/Source/Runtime/Engine/Private/ReplayHelper.cpp b/Engine/Source/Runtime/Engine/Private/ReplayHelper.cpp
--- a/Engine/Source/Runtime/Engine/Private/ReplayHelper.cpp
+++ b/Engine/Source/Runtime/Engine/Private/ReplayHelper.cpp
@@ -1420,7 +1420,8 @@ void FReplayHelper::SaveExternalData(UNetConnection* Connection, FArchive& Ar)
 			else
 			{
 				PRAGMA_DISABLE_DEPRECATION_WARNINGS
-				FRepChangedPropertyTracker* PropertyTracker = Connection->Driver->RepChangedPropertyTrackerMap.FindChecked(Object).Get();
+				TSharedPtr<FRepChangedPropertyTracker> PropertyTracker = Connection->Driver->FindRepChangedPropertyTracker(Object);
+				check(PropertyTracker.IsValid());
 
 				uint32 ExternalDataNumBits = PropertyTracker->ExternalDataNumBits;
 				if (ExternalDataNumBits > 0)
@@ -1603,14 +1604,12 @@ bool FReplayHelper::UpdateExternalDataForObject(UNetConnection* Connection, UObj
 {
 	check(Connection && Connection->Driver);
 
-	UNetDriver::FRepChangedPropertyTrackerWrapper* PropertyTrackerWrapper = Connection->Driver->RepChangedPropertyTrackerMap.Find(OwningObject);
-	if (PropertyTrackerWrapper == nullptr)
+	TSharedPtr<FRepChangedPropertyTracker> PropertyTracker = Connection->Driver->FindRepChangedPropertyTracker(OwningObject);
+	if (!PropertyTracker.IsValid())
 	{
 		return false;
 	}
 
-	FRepChangedPropertyTracker* PropertyTracker = PropertyTrackerWrapper->Get();
-
 	if (PropertyTracker->ExternalData.Num() == 0)
 	{
 		return false;
diff --git a/Engine/Source/Runtime/Engine/Private/ReplayNetConnection.cpp b/Engine/Source/Runtime/Engine/Private/ReplayNetConnection.cpp
--- a/Engine/Source/Runtime/Engine/Private/ReplayNetConnection.cpp
+++ b/Engine/Source/Runtime/Engine/Private/ReplayNetConnection.cpp
@@ -134,7 +134,6 @@ void UReplayNetConnection::LowLevelSend(void* Data, int32 CountBits, FOutPacketT
 			if (!Actor->IsPendingKillPending())
 			{
 				//@todo: unique this in tick?
-				// RepChangedPropertyTrackerMap.Find is expensive
 				ReplayHelper.UpdateExternalDataForObject(this, Actor);
 			}
 
diff --git a/Engine/Source/Runtime/Engine/Private/UnrealNetwork.cpp b/Engine/Source/Runtime/Engine/Private/UnrealNetwork.cpp
--- a/Engine/Source/Runtime/Engine/Private/UnrealNetwork.cpp
+++ b/Engine/Source/Runtime/Engine/Private/UnrealNetwork.cpp
@@ -217,23 +217,4 @@ void DisableAllReplicatedPropertiesOfClass(const UClass* ThisClass, const UClass
 			SetReplicatedPropertyToDisabled(Prop, OutLifetimeProps);
 		}
 	}
-}
-
-
-void DeprecatedChangeCondition(const  FProperty* ReplicatedProperty, TArray<FLifetimeProperty>& OutLifetimeProps, ELifetimeCondition InCondition)
-{
-	bool bFound = false;
-	for (int32 i = 0; i < OutLifetimeProps.Num(); i++)
-	{
-		if (OutLifetimeProps[i].RepIndex == ReplicatedProperty->RepIndex)
-		{
-			for ( int32 j = 0; j < ReplicatedProperty->ArrayDim; j++ )
-			{
-				OutLifetimeProps[i + j].Condition = InCondition;
-			}
-			bFound = true;
-			break;
-		}
-	}
-	check( bFound );
 }
\ No newline at end of file
diff --git a/Engine/Source/Runtime/Engine/Public/Net/RepLayout.h b/Engine/Source/Runtime/Engine/Public/Net/RepLayout.h
--- a/Engine/Source/Runtime/Engine/Public/Net/RepLayout.h
+++ b/Engine/Source/Runtime/Engine/Public/Net/RepLayout.h
@@ -22,6 +22,7 @@
 #include "Containers/StaticBitArray.h"
 #include "Net/Core/Misc/GuidReferences.h"
 #include "Net/Core/PushModel/PushModel.h"
+#include "Net/Core/PropertyConditions/RepChangedPropertyTracker.h"
 #include "Templates/CopyQualifiersFromTo.h"
 
 class FNetFieldExportGroup;
@@ -141,83 +142,7 @@ public:
 	uint32 IsConditional: 1;
 };
 
-/**
- * This class is used to store meta data about properties that is shared between connections,
- * including whether or not a given property is Conditional, Active, and any external data
- * that may be needed for Replays.
- *
- * TODO: This class (and arguably IRepChangedPropertyTracker) should be renamed to reflect
- *			what they actually do now.
- */
-class FRepChangedPropertyTracker : public IRepChangedPropertyTracker
-{
-public:
-	FRepChangedPropertyTracker() = delete;
-	FRepChangedPropertyTracker(const bool InbIsReplay, const bool InbIsClientReplayRecording);
-
-	PRAGMA_DISABLE_DEPRECATION_WARNINGS
-	virtual ~FRepChangedPropertyTracker() = default;
-	PRAGMA_ENABLE_DEPRECATION_WARNINGS
-
-	//~ Begin IRepChangedPropertyTracker Interface.
-	/**
-	 * Manually set whether or not Property should be marked inactive.
-	 * This will change the Active status for all connections.
-	 *
-	 * @see DOREPLIFETIME_ACTIVE_OVERRIDE
-	 *
-	 * @param OwningObject	The object that we're tracking.
-	 * @param RepIndex		Replication index for the Property.
-	 * @param bIsActive		The new Active state.
-	 */
-	virtual void SetCustomIsActiveOverride(
-		UObject* OwningObject,
-		const uint16 RepIndex,
-		const bool bIsActive) override;
-
-	/**
-	 * Sets (or resets) the External Data.
-	 * External Data is primarily used for Replays, and is used to track additional non-replicated
-	 * data or state about an object.
-	 *
-	 * @param Src		Memory containing the external data.
-	 * @param NumBits	Size of the memory, in bits.
-	 */
-	UE_DEPRECATED(5.0, "Please use UReplaySubsystem::SetExternalDataForObject instead.")
-	virtual void SetExternalData(const uint8* Src, const int32 NumBits) override;
-
-	virtual void CountBytes(FArchive& Ar) const override;
-	//~ End IRepChangedPropertyTracker Interface
-
-	void InitActiveParents(int32 ParentCount)
-	{
-		ActiveParents.Init(true, ParentCount);
-	}
-
-	bool IsParentActive(int32 ParentIndex) const 
-	{ 
-		return ActiveParents[ParentIndex]; 
-	}
-
-	int32 GetParentCount() const
-	{
-		return ActiveParents.Num();
-	}
-
-private:
-	/** Whether or not this is being used for a client replay recording. */
-	bool bIsClientReplayRecording;
-
-	/** Activation data for top level Properties on the given Actor / Object. */
-	TBitArray<> ActiveParents;
-
-public:
-	UE_DEPRECATED(5.0, "No longer used, see UReplaySubsystem::SetExternalDataForObject")
-	TArray<uint8> ExternalData;
-
-	UE_DEPRECATED(5.0, "No longer used, see UReplaySubsystem::SetExternalDataForObject")
-	uint32 ExternalDataNumBits;
-};
+/** FRepChangedPropertyTracker moved to NetCore module */
 
 class FRepLayout;
 class FRepLayoutCmd;
@@ -1287,6 +1212,7 @@ public:
 		TSharedPtr<FRepChangedPropertyTracker>& InRepChangedPropertyTracker,
 		ECreateRepStateFlags Flags) const;
 
+	UE_DEPRECATED(5.1, "No longer used, trackers are initialized by the replication subsystem.")
 	void InitChangedTracker(FRepChangedPropertyTracker * ChangedTracker) const;
 
 	/**
diff --git a/Engine/Source/Runtime/Engine/Public/Net/UnrealNetwork.h b/Engine/Source/Runtime/Engine/Public/Net/UnrealNetwork.h
--- a/Engine/Source/Runtime/Engine/Public/Net/UnrealNetwork.h
+++ b/Engine/Source/Runtime/Engine/Public/Net/UnrealNetwork.h
@@ -16,6 +16,7 @@
 #include "EngineLogs.h"
 #include "UObject/UnrealType.h"
 #include "Engine/EngineBaseTypes.h"
+#include "Net/Core/PropertyConditions/PropertyConditions.h"
 
 class AActor;
 
@@ -260,7 +261,6 @@ static FProperty* GetReplicatedProperty(const UClass* CallingClass, const UClass
 	DOREPLIFETIME_WITH_PARAMS(c,v,LocalDoRepParams); \
 }
 
-
 #define DOREPLIFETIME_ACTIVE_OVERRIDE_FAST(c,v,active) \
 { \
 	static const bool bIsValid_##c_##v = ValidateReplicatedClassInheritance(StaticClass(), c::StaticClass(), TEXT(#v)); \
@@ -285,23 +285,25 @@ static FProperty* GetReplicatedProperty(const UClass* CallingClass, const UClass
 	} \
 }
 
-UE_DEPRECATED(4.24, "Please use the RESET_REPLIFETIME_CONDITION macro")
-ENGINE_API void DeprecatedChangeCondition(FProperty* ReplicatedProperty, TArray< FLifetimeProperty >& OutLifetimeProps, ELifetimeCondition InCondition);
+#define DOREPCUSTOMCONDITION_ACTIVE_FAST(c,v,active) \
+{ \
+	static const bool bIsValid_##c_##v = ValidateReplicatedClassInheritance(StaticClass(), c::StaticClass(), TEXT(#v)); \
+	OutActiveState.SetActiveState((uint16)c::ENetFields_Private::v, active); \
+}
 
+#define SET_REPCUSTOMCONDITION_ACTIVE_FAST(c,v,active) \
+{ \
+	FNetPropertyConditionManager::Get().SetPropertyActive(this, (uint16)c::ENetFields_Private::v, active); \
+}
 
-//~ This is already using a deprecated method, don't bother updating it.
-#define DOREPLIFETIME_CHANGE_CONDITION(c,v,cond) \
+#define SET_REPCUSTOMCONDITION_ACTIVE_FAST_STATIC_ARRAY(c,v,active) \
 { \
-	FProperty* sp##v = GetReplicatedProperty(StaticClass(), c::StaticClass(),GET_MEMBER_NAME_CHECKED(c,v));			\
-	DeprecatedChangeCondition(sp##v, OutLifetimeProps, cond);														\
+	for (int32 i = 0; i < (int32)c::EArrayDims_Private::v; ++i) \
+	{ \
+		FNetPropertyConditionManager::Get().SetPropertyActive(this, (uint16)c::ENetFields_Private::v##_STATIC_ARRAY + i, active); \
+	} \
 }
 
-UE_DEPRECATED(4.24, "Use RegisterReplicatedLifetimeProperty that takes FDoRepLifetimeParams.")
-ENGINE_API void RegisterReplicatedLifetimeProperty(
-	const FProperty* ReplicatedProperty,
-	TArray<FLifetimeProperty>& OutLifetimeProps,
-	ELifetimeCondition InCondition,
-	ELifetimeRepNotifyCondition InRepNotifyCondition = REPNOTIFY_OnChanged);
 
 ENGINE_API void RegisterReplicatedLifetimeProperty(
 	const FProperty* ReplicatedProperty,
diff --git a/Engine/Source/Runtime/Net/Core/Private/Net/Core/PropertyConditions/PropertyConditions.cpp b/Engine/Source/Runtime/Net/Core/Private/Net/Core/PropertyConditions/PropertyConditions.cpp
new file mode 100644
--- /dev/null
+++ b/Engine/Source/Runtime/Net/Core/Private/Net/Core/PropertyConditions/PropertyConditions.cpp
@@ -0,0 +1,106 @@
+// Copyright Epic Games, Inc. All Rights Reserved.
+
+#include "Net/Core/PropertyConditions/PropertyConditions.h"
+#include "Net/Core/PropertyConditions/RepChangedPropertyTracker.h"
+#include "UObject/UObjectGlobals.h"
+#include "UObject/Object.h"
+#include "Serialization/ArchiveCountMem.h"
+#include "Stats/Stats2.h"
+
+DECLARE_CYCLE_STAT(TEXT("PropertyConditions PostGarbageCollect"), STAT_PropertyConditions_PostGarbageCollect, STATGROUP_Net);
+	
+FNetPropertyConditionManager::FNetPropertyConditionManager()
+{
+	PostGarbageCollectHandle = FCoreUObjectDelegates::GetPostGarbageCollect().AddRaw(this, &FNetPropertyConditionManager::PostGarbageCollect);
+}
+
+FNetPropertyConditionManager::~FNetPropertyConditionManager()
+{
+	FCoreUObjectDelegates::GetPostGarbageCollect().Remove(PostGarbageCollectHandle);
+}
+
+FNetPropertyConditionManager& FNetPropertyConditionManager::Get()
+{
+	static FNetPropertyConditionManager Singleton;
+	return Singleton;
+}
+
+void FNetPropertyConditionManager::SetPropertyActive(const FObjectKey ObjectKey, const uint16 RepIndex, const bool bActive)
+{
+	TSharedPtr<FRepChangedPropertyTracker> Tracker = FindPropertyTracker(ObjectKey);
+	if (Tracker.IsValid())
+	{
+		Tracker->SetCustomIsActiveOverride(ObjectKey.ResolveObjectPtr(), RepIndex, bActive);
+	}
+}
+
+void FNetPropertyConditionManager::NotifyObjectDestroyed(const FObjectKey ObjectKey)
+{
+	PropertyTrackerMap.Remove(ObjectKey);
+}
+
+TSharedPtr<FRepChangedPropertyTracker> FNetPropertyConditionManager::FindOrCreatePropertyTracker(const FObjectKey ObjectKey)
+{
+	TSharedPtr<FRepChangedPropertyTracker> Tracker = FindPropertyTracker(ObjectKey);
+	if (!Tracker.IsValid())
+	{
+		if (UObject* Obj = ObjectKey.ResolveObjectPtr())
+		{
+			UClass* ObjectClass = Obj->GetClass();
+			check(ObjectClass);
+			ObjectClass->SetUpRuntimeReplicationData();
+
+			const int32 NumProperties = ObjectClass->ClassReps.Num();
+
+			FCustomPropertyConditionState ActiveState(NumProperties);
+			Obj->GetReplicatedCustomConditionState(ActiveState);
+
+			Tracker = MakeShared<FRepChangedPropertyTracker>(MoveTemp(ActiveState));
+
+			PropertyTrackerMap.Add(ObjectKey, Tracker);
+		}
+		else
+		{
+			ensureMsgf(false, TEXT("FindOrCreatePropertyTracker: Unable to resolve object key."));
+		}
+	}
+
+	return Tracker;
+}
+
+TSharedPtr<FRepChangedPropertyTracker> FNetPropertyConditionManager::FindPropertyTracker(const FObjectKey ObjectKey) const
+{
+	return PropertyTrackerMap.FindRef(ObjectKey);
+}
+
+void FNetPropertyConditionManager::PostGarbageCollect()
+{
+	SCOPE_CYCLE_COUNTER(STAT_PropertyConditions_PostGarbageCollect);
+
+	for (auto It = PropertyTrackerMap.CreateIterator(); It; ++It)
+	{
+		if (!It.Key().ResolveObjectPtr())
+		{
+			It.RemoveCurrent();
+		}
+	}
+}
+
+void FNetPropertyConditionManager::LogMemory(FOutputDevice& Ar)
+{
+	FArchiveCountMem CountAr(nullptr);
+
+	PropertyTrackerMap.CountBytes(CountAr);
+
+	for (auto It = PropertyTrackerMap.CreateConstIterator(); It; ++It)
+	{
+		if (It->Value.IsValid())
+		{
+			It->Value->CountBytes(CountAr);
+		}
+	}
+
+	const int32 CountBytes = sizeof(*this) + CountAr.GetNum();
+
+	Ar.Logf(TEXT("  Property Condition Memory: %u"), CountBytes);
+}
diff --git a/Engine/Source/Runtime/Net/Core/Private/Net/Core/PropertyConditions/RepChangedPropertyTracker.cpp b/Engine/Source/Runtime/Net/Core/Private/Net/Core/PropertyConditions/RepChangedPropertyTracker.cpp
new file mode 100644
--- /dev/null
+++ b/Engine/Source/Runtime/Net/Core/Private/Net/Core/PropertyConditions/RepChangedPropertyTracker.cpp
@@ -0,0 +1,52 @@
+// Copyright Epic Games, Inc. All Rights Reserved.
+
+#include "Net/Core/PropertyConditions/RepChangedPropertyTracker.h"
+#include "Net/Core/PushModel/PushModel.h"
+
+PRAGMA_DISABLE_DEPRECATION_WARNINGS
+FRepChangedPropertyTracker::FRepChangedPropertyTracker(FCustomPropertyConditionState&& InActiveState)
+	: ActiveState(MoveTemp(InActiveState))
+	, ExternalDataNumBits(0)
+{
+}
+
+FRepChangedPropertyTracker::FRepChangedPropertyTracker(const bool InbIsReplay, const bool InbIsClientReplayRecording)
+	: FRepChangedPropertyTracker(0)
+{}
+PRAGMA_ENABLE_DEPRECATION_WARNINGS
+
+void FRepChangedPropertyTracker::SetCustomIsActiveOverride(UObject* OwningObject, const uint16 RepIndex, const bool bIsActive)
+{
+	const bool bOldActive = ActiveState.GetActiveState(RepIndex);
+	ActiveState.SetActiveState(RepIndex, bIsActive);	// check for client replay recording moved to FReplicationFlags
+
+#if WITH_PUSH_MODEL
+	if (!bOldActive && bIsActive)
+	{
+		MARK_PROPERTY_DIRTY_UNSAFE(OwningObject, RepIndex);
+	}
+#endif	
+}
+
+PRAGMA_DISABLE_DEPRECATION_WARNINGS
+void FRepChangedPropertyTracker::SetExternalData(const uint8* Src, const int32 NumBits)
+{
+	ExternalDataNumBits = NumBits;
+	const int32 NumBytes = (NumBits + 7) >> 3;
+	ExternalData.Reset(NumBytes);
+	ExternalData.AddUninitialized(NumBytes);
+	FMemory::Memcpy(ExternalData.GetData(), Src, NumBytes);
+}
+PRAGMA_ENABLE_DEPRECATION_WARNINGS
+
+void FRepChangedPropertyTracker::CountBytes(FArchive& Ar) const
+{
+	// Include our size here, because the caller won't know.
+	Ar.CountBytes(sizeof(FRepChangedPropertyTracker), sizeof(FRepChangedPropertyTracker));
+	ActiveState.CountBytes(Ar);
+
+	PRAGMA_DISABLE_DEPRECATION_WARNINGS
+	ExternalData.CountBytes(Ar);
+	PRAGMA_ENABLE_DEPRECATION_WARNINGS
+}
+
diff --git a/Engine/Source/Runtime/Net/Core/Public/Net/Core/PropertyConditions/PropertyConditions.h b/Engine/Source/Runtime/Net/Core/Public/Net/Core/PropertyConditions/PropertyConditions.h
new file mode 100644
--- /dev/null
+++ b/Engine/Source/Runtime/Net/Core/Public/Net/Core/PropertyConditions/PropertyConditions.h
@@ -0,0 +1,33 @@
+// Copyright Epic Games, Inc. All Rights Reserved.
+
+#pragma once
+
+#include "CoreMinimal.h"
+#include "UObject/CoreNet.h"
+#include "UObject/ObjectKey.h"
+
+class FRepChangedPropertyTracker;
+
+class NETCORE_API FNetPropertyConditionManager
+{
+public:
+	FNetPropertyConditionManager();
+	~FNetPropertyConditionManager();
+
+	static FNetPropertyConditionManager& Get();
+
+	void SetPropertyActive(const FObjectKey ObjectKey, const uint16 RepIndex, const bool bActive);
+	void NotifyObjectDestroyed(const FObjectKey ObjectKey);
+
+	TSharedPtr<FRepChangedPropertyTracker> FindOrCreatePropertyTracker(const FObjectKey ObjectKey);
+	TSharedPtr<FRepChangedPropertyTracker> FindPropertyTracker(const FObjectKey ObjectKey) const;
+
+	void LogMemory(FOutputDevice& Ar);
+
+private:
+	void PostGarbageCollect();
+
+	FDelegateHandle PostGarbageCollectHandle;
+
+	TMap<FObjectKey, TSharedPtr<FRepChangedPropertyTracker>> PropertyTrackerMap;
+};
diff --git a/Engine/Source/Runtime/Net/Core/Public/Net/Core/PropertyConditions/RepChangedPropertyTracker.h b/Engine/Source/Runtime/Net/Core/Public/Net/Core/PropertyConditions/RepChangedPropertyTracker.h
new file mode 100644
--- /dev/null
+++ b/Engine/Source/Runtime/Net/Core/Public/Net/Core/PropertyConditions/RepChangedPropertyTracker.h
@@ -0,0 +1,79 @@
+// Copyright Epic Games, Inc. All Rights Reserved.
+
+#pragma once
+
+#include "CoreMinimal.h"
+#include "UObject/CoreNet.h"
+
+/**
+ * This class is used to store meta data about properties that is shared between connections,
+ * including whether or not a given property is Conditional, Active, and any external data
+ * that may be needed for Replays.
+ *
+ * TODO: This class (and arguably IRepChangedPropertyTracker) should be renamed to reflect
+ *			what they actually do now.
+ */
+class NETCORE_API FRepChangedPropertyTracker : public IRepChangedPropertyTracker
+{
+public:
+	FRepChangedPropertyTracker() = delete;
+	FRepChangedPropertyTracker(FCustomPropertyConditionState&& InActiveState);
+
+	UE_DEPRECATED(5.1, "Replay arguments no longer used.")
+	FRepChangedPropertyTracker(const bool InbIsReplay, const bool InbIsClientReplayRecording);
+
+	PRAGMA_DISABLE_DEPRECATION_WARNINGS
+	virtual ~FRepChangedPropertyTracker() = default;
+	PRAGMA_ENABLE_DEPRECATION_WARNINGS
+
+	//~ Begin IRepChangedPropertyTracker Interface.
+	/**
+		* Manually set whether or not Property should be marked inactive.
+		* This will change the Active status for all connections.
+		*
+		* @see DOREPLIFETIME_ACTIVE_OVERRIDE
+		*
+		* @param OwningObject	The object that we're tracking.
+		* @param RepIndex		Replication index for the Property.
+		* @param bIsActive		The new Active state.
+		*/
+	virtual void SetCustomIsActiveOverride(UObject* OwningObject, const uint16 RepIndex, const bool bIsActive) override;
+
+	/**
+	 * Sets (or resets) the External Data.
+	 * External Data is primarily used for Replays, and is used to track additional non-replicated
+	 * data or state about an object.
+	 *
+	 * @param Src		Memory containing the external data.
+	 * @param NumBits	Size of the memory, in bits.
+	 */
+	UE_DEPRECATED(5.0, "Please use UReplaySubsystem::SetExternalDataForObject instead.")
+	virtual void SetExternalData(const uint8* Src, const int32 NumBits) override;
+
+	virtual void CountBytes(FArchive& Ar) const override;
+	//~ End IRepChangedPropertyTracker Interface
+
+	UE_DEPRECATED(5.1, "No longer used, ActiveState must be constructed with the correct number of properties.")
+	void InitActiveParents(int32 ParentCount) {}
+
+	bool IsParentActive(int32 ParentIndex) const
+	{
+		return ActiveState.GetActiveState(ParentIndex);
+	}
+
+	int32 GetParentCount() const
+	{
+		return ActiveState.GetNumProperties();
+	}
+
+private:
+	/** Activation data for top level Properties on the given Actor / Object. */
+	FCustomPropertyConditionState ActiveState;
+
+public:
+	UE_DEPRECATED(5.0, "No longer used, see UReplaySubsystem::SetExternalDataForObject")
+	TArray<uint8> ExternalData;
+
+	UE_DEPRECATED(5.0, "No longer used, see UReplaySubsystem::SetExternalDataForObject")
+	uint32 ExternalDataNumBits;
+};
