fixed the order of destruction issue, caused by audio internal task event being destroyed after its allocator

diff --git a/Engine/Source/Runtime/Engine/Private/AudioThread.cpp b/Engine/Source/Runtime/Engine/Private/AudioThread.cpp
--- a/Engine/Source/Runtime/Engine/Private/AudioThread.cpp
+++ b/Engine/Source/Runtime/Engine/Private/AudioThread.cpp
@@ -656,6 +656,7 @@ void FAudioThread::StopAudioThread()
 
 	check(ResumeEvent.IsValid());
 	ResumeEvent->Trigger(); // every FTaskEvent must be triggered before destruction to pass the check for completion
+	ResumeEvent.Reset();
 }
 
 FAudioCommandFence::~FAudioCommandFence()
