Skeletal mesh does not update to show the pose with additive layer tracks when the weight of the pose is set to 1.0
#fix apply additive transform curves when calling ::GetBoneTransformFromModel to match old behaviour (RawTrackData would have incorporated the baked Transform curve, which is now baked during compression)

diff --git a/Engine/Source/Runtime/Engine/Private/Animation/AnimSequenceHelpers.cpp b/Engine/Source/Runtime/Engine/Private/Animation/AnimSequenceHelpers.cpp
--- a/Engine/Source/Runtime/Engine/Private/Animation/AnimSequenceHelpers.cpp
+++ b/Engine/Source/Runtime/Engine/Private/Animation/AnimSequenceHelpers.cpp
@@ -315,6 +315,20 @@ void GetBoneTransformFromModel(const UAnimDataModel* Model, FTransform& OutTrans
 	{
 		OutTransform.SetScale3D(DefaultScale3D);
 	}
+
+	for (const FTransformCurve& AdditiveTransformCurve : Model->GetTransformCurves())
+	{
+		if (AdditiveTransformCurve.Name.DisplayName == TrackData.Name)
+		{
+			const float TimeInterval = Model->GetFrameRate().AsSeconds(KeyIndex);
+			const FTransform AdditiveTransform = AdditiveTransformCurve.Evaluate(TimeInterval, 1.f);
+			const FTransform LocalTransform = OutTransform;
+			OutTransform.SetRotation(LocalTransform.GetRotation() * AdditiveTransform.GetRotation());
+			OutTransform.SetTranslation(LocalTransform.TransformPosition(AdditiveTransform.GetTranslation()));
+			OutTransform.SetScale3D(LocalTransform.GetScale3D() * AdditiveTransform.GetScale3D());			
+			break;
+		}
+	}
 }
 
 void CopyCurveDataToModel(const FRawCurveTracks& CurveData, const USkeleton* Skeleton, IAnimationDataController& Controller)
