World: Allow Subsystem initialization on EWorldType::Inactive

-  UWorldSubsystem::DoesSupportWorldType still doesn't support EWorldType::Inactive so behavior remains unchanged
-  WorldPartitionCookPackageSplitter now initializes the world if it isn't already so that UDataLayerSubsystem is initialized at cook time
-  Remove remaining code path where same world could be reused if it was initialized as inactive (Opening unsaved new level in content browser will now require saving first)

#test CitySample, AncientWorld

diff --git a/Engine/Source/Developer/AssetTools/Private/AssetTypeActions/AssetTypeActions_World.cpp b/Engine/Source/Developer/AssetTools/Private/AssetTypeActions/AssetTypeActions_World.cpp
--- a/Engine/Source/Developer/AssetTools/Private/AssetTypeActions/AssetTypeActions_World.cpp
+++ b/Engine/Source/Developer/AssetTools/Private/AssetTypeActions/AssetTypeActions_World.cpp
@@ -35,13 +35,15 @@ void FAssetTypeActions_World::OpenAssetEditor( const TArray<UObject*>& InObjects
 	for (auto ObjIt = InObjects.CreateConstIterator(); ObjIt; ++ObjIt)
 	{
 		UWorld* World = Cast<UWorld>(*ObjIt);
-		if (World != nullptr && ensureMsgf(World->GetTypedOuter<UPackage>(), TEXT("World(%s) is not in a package and cannot be opened"), *World->GetFullName()))
+		if (World != nullptr && 
+			ensureMsgf(World->GetTypedOuter<UPackage>(), TEXT("World(%s) is not in a package and cannot be opened"), *World->GetFullName()) && 
+			ensureMsgf(!World->GetPackage()->IsDirty(), TEXT("World(%s) is unsaved and cannot be opened")))
 		{
 			const FString FileToOpen = FPackageName::LongPackageNameToFilename(World->GetOutermost()->GetName(), FPackageName::GetMapPackageExtension());
 			const bool bLoadAsTemplate = false;
 			const bool bShowProgress = true;
 			FEditorFileUtils::LoadMap(FileToOpen, bLoadAsTemplate, bShowProgress);
-			
+
 			// We can only edit one world at a time... so just break after the first valid world to load
 			break;
 		}
@@ -73,9 +75,17 @@ TArray<FAssetData> FAssetTypeActions_World::GetValidAssetsForPreviewOrEdit(TArra
 		constexpr bool bPromptUserToSave = true;
 		constexpr bool bSaveMapPackages = true;
 		constexpr bool bSaveContentPackages = true;
-		if (FEditorFileUtils::SaveDirtyPackages(bPromptUserToSave, bSaveMapPackages, bSaveContentPackages))
+		constexpr bool bFastSave = false;
+		constexpr bool bNotifyNoPackagesSaved = false;
+		constexpr bool bCanBeDeclined = false;
+		if (FEditorFileUtils::SaveDirtyPackages(bPromptUserToSave, bSaveMapPackages, bSaveContentPackages, bFastSave, bNotifyNoPackagesSaved, bCanBeDeclined))
 		{
-			AssetsToOpen.Add(AssetData);
+			// Validate that Asset was saved or isn't loaded meaning it can be loaded
+			const bool bLoad = false;
+			if (UWorld* World = Cast<UWorld>(AssetData.FastGetAsset(bLoad)); !World || !World->GetPackage()->IsDirty())
+			{
+				AssetsToOpen.Add(AssetData);
+			}
 		}
 	}
 
diff --git a/Engine/Source/Editor/UnrealEd/Private/EditorServer.cpp b/Engine/Source/Editor/UnrealEd/Private/EditorServer.cpp
--- a/Engine/Source/Editor/UnrealEd/Private/EditorServer.cpp
+++ b/Engine/Source/Editor/UnrealEd/Private/EditorServer.cpp
@@ -2544,11 +2544,11 @@ bool UEditorEngine::Map_Load(const TCHAR* Str, FOutputDevice& Ar)
 						FMessageLog("LoadErrors").NewPage( FText::Format( LOCTEXT("LoadMapLogPage", "Loading map: {MapFileName}"), Arguments ) );
 					}
 
-					// Only worlds that are uninitialized or reside within newly created packages (so not yet saved) may be considered valid for re-use.
+					// Only worlds that are uninitialized may be considered valid for re-use.
 					// All other worlds need to be reloaded from disk, as a world is only initialized correctly as part of the level loading process.
 					auto IsWorldValidForReuse = [](UWorld* WorldToConsider)
 					{
-						return !WorldToConsider->bIsWorldInitialized || WorldToConsider->GetPackage()->HasAnyPackageFlags(PKG_NewlyCreated);
+						return !WorldToConsider->bIsWorldInitialized;
 					};
 
 					// If we are loading the same world again (reloading) then we must not specify that we want to keep this world in memory.
@@ -2738,37 +2738,10 @@ bool UEditorEngine::Map_Load(const TCHAR* Str, FOutputDevice& Ar)
 				FParse::Value(Str, TEXT("FEATURELEVEL="), FeatureLevelIndex);
 				FeatureLevelIndex = FMath::Clamp(FeatureLevelIndex, 0, (int32)ERHIFeatureLevel::Num);
 
-				if (World->bIsWorldInitialized)
-				{
-					// If we are using a previously initialized world, make sure it has a physics scene and FXSystem.
-					// Inactive worlds are already initialized but lack these two objects for memory reasons.
-					World->ClearWorldComponents();
-
-					// If the world was inactive subsystems would not have been initialized (happens when loading a world trough the asset browser).
-					// When we transition to the editor world initialize them.
-					World->InitializeSubsystems();
-
-					if (World->FeatureLevel != FeatureLevelIndex)
-					{
-						World->ChangeFeatureLevel((ERHIFeatureLevel::Type)FeatureLevelIndex);
-					}
-
-					if (World->GetPhysicsScene() == nullptr)
-					{
-						World->CreatePhysicsScene();
-					}
-
-					if (World->FXSystem == nullptr)
-					{
-						World->CreateFXSystem();
-					}
-				}
-				else
-				{
-					World->FeatureLevel = (ERHIFeatureLevel::Type)FeatureLevelIndex;
-					World->InitWorld( GetEditorWorldInitializationValues() );
-				}
-
+				check(!World->bIsWorldInitialized);
+				World->FeatureLevel = (ERHIFeatureLevel::Type)FeatureLevelIndex;
+				World->InitWorld( GetEditorWorldInitializationValues() );
+				
 				SlowTask.EnterProgressFrame(20, FText::Format( LOCTEXT( "LoadingMapStatus_Initializing", "Loading map: {0}... (Initializing world)" ), FText::FromString(MapFileName) ));
 				{
 					FBSPOps::bspValidateBrush(Context.World()->GetDefaultBrush()->Brush, 0, 1);
diff --git a/Engine/Source/Runtime/Engine/Private/World.cpp b/Engine/Source/Runtime/Engine/Private/World.cpp
--- a/Engine/Source/Runtime/Engine/Private/World.cpp
+++ b/Engine/Source/Runtime/Engine/Private/World.cpp
@@ -8534,33 +8534,19 @@ void UWorld::RemovePostProcessVolume(IInterface_PostProcessVolume* InVolume)
 
 void UWorld::InitializeSubsystems()
 {
-	// In a specific code path (opening level from ContentBrowser) the InitializeSubsystems gets called twice.
-	// - First time is through UEditorEngine::InitializeNewlyCreatedInactiveWorld with EWorldType::Inactive and we want to prevent initializing 
-	//   in that specific case because SubsystemCollection::Initialize can only be called once and we want to call it with the proper WorldType
-	// - Second time is through UEditorEngine::Map_Load and at that point the WorldType is valid.
-	if (WorldType != EWorldType::Inactive)
-	{
-		SubsystemCollection.Initialize(this);
-
-		if (bIsWorldInitialized)
-		{
-			PostInitializeSubsystems();
-		}
-	}
+	check(!SubsystemCollection.IsInitialized());
+	SubsystemCollection.Initialize(this);
 }
 
 void UWorld::PostInitializeSubsystems()
 {
-	if (WorldType == EWorldType::Game || WorldType == EWorldType::Editor || WorldType == EWorldType::PIE)
-	{
-		check(bIsWorldInitialized);
+	check(bIsWorldInitialized);
 
-		const TArray<UWorldSubsystem*>& WorldSubsystems = SubsystemCollection.GetSubsystemArray<UWorldSubsystem>(UWorldSubsystem::StaticClass());
+	const TArray<UWorldSubsystem*>& WorldSubsystems = SubsystemCollection.GetSubsystemArray<UWorldSubsystem>(UWorldSubsystem::StaticClass());
 
-		for (UWorldSubsystem* WorldSubsystem : WorldSubsystems)
-		{
-			WorldSubsystem->PostInitialize();
-		}
+	for (UWorldSubsystem* WorldSubsystem : WorldSubsystems)
+	{
+		WorldSubsystem->PostInitialize();
 	}
 }
 
diff --git a/Engine/Source/Runtime/Engine/Private/WorldPartition/WorldPartitionCookPackageSplitter.cpp b/Engine/Source/Runtime/Engine/Private/WorldPartition/WorldPartitionCookPackageSplitter.cpp
--- a/Engine/Source/Runtime/Engine/Private/WorldPartition/WorldPartitionCookPackageSplitter.cpp
+++ b/Engine/Source/Runtime/Engine/Private/WorldPartition/WorldPartitionCookPackageSplitter.cpp
@@ -11,6 +11,7 @@
 #include "WorldPartition/WorldPartitionRuntimeHash.h"
 #include "WorldPartition/WorldPartitionSubsystem.h"
 #include "WorldPartition/WorldPartition.h"
+#include "Editor.h"
 
 // Register FWorldPartitionCookPackageSplitter for UWorld class
 REGISTER_COOKPACKAGE_SPLITTER(FWorldPartitionCookPackageSplitter, UWorld);
@@ -23,23 +24,16 @@ bool FWorldPartitionCookPackageSplitter::ShouldSplit(UObject* SplitData)
 
 FWorldPartitionCookPackageSplitter::FWorldPartitionCookPackageSplitter()
 {
-	FCoreUObjectDelegates::GetPreGarbageCollectDelegate().AddRaw(this, &FWorldPartitionCookPackageSplitter::PreGarbageCollect);
 }
 
 FWorldPartitionCookPackageSplitter::~FWorldPartitionCookPackageSplitter()
 {
-	FCoreUObjectDelegates::GetPreGarbageCollectDelegate().RemoveAll(this);
-	TeardownWorldPartition();
+	check(!ReferencedWorld);
 }
 
-void FWorldPartitionCookPackageSplitter::PreGarbageCollect()
+void FWorldPartitionCookPackageSplitter::Teardown(ETeardown Status)
 {
-	TeardownWorldPartition();
-}
-
-void FWorldPartitionCookPackageSplitter::TeardownWorldPartition()
-{
-	if (bWorldPartitionNeedsTeardown)
+ 	if (bInitializedWorldPartition)
 	{
 		if (UWorld* LocalWorld = ReferencedWorld.Get())
 		{
@@ -49,8 +43,27 @@ void FWorldPartitionCookPackageSplitter::TeardownWorldPartition()
 				WorldPartition->Uninitialize();
 			}
 		}
-		bWorldPartitionNeedsTeardown = false;
+		bInitializedWorldPartition = false;
+	}
+
+	if (bInitializedPhysicsSceneForSave)
+	{
+		GEditor->CleanupPhysicsSceneThatWasInitializedForSave(ReferencedWorld.Get(), bForceInitializedWorld);
+		bInitializedPhysicsSceneForSave = false;
+		bForceInitializedWorld = false;
 	}
+
+	ReferencedWorld = nullptr;
+}
+
+void FWorldPartitionCookPackageSplitter::AddReferencedObjects(FReferenceCollector& Collector)
+{
+	Collector.AddReferencedObject(ReferencedWorld);
+}
+
+FString FWorldPartitionCookPackageSplitter::GetReferencerName() const
+{
+	return TEXT("FWorldPartitionCookPackageSplitter");
 }
 
 UWorld* FWorldPartitionCookPackageSplitter::ValidateDataObject(UObject* SplitData)
@@ -78,11 +91,15 @@ TArray<ICookPackageSplitter::FGeneratedPackage> FWorldPartitionCookPackageSplitt
 	// that is necessary for populate 
 	ReferencedWorld = PartitionedWorld;
 
+	check(!bInitializedPhysicsSceneForSave && !bForceInitializedWorld);
+	bInitializedPhysicsSceneForSave = GEditor->InitializePhysicsSceneForSaveIfNecessary(PartitionedWorld, bForceInitializedWorld);
+
 	// Manually initialize WorldPartition
 	UWorldPartition* WorldPartition = PartitionedWorld->PersistentLevel->GetWorldPartition();
 	// We expect the WorldPartition has not yet been initialized
 	ensure(!WorldPartition->IsInitialized());
 	WorldPartition->Initialize(PartitionedWorld, FTransform::Identity);
+	bInitializedWorldPartition = true;
 
 	TArray<FString> WorldPartitionGeneratedPackages;
 	WorldPartition->GenerateStreaming(&WorldPartitionGeneratedPackages);
@@ -103,15 +120,7 @@ TArray<ICookPackageSplitter::FGeneratedPackage> FWorldPartitionCookPackageSplitt
 bool FWorldPartitionCookPackageSplitter::TryPopulatePackage(const UPackage* OwnerPackage, const UObject* OwnerObject,
 	const FGeneratedPackageForPopulate& GeneratedPackage, bool bWasOwnerReloaded)
 {
-	if (bWasOwnerReloaded)
-	{
-		GetGenerateList(OwnerPackage, OwnerObject);
-		// When GetGenerateList is called by CookOnTheFlyServer, it is followed up with a call to 
-		// CleanupWorld when the Generator package is done saving, which calls WorldPartition->Uninitialize
-		// But when we call it here to initialize our data for the populate, we need to take responsibility
-		// for calling WorldPartition->Uninitialize before the next GarbageCollection
-		bWorldPartitionNeedsTeardown = true;
-	}
+	check(!bWasOwnerReloaded);
 
 	// TODO: Make PopulateGeneratedPackageForCook const so we can honor the constness of the OwnerObject in this API function
 	const UWorld* ConstPartitionedWorld = ValidateDataObject(OwnerObject);
diff --git a/Engine/Source/Runtime/Engine/Private/WorldPartition/WorldPartitionCookPackageSplitter.h b/Engine/Source/Runtime/Engine/Private/WorldPartition/WorldPartitionCookPackageSplitter.h
--- a/Engine/Source/Runtime/Engine/Private/WorldPartition/WorldPartitionCookPackageSplitter.h
+++ b/Engine/Source/Runtime/Engine/Private/WorldPartition/WorldPartitionCookPackageSplitter.h
@@ -8,7 +8,7 @@
 #include "Engine/World.h"
 #include "UObject/WeakObjectPtrTemplates.h"
 
-class FWorldPartitionCookPackageSplitter : public ICookPackageSplitter
+class FWorldPartitionCookPackageSplitter : public FGCObject, public ICookPackageSplitter
 {
 public:
 	//~ Begin of ICookPackageSplitter
@@ -17,6 +17,8 @@ public:
 	FWorldPartitionCookPackageSplitter();
 	virtual ~FWorldPartitionCookPackageSplitter();
 
+	virtual void Teardown(ETeardown Status) override;
+	virtual bool UseInternalReferenceToAvoidGarbageCollect() override { return true; }
 	virtual TArray<ICookPackageSplitter::FGeneratedPackage> GetGenerateList(const UPackage* OwnerPackage, const UObject* OwnerObject) override;
 	virtual bool TryPopulatePackage(const UPackage* OwnerPackage, const UObject* OwnerObject,
 		const ICookPackageSplitter::FGeneratedPackageForPopulate& GeneratedPackage, bool bWasOwnerReloaded) override;
@@ -25,13 +27,18 @@ public:
 	//~ End of ICookPackageSplitter
 
 private:
+	/** FGCObject interface */
+	virtual void AddReferencedObjects(FReferenceCollector& Collector) override;
+	virtual FString GetReferencerName() const override;
+
 	const UWorld* ValidateDataObject(const UObject* SplitData);
 	UWorld* ValidateDataObject(UObject* SplitData);
-	void PreGarbageCollect();
-	void TeardownWorldPartition();
+	
+	TObjectPtr<UWorld> ReferencedWorld;
 
-	TWeakObjectPtr<UWorld> ReferencedWorld;
-	bool bWorldPartitionNeedsTeardown = false;
+	bool bInitializedWorldPartition = false;
+	bool bForceInitializedWorld = false;
+	bool bInitializedPhysicsSceneForSave = false;
 };
 
 #endif
\ No newline at end of file
diff --git a/Engine/Source/Runtime/Engine/Public/Subsystems/SubsystemCollection.h b/Engine/Source/Runtime/Engine/Public/Subsystems/SubsystemCollection.h
--- a/Engine/Source/Runtime/Engine/Public/Subsystems/SubsystemCollection.h
+++ b/Engine/Source/Runtime/Engine/Public/Subsystems/SubsystemCollection.h
@@ -20,6 +20,9 @@ public:
 	/* Clears the collection, while deinitializing the systems */
 	void Deinitialize();
 
+	/** Returns true if collection was already initialized */
+	bool IsInitialized() const { return Outer != nullptr; }
+
 	/** 
 	 * Only call from Initialize() of Systems to ensure initialization order
 	 * Note: Dependencies only work within a collection
