Allow DataLayerSubsystem to be initialized in Inactive worlds

#test CitySample, AncientWorld

diff --git a/Engine/Source/Runtime/Engine/Private/WorldPartition/DataLayer/DataLayerSubsystem.cpp b/Engine/Source/Runtime/Engine/Private/WorldPartition/DataLayer/DataLayerSubsystem.cpp
--- a/Engine/Source/Runtime/Engine/Private/WorldPartition/DataLayer/DataLayerSubsystem.cpp
+++ b/Engine/Source/Runtime/Engine/Private/WorldPartition/DataLayer/DataLayerSubsystem.cpp
@@ -80,6 +80,12 @@ void UDataLayerSubsystem::Deinitialize()
 #endif
 }
 
+bool UDataLayerSubsystem::DoesSupportWorldType(const EWorldType::Type WorldType) const
+{
+	return Super::DoesSupportWorldType(WorldType) || WorldType == EWorldType::Inactive;
+}
+
+
 const TSet<FName>& UDataLayerSubsystem::GetEffectiveActiveDataLayerNames() const
 {
 	static TSet<FName> EmptySet;
diff --git a/Engine/Source/Runtime/Engine/Public/WorldPartition/DataLayer/DataLayerSubsystem.h b/Engine/Source/Runtime/Engine/Public/WorldPartition/DataLayer/DataLayerSubsystem.h
--- a/Engine/Source/Runtime/Engine/Public/WorldPartition/DataLayer/DataLayerSubsystem.h
+++ b/Engine/Source/Runtime/Engine/Public/WorldPartition/DataLayer/DataLayerSubsystem.h
@@ -49,6 +49,10 @@ public:
 	virtual void Deinitialize() override;
 	//~ End USubsystem Interface.
 
+protected:
+	virtual bool DoesSupportWorldType(const EWorldType::Type WorldType) const override;
+
+public:
 	//~ Begin Blueprint callable functions
 
 	/** Find a Data Layer by its asset. */
