Crash Occurs In Skeletal Mesh Editor When Adding a Bone Transform in a New Animation Sequence
#fix ensure we set _two_ bone keys when creating a new single frame animation (AnimDataModel expects uniform key distribution)

diff --git a/Engine/Source/Runtime/Engine/Private/Animation/AnimSequence.cpp b/Engine/Source/Runtime/Engine/Private/Animation/AnimSequence.cpp
--- a/Engine/Source/Runtime/Engine/Private/Animation/AnimSequence.cpp
+++ b/Engine/Source/Runtime/Engine/Private/Animation/AnimSequence.cpp
@@ -3854,7 +3854,7 @@ bool UAnimSequence::CreateAnimation(USkeletalMesh* Mesh)
 		{
 			const FName& BoneName = RefSkeleton.GetBoneName(BoneIndex);
 			Controller->AddBoneTrack(BoneName);
-			Controller->SetBoneTrackKeys(BoneName, { FVector3f(RefBonePose[BoneIndex].GetTranslation()) }, { FQuat4f(RefBonePose[BoneIndex].GetRotation()) }, { FVector3f(RefBonePose[BoneIndex].GetScale3D()) });
+			Controller->SetBoneTrackKeys(BoneName, { FVector3f(RefBonePose[BoneIndex].GetTranslation()), FVector3f(RefBonePose[BoneIndex].GetTranslation()) }, { FQuat4f(RefBonePose[BoneIndex].GetRotation()), FQuat4f(RefBonePose[BoneIndex].GetRotation()) }, { FVector3f(RefBonePose[BoneIndex].GetScale3D()), FVector3f(RefBonePose[BoneIndex].GetScale3D()) });
 		}
 
 		Controller->NotifyPopulated();
@@ -3891,7 +3891,7 @@ bool UAnimSequence::CreateAnimation(USkeletalMeshComponent* MeshComponent)
 		{
 			const FName& BoneName = RefSkeleton.GetBoneName(BoneIndex);
 			Controller->AddBoneTrack(BoneName);
-			Controller->SetBoneTrackKeys(BoneName, TArray<FVector>{ BoneSpaceTransforms[BoneIndex].GetTranslation() }, { BoneSpaceTransforms[BoneIndex].GetRotation() }, TArray<FVector>{ BoneSpaceTransforms[BoneIndex].GetScale3D() });
+			Controller->SetBoneTrackKeys(BoneName, TArray<FVector>{ BoneSpaceTransforms[BoneIndex].GetTranslation(), BoneSpaceTransforms[BoneIndex].GetTranslation() }, { BoneSpaceTransforms[BoneIndex].GetRotation(), BoneSpaceTransforms[BoneIndex].GetRotation() }, TArray<FVector>{ BoneSpaceTransforms[BoneIndex].GetScale3D(), BoneSpaceTransforms[BoneIndex].GetScale3D() });
 		}
 
 		Controller->NotifyPopulated();
