Fix crash opening a control rig using a preview a mesh that is missing its skeleton

diff --git a/Engine/Source/Runtime/Engine/Private/Animation/AnimInstanceProxy.cpp b/Engine/Source/Runtime/Engine/Private/Animation/AnimInstanceProxy.cpp
--- a/Engine/Source/Runtime/Engine/Private/Animation/AnimInstanceProxy.cpp
+++ b/Engine/Source/Runtime/Engine/Private/Animation/AnimInstanceProxy.cpp
@@ -3102,11 +3102,11 @@ void FAnimInstanceProxy::UpdateCurvesToEvaluationContext(const FAnimationEvaluat
 
 	ResetAnimationCurves();
 
-	if(InContext.Curve.UIDToArrayIndexLUT != nullptr)
+	if(InContext.Curve.UIDToArrayIndexLUT != nullptr && InContext.Curve.UIDToArrayIndexLUT->Num() > 0)
 	{
 		const TArray<uint16>& UIDToArrayIndexLookupTable = RequiredBones.GetUIDToArrayLookupTable();
 		const FSmartNameMapping* Mapping = RequiredBones.GetSkeletonAsset()->GetSmartNameContainer(USkeleton::AnimCurveMappingName);
-	
+		
 		check(UIDToArrayIndexLookupTable.Num() == InContext.Curve.UIDToArrayIndexLUT->Num());
 
 		for (int32 CurveUID = 0; CurveUID < UIDToArrayIndexLookupTable.Num(); ++CurveUID)
