Missed this when fixing bad merge in CL 19902582

#rnx

diff --git a/Engine/Source/Runtime/Engine/Private/World.cpp b/Engine/Source/Runtime/Engine/Private/World.cpp
--- a/Engine/Source/Runtime/Engine/Private/World.cpp
+++ b/Engine/Source/Runtime/Engine/Private/World.cpp
@@ -5052,26 +5052,27 @@ void UWorld::CleanupWorldInternal(bool bSessionEnded, bool bCleanupResources, bo
 				});
 			}
 		}
+	}
 
-		// Cleanup Persistent level outside of following loop because uninitialized worlds don't have a valid Levels array
-		// StreamingLevels are not initialized.
-		if (PersistentLevel)
-		{
-			PersistentLevel->CleanupLevel(bCleanupResources);
-			PersistentLevel->CleanupReferences();
-		}
+	// Cleanup Persistent level outside of following loop because uninitialized worlds don't have a valid Levels array
+	// StreamingLevels are not initialized.
+	if (PersistentLevel)
+	{
+		PersistentLevel->CleanupLevel(bCleanupResources, bUnloadFromEditor);
+		PersistentLevel->CleanupReferences();
+	}
 
-		if (GetNumLevels() > 1)
+	if (GetNumLevels() > 1)
+	{
+		check(GetLevel(0) == PersistentLevel);
+		for (int32 LevelIndex = 1; LevelIndex < GetNumLevels(); ++LevelIndex)
 		{
-			check(GetLevel(0) == PersistentLevel);
-			for (int32 LevelIndex = 1; LevelIndex < GetNumLevels(); ++LevelIndex)
-			{
-				ULevel* Level = GetLevel(LevelIndex);
-				Level->CleanupLevel(bCleanupResources);
-				Level->CleanupReferences();
-			}
+			ULevel* Level = GetLevel(LevelIndex);
+			Level->CleanupLevel(bCleanupResources, bUnloadFromEditor);
+			Level->CleanupReferences();
 		}
 	}
+
 #else
 	if (PersistentLevel)
 	{
