Engine changes required to get modular (non-monolithic) console builds working.

#rnx

diff --git a/Engine/Source/Runtime/Engine/Private/UnrealEngine.cpp b/Engine/Source/Runtime/Engine/Private/UnrealEngine.cpp
--- a/Engine/Source/Runtime/Engine/Private/UnrealEngine.cpp
+++ b/Engine/Source/Runtime/Engine/Private/UnrealEngine.cpp
@@ -825,7 +825,7 @@ bool ParseResolution(const TCHAR* InResolution, uint32& OutX, uint32& OutY, int3
 	return false;
 }
 
-void HDRSettingChangedSinkCallback()
+void ENGINE_API HDRSettingChangedSinkCallback()
 {
 	static const auto CVarHDROutputEnabled = IConsoleManager::Get().FindTConsoleVariableDataInt(TEXT("r.HDR.EnableHDROutput"));
 	check(CVarHDROutputEnabled);
diff --git a/Engine/Source/Runtime/Engine/Public/Audio.h b/Engine/Source/Runtime/Engine/Public/Audio.h
--- a/Engine/Source/Runtime/Engine/Public/Audio.h
+++ b/Engine/Source/Runtime/Engine/Public/Audio.h
@@ -48,7 +48,7 @@ DECLARE_CYCLE_STAT_EXTERN(TEXT("Updating Effects"), STAT_AudioUpdateEffects, STA
 DECLARE_CYCLE_STAT_EXTERN(TEXT("Source Init"), STAT_AudioSourceInitTime, STATGROUP_Audio, ENGINE_API);
 DECLARE_CYCLE_STAT_EXTERN(TEXT("Source Create"), STAT_AudioSourceCreateTime, STATGROUP_Audio, ENGINE_API);
 DECLARE_CYCLE_STAT_EXTERN(TEXT("Submit Buffers"), STAT_AudioSubmitBuffersTime, STATGROUP_Audio, ENGINE_API);
-DECLARE_CYCLE_STAT_EXTERN(TEXT("Decompress Audio"), STAT_AudioDecompressTime, STATGROUP_Audio, );
+DECLARE_CYCLE_STAT_EXTERN(TEXT("Decompress Audio"), STAT_AudioDecompressTime, STATGROUP_Audio, ENGINE_API);
 DECLARE_CYCLE_STAT_EXTERN(TEXT("Decompress Vorbis"), STAT_VorbisDecompressTime, STATGROUP_Audio, );
 DECLARE_CYCLE_STAT_EXTERN(TEXT("Prepare Audio Decompression"), STAT_AudioPrepareDecompressionTime, STATGROUP_Audio, );
 DECLARE_CYCLE_STAT_EXTERN(TEXT("Prepare Vorbis Decompression"), STAT_VorbisPrepareDecompressionTime, STATGROUP_Audio, );
diff --git a/Engine/Source/Runtime/Engine/Public/AudioDecompress.h b/Engine/Source/Runtime/Engine/Public/AudioDecompress.h
--- a/Engine/Source/Runtime/Engine/Public/AudioDecompress.h
+++ b/Engine/Source/Runtime/Engine/Public/AudioDecompress.h
@@ -233,7 +233,7 @@ public:
 protected:
 
 	/** Reads from the internal source audio buffer stream of the given data size. */
-	uint32	Read(void *Outbuffer, uint32 DataSize);
+	ENGINE_API uint32	Read(void *Outbuffer, uint32 DataSize);
 
 	/**
 	* Decompresses a frame of data to PCM buffer
@@ -276,7 +276,7 @@ protected:
 	 * @param[out] OutChunkSize the size of the chunk.
 	 * @return a pointer to the chunk if it's loaded, nullptr otherwise.
 	 */
-	const uint8* GetLoadedChunk(const FSoundWaveProxyPtr& InSoundWave, uint32 ChunkIndex, uint32& OutChunkSize);
+	ENGINE_API const uint8* GetLoadedChunk(const FSoundWaveProxyPtr& InSoundWave, uint32 ChunkIndex, uint32& OutChunkSize);
 
 	/** bool set before ParseHeader. Whether we are streaming a file or not. */
 	bool bIsStreaming;
diff --git a/Engine/Source/Runtime/RHICore/Public/RHICoreTransientResourceAllocator.h b/Engine/Source/Runtime/RHICore/Public/RHICoreTransientResourceAllocator.h
--- a/Engine/Source/Runtime/RHICore/Public/RHICoreTransientResourceAllocator.h
+++ b/Engine/Source/Runtime/RHICore/Public/RHICoreTransientResourceAllocator.h
@@ -976,7 +976,7 @@ public:
 	struct FInitializer
 	{
 		// Creates a default initializer using common RHI CVars.
-		static FInitializer CreateDefault();
+		static RHICORE_API FInitializer CreateDefault();
 
 		static const uint32 kDefaultResourceCacheSize = 256;
 
