Fix APlayerController default for motion controls to enabled

#android
#ios

diff --git a/Engine/Source/Runtime/Engine/Private/PlayerController.cpp b/Engine/Source/Runtime/Engine/Private/PlayerController.cpp
--- a/Engine/Source/Runtime/Engine/Private/PlayerController.cpp
+++ b/Engine/Source/Runtime/Engine/Private/PlayerController.cpp
@@ -128,6 +128,9 @@ APlayerController::APlayerController(const FObjectInitializer& ObjectInitializer
 	bForceFeedbackEnabled = true;
 	ForceFeedbackScale = 1.f;
 
+	// default to true; won't do anything if enable motion controls in input settings isn't also true
+	SetMotionControlsEnabled(true);
+
 	bEnableStreamingSource = true;
 	bStreamingSourceShouldActivate = true;
 	bStreamingSourceShouldBlockOnSlowStreaming = true;
