Chaos: Add minimum scale to convex objects to match other geometries to avoid pushing zero scale objects into the solver resolving behaviour issues and crashes inside the solver. Currently this crashes but the fix is a behaviour change from 4.27 where we would fail to create any geometry for these shapes.

diff --git a/Engine/Source/Runtime/Engine/Private/PhysicsEngine/Experimental/ChaosInterfaceUtils.cpp b/Engine/Source/Runtime/Engine/Private/PhysicsEngine/Experimental/ChaosInterfaceUtils.cpp
--- a/Engine/Source/Runtime/Engine/Private/PhysicsEngine/Experimental/ChaosInterfaceUtils.cpp
+++ b/Engine/Source/Runtime/Engine/Private/PhysicsEngine/Experimental/ChaosInterfaceUtils.cpp
@@ -353,7 +353,13 @@ namespace ChaosInterface
 				if (const auto& ConvexImplicit = CollisionBody.GetChaosConvexMesh())
 				{
 					// Extract the scale from the transform - we have separate wrapper classes for scale versus translate/rotate 
-					const FVector NetScale = Scale * InParams.LocalTransform.GetScale3D();
+					FVector NetScale = Scale * InParams.LocalTransform.GetScale3D();
+
+					// If Scale is zero in any component, set minimum positive instead
+					NetScale.X = FMath::Abs(NetScale.X) < KINDA_SMALL_NUMBER ? KINDA_SMALL_NUMBER : NetScale.X;
+					NetScale.Y = FMath::Abs(NetScale.Y) < KINDA_SMALL_NUMBER ? KINDA_SMALL_NUMBER : NetScale.Y;
+					NetScale.Z = FMath::Abs(NetScale.Z) < KINDA_SMALL_NUMBER ? KINDA_SMALL_NUMBER : NetScale.Z;
+
 					FTransform ConvexTransform = FTransform(InParams.LocalTransform.GetRotation(), Scale * InParams.LocalTransform.GetLocation(), FVector(1, 1, 1));
 					const FVector ScaledSize = (NetScale.GetAbs() * CollisionBody.ElemBox.GetSize());	// Note: Scale can be negative
 					const Chaos::FReal CollisionMargin = FMath::Min<Chaos::FReal>(ScaledSize.GetMin() * CollisionMarginFraction, CollisionMarginMax);
