World Partition:
- Set streaming source components default streaming priority to normal and set a default debug color.
- Expose player controller streaming source priority so it can be tweaked by the user.

#rnx

diff --git a/Engine/Source/Runtime/Engine/Classes/GameFramework/PlayerController.h b/Engine/Source/Runtime/Engine/Classes/GameFramework/PlayerController.h
--- a/Engine/Source/Runtime/Engine/Classes/GameFramework/PlayerController.h
+++ b/Engine/Source/Runtime/Engine/Classes/GameFramework/PlayerController.h
@@ -524,6 +524,10 @@ public:
 	UPROPERTY(EditAnywhere, BlueprintReadWrite, Category = WorldPartition, meta=(EditCondition="bEnableStreamingSource"))
 	uint32 bStreamingSourceShouldBlockOnSlowStreaming:1;
 
+	/** PlayerController streaming source priority. */
+	UPROPERTY(EditAnywhere, BlueprintReadWrite, Category = WorldPartition, meta=(EditCondition="bEnableStreamingSource"))
+	EStreamingSourcePriority StreamingSourcePriority;
+
 	/** Scale applied to force feedback values */
 	UPROPERTY(config)
 	float ForceFeedbackScale;
@@ -740,6 +744,14 @@ public:
 	UFUNCTION(BlueprintCallable, Category = WorldPartition)
 	virtual bool StreamingSourceShouldBlockOnSlowStreaming() const { return bEnableStreamingSource && bStreamingSourceShouldBlockOnSlowStreaming; }
 
+	/**
+	* Gets the streaming source priority.
+	* Default implementation returns StreamingSourcePriority but can be overriden in child classes.
+	* @return the streaming source priority.
+	*/
+	UFUNCTION(BlueprintCallable, Category = WorldPartition)
+	virtual EStreamingSourcePriority GetStreamingSourcePriority() const { return StreamingSourcePriority; }
+
 protected:
 	/** Pawn has been possessed, so changing state to NAME_Playing. Start it walking and begin playing with it. */
 	virtual void BeginPlayingState();
diff --git a/Engine/Source/Runtime/Engine/Private/Components/WorldPartitionStreamingSourceComponent.cpp b/Engine/Source/Runtime/Engine/Private/Components/WorldPartitionStreamingSourceComponent.cpp
--- a/Engine/Source/Runtime/Engine/Private/Components/WorldPartitionStreamingSourceComponent.cpp
+++ b/Engine/Source/Runtime/Engine/Private/Components/WorldPartitionStreamingSourceComponent.cpp
@@ -1,6 +1,8 @@
 // Copyright Epic Games, Inc. All Rights Reserved.
 
 #include "Components/WorldPartitionStreamingSourceComponent.h"
+#include "Math/Color.h"
+#include "Math/RandomStream.h"
 #include "WorldPartition/WorldPartition.h"
 #include "WorldPartition/WorldPartitionRuntimeCell.h"
 #include "WorldPartition/WorldPartitionSubsystem.h"
@@ -16,8 +18,8 @@ UWorldPartitionStreamingSourceComponent::UWorldPartitionStreamingSourceComponent
 	, DefaultVisualizerLoadingRange(10000.f)
 #endif
 	, TargetGrid(NAME_None)
-	, DebugColor(ForceInit)
-	, Priority(EStreamingSourcePriority::Low)
+	, DebugColor(FColor::MakeRedToGreenColorFromScalar(FRandomStream(GetFName()).GetFraction()))
+	, Priority(EStreamingSourcePriority::Normal)
 	, bStreamingSourceEnabled(true)
 	, TargetState(EStreamingSourceTargetState::Activated)
 {
diff --git a/Engine/Source/Runtime/Engine/Private/PlayerController.cpp b/Engine/Source/Runtime/Engine/Private/PlayerController.cpp
--- a/Engine/Source/Runtime/Engine/Private/PlayerController.cpp
+++ b/Engine/Source/Runtime/Engine/Private/PlayerController.cpp
@@ -66,6 +66,7 @@
 #include "GameFramework/GameSession.h"
 #include "GameMapsSettings.h"
 #include "Particles/EmitterCameraLensEffectBase.h"
+#include "WorldPartition/WorldPartitionStreamingSource.h"
 
 DEFINE_LOG_CATEGORY(LogPlayerController);
 
@@ -134,6 +135,7 @@ APlayerController::APlayerController(const FObjectInitializer& ObjectInitializer
 	bEnableStreamingSource = true;
 	bStreamingSourceShouldActivate = true;
 	bStreamingSourceShouldBlockOnSlowStreaming = true;
+	StreamingSourcePriority = EStreamingSourcePriority::Default;
 
 	bAutoManageActiveCameraTarget = true;
 	bRenderPrimitiveComponents = true;
diff --git a/Engine/Source/Runtime/Engine/Private/WorldPartition/WorldPartitionStreamingPolicy.cpp b/Engine/Source/Runtime/Engine/Private/WorldPartition/WorldPartitionStreamingPolicy.cpp
--- a/Engine/Source/Runtime/Engine/Private/WorldPartition/WorldPartitionStreamingPolicy.cpp
+++ b/Engine/Source/Runtime/Engine/Private/WorldPartition/WorldPartitionStreamingPolicy.cpp
@@ -134,7 +134,8 @@ void UWorldPartitionStreamingPolicy::UpdateStreamingSources()
 						ViewRotation = WorldToLocal.TransformRotation(ViewRotation.Quaternion()).Rotator();
 						const EStreamingSourceTargetState TargetState = Player->PlayerController->StreamingSourceShouldActivate() ? EStreamingSourceTargetState::Activated : EStreamingSourceTargetState::Loaded;
 						const bool bBlockOnSlowLoading = Player->PlayerController->StreamingSourceShouldBlockOnSlowStreaming();
-						StreamingSources.Add(FWorldPartitionStreamingSource(Player->PlayerController->GetFName(), ViewLocation, ViewRotation, TargetState, bBlockOnSlowLoading, EStreamingSourcePriority::Default));
+						const EStreamingSourcePriority StreamingSourcePriority = Player->PlayerController->GetStreamingSourcePriority();
+						StreamingSources.Add(FWorldPartitionStreamingSource(Player->PlayerController->GetFName(), ViewLocation, ViewRotation, TargetState, bBlockOnSlowLoading, StreamingSourcePriority));
 					}
 				}
 			}
