[Navmesh] Increase streaming radius by the size of the largest supported agent

diff --git a/Engine/Source/Runtime/Engine/Classes/AI/NavigationSystemBase.h b/Engine/Source/Runtime/Engine/Classes/AI/NavigationSystemBase.h
--- a/Engine/Source/Runtime/Engine/Classes/AI/NavigationSystemBase.h
+++ b/Engine/Source/Runtime/Engine/Classes/AI/NavigationSystemBase.h
@@ -146,6 +146,7 @@ namespace FNavigationSystem
 	ENGINE_API bool HasComponentData(UActorComponent& Comp);
 	
 	ENGINE_API const FNavDataConfig& GetDefaultSupportedAgent();
+	ENGINE_API const FNavDataConfig& GetBiggestSupportedAgent(const UWorld* World);
 
 	ENGINE_API TSubclassOf<UNavAreaBase> GetDefaultWalkableArea();
 	ENGINE_API TSubclassOf<UNavAreaBase> GetDefaultObstacleArea();
@@ -194,7 +195,8 @@ namespace FNavigationSystem
 	DECLARE_DELEGATE_RetVal_OneParam(bool, FBoolControllerBasedSignature, const AController& /*Controller*/);
 	DECLARE_DELEGATE_RetVal_OneParam(bool, FBoolActorComponentBasedSignature, UActorComponent& /*Comp*/);
 	DECLARE_DELEGATE_RetVal(TSubclassOf<UNavAreaBase>, FNavAreaBasedSignature);
-	DECLARE_DELEGATE_RetVal(const FNavDataConfig&, FNavDatConfigBasedSignature);
+	DECLARE_DELEGATE_RetVal(const FNavDataConfig&, FNavDataConfigBasedSignature);
+	DECLARE_DELEGATE_RetVal_OneParam(const FNavDataConfig&, FNavDataConfigAndWorldSignature, const UWorld* /*World*/);
 	DECLARE_DELEGATE_TwoParams(FWorldByteBasedSignature, UWorld& /*World*/, uint8 /*Flags*/);
 	DECLARE_DELEGATE_TwoParams(FActorBooleBasedSignature, AActor& /*Actor*/, bool /*bUpdateAttachedActors*/);
 	DECLARE_DELEGATE_ThreeParams(FComponentBoundsChangeSignature, UActorComponent& /*Comp*/, const FBox& /*NewBounds*/, const FBox& /*DirtyArea*/)
@@ -284,7 +286,8 @@ protected:
 	static FNavigationSystem::FActorComponentBasedSignature& OnComponentUnregisteredDelegate();
 	static FNavigationSystem::FActorBasedSignature& RemoveActorDataDelegate();
 	static FNavigationSystem::FBoolActorComponentBasedSignature& HasComponentDataDelegate();
-	static FNavigationSystem::FNavDatConfigBasedSignature& GetDefaultSupportedAgentDelegate();
+	static FNavigationSystem::FNavDataConfigBasedSignature& GetDefaultSupportedAgentDelegate();
+	static FNavigationSystem::FNavDataConfigAndWorldSignature& GetBiggestSupportedAgentDelegate();
 	static FNavigationSystem::FActorBooleBasedSignature& UpdateActorAndComponentDataDelegate();
 	static FNavigationSystem::FComponentBoundsChangeSignature& OnComponentBoundsChangedDelegate();
 	static FNavigationSystem::FNavDataForActorSignature& GetNavDataForActorDelegate();
diff --git a/Engine/Source/Runtime/Engine/Private/AI/NavigationSystemBase.cpp b/Engine/Source/Runtime/Engine/Private/AI/NavigationSystemBase.cpp
--- a/Engine/Source/Runtime/Engine/Private/AI/NavigationSystemBase.cpp
+++ b/Engine/Source/Runtime/Engine/Private/AI/NavigationSystemBase.cpp
@@ -111,6 +111,12 @@ namespace FNavigationSystem
 		return FallbackSupportedAgent; 
 	}
 
+	const FNavDataConfig& GetFallbackSupportedAgent(const UWorld* World) 
+	{ 
+		static FNavDataConfig FallbackSupportedAgent;
+		return FallbackSupportedAgent; 
+	}
+	
 	bool bWantsComponentChangeNotifies = true;
 	
 	class FDelegates
@@ -130,7 +136,8 @@ namespace FNavigationSystem
 		FControllerBasedSignature StopMovement;
 		FBoolControllerBasedSignature IsFollowingAPath;
 		FBoolActorComponentBasedSignature HasComponentData;
-		FNavDatConfigBasedSignature GetDefaultSupportedAgent;
+		FNavDataConfigBasedSignature GetDefaultSupportedAgent;
+		FNavDataConfigAndWorldSignature GetBiggestSupportedAgent;
 		FActorBooleBasedSignature UpdateActorAndComponentData;
 		FComponentBoundsChangeSignature OnComponentBoundsChanged;
 		FNavDataForActorSignature GetNavDataForActor;
@@ -165,6 +172,7 @@ namespace FNavigationSystem
 			IsFollowingAPath.BindLambda([](const AController&) { return false; });
 			HasComponentData.BindLambda([](UActorComponent&) { return false; });
 			GetDefaultSupportedAgent.BindStatic(&GetFallbackSupportedAgent);
+			GetBiggestSupportedAgent.BindStatic(&GetFallbackSupportedAgent);			
 			UpdateActorAndComponentData.BindLambda([](AActor&, bool) {});
 			OnComponentBoundsChanged.BindLambda([](UActorComponent&, const FBox&, const FBox&) {});
 			GetNavDataForActor.BindLambda([](const AActor&) { return nullptr; });
@@ -201,6 +209,7 @@ namespace FNavigationSystem
 	void RemoveActorData(AActor& Actor) { Delegates.RemoveActorData.Execute(Actor); }
 	bool HasComponentData(UActorComponent& Comp) { return Delegates.HasComponentData.Execute(Comp);	}
 	const FNavDataConfig& GetDefaultSupportedAgent() { return Delegates.GetDefaultSupportedAgent.Execute(); }
+	const FNavDataConfig& GetBiggestSupportedAgent(const UWorld* World) { return Delegates.GetBiggestSupportedAgent.Execute(World); }
 
 
 	TSubclassOf<UNavAreaBase> DefaultWalkableArea; 
@@ -393,7 +402,8 @@ FNavigationSystem::FActorComponentBasedSignature& UNavigationSystemBase::OnCompo
 FNavigationSystem::FActorComponentBasedSignature& UNavigationSystemBase::OnComponentUnregisteredDelegate() { return FNavigationSystem::Delegates.OnComponentUnregistered; }
 FNavigationSystem::FActorBasedSignature& UNavigationSystemBase::RemoveActorDataDelegate() { return FNavigationSystem::Delegates.RemoveActorData; }
 FNavigationSystem::FBoolActorComponentBasedSignature& UNavigationSystemBase::HasComponentDataDelegate() { return FNavigationSystem::Delegates.HasComponentData; }
-FNavigationSystem::FNavDatConfigBasedSignature& UNavigationSystemBase::GetDefaultSupportedAgentDelegate() { return FNavigationSystem::Delegates.GetDefaultSupportedAgent; }
+FNavigationSystem::FNavDataConfigBasedSignature& UNavigationSystemBase::GetDefaultSupportedAgentDelegate() { return FNavigationSystem::Delegates.GetDefaultSupportedAgent; }
+FNavigationSystem::FNavDataConfigAndWorldSignature& UNavigationSystemBase::GetBiggestSupportedAgentDelegate() { return FNavigationSystem::Delegates.GetBiggestSupportedAgent; }
 FNavigationSystem::FActorBooleBasedSignature& UNavigationSystemBase::UpdateActorAndComponentDataDelegate() { return FNavigationSystem::Delegates.UpdateActorAndComponentData; }
 FNavigationSystem::FComponentBoundsChangeSignature& UNavigationSystemBase::OnComponentBoundsChangedDelegate() { return FNavigationSystem::Delegates.OnComponentBoundsChanged; }
 FNavigationSystem::FNavDataForActorSignature& UNavigationSystemBase::GetNavDataForActorDelegate() { return FNavigationSystem::Delegates.GetNavDataForActor; }
diff --git a/Engine/Source/Runtime/Engine/Private/Components/SceneComponent.cpp b/Engine/Source/Runtime/Engine/Private/Components/SceneComponent.cpp
--- a/Engine/Source/Runtime/Engine/Private/Components/SceneComponent.cpp
+++ b/Engine/Source/Runtime/Engine/Private/Components/SceneComponent.cpp
@@ -2397,7 +2397,15 @@ TStructOnScope<FActorComponentInstanceData> USceneComponent::GetComponentInstanc
 #if WITH_EDITOR
 FBox USceneComponent::GetStreamingBounds() const
 {
-	return Bounds.GetBox();
+	FBox Box = Bounds.GetBox();
+	
+	if (IsNavigationRelevant())
+	{
+		const FNavDataConfig& Config = FNavigationSystem::GetBiggestSupportedAgent(GetWorld());
+		Box = Box.ExpandBy(Config.AgentRadius);
+	}
+
+	return Box;
 }
 #endif
 
diff --git a/Engine/Source/Runtime/NavigationSystem/Private/NavigationSystem.cpp b/Engine/Source/Runtime/NavigationSystem/Private/NavigationSystem.cpp
--- a/Engine/Source/Runtime/NavigationSystem/Private/NavigationSystem.cpp
+++ b/Engine/Source/Runtime/NavigationSystem/Private/NavigationSystem.cpp
@@ -34,6 +34,7 @@
 #include "EditorModeManager.h"
 #include "EditorModes.h"
 #include "Editor/LevelEditor/Public/LevelEditor.h"
+#include "Misc/MessageDialog.h"
 #endif
 
 #include "NavAreas/NavArea_Null.h"
@@ -449,6 +450,7 @@ UNavigationSystemV1::UNavigationSystemV1(const FObjectInitializer& ObjectInitial
 				return (NavSys && (NavSys->GetObjectsNavOctreeId(Comp) || NavSys->HasPendingObjectNavOctreeId(&Comp)));
 			});
 			UNavigationSystemBase::GetDefaultSupportedAgentDelegate().BindStatic(&UNavigationSystemV1::GetDefaultSupportedAgent);
+			UNavigationSystemBase::GetBiggestSupportedAgentDelegate().BindStatic(&UNavigationSystemV1::GetBiggestSupportedAgent);
 			UNavigationSystemBase::UpdateActorAndComponentDataDelegate().BindStatic(&UNavigationSystemV1::UpdateActorAndComponentsInNavOctree);
 			UNavigationSystemBase::OnComponentBoundsChangedDelegate().BindLambda([](UActorComponent& Comp, const FBox& NewBounds, const FBox& DirtyArea) {
 				UNavigationSystemV1* NavSys = FNavigationSystem::GetCurrent<UNavigationSystemV1>(Comp.GetWorld());
@@ -834,6 +836,21 @@ void UNavigationSystemV1::PostEditChangeProperty(FPropertyChangedEvent& Property
 				}
 			}
 		}
+		else if (PropName == GET_MEMBER_NAME_CHECKED(FNavDataConfig, AgentRadius))
+		{
+			const bool bIsCDO = HasAnyFlags(RF_ClassDefaultObject);
+			if (!bIsCDO)
+			{
+				const UWorld* World = GetWorld();
+				if (World && World->IsPartitionedWorld())
+				{
+					FMessageDialog::Open(EAppMsgType::Ok, LOCTEXT("NeedToRunPartitionResaveActorsBuilder",
+						"In a world partitioned map, changing this property changes the partitioning of actors.\n"
+						"For the change to take effect on partitioning, actors needs to be resaved.\n"
+						"Run the WorldPartitionResaveActorsBuilder to update the whole map."));	
+				}
+			}
+		}
 	}
 }
 #endif // WITH_EDITOR
@@ -4818,6 +4835,39 @@ const FNavDataConfig& UNavigationSystemV1::GetDefaultSupportedAgent()
 		: DefaultAgent;
 }
 
+const FNavDataConfig& UNavigationSystemV1::GetBiggestSupportedAgent(const UWorld* World) 
+{
+	const UNavigationSystemV1* NavSys = nullptr;
+	if (World != nullptr)
+	{
+		NavSys = FNavigationSystem::GetCurrent<UNavigationSystemV1>(World);		
+	}
+
+	if (NavSys == nullptr)
+	{
+		// If no world is available, use the CDO.
+		NavSys = GetDefault<UNavigationSystemV1>();
+	}
+	check(NavSys);
+
+	if (NavSys->GetSupportedAgents().IsEmpty())
+	{
+		static const FNavDataConfig DefaultAgent;
+		return DefaultAgent;
+	}
+
+	const FNavDataConfig* BiggestAgent = nullptr;
+	for (const FNavDataConfig& Config : NavSys->GetSupportedAgents())
+	{
+		if (BiggestAgent == nullptr || Config.AgentRadius > BiggestAgent->AgentRadius)
+		{
+			BiggestAgent = &Config; 
+		}
+	}
+
+	return *BiggestAgent;
+}
+
 const FNavDataConfig& UNavigationSystemV1::GetDefaultSupportedAgentConfig() const 
 { 
 	static const FNavDataConfig DefaultAgent;
diff --git a/Engine/Source/Runtime/NavigationSystem/Public/NavigationSystem.h b/Engine/Source/Runtime/NavigationSystem/Public/NavigationSystem.h
--- a/Engine/Source/Runtime/NavigationSystem/Public/NavigationSystem.h
+++ b/Engine/Source/Runtime/NavigationSystem/Public/NavigationSystem.h
@@ -642,6 +642,7 @@ public:
 	const TSet<FNavigationBounds>& GetNavigationBounds() const;
 
 	static const FNavDataConfig& GetDefaultSupportedAgent();
+	static const FNavDataConfig& GetBiggestSupportedAgent(const UWorld* World);
 	const FNavDataConfig& GetDefaultSupportedAgentConfig() const;
 	FORCEINLINE const TArray<FNavDataConfig>& GetSupportedAgents() const { return SupportedAgents; }
 	void OverrideSupportedAgents(const TArray<FNavDataConfig>& NewSupportedAgents);
