UHLODLayer::DuplicateHLODLayersSetup add MarkAsFullyLoaded call to created package to allow overwriting existing package

#rnx

diff --git a/Engine/Source/Runtime/Engine/Private/WorldPartition/HLOD/HLODLayer.cpp b/Engine/Source/Runtime/Engine/Private/WorldPartition/HLOD/HLODLayer.cpp
--- a/Engine/Source/Runtime/Engine/Private/WorldPartition/HLOD/HLODLayer.cpp
+++ b/Engine/Source/Runtime/Engine/Private/WorldPartition/HLOD/HLODLayer.cpp
@@ -126,6 +126,8 @@ UHLODLayer* UHLODLayer::DuplicateHLODLayersSetup(UHLODLayer* HLODLayer, const FS
 	{
 		const FString PackageName = DestinationPath + TEXT("_") + CurrentHLODLayer->GetName();
 		UPackage* Package = CreatePackage(*PackageName);
+		// In case Package already exists setting this flag will allow overwriting it
+		Package->MarkAsFullyLoaded();
 
 		FString NewHLODLayerName = Prefix + TEXT("_") + CurrentHLODLayer->GetName();
 		UHLODLayer* NewHLODLayer = CastChecked<UHLODLayer>(StaticDuplicateObject(CurrentHLODLayer, Package, *NewHLODLayerName));
