Fix chaos CIS
#rnx

diff --git a/Engine/Source/Runtime/Engine/Private/PhysicsEngine/Experimental/ChaosInterfaceUtils.cpp b/Engine/Source/Runtime/Engine/Private/PhysicsEngine/Experimental/ChaosInterfaceUtils.cpp
--- a/Engine/Source/Runtime/Engine/Private/PhysicsEngine/Experimental/ChaosInterfaceUtils.cpp
+++ b/Engine/Source/Runtime/Engine/Private/PhysicsEngine/Experimental/ChaosInterfaceUtils.cpp
@@ -356,9 +356,9 @@ namespace ChaosInterface
 					FVector NetScale = Scale * InParams.LocalTransform.GetScale3D();
 
 					// If Scale is zero in any component, set minimum positive instead
-					NetScale.X = FMath::Abs(NetScale.X) < KINDA_SMALL_NUMBER ? KINDA_SMALL_NUMBER : NetScale.X;
-					NetScale.Y = FMath::Abs(NetScale.Y) < KINDA_SMALL_NUMBER ? KINDA_SMALL_NUMBER : NetScale.Y;
-					NetScale.Z = FMath::Abs(NetScale.Z) < KINDA_SMALL_NUMBER ? KINDA_SMALL_NUMBER : NetScale.Z;
+					NetScale.X = FMath::Abs(NetScale.X) < UE_KINDA_SMALL_NUMBER ? UE_KINDA_SMALL_NUMBER : NetScale.X;
+					NetScale.Y = FMath::Abs(NetScale.Y) < UE_KINDA_SMALL_NUMBER ? UE_KINDA_SMALL_NUMBER : NetScale.Y;
+					NetScale.Z = FMath::Abs(NetScale.Z) < UE_KINDA_SMALL_NUMBER ? UE_KINDA_SMALL_NUMBER : NetScale.Z;
 
 					FTransform ConvexTransform = FTransform(InParams.LocalTransform.GetRotation(), Scale * InParams.LocalTransform.GetLocation(), FVector(1, 1, 1));
 					const FVector ScaledSize = (NetScale.GetAbs() * CollisionBody.ElemBox.GetSize());	// Note: Scale can be negative
