Allow components to register and unregister async physics tick.
Make it so GeometryCollection doesn't register async physics tick unless replication is needed and the object simulates

diff --git a/Engine/Source/Runtime/Engine/Classes/Components/ActorComponent.h b/Engine/Source/Runtime/Engine/Classes/Components/ActorComponent.h
--- a/Engine/Source/Runtime/Engine/Classes/Components/ActorComponent.h
+++ b/Engine/Source/Runtime/Engine/Classes/Components/ActorComponent.h
@@ -297,11 +297,11 @@ private:
 	/** Tracks whether the component has been added to the world's pre end of frame sync list */
 	uint8 bMarkedForPreEndOfFrameSync : 1;
 
-public:
-
 	/** Whether to use use the async physics tick with this component. */
 	uint8 bAsyncPhysicsTickEnabled : 1;
 
+public:
+
 	/** Describes how a component instance will be created */
 	UPROPERTY()
 	EComponentCreationMethod CreationMethod;
@@ -1085,6 +1085,9 @@ protected:
 	 *		already been registered. Setting bForceUpdate to true overrides that check */
 	void HandleCanEverAffectNavigationChange(bool bForceUpdate = false);
 
+	/** Sets bAsyncPhysicsTickEnabled which determines whether to use use the async physics tick with this component. */
+	void SetAsyncPhysicsTickEnabled(bool bEnabled);
+
 private:
 #if !(UE_BUILD_SHIPPING || UE_BUILD_TEST)
 	/** This is the old name of the tick function. We just want to avoid mistakes with an attempt to override this */
@@ -1093,6 +1096,8 @@ private:
 
 	void ClearNeedEndOfFrameUpdate_Internal();
 
+	void RegisterAsyncPhysicsTickEnabled(bool bRegister);
+
 	friend struct FMarkComponentEndOfFrameUpdateState;
 	friend struct FSetUCSSerializationIndex;
 	friend class FActorComponentDetails;
diff --git a/Engine/Source/Runtime/Engine/Private/Components/ActorComponent.cpp b/Engine/Source/Runtime/Engine/Private/Components/ActorComponent.cpp
--- a/Engine/Source/Runtime/Engine/Private/Components/ActorComponent.cpp
+++ b/Engine/Source/Runtime/Engine/Private/Components/ActorComponent.cpp
@@ -1208,25 +1208,43 @@ void UActorComponent::RegisterAllComponentTickFunctions(bool bRegister)
 			GTestRegisterComponentTickFunctions = NULL;
 		}
 
+		if (bAsyncPhysicsTickEnabled)
+		{
+			RegisterAsyncPhysicsTickEnabled(bRegister);
+		}
+	}
+}
+
+void UActorComponent::RegisterAsyncPhysicsTickEnabled(bool bRegister)
+{
 #if WITH_CHAOS
-		if(bAsyncPhysicsTickEnabled)
+	if (FPhysScene_Chaos* Scene = static_cast<FPhysScene_Chaos*>(WorldPrivate->GetPhysicsScene()))
+	{
+		if (bRegister)
 		{
-			if(FPhysScene_Chaos* Scene = static_cast<FPhysScene_Chaos*>(WorldPrivate->GetPhysicsScene()))
-			{
-				if(bRegister)
-				{
-					Scene->RegisterAsyncPhysicsTickComponent(this);
-				}
-				else
-				{
-					Scene->UnregisterAsyncPhysicsTickComponent(this);
-				}
-			}
+			Scene->RegisterAsyncPhysicsTickComponent(this);
+		}
+		else
+		{
+			Scene->UnregisterAsyncPhysicsTickComponent(this);
 		}
 	}
 #endif
 }
 
+void UActorComponent::SetAsyncPhysicsTickEnabled(bool bEnable)
+{
+#if WITH_CHAOS
+	// Components don't have async physics functions until they are registered with the world
+	if(bRegistered)
+	{
+		RegisterAsyncPhysicsTickEnabled(bEnable);
+	}
+	
+	bAsyncPhysicsTickEnabled = bEnable;
+#endif
+}
+
 void UActorComponent::TickComponent(float DeltaTime, enum ELevelTick TickType, FActorComponentTickFunction *ThisTickFunction)
 {
 	check(bRegistered);
diff --git a/Engine/Source/Runtime/Engine/Private/PhysicsEngine/AsyncPhysicsInputComponent.cpp b/Engine/Source/Runtime/Engine/Private/PhysicsEngine/AsyncPhysicsInputComponent.cpp
--- a/Engine/Source/Runtime/Engine/Private/PhysicsEngine/AsyncPhysicsInputComponent.cpp
+++ b/Engine/Source/Runtime/Engine/Private/PhysicsEngine/AsyncPhysicsInputComponent.cpp
@@ -25,6 +25,9 @@ namespace InputCmdCVars
 	static int32 ForceFault = 0;
 	static FAutoConsoleVariableRef CVarForceFault(TEXT("p.net.ForceFault"), ForceFault, TEXT("Forces server side input fault"));
 
+	static int32 ForceInputDrop = 0;
+	static FAutoConsoleVariableRef CVarForceInputDrop(TEXT("p.net.ForceInputDrop"), ForceInputDrop, TEXT("Forces client to drop inputs. Useful for simulating desync"));
+
 	static int32 MaxBufferedCmds = 16;
 	static FAutoConsoleVariableRef CVarMaxBufferedCmds(TEXT("p.net.MaxBufferedCmds"), MaxBufferedCmds, TEXT("MaxNumber of buffered server side commands"));
 
@@ -292,10 +295,10 @@ struct FAsyncPhysicsInputRewindCallback : public Chaos::IRewindCallback
 
 UAsyncPhysicsInputComponent::UAsyncPhysicsInputComponent()
 {
-	bAsyncPhysicsTickEnabled = true;
 	bWantsInitializeComponent = true;
 	bAutoActivate = true;
 	SetIsReplicatedByDefault(true);
+	SetAsyncPhysicsTickEnabled(true);
 }
 
 void UAsyncPhysicsInputComponent::AsyncPhysicsTickComponent(float DeltaTime, float SimTime)
@@ -330,7 +333,7 @@ void UAsyncPhysicsInputComponent::AsyncPhysicsTickComponent(float DeltaTime, flo
 		bool bFreeInput = Input->ServerFrame < ServerFrame;	//for most cases once we're passed this frame we can free
 		if (APlayerController* PC = GetPlayerController())
 		{
-			if (PC->IsLocalController())
+			if (PC->IsLocalController() && !InputCmdCVars::ForceInputDrop)
 			{
 				FAsyncPhysicsInputWrapper Wrapper;
 				Wrapper.Input = Input;
diff --git a/Engine/Source/Runtime/Experimental/GeometryCollectionEngine/Private/GeometryCollection/GeometryCollectionComponent.cpp b/Engine/Source/Runtime/Experimental/GeometryCollectionEngine/Private/GeometryCollection/GeometryCollectionComponent.cpp
--- a/Engine/Source/Runtime/Experimental/GeometryCollectionEngine/Private/GeometryCollection/GeometryCollectionComponent.cpp
+++ b/Engine/Source/Runtime/Experimental/GeometryCollectionEngine/Private/GeometryCollection/GeometryCollectionComponent.cpp
@@ -271,8 +271,6 @@ UGeometryCollectionComponent::UGeometryCollectionComponent(const FObjectInitiali
 	bTickInEditor = true;
 	bAutoActivate = true;
 
-	bAsyncPhysicsTickEnabled = true;
-
 	static uint32 GlobalNavMeshInvalidationCounter = 0;
 	//space these out over several frames (3 is arbitrary)
 	GlobalNavMeshInvalidationCounter += 3;
@@ -2241,6 +2239,7 @@ void UGeometryCollectionComponent::RegisterAndInitializePhysicsProxy()
 	Scene->AddObject(this, PhysicsProxy);
 
 	RegisterForEvents();
+	SetAsyncPhysicsTickEnabled(GetIsReplicated());
 #endif
 }
 
