Avoid allocating an FString just for the sake of comparing device profile names.  Shows up in profiling at cook time due to high call frequency.

diff --git a/Engine/Source/Runtime/Engine/Private/DeviceProfiles/DeviceProfileManager.cpp b/Engine/Source/Runtime/Engine/Private/DeviceProfiles/DeviceProfileManager.cpp
--- a/Engine/Source/Runtime/Engine/Private/DeviceProfiles/DeviceProfileManager.cpp
+++ b/Engine/Source/Runtime/Engine/Private/DeviceProfiles/DeviceProfileManager.cpp
@@ -722,7 +722,7 @@ UDeviceProfile* UDeviceProfileManager::FindProfile( const FString& ProfileName,
 	for( int32 Idx = 0; Idx < Profiles.Num(); Idx++ )
 	{
 		UDeviceProfile* CurrentDevice = CastChecked<UDeviceProfile>( Profiles[Idx] );
-		if( CurrentDevice->GetName() == ProfileName )
+		if( CurrentDevice->GetFName() == *ProfileName )
 		{
 			FoundProfile = CurrentDevice;
 			break;
