Texture: Fix population of FTexturePlatformData from the cache

GetNumNonStreamingMips is incorrect for some textures and was causing the platform data to incorrectly identify which mips are streaming.

#rnx

diff --git a/Engine/Source/Runtime/Engine/Private/TextureDerivedDataTask.cpp b/Engine/Source/Runtime/Engine/Private/TextureDerivedDataTask.cpp
--- a/Engine/Source/Runtime/Engine/Private/TextureDerivedDataTask.cpp
+++ b/Engine/Source/Runtime/Engine/Private/TextureDerivedDataTask.cpp
@@ -24,6 +24,7 @@
 #if WITH_EDITOR
 
 #include "Algo/Accumulate.h"
+#include "Algo/AnyOf.h"
 #include "DerivedDataBuild.h"
 #include "DerivedDataBuildAction.h"
 #include "DerivedDataBuildInputResolver.h"
@@ -826,12 +827,17 @@ void FTextureCacheDerivedDataWorker::DoWork()
 		}
 		else
 		{
-			if (const int32 StreamingMipCount = FMath::Max(0, DerivedData->Mips.Num() - DerivedData->GetNumNonStreamingMips()))
+			if (Algo::AnyOf(DerivedData->Mips, [](const FTexture2DMipMap& Mip) { return !Mip.BulkData.IsBulkDataLoaded(); }))
 			{
+				int32 MipIndex = 0;
 				const FSharedString Name(TexturePathName);
-				for (int32 MipIndex = 0; MipIndex < StreamingMipCount; ++MipIndex)
+				for (FTexture2DMipMap& Mip : DerivedData->Mips)
 				{
-					DerivedData->Mips[MipIndex].DerivedData = FDerivedData(Name, ConvertLegacyCacheKey(DerivedData->GetDerivedDataMipKeyString(MipIndex, DerivedData->Mips[MipIndex])));
+					if (!Mip.BulkData.IsBulkDataLoaded())
+					{
+						Mip.DerivedData = FDerivedData(Name, ConvertLegacyCacheKey(DerivedData->GetDerivedDataMipKeyString(MipIndex, Mip)));
+					}
+					++MipIndex;
 				}
 			}
 		}
