Remove aggressive subobject replacement
Leaving a toggleable audit in place (defaulting to on in non-shipping builds) to report any cases where the aggressive replacement was in use via the replacement archiver. This won't catch any cases that NotifyToolsOfObjectReplacement, but it would have made even less sense that those replacements were being applied (and were the cause of the bug we were investigating).
Also leaving a toggle to re-renable the behavior if issues arise/ A/B testing is required.

diff --git a/Engine/Plugins/Runtime/nDisplay/Source/DisplayCluster/Private/Game/EngineClasses/Scene/DisplayClusterRootActor.cpp b/Engine/Plugins/Runtime/nDisplay/Source/DisplayCluster/Private/Game/EngineClasses/Scene/DisplayClusterRootActor.cpp
--- a/Engine/Plugins/Runtime/nDisplay/Source/DisplayCluster/Private/Game/EngineClasses/Scene/DisplayClusterRootActor.cpp
+++ b/Engine/Plugins/Runtime/nDisplay/Source/DisplayCluster/Private/Game/EngineClasses/Scene/DisplayClusterRootActor.cpp
@@ -329,7 +329,10 @@ void ADisplayClusterRootActor::UpdateConfigDataInstance(UDisplayClusterConfigura
 		else if (bForceRecreate)
 		{
 			UEngine::FCopyPropertiesForUnrelatedObjectsParams Params;
+			PRAGMA_DISABLE_DEPRECATION_WARNINGS
+			// Leaving this enabled for now for the purposes of the aggressive replacement auditing
 			Params.bAggressiveDefaultSubobjectReplacement = true;
+			PRAGMA_ENABLE_DEPRECATION_WARNINGS
 			Params.bNotifyObjectReplacement = false;
 			Params.bDoDelta = false;
 			UEngine::CopyPropertiesForUnrelatedObjects(ConfigDataTemplate, CurrentConfigData, Params);
diff --git a/Engine/Source/Editor/UnrealEd/Private/Kismet2/KismetReinstanceUtilities.cpp b/Engine/Source/Editor/UnrealEd/Private/Kismet2/KismetReinstanceUtilities.cpp
--- a/Engine/Source/Editor/UnrealEd/Private/Kismet2/KismetReinstanceUtilities.cpp
+++ b/Engine/Source/Editor/UnrealEd/Private/Kismet2/KismetReinstanceUtilities.cpp
@@ -2084,7 +2084,10 @@ static void ReplaceActorHelper(AActor* OldActor, UClass* OldClass, UObject*& New
 
 	UEngine::FCopyPropertiesForUnrelatedObjectsParams Params;
 	Params.bPreserveRootComponent = bPreserveRootComponent;
+	PRAGMA_DISABLE_DEPRECATION_WARNINGS
+	// Leaving this enabled for now for the purposes of the aggressive replacement auditing
 	Params.bAggressiveDefaultSubobjectReplacement = true;
+	PRAGMA_ENABLE_DEPRECATION_WARNINGS
 	Params.bNotifyObjectReplacement = true;
 	UEngine::CopyPropertiesForUnrelatedObjects(OldActor, NewActor, Params);
 
@@ -2722,7 +2725,6 @@ void FBlueprintCompileReinstancer::CopyPropertiesForUnrelatedObjects(UObject* Ol
 	InstancedPropertyUtils::FArchiveInstancedSubObjCollector  InstancedSubObjCollector(OldObject, InstancedPropertyMap);
 
 	UEngine::FCopyPropertiesForUnrelatedObjectsParams Params;
-	Params.bAggressiveDefaultSubobjectReplacement = false;
 	// During a blueprint reparent, delta serialization must be enabled to correctly copy all properties
 	Params.bDoDelta = bForceDeltaSerialization || !OldObject->HasAnyFlags(RF_ClassDefaultObject);
 	Params.bCopyDeprecatedProperties = true;
diff --git a/Engine/Source/Runtime/Engine/Classes/Engine/Engine.h b/Engine/Source/Runtime/Engine/Classes/Engine/Engine.h
--- a/Engine/Source/Runtime/Engine/Classes/Engine/Engine.h
+++ b/Engine/Source/Runtime/Engine/Classes/Engine/Engine.h
@@ -2771,6 +2771,7 @@ public:
 	/** Makes a strong effort to copy everything possible from and old object to a new object of a different class, used for blueprint to update things after a recompile. */
 	struct FCopyPropertiesForUnrelatedObjectsParams
 	{
+		UE_DEPRECATED(5.1, "Aggressive Default Subobject Replacement is no longer being done. An ensure has been left in place to catch any cases that was making use of this feature.")
 		bool bAggressiveDefaultSubobjectReplacement;
 		bool bDoDelta;
 		bool bReplaceObjectClassReferences;
@@ -2784,20 +2785,10 @@ public:
 		bool bClearReferences;
 		bool bDontClearReferenceIfNewerClassExists;
 
-		FCopyPropertiesForUnrelatedObjectsParams()
-			: bAggressiveDefaultSubobjectReplacement(false)
-			, bDoDelta(true)
-			, bReplaceObjectClassReferences(true)
-			, bCopyDeprecatedProperties(false)
-			, bPreserveRootComponent(true)
-			, bPerformDuplication(false)
-			, bSkipCompilerGeneratedDefaults(false)
-			, bNotifyObjectReplacement(false)
-			, bClearReferences(true)
-			, bDontClearReferenceIfNewerClassExists(false)
-		{}
+		ENGINE_API FCopyPropertiesForUnrelatedObjectsParams();
+		ENGINE_API FCopyPropertiesForUnrelatedObjectsParams(const FCopyPropertiesForUnrelatedObjectsParams&);
 	};
-	static void CopyPropertiesForUnrelatedObjects(UObject* OldObject, UObject* NewObject, FCopyPropertiesForUnrelatedObjectsParams Params = FCopyPropertiesForUnrelatedObjectsParams());//bool bAggressiveDefaultSubobjectReplacement = false, bool bDoDelta = true);
+	static void CopyPropertiesForUnrelatedObjects(UObject* OldObject, UObject* NewObject, FCopyPropertiesForUnrelatedObjectsParams Params = FCopyPropertiesForUnrelatedObjectsParams());
 	virtual void NotifyToolsOfObjectReplacement(const TMap<UObject*, UObject*>& OldToNewInstanceMap) { }
 
 	virtual bool UseSound() const;
diff --git a/Engine/Source/Runtime/Engine/Private/UnrealEngine.cpp b/Engine/Source/Runtime/Engine/Private/UnrealEngine.cpp
--- a/Engine/Source/Runtime/Engine/Private/UnrealEngine.cpp
+++ b/Engine/Source/Runtime/Engine/Private/UnrealEngine.cpp
@@ -15977,6 +15977,65 @@ public:
 	// ~End FArchive Interface
 };
 
+PRAGMA_DISABLE_DEPRECATION_WARNINGS
+UEngine::FCopyPropertiesForUnrelatedObjectsParams::FCopyPropertiesForUnrelatedObjectsParams()
+	: bAggressiveDefaultSubobjectReplacement(false)
+	, bDoDelta(true)
+	, bReplaceObjectClassReferences(true)
+	, bCopyDeprecatedProperties(false)
+	, bPreserveRootComponent(true)
+	, bPerformDuplication(false)
+	, bSkipCompilerGeneratedDefaults(false)
+	, bNotifyObjectReplacement(false)
+	, bClearReferences(true)
+	, bDontClearReferenceIfNewerClassExists(false)
+{}
+UEngine::FCopyPropertiesForUnrelatedObjectsParams::FCopyPropertiesForUnrelatedObjectsParams(const FCopyPropertiesForUnrelatedObjectsParams&) = default;
+PRAGMA_ENABLE_DEPRECATION_WARNINGS
+
+
+TAutoConsoleVariable<bool> CVarAuditAggressiveReferenceReplacement(
+	TEXT("cpfuo.AuditAggressiveReferenceReplacment"),
+	1,
+	TEXT("Whether to audit and report on reference replacements that come from the aggressive replacement path.\n"),
+	ECVF_Cheat);
+
+TAutoConsoleVariable<bool> CVarUseAggressiveReferenceReplacement(
+	TEXT("cpfuo.UseAggressiveReferenceReplacment"),
+	0,
+	TEXT("Whether to aggressively replace references. This behavior is being deprecated but being left with the ability to toggle back on in case issues arise.\n"),
+	ECVF_Default);
+
+class FAggressiveReplacementAuditArchive : public FArchiveReplaceOrClearExternalReferences<UObject>
+{
+public:
+
+	typedef FArchiveReplaceOrClearExternalReferences<UObject> TSuper;
+
+	FAggressiveReplacementAuditArchive(UObject* InSearchObject, const TMap<UObject*, UObject*>& InReplacementMap, UPackage* InDestPackage, const TMap<UObject*, UObject*>& InAggressiveReplacements)
+		: TSuper(InSearchObject, InReplacementMap, InDestPackage, EArchiveReplaceObjectFlags::DelayStart)
+		, AggressiveReplacements(InAggressiveReplacements)
+	{
+		this->SerializeSearchObject();
+	}
+
+	FArchive& operator<<(UObject*& Obj)
+	{
+		UObject* PreResolveObj = Obj;
+		TSuper::operator<<(Obj);
+		if (UObject* const * AggressiveReplacementObj = AggressiveReplacements.Find(PreResolveObj))
+		{
+			ensureAlwaysMsgf(false, TEXT("UE-91735: Attempted aggressive reference replacement. '%s' would have been replaced by '%s' in '%s' for property '%s'"), *GetNameSafe(PreResolveObj), *GetNameSafe(*AggressiveReplacementObj), *GetNameSafe(GetSearchObject()), *GetNameSafe(GetSerializedProperty()));
+		}
+
+		return *this;
+	}
+
+private:
+
+	const TMap<UObject*, UObject*>& AggressiveReplacements;
+};
+
 void UEngine::CopyPropertiesForUnrelatedObjects(UObject* OldObject, UObject* NewObject, FCopyPropertiesForUnrelatedObjectsParams Params)
 {
 	// UObject::CollectDefaultSubobjects will not return all DSOs for the CDO (only returns inherited subobjects). It also
@@ -16064,6 +16123,16 @@ void UEngine::CopyPropertiesForUnrelatedObjects(UObject* OldObject, UObject* New
 	ReferenceReplacementMap.Add(OldObject->GetClass()->GetDefaultObject(), NewObject->GetClass()->GetDefaultObject());
 
 	TArray<UObject*> ComponentsOnNewObject;
+	TMap<UObject*, UObject*> AggressiveReplaceReferences;
+#if !UE_BUILD_SHIPPING
+	PRAGMA_DISABLE_DEPRECATION_WARNINGS
+	const bool bAuditAggressiveReplacement = Params.bAggressiveDefaultSubobjectReplacement && CVarAuditAggressiveReferenceReplacement.GetValueOnAnyThread();
+	PRAGMA_ENABLE_DEPRECATION_WARNINGS
+#else
+	constexpr bool bAuditAggressiveReplacement = false;
+#endif
+	const bool bDoAggressiveReplacement = CVarUseAggressiveReferenceReplacement.GetValueOnAnyThread();
+
 	{
 		// Serialize in the modified properties from the old CDO to the new CDO
 		if (Writer.SavedPropertyData.Num() > 0)
@@ -16081,24 +16150,38 @@ void UEngine::CopyPropertiesForUnrelatedObjects(UObject* OldObject, UObject* New
 			{
 				FInstancedObjectRecord& Record = SavedInstances[*pOldInstanceIndex];
 				ReferenceReplacementMap.Add(Record.OldInstance, NewInstance);
-				if (Params.bAggressiveDefaultSubobjectReplacement)
+				if (bDoAggressiveReplacement || bAuditAggressiveReplacement)
 				{
 					UClass* Class = OldObject->GetClass()->GetSuperClass();
 					if (Class)
 					{
-						UObject *CDOInst = Class->GetDefaultSubobjectByName(NewInstance->GetFName());
+						UObject* CDOInst = Class->GetDefaultSubobjectByName(NewInstance->GetFName());
 						if (CDOInst)
 						{
-							ReferenceReplacementMap.Add(CDOInst, NewInstance);
+							if (bAuditAggressiveReplacement)
+							{
+								AggressiveReplaceReferences.Add(CDOInst, NewInstance);
+							}
+							if (bDoAggressiveReplacement)
+							{
+								ReferenceReplacementMap.Add(CDOInst, NewInstance);
+							}
 #if WITH_EDITOR
-							if (Class->ClassGeneratedBy && Cast<UBlueprint>(Class->ClassGeneratedBy)->SkeletonGeneratedClass)
+							UBlueprint* BPGeneratedBy = CastChecked<UBlueprint>(Class->ClassGeneratedBy, ECastCheckedType::NullAllowed);
+							if (BPGeneratedBy && BPGeneratedBy->SkeletonGeneratedClass)
 							{
-								UObject *CDOInstS = Cast<UBlueprint>(Class->ClassGeneratedBy)->SkeletonGeneratedClass->GetDefaultSubobjectByName(NewInstance->GetFName());
+								UObject* CDOInstS = BPGeneratedBy->SkeletonGeneratedClass->GetDefaultSubobjectByName(NewInstance->GetFName());
 								if (CDOInstS)
 								{
-									ReferenceReplacementMap.Add(CDOInstS, NewInstance);
+									if (bAuditAggressiveReplacement)
+									{
+										AggressiveReplaceReferences.Add(CDOInstS, NewInstance);
+									}
+									if (bDoAggressiveReplacement)
+									{
+										ReferenceReplacementMap.Add(CDOInstS, NewInstance);
+									}
 								}
-
 							}
 #endif // WITH_EDITOR
 						}
@@ -16156,13 +16239,27 @@ void UEngine::CopyPropertiesForUnrelatedObjects(UObject* OldObject, UObject* New
 	{
 		UPackage* NewPackage = NewObject->GetOutermost();
 		// Replace references to old classes and instances on this object with the corresponding new ones
-		FArchiveReplaceOrClearExternalReferences<UObject> ReplaceInCDOAr(NewObject, ReferenceReplacementMap, NewPackage);
+		if (bAuditAggressiveReplacement)
+		{
+			FAggressiveReplacementAuditArchive ReplaceInCDOAr(NewObject, ReferenceReplacementMap, NewPackage, AggressiveReplaceReferences);
+		}
+		else
+		{
+			FArchiveReplaceOrClearExternalReferences<UObject> ReplaceInCDOAr(NewObject, ReferenceReplacementMap, NewPackage);
+		}
 
 		// Replace references inside each individual component. This is always required because if something is in ReferenceReplacementMap, the above replace code will skip fixing child properties
 		for (int32 ComponentIndex = 0; ComponentIndex < ComponentsOnNewObject.Num(); ++ComponentIndex)
 		{
 			UObject* NewComponent = ComponentsOnNewObject[ComponentIndex];
-			FArchiveReplaceOrClearExternalReferences<UObject> ReplaceInComponentAr(NewComponent, ReferenceReplacementMap, NewPackage);
+			if (bAuditAggressiveReplacement)
+			{
+				FAggressiveReplacementAuditArchive ReplaceInComponentAr(NewComponent, ReferenceReplacementMap, NewPackage, AggressiveReplaceReferences);
+			}
+			else
+			{
+				FArchiveReplaceOrClearExternalReferences<UObject> ReplaceInComponentAr(NewComponent, ReferenceReplacementMap, NewPackage);
+			}
 		}
 	}
 
