Make FPrimaryAssetTypeInfo::bIsEditorOnly comment more clear

diff --git a/Engine/Source/Runtime/Engine/Classes/Engine/AssetManagerTypes.h b/Engine/Source/Runtime/Engine/Classes/Engine/AssetManagerTypes.h
--- a/Engine/Source/Runtime/Engine/Classes/Engine/AssetManagerTypes.h
+++ b/Engine/Source/Runtime/Engine/Classes/Engine/AssetManagerTypes.h
@@ -125,7 +125,7 @@ public:
 	UPROPERTY(EditAnywhere, Category = AssetType)
 	bool bHasBlueprintClasses;
 
-	/** True if this type is editor only */
+	/** True if this type is editor only. Disable this to consider content that is part of the engine or engine plugins. */
 	UPROPERTY(EditAnywhere, Category = AssetType)
 	bool bIsEditorOnly;
 
