Texture: Fixed streaming mips to make IsBulkDataLoaded() return false when the mips are not loaded

The previous fix only removed bulk data during load, but saving an empty, unloaded bulk data will modify its state to be loaded but with a null buffer.

#rnx

diff --git a/Engine/Source/Runtime/Engine/Private/Texture2D.cpp b/Engine/Source/Runtime/Engine/Private/Texture2D.cpp
--- a/Engine/Source/Runtime/Engine/Private/Texture2D.cpp
+++ b/Engine/Source/Runtime/Engine/Private/Texture2D.cpp
@@ -151,7 +151,7 @@ void FTexture2DMipMap::Serialize(FArchive& Ar, UObject* Owner, int32 MipIdx)
 	// Streaming mips are saved with a size of 0 because they are stored separately.
 	// IsBulkDataLoaded() returns true for this empty bulk data. Remove the empty
 	// bulk data to allow unloaded streaming mips to be detected.
-	if (Ar.IsLoading() && BulkData.GetBulkDataSize() == 0)
+	if (BulkData.GetBulkDataSize() == 0)
 	{
 		BulkData.RemoveBulkData();
 	}
diff --git a/Engine/Source/Runtime/Engine/Private/VT/VirtualTextureBuiltData.cpp b/Engine/Source/Runtime/Engine/Private/VT/VirtualTextureBuiltData.cpp
--- a/Engine/Source/Runtime/Engine/Private/VT/VirtualTextureBuiltData.cpp
+++ b/Engine/Source/Runtime/Engine/Private/VT/VirtualTextureBuiltData.cpp
@@ -378,7 +378,7 @@ void FVirtualTextureBuiltData::Serialize(FArchive& Ar, UObject* Owner, int32 Fir
 		// Streaming chunks are saved with a size of 0 because they are stored separately.
 		// IsBulkDataLoaded() returns true for this empty bulk data. Remove the empty
 		// bulk data to allow unloaded streaming chunks to be detected.
-		if (Ar.IsLoading() && Chunk.BulkData.GetBulkDataSize() == 0)
+		if (Chunk.BulkData.GetBulkDataSize() == 0)
 		{
 			Chunk.BulkData.RemoveBulkData();
 		}
