Fixed some compilation errors in non-unity builds.

diff --git a/Engine/Source/Runtime/Engine/Private/ShaderCompiler/ShaderCompiler.cpp b/Engine/Source/Runtime/Engine/Private/ShaderCompiler/ShaderCompiler.cpp
--- a/Engine/Source/Runtime/Engine/Private/ShaderCompiler/ShaderCompiler.cpp
+++ b/Engine/Source/Runtime/Engine/Private/ShaderCompiler/ShaderCompiler.cpp
@@ -66,6 +66,7 @@
 #include "ClearReplacementShaders.h"
 #include "Engine/RendererSettings.h"
 #include "Misc/LargeWorldRenderPosition.h"
+#include "UnrealEngine.h"
 
 #if WITH_EDITOR
 #include "Compression/OodleDataCompression.h"
