AnimInstanceProxy.h dependency cleanup.
The original target was its include of AnimBlueprintGeneratedClass.h but more opportunities presented themselves.
Moved FNodeVisit and FAttributeRecord out of FAnimBlueprintDebugData so they could be forward declared.

diff --git a/Engine/Plugins/Experimental/Animation/PoseSearch/Source/Runtime/Private/PoseSearchLibrary.cpp b/Engine/Plugins/Experimental/Animation/PoseSearch/Source/Runtime/Private/PoseSearchLibrary.cpp
--- a/Engine/Plugins/Experimental/Animation/PoseSearch/Source/Runtime/Private/PoseSearchLibrary.cpp
+++ b/Engine/Plugins/Experimental/Animation/PoseSearch/Source/Runtime/Private/PoseSearchLibrary.cpp
@@ -6,6 +6,7 @@
 #include "Animation/AnimNode_Inertialization.h"
 #include "Animation/AnimNode_SequencePlayer.h"
 #include "Animation/AnimRootMotionProvider.h"
+#include "Animation/BlendSpace.h"
 #include "PoseSearch/AnimNode_MotionMatching.h"
 #include "Trace/PoseSearchTraceLogger.h"
 
diff --git a/Engine/Plugins/Experimental/Animation/PoseSearch/Source/Runtime/Private/Trace/PoseSearchTraceLogger.cpp b/Engine/Plugins/Experimental/Animation/PoseSearch/Source/Runtime/Private/Trace/PoseSearchTraceLogger.cpp
--- a/Engine/Plugins/Experimental/Animation/PoseSearch/Source/Runtime/Private/Trace/PoseSearchTraceLogger.cpp
+++ b/Engine/Plugins/Experimental/Animation/PoseSearch/Source/Runtime/Private/Trace/PoseSearchTraceLogger.cpp
@@ -3,6 +3,7 @@
 #include "PoseSearchTraceLogger.h"
 #include "Animation/AnimInstanceProxy.h"
 #include "Trace/Trace.inl"
+#include "Animation/AnimNodeBase.h"
 
 #if UE_POSE_SEARCH_TRACE_ENABLED
 
diff --git a/Engine/Source/Runtime/AnimGraphRuntime/Private/AnimGraphRuntimeModule.cpp b/Engine/Source/Runtime/AnimGraphRuntime/Private/AnimGraphRuntimeModule.cpp
--- a/Engine/Source/Runtime/AnimGraphRuntime/Private/AnimGraphRuntimeModule.cpp
+++ b/Engine/Source/Runtime/AnimGraphRuntime/Private/AnimGraphRuntimeModule.cpp
@@ -7,6 +7,7 @@
 #include "Animation/AnimInstance.h"
 #include "Animation/AttributeTypes.h"
 #include "Algo/Sort.h"
+#include "Animation/AnimBlueprintGeneratedClass.h"
 
 //////////////////////////////////////////////////////////////////////////
 // FAnimGraphRuntimeModule
diff --git a/Engine/Source/Runtime/AnimGraphRuntime/Private/AnimNodes/AnimNode_BlendSpaceGraphBase.cpp b/Engine/Source/Runtime/AnimGraphRuntime/Private/AnimNodes/AnimNode_BlendSpaceGraphBase.cpp
--- a/Engine/Source/Runtime/AnimGraphRuntime/Private/AnimNodes/AnimNode_BlendSpaceGraphBase.cpp
+++ b/Engine/Source/Runtime/AnimGraphRuntime/Private/AnimNodes/AnimNode_BlendSpaceGraphBase.cpp
@@ -5,6 +5,9 @@
 #include "Animation/AnimInstanceProxy.h"
 #include "AnimGraphRuntimeTrace.h"
 #include "Animation/AnimSyncScope.h"
+#if WITH_EDITORONLY_DATA
+#include "Animation/AnimBlueprintGeneratedClass.h"
+#endif
 
 void FAnimNode_BlendSpaceGraphBase::Initialize_AnyThread(const FAnimationInitializeContext& Context)
 {
diff --git a/Engine/Source/Runtime/AnimGraphRuntime/Private/AnimNodes/AnimNode_BlendSpacePlayer.cpp b/Engine/Source/Runtime/AnimGraphRuntime/Private/AnimNodes/AnimNode_BlendSpacePlayer.cpp
--- a/Engine/Source/Runtime/AnimGraphRuntime/Private/AnimNodes/AnimNode_BlendSpacePlayer.cpp
+++ b/Engine/Source/Runtime/AnimGraphRuntime/Private/AnimNodes/AnimNode_BlendSpacePlayer.cpp
@@ -7,6 +7,9 @@
 #include "AnimGraphRuntimeTrace.h"
 #include "Animation/AnimSync.h"
 #include "Animation/AnimSyncScope.h"
+#if WITH_EDITORONLY_DATA
+#include "Animation/AnimBlueprintGeneratedClass.h"
+#endif
 
 /////////////////////////////////////////////////////
 // FAnimNode_BlendSpacePlayerBase
diff --git a/Engine/Source/Runtime/Engine/Classes/Animation/AnimBlueprintGeneratedClass.h b/Engine/Source/Runtime/Engine/Classes/Animation/AnimBlueprintGeneratedClass.h
--- a/Engine/Source/Runtime/Engine/Classes/Animation/AnimBlueprintGeneratedClass.h
+++ b/Engine/Source/Runtime/Engine/Classes/Animation/AnimBlueprintGeneratedClass.h
@@ -129,6 +129,31 @@ public:
 #endif
 };
 
+struct FAnimBlueprintDebugData_NodeVisit
+{
+	int32 SourceID;
+	int32 TargetID;
+	float Weight;
+
+	FAnimBlueprintDebugData_NodeVisit(int32 InSourceID, int32 InTargetID, float InWeight)
+		: SourceID(InSourceID)
+		, TargetID(InTargetID)
+		, Weight(InWeight)
+	{
+	}
+};
+
+struct FAnimBlueprintDebugData_AttributeRecord
+{
+	FName Attribute;
+	int32 OtherNode;
+
+	FAnimBlueprintDebugData_AttributeRecord(int32 InOtherNode, FName InAttribute)
+		: Attribute(InAttribute)
+		, OtherNode(InOtherNode)
+	{}
+};
+
 // This structure represents animation-related debugging information for an entire AnimBlueprint
 // (general debug information for the event graph, etc... is still contained in a FBlueprintDebugData structure)
 USTRUCT()
@@ -180,34 +205,13 @@ public:
 	TMap<TFieldPath<const FProperty>, TFieldPath<const FProperty>> NodeToFoldedPropertyMap;
 
 	// Node visit structure
-	struct FNodeVisit
-	{
-		int32 SourceID;
-		int32 TargetID;
-		float Weight;
-
-		FNodeVisit(int32 InSourceID, int32 InTargetID, float InWeight)
-			: SourceID(InSourceID)
-			, TargetID(InTargetID)
-			, Weight(InWeight)
-		{
-		}
-	};
+	using FNodeVisit = FAnimBlueprintDebugData_NodeVisit;
 
 	// History of activated nodes
 	TArray<FNodeVisit> UpdatedNodesThisFrame;
 
 	// Record of attribute transfer between nodes
-	struct FAttributeRecord
-	{
-		FName Attribute;
-		int32 OtherNode;
-
-		FAttributeRecord(int32 InOtherNode, FName InAttribute)
-			: Attribute(InAttribute)
-			, OtherNode(InOtherNode)
-		{}
-	};
+	using FAttributeRecord = FAnimBlueprintDebugData_AttributeRecord;
 
 	// History of node attributes that are output from and input to nodes
 	TMap<int32, TArray<FAttributeRecord>> NodeInputAttributesThisFrame;
diff --git a/Engine/Source/Runtime/Engine/Private/Animation/AnimInstanceProxy.cpp b/Engine/Source/Runtime/Engine/Private/Animation/AnimInstanceProxy.cpp
--- a/Engine/Source/Runtime/Engine/Private/Animation/AnimInstanceProxy.cpp
+++ b/Engine/Source/Runtime/Engine/Private/Animation/AnimInstanceProxy.cpp
@@ -25,6 +25,7 @@
 #include "Animation/AnimSyncScope.h"
 #include "Animation/AnimNotifyStateMachineInspectionLibrary.h"
 #include "Animation/MirrorDataTable.h"
+#include "Animation/AnimBlueprintGeneratedClass.h"
 
 #define DO_ANIMSTAT_PROCESSING(StatName) DEFINE_STAT(STAT_ ## StatName)
 #include "Animation/AnimMTStats.h"
@@ -41,6 +42,65 @@ const FName NAME_Evaluate(TEXT("Evaluate"));
 const FName NAME_Update(TEXT("Update"));
 const FName NAME_AnimGraph(TEXT("AnimGraph"));
 
+FAnimInstanceProxy::FAnimInstanceProxy()
+	: AnimInstanceObject(nullptr)
+	, AnimClassInterface(nullptr)
+	, Skeleton(nullptr)
+	, SkeletalMeshComponent(nullptr)
+	, MainInstanceProxy(nullptr)
+	, CurrentDeltaSeconds(0.0f)
+	, CurrentTimeDilation(1.0f)
+	, RootNode(nullptr)
+	, DefaultLinkedInstanceInputNode(nullptr)
+	, BufferWriteIndex(0)
+	, RootMotionMode(ERootMotionMode::NoRootMotionExtraction)
+	, FrameCounterForUpdate(0)
+	, FrameCounterForNodeUpdate(0)
+	, CacheBonesRecursionCounter(0)
+	, MainMontageEvaluationData(&MontageEvaluationData)
+	, bUpdatingRoot(false)
+	, bBoneCachesInvalidated(false)
+	, bShouldExtractRootMotion(false)
+	, bDeferRootNodeInitialization(false)
+#if WITH_EDITORONLY_DATA
+	, bIsBeingDebugged(false)
+#endif
+	, bInitializeSubsystems(false)
+{
+}
+
+FAnimInstanceProxy::FAnimInstanceProxy(UAnimInstance* Instance)
+	: AnimInstanceObject(Instance)
+	, AnimClassInterface(IAnimClassInterface::GetFromClass(Instance->GetClass()))
+	, Skeleton(nullptr)
+	, SkeletalMeshComponent(nullptr)
+	, MainInstanceProxy(nullptr)
+	, CurrentDeltaSeconds(0.0f)
+	, CurrentTimeDilation(1.0f)
+	, RootNode(nullptr)
+	, DefaultLinkedInstanceInputNode(nullptr)
+	, BufferWriteIndex(0)
+	, RootMotionMode(ERootMotionMode::NoRootMotionExtraction)
+	, FrameCounterForUpdate(0)
+	, FrameCounterForNodeUpdate(0)
+	, CacheBonesRecursionCounter(0)
+	, MainMontageEvaluationData(&MontageEvaluationData)
+	, bUpdatingRoot(false)
+	, bBoneCachesInvalidated(false)
+	, bShouldExtractRootMotion(false)
+	, bDeferRootNodeInitialization(false)
+#if WITH_EDITORONLY_DATA
+	, bIsBeingDebugged(false)
+#endif
+	, bInitializeSubsystems(false)
+{
+}
+
+FAnimInstanceProxy::FAnimInstanceProxy(const FAnimInstanceProxy&) = default;
+FAnimInstanceProxy& FAnimInstanceProxy::operator=(FAnimInstanceProxy&&) = default;
+FAnimInstanceProxy& FAnimInstanceProxy::operator=(const FAnimInstanceProxy&) = default;
+FAnimInstanceProxy::~FAnimInstanceProxy() = default;
+
 void FAnimInstanceProxy::UpdateAnimationNode(const FAnimationUpdateContext& InContext)
 {
 	UpdateAnimationNode_WithRoot(InContext, RootNode, NAME_AnimGraph);
@@ -2999,6 +3059,11 @@ TArray<FAnimNode_AssetPlayerBase*> FAnimInstanceProxy::GetMutableInstanceAssetPl
 }
 
 #if WITH_EDITOR
+void FAnimInstanceProxy::RecordNodeVisit(int32 TargetNodeIndex, int32 SourceNodeIndex, float BlendWeight)
+{
+	UpdatedNodesThisFrame.Emplace(SourceNodeIndex, TargetNodeIndex, BlendWeight);
+}
+
 void FAnimInstanceProxy::RecordNodeAttribute(const FAnimInstanceProxy& InSourceProxy, int32 InTargetNodeIndex, int32 InSourceNodeIndex, FName InAttribute)
 {
 	TArray<FAnimBlueprintDebugData::FAttributeRecord>& InputAttributeRecords = NodeInputAttributesThisFrame.FindOrAdd(InTargetNodeIndex);
diff --git a/Engine/Source/Runtime/Engine/Private/Animation/AnimNode_ApplyMeshSpaceAdditive.cpp b/Engine/Source/Runtime/Engine/Private/Animation/AnimNode_ApplyMeshSpaceAdditive.cpp
--- a/Engine/Source/Runtime/Engine/Private/Animation/AnimNode_ApplyMeshSpaceAdditive.cpp
+++ b/Engine/Source/Runtime/Engine/Private/Animation/AnimNode_ApplyMeshSpaceAdditive.cpp
@@ -2,8 +2,6 @@
 
 #include "Animation/AnimNode_ApplyMeshSpaceAdditive.h"
 #include "Animation/AnimInstanceProxy.h"
-#include "AnimationRuntime.h"
-#include "Animation/AnimTrace.h"
 
 /////////////////////////////////////////////////////
 // FAnimNode_ApplyMeshSpaceAdditive
diff --git a/Engine/Source/Runtime/Engine/Private/Animation/AnimNode_AssetPlayerBase.cpp b/Engine/Source/Runtime/Engine/Private/Animation/AnimNode_AssetPlayerBase.cpp
--- a/Engine/Source/Runtime/Engine/Private/Animation/AnimNode_AssetPlayerBase.cpp
+++ b/Engine/Source/Runtime/Engine/Private/Animation/AnimNode_AssetPlayerBase.cpp
@@ -1,7 +1,6 @@
 // Copyright Epic Games, Inc. All Rights Reserved.
 
 #include "Animation/AnimNode_AssetPlayerBase.h"
-#include "Animation/AnimInstanceProxy.h"
 #include "Animation/AnimSyncScope.h"
 
 void FAnimNode_AssetPlayerBase::Initialize_AnyThread(const FAnimationInitializeContext& Context)
diff --git a/Engine/Source/Runtime/Engine/Private/Animation/AnimNode_Inertialization.cpp b/Engine/Source/Runtime/Engine/Private/Animation/AnimNode_Inertialization.cpp
--- a/Engine/Source/Runtime/Engine/Private/Animation/AnimNode_Inertialization.cpp
+++ b/Engine/Source/Runtime/Engine/Private/Animation/AnimNode_Inertialization.cpp
@@ -8,6 +8,7 @@
 #include "Animation/BlendProfile.h"
 #include "Algo/MaxElement.h"
 #include "HAL/LowLevelMemTracker.h"
+#include "Animation/AnimBlueprintGeneratedClass.h"
 
 LLM_DEFINE_TAG(Animation_Inertialization);
 
diff --git a/Engine/Source/Runtime/Engine/Private/Animation/AnimNode_SequencePlayer.cpp b/Engine/Source/Runtime/Engine/Private/Animation/AnimNode_SequencePlayer.cpp
--- a/Engine/Source/Runtime/Engine/Private/Animation/AnimNode_SequencePlayer.cpp
+++ b/Engine/Source/Runtime/Engine/Private/Animation/AnimNode_SequencePlayer.cpp
@@ -6,6 +6,9 @@
 #include "Animation/AnimInstanceProxy.h"
 #include "Animation/AnimTrace.h"
 #include "Animation/AnimPoseSearchProvider.h"
+#if WITH_EDITORONLY_DATA
+#include "Animation/AnimBlueprintGeneratedClass.h"
+#endif
 
 #define LOCTEXT_NAMESPACE "AnimNode_SequencePlayer"
 
diff --git a/Engine/Source/Runtime/Engine/Private/Animation/AnimSubsystem_NodeRelevancy.cpp b/Engine/Source/Runtime/Engine/Private/Animation/AnimSubsystem_NodeRelevancy.cpp
--- a/Engine/Source/Runtime/Engine/Private/Animation/AnimSubsystem_NodeRelevancy.cpp
+++ b/Engine/Source/Runtime/Engine/Private/Animation/AnimSubsystem_NodeRelevancy.cpp
@@ -2,6 +2,7 @@
 
 #include "Animation/AnimSubsystem_NodeRelevancy.h"
 #include "Animation/AnimInstanceProxy.h"
+#include "Animation/AnimNodeBase.h"
 
 float FAnimNodeRelevancyStatus::GetCurrentWeight() const
 {
diff --git a/Engine/Source/Runtime/Engine/Private/Animation/AnimSyncScope.cpp b/Engine/Source/Runtime/Engine/Private/Animation/AnimSyncScope.cpp
--- a/Engine/Source/Runtime/Engine/Private/Animation/AnimSyncScope.cpp
+++ b/Engine/Source/Runtime/Engine/Private/Animation/AnimSyncScope.cpp
@@ -3,6 +3,7 @@
 #include "Animation/AnimSyncScope.h"
 #include "Animation/AnimInstanceProxy.h"
 #include "Animation/AnimSync.h"
+#include "Animation/AnimNodeBase.h"
 
 IMPLEMENT_ANIMGRAPH_MESSAGE(UE::Anim::FAnimSyncGroupScope);
 
diff --git a/Engine/Source/Runtime/Engine/Private/Animation/AnimTrace.cpp b/Engine/Source/Runtime/Engine/Private/Animation/AnimTrace.cpp
--- a/Engine/Source/Runtime/Engine/Private/Animation/AnimTrace.cpp
+++ b/Engine/Source/Runtime/Engine/Private/Animation/AnimTrace.cpp
@@ -15,6 +15,7 @@
 #include "Animation/AnimNodeBase.h"
 #include "Animation/AnimMontage.h"
 #include "Animation/BlendSpace.h"
+#include "Animation/AnimBlueprintGeneratedClass.h"
 #include "Animation/AnimNode_SequencePlayer.h"
 #include "Animation/AnimNotifies/AnimNotify.h"
 #include "Animation/AnimNotifies/AnimNotifyState.h"
diff --git a/Engine/Source/Runtime/Engine/Private/Animation/MirrorSyncScope.cpp b/Engine/Source/Runtime/Engine/Private/Animation/MirrorSyncScope.cpp
--- a/Engine/Source/Runtime/Engine/Private/Animation/MirrorSyncScope.cpp
+++ b/Engine/Source/Runtime/Engine/Private/Animation/MirrorSyncScope.cpp
@@ -4,6 +4,7 @@
 #include "Animation/AnimSyncScope.h"
 #include "Animation/AnimInstanceProxy.h"
 #include "Animation/MirrorDataTable.h"
+#include "Animation/AnimNodeBase.h"
 
 IMPLEMENT_ANIMGRAPH_MESSAGE(UE::Anim::FMirrorSyncScope);
 
diff --git a/Engine/Source/Runtime/Engine/Public/Animation/AnimInstanceProxy.h b/Engine/Source/Runtime/Engine/Public/Animation/AnimInstanceProxy.h
--- a/Engine/Source/Runtime/Engine/Public/Animation/AnimInstanceProxy.h
+++ b/Engine/Source/Runtime/Engine/Public/Animation/AnimInstanceProxy.h
@@ -4,33 +4,57 @@
 
 #include "CoreMinimal.h"
 #include "UObject/ObjectMacros.h"
+#include "Engine/EngineTypes.h"
 #include "Animation/AnimTypes.h"
-#include "BoneContainer.h"
-#include "Animation/Skeleton.h"
 #include "Animation/AnimationAsset.h"
 #include "Animation/AnimBlueprint.h"
+#include "Animation/AnimClassInterface.h"
+#include "Animation/AnimTrace.h"
+#include "Animation/AnimSync.h"
+#if UE_ENABLE_INCLUDE_ORDER_DEPRECATED_IN_5_1
+#include "BoneContainer.h"
+#include "Animation/Skeleton.h"
 #include "BonePose.h"
 #include "Animation/AnimNotifyQueue.h"
 #include "Animation/PoseSnapshot.h"
 #include "Animation/AnimInstance.h"
 #include "Engine/PoseWatch.h"
-#include "Animation/AnimClassInterface.h"
 #include "Animation/AnimBlueprintGeneratedClass.h"
 #include "Logging/TokenizedMessage.h"
-#include "Animation/AnimTrace.h"
-#include "Animation/AnimSync.h"
 #include "Animation/AnimSyncScope.h"
 #include "Animation/ActiveStateMachineScope.h"
+#endif
 #include "AnimInstanceProxy.generated.h"
 
+class UAnimInstance;
+class UBlendProfile;
 class UBlendSpace;
+struct FAnimationPoseData;
+struct FAnimBlueprintDebugData_NodeVisit;
+struct FAnimBlueprintDebugData_AttributeRecord;
+struct FAnimGroupInstance;
+struct FAnimNodePoseWatch;
 struct FAnimNode_AssetPlayerBase;
 struct FAnimNode_Base;
+struct FAnimNode_LinkedInputPose;
 struct FAnimNode_SaveCachedPose;
 struct FAnimNode_StateMachine;
-struct FAnimNode_LinkedInputPose;
+struct FAnimTickRecord;
+struct FBakedAnimationStateMachine;
+struct FCompactPose;
 struct FNodeDebugData;
 struct FPoseContext;
+struct FPoseSnapshot;
+
+namespace UE::Anim
+{
+	class FAnimSyncGroupScope;
+	class FActiveStateMachineScope;
+	struct FAnimSyncParams;
+	using FSlotInertializationRequest = TPair<float, const UBlendProfile*>;
+}
+
+template<class PoseType> struct FCSPose;
 
 // Disable debugging information for shipping and test builds.
 #define ENABLE_ANIM_DRAW_DEBUG (1 && !(UE_BUILD_SHIPPING || UE_BUILD_TEST))
@@ -111,66 +135,18 @@ struct FQueuedDrawDebugItem
 USTRUCT(meta = (DisplayName = "Native Variables"))
 struct ENGINE_API FAnimInstanceProxy
 {
-	GENERATED_BODY()
+	GENERATED_USTRUCT_BODY()
 
 public:
 	using FSyncGroupMap = TMap<FName, FAnimGroupInstance>;
 
-	FAnimInstanceProxy()
-		: AnimInstanceObject(nullptr)
-		, AnimClassInterface(nullptr)
-		, Skeleton(nullptr)
-		, SkeletalMeshComponent(nullptr)
-		, MainInstanceProxy(nullptr)
-		, CurrentDeltaSeconds(0.0f)
-		, CurrentTimeDilation(1.0f)
-		, RootNode(nullptr)
-		, DefaultLinkedInstanceInputNode(nullptr)
-		, BufferWriteIndex(0)
-		, RootMotionMode(ERootMotionMode::NoRootMotionExtraction)
-		, FrameCounterForUpdate(0)
-		, FrameCounterForNodeUpdate(0)
-		, CacheBonesRecursionCounter(0)
-		, MainMontageEvaluationData(&MontageEvaluationData)
-		, bUpdatingRoot(false)
-		, bBoneCachesInvalidated(false)
-		, bShouldExtractRootMotion(false)
-		, bDeferRootNodeInitialization(false)
-#if WITH_EDITORONLY_DATA
-		, bIsBeingDebugged(false)
-#endif
-		, bInitializeSubsystems(false)
-	{
-	}
+	FAnimInstanceProxy();
+	FAnimInstanceProxy(UAnimInstance* Instance);
 
-	FAnimInstanceProxy(UAnimInstance* Instance)
-		: AnimInstanceObject(Instance)
-		, AnimClassInterface(IAnimClassInterface::GetFromClass(Instance->GetClass()))
-		, Skeleton(nullptr)
-		, SkeletalMeshComponent(nullptr)
-		, MainInstanceProxy(nullptr)
-		, CurrentDeltaSeconds(0.0f)
-		, CurrentTimeDilation(1.0f)
-		, RootNode(nullptr)
-		, DefaultLinkedInstanceInputNode(nullptr)
-		, BufferWriteIndex(0)
-		, RootMotionMode(ERootMotionMode::NoRootMotionExtraction)
-		, FrameCounterForUpdate(0)
-		, FrameCounterForNodeUpdate(0)
-		, CacheBonesRecursionCounter(0)
-		, MainMontageEvaluationData(&MontageEvaluationData)
-		, bUpdatingRoot(false)
-		, bBoneCachesInvalidated(false)
-		, bShouldExtractRootMotion(false)
-		, bDeferRootNodeInitialization(false)
-#if WITH_EDITORONLY_DATA
-		, bIsBeingDebugged(false)
-#endif
-		, bInitializeSubsystems(false)
-	{
-	}
-
-	virtual ~FAnimInstanceProxy() {}
+	FAnimInstanceProxy(const FAnimInstanceProxy&);
+	FAnimInstanceProxy& operator=(FAnimInstanceProxy&&);
+	FAnimInstanceProxy& operator=(const FAnimInstanceProxy&);
+	virtual ~FAnimInstanceProxy();
 
 	// Get the IAnimClassInterface associated with this context, if there is one.
 	// Note: This can return NULL, so check the result.
@@ -199,10 +175,7 @@ public:
 	}
 
 	/** Record a visited node in the debugger */
-	void RecordNodeVisit(int32 TargetNodeIndex, int32 SourceNodeIndex, float BlendWeight)
-	{
-		UpdatedNodesThisFrame.Emplace(SourceNodeIndex, TargetNodeIndex, BlendWeight);
-	}
+	void RecordNodeVisit(int32 TargetNodeIndex, int32 SourceNodeIndex, float BlendWeight);
 
 	/** Record a node attribute in the debugger */
 	void RecordNodeAttribute(const FAnimInstanceProxy& InSourceProxy, int32 InTargetNodeIndex, int32 InSourceNodeIndex, FName InAttribute);
@@ -598,23 +571,9 @@ protected:
 	/** Update override point */
 	virtual void Update(float DeltaSeconds) {}
 
-	UE_DEPRECATED(4.24, "Please use the overload that takes an FAnimationUpdateContext")
-	virtual void UpdateAnimationNode(float DeltaSeconds)
-	{
-		FAnimationUpdateContext Context(this, DeltaSeconds);
-		UpdateAnimationNode(Context);
-	}
-
 	/** Updates the anim graph */
 	virtual void UpdateAnimationNode(const FAnimationUpdateContext& InContext);
 
-	UE_DEPRECATED(4.24, "Please use the overload that takes an FAnimationUpdateContext")
-	virtual void UpdateAnimationNode_WithRoot(float DeltaSeconds, FAnimNode_Base* InRootNode, FName InLayerName) 
-	{
-		FAnimationUpdateContext Context(this, DeltaSeconds);
-		UpdateAnimationNode_WithRoot(Context, InRootNode, InLayerName);
-	}
-
 	/** Updates the anim graph using a specified root node */
 	virtual void UpdateAnimationNode_WithRoot(const FAnimationUpdateContext& InContext, FAnimNode_Base* InRootNode, FName InLayerName);
 
@@ -679,13 +638,6 @@ protected:
 	/** Calls Update(), updates the anim graph, ticks asset players */
 	void UpdateAnimation();
 
-	UE_DEPRECATED(4.24, "Please use the overload that takes an FAnimationUpdateContext")
-	void UpdateAnimation_WithRoot(FAnimNode_Base* InRootNode, FName InLayerName)
-	{
-		FAnimationUpdateContext Context(this, CurrentDeltaSeconds);
-		UpdateAnimation_WithRoot(Context, InRootNode, InLayerName);
-	}
-
 	/** Calls Update(), updates the anim graph from the specified root, ticks asset players */
 	void UpdateAnimation_WithRoot(const FAnimationUpdateContext& InContext, FAnimNode_Base* InRootNode, FName InLayerName);
 
@@ -995,11 +947,11 @@ private:
 
 #if WITH_EDITORONLY_DATA
 	/** Array of visited nodes this frame */
-	TArray<FAnimBlueprintDebugData::FNodeVisit> UpdatedNodesThisFrame;
+	TArray<FAnimBlueprintDebugData_NodeVisit> UpdatedNodesThisFrame;
 
 	/** Map of node attributes this frame */
-	TMap<int32, TArray<FAnimBlueprintDebugData::FAttributeRecord>> NodeInputAttributesThisFrame;
-	TMap<int32, TArray<FAnimBlueprintDebugData::FAttributeRecord>> NodeOutputAttributesThisFrame;
+	TMap<int32, TArray<FAnimBlueprintDebugData_AttributeRecord>> NodeInputAttributesThisFrame;
+	TMap<int32, TArray<FAnimBlueprintDebugData_AttributeRecord>> NodeOutputAttributesThisFrame;
 
 	/** Map of node syncs this frame - maps from player node index to graph-determined group name */
 	TMap<int32, FName> NodeSyncsThisFrame;
