Texture: Improved the estimate for how many streaming mips to prefetch to avoid a blocking cache fetch on the main thread

diff --git a/Engine/Source/Runtime/Engine/Private/TextureDerivedDataTask.cpp b/Engine/Source/Runtime/Engine/Private/TextureDerivedDataTask.cpp
--- a/Engine/Source/Runtime/Engine/Private/TextureDerivedDataTask.cpp
+++ b/Engine/Source/Runtime/Engine/Private/TextureDerivedDataTask.cpp
@@ -723,6 +723,52 @@ FTextureCacheDerivedDataWorker::FTextureCacheDerivedDataWorker(
 	}
 }
 
+bool FTextureCacheDerivedDataWorker::TryCacheStreamingMips(int32 FirstMipToLoad, int32 FirstMipToPrefetch)
+{
+	using namespace UE::DerivedData;
+	check(DerivedData->DerivedDataKey.IsType<FString>());
+
+	TArray<FCacheGetValueRequest, TInlineAllocator<16>> MipRequests;
+
+	int32 MipIndex = -1;
+	const FSharedString Name(TexturePathName);
+	for (const FTexture2DMipMap& Mip : DerivedData->Mips)
+	{
+		++MipIndex;
+		if (Mip.IsPagedToDerivedData())
+		{
+			const FCacheKey MipKey = ConvertLegacyCacheKey(DerivedData->GetDerivedDataMipKeyString(MipIndex, Mip));
+			const ECachePolicy Policy
+				= (MipIndex >= FirstMipToLoad) ? ECachePolicy::Default
+				: (MipIndex >= FirstMipToPrefetch) ? ECachePolicy::Default | ECachePolicy::SkipData
+				: ECachePolicy::Query | ECachePolicy::SkipData;
+			MipRequests.Add({Name, MipKey, Policy, uint64(MipIndex)});
+		}
+	}
+
+	if (MipRequests.IsEmpty())
+	{
+		return true;
+	}
+
+	bool bOk = true;
+	FRequestOwner BlockingOwner(EPriority::Blocking);
+	GetCache().GetValue(MipRequests, BlockingOwner, [this, &bOk](FCacheGetValueResponse&& Response)
+	{
+		bOk &= Response.Status == EStatus::Ok;
+		if (const FSharedBuffer MipBuffer = Response.Value.GetData().Decompress())
+		{
+			FTexture2DMipMap& Mip = DerivedData->Mips[int32(Response.UserData)];
+			Mip.BulkData.Lock(LOCK_READ_WRITE);
+			void* MipData = Mip.BulkData.Realloc(int64(MipBuffer.GetSize()));
+			FMemory::Memcpy(MipData, MipBuffer.GetData(), MipBuffer.GetSize());
+			Mip.BulkData.Unlock();
+		}
+	});
+	BlockingOwner.Wait();
+	return bOk;
+}
+
 // DDC1 primary fetch/build work function
 void FTextureCacheDerivedDataWorker::DoWork()
 {
@@ -805,6 +851,7 @@ void FTextureCacheDerivedDataWorker::DoWork()
 	{
 		const bool bInlineMips = EnumHasAnyFlags(CacheFlags, ETextureCacheFlags::InlineMips);
 		const bool bForDDC = EnumHasAnyFlags(CacheFlags, ETextureCacheFlags::ForDDCBuild);
+		int32 FirstResidentMipIndex = 0;
 
 		BytesCached = RawDerivedData.GetSize();
 		FMemoryReaderView Ar(RawDerivedData.GetView(), /*bIsPersistent=*/ true);
@@ -830,6 +877,7 @@ void FTextureCacheDerivedDataWorker::DoWork()
 			if (Algo::AnyOf(DerivedData->Mips, [](const FTexture2DMipMap& Mip) { return !Mip.BulkData.IsBulkDataLoaded(); }))
 			{
 				int32 MipIndex = 0;
+				FirstResidentMipIndex = DerivedData->Mips.Num();
 				const FSharedString Name(TexturePathName);
 				for (FTexture2DMipMap& Mip : DerivedData->Mips)
 				{
@@ -837,6 +885,10 @@ void FTextureCacheDerivedDataWorker::DoWork()
 					{
 						Mip.DerivedData = FDerivedData(Name, ConvertLegacyCacheKey(DerivedData->GetDerivedDataMipKeyString(MipIndex, Mip)));
 					}
+					else
+					{
+						FirstResidentMipIndex = FMath::Min(FirstResidentMipIndex, MipIndex);
+					}
 					++MipIndex;
 				}
 			}
@@ -900,7 +952,11 @@ void FTextureCacheDerivedDataWorker::DoWork()
 			}
 			else
 			{
-				bSucceeded = DerivedData->AreDerivedMipsAvailable(TexturePathName);
+				const bool bDisableStreaming = (Texture.NeverStream || Texture.LODGroup == TEXTUREGROUP_UI);
+				const int32 FirstMipToLoad = FirstResidentMipIndex;
+				const int32 FirstNonStreamingMipIndex = DerivedData->Mips.Num() - DerivedData->GetNumNonStreamingMips();
+				const int32 FirstMipToPrefetch = IsInGameThread() ? FirstMipToLoad : bDisableStreaming ? 0 : FirstNonStreamingMipIndex;
+				bSucceeded = TryCacheStreamingMips(FirstMipToLoad, FirstMipToPrefetch);
 				if (!bSucceeded)
 				{
 					UE_LOG(LogTexture, Display, TEXT("Texture %s is missing derived mips. The texture will be rebuilt."), *TexturePathName);
diff --git a/Engine/Source/Runtime/Engine/Private/TextureDerivedDataTask.h b/Engine/Source/Runtime/Engine/Private/TextureDerivedDataTask.h
--- a/Engine/Source/Runtime/Engine/Private/TextureDerivedDataTask.h
+++ b/Engine/Source/Runtime/Engine/Private/TextureDerivedDataTask.h
@@ -154,6 +154,9 @@ class FTextureCacheDerivedDataWorker : public FNonAbandonableTask
 	/** true if the derived data was pulled from DDC */
 	bool bLoadedFromDDC = false;
 
+	/** Check for existence of the streaming mips in the cache, and optionally prefetch or load them. */
+	bool TryCacheStreamingMips(int32 FirstMipToLoad, int32 FirstMipToPrefetch);
+
 	/** Build the texture. This function is safe to call from any thread. */
 	void BuildTexture(TArray<FTextureBuildSettings>& InBuildSettingsPerLayer, bool bReplaceExistingDDC = false);
 public:
