Adding project level default codec settings

diff --git a/Engine/Source/Runtime/AudioMixerCore/Public/AudioMixerTypes.h b/Engine/Source/Runtime/AudioMixerCore/Public/AudioMixerTypes.h
--- a/Engine/Source/Runtime/AudioMixerCore/Public/AudioMixerTypes.h
+++ b/Engine/Source/Runtime/AudioMixerCore/Public/AudioMixerTypes.h
@@ -44,6 +44,7 @@ namespace Audio {
 
 	// Indicates a platform-specific format
 	static FName NAME_PLATFORM_SPECIFIC(TEXT("PLATFORM_SPECIFIC"));
+	static FName NAME_PROJECT_DEFINED(TEXT("PROJECT_DEFINED"));
 
 	// Supported on all platforms:
 	static FName NAME_BINKA(TEXT("BINKA"));
diff --git a/Engine/Source/Runtime/Engine/Classes/Sound/AudioSettings.h b/Engine/Source/Runtime/Engine/Classes/Sound/AudioSettings.h
--- a/Engine/Source/Runtime/Engine/Classes/Sound/AudioSettings.h
+++ b/Engine/Source/Runtime/Engine/Classes/Sound/AudioSettings.h
@@ -48,6 +48,25 @@ enum class EMonoChannelUpmixMethod : int8
 	FullVolume
 };
 
+
+// The sound asset compression type to use for assets using the compression type "project defined".
+UENUM()
+enum class EDefaultAudioCompressionType : uint8
+{
+	// Perceptual-based codec which supports all features across all platforms.
+	BinkAudio,
+
+	// Will encode the asset using ADPCM, a time-domain codec that has fixed-sized quality and about ~4x compression ratio, but is relatively cheap to decode.
+	ADPCM,
+
+	// Uncompressed audio. Large memory usage (streamed chunks contain less audio per chunk) but extremely cheap to decode and supports all features. 
+	PCM,
+
+	// Encodes the asset to a platform specific format and will be different depending on the platform. It does not currently support seeking.
+	PlatformSpecific,
+};
+
+
 USTRUCT()
 struct ENGINE_API FAudioQualitySettings
 {
@@ -147,6 +166,10 @@ class ENGINE_API UAudioSettings : public UDeveloperSettings
 	UPROPERTY(config, EditAnywhere, Category = "Audio", meta = (DisplayName = "VOIP Sample Rate"))
 	EVoiceSampleRate VoiPSampleRate;
 
+	/** Default audio compression type to use for audio assets. */
+	UPROPERTY(config, EditAnywhere, Category = "Audio")
+	EDefaultAudioCompressionType DefaultAudioCompressionType;
+
 	/** The amount of audio to send to reverb submixes if no reverb send is setup for the source through attenuation settings. Only used in audio mixer. */
 	UPROPERTY(config)
 	float DefaultReverbSendLevel_DEPRECATED;
diff --git a/Engine/Source/Runtime/Engine/Classes/Sound/SoundWave.h b/Engine/Source/Runtime/Engine/Classes/Sound/SoundWave.h
--- a/Engine/Source/Runtime/Engine/Classes/Sound/SoundWave.h
+++ b/Engine/Source/Runtime/Engine/Classes/Sound/SoundWave.h
@@ -316,6 +316,9 @@ enum class ESoundAssetCompressionType : uint8
 
 	// Encodes the asset to a platform specific format and will be different depending on the platform. It does not currently support seeking.
 	PlatformSpecific,
+
+	// The project defines the codec used for this asset.
+	ProjectDefined,
 };
 
 
@@ -325,16 +328,30 @@ namespace Audio
 	{
 		switch (InDecoderType)
 		{
-		case ESoundAssetCompressionType::PlatformSpecific:		return NAME_PLATFORM_SPECIFIC;
-
 		case ESoundAssetCompressionType::BinkAudio:				return NAME_BINKA;
 		case ESoundAssetCompressionType::ADPCM:					return NAME_ADPCM;
 		case ESoundAssetCompressionType::PCM:					return NAME_PCM;
+		case ESoundAssetCompressionType::PlatformSpecific:		return NAME_PLATFORM_SPECIFIC;
+		case ESoundAssetCompressionType::ProjectDefined:		return NAME_PROJECT_DEFINED;
 		default:
 			ensure(false);
 			return TEXT("UNKNOWN");
 		}
 	}
+
+	static ESoundAssetCompressionType ToSoundAssetCompressionType(EDefaultAudioCompressionType InDefaultCompressionType)
+	{
+		switch (InDefaultCompressionType)
+		{
+			case EDefaultAudioCompressionType::BinkAudio:			return ESoundAssetCompressionType::BinkAudio;
+			case EDefaultAudioCompressionType::ADPCM:				return ESoundAssetCompressionType::ADPCM;
+			case EDefaultAudioCompressionType::PCM:					return ESoundAssetCompressionType::PCM;
+			case EDefaultAudioCompressionType::PlatformSpecific:	return ESoundAssetCompressionType::PlatformSpecific;
+			default:
+				ensure(false);
+				return ESoundAssetCompressionType::PlatformSpecific;
+		}
+	}
 }
 
 // Struct defining a cue point in a sound wave asset
diff --git a/Engine/Source/Runtime/Engine/Private/AudioSettings.cpp b/Engine/Source/Runtime/Engine/Private/AudioSettings.cpp
--- a/Engine/Source/Runtime/Engine/Private/AudioSettings.cpp
+++ b/Engine/Source/Runtime/Engine/Private/AudioSettings.cpp
@@ -35,6 +35,8 @@ UAudioSettings::UAudioSettings(const FObjectInitializer& ObjectInitializer)
 
 	GlobalMinPitchScale = 0.4F;
 	GlobalMaxPitchScale = 2.0F;
+
+	DefaultAudioCompressionType = EDefaultAudioCompressionType::BinkAudio;
 }
 
 void UAudioSettings::AddDefaultSettings()
diff --git a/Engine/Source/Runtime/Engine/Private/SoundWave.cpp b/Engine/Source/Runtime/Engine/Private/SoundWave.cpp
--- a/Engine/Source/Runtime/Engine/Private/SoundWave.cpp
+++ b/Engine/Source/Runtime/Engine/Private/SoundWave.cpp
@@ -681,7 +681,7 @@ USoundWave::USoundWave(const FObjectInitializer& ObjectInitializer)
 	bProcedural = false;
 	bRequiresStopFade = false;
 
-	SoundAssetCompressionType = ESoundAssetCompressionType::BinkAudio;
+	SoundAssetCompressionType = ESoundAssetCompressionType::ProjectDefined;	
 
 #if WITH_EDITOR
 	bWasStreamCachingEnabledOnLastCook = FPlatformCompressionUtilities::IsCurrentPlatformUsingStreamCaching();
@@ -1068,7 +1068,30 @@ void USoundWave::Serialize( FArchive& Ar )
 
 ESoundAssetCompressionType USoundWave::GetSoundAssetCompressionType() const
 {
-	return SoundAssetCompressionType;
+	// If we are told to use the proejct-defined asset compression type, then we need to look it up
+	if (SoundAssetCompressionType == ESoundAssetCompressionType::ProjectDefined)
+	{
+#if WITH_EDITOR
+		// In the editor, always retrieve the settings
+		UAudioSettings* Settings = GetMutableDefault<UAudioSettings>();
+		return Audio::ToSoundAssetCompressionType(Settings->DefaultAudioCompressionType);
+#else
+		// Only query the project settings once when retrieving the asset compression type when not running with the editor
+		static bool bDefaultCompressionTypeCached = false;
+		static ESoundAssetCompressionType DefaultAudioCompressionType = ESoundAssetCompressionType::PlatformSpecific;
+
+		if (!bDefaultCompressionTypeCached)
+		{
+			DefaultAudioCompressionType = Audio::ToSoundAssetCompressionType(Settings->DefaultAudioCompressionType);
+			bDefaultCompressionTypeCached = true;
+		}
+		return DefaultAudioCompressionType;
+#endif
+	}
+	else
+	{
+		return SoundAssetCompressionType;
+	}
 }
 
 void USoundWave::SetSoundAssetCompressionType(ESoundAssetCompressionType InSoundAssetCompressionType)
