Fixing non-editor builds

diff --git a/Engine/Source/Runtime/Engine/Private/SoundWave.cpp b/Engine/Source/Runtime/Engine/Private/SoundWave.cpp
--- a/Engine/Source/Runtime/Engine/Private/SoundWave.cpp
+++ b/Engine/Source/Runtime/Engine/Private/SoundWave.cpp
@@ -1082,6 +1082,7 @@ ESoundAssetCompressionType USoundWave::GetSoundAssetCompressionType() const
 
 		if (!bDefaultCompressionTypeCached)
 		{
+			UAudioSettings* Settings = GetMutableDefault<UAudioSettings>();
 			DefaultAudioCompressionType = Audio::ToSoundAssetCompressionType(Settings->DefaultAudioCompressionType);
 			bDefaultCompressionTypeCached = true;
 		}
