Strata - added check that operators are only used with Slabs for now.

diff --git a/Engine/Source/Runtime/Engine/Private/Materials/HLSLMaterialTranslator.cpp b/Engine/Source/Runtime/Engine/Private/Materials/HLSLMaterialTranslator.cpp
--- a/Engine/Source/Runtime/Engine/Private/Materials/HLSLMaterialTranslator.cpp
+++ b/Engine/Source/Runtime/Engine/Private/Materials/HLSLMaterialTranslator.cpp
@@ -9742,6 +9742,7 @@ bool FHLSLMaterialTranslator::StrataGenerateDerivedMaterialOperatorData()
 		int VOpTopBranchCountTaken = 0;
 		int VOpBottomBranchCountTaken = 0;
 		bool bStrataUsesVerticalLayering = false;
+		bool bOperatorEncountered = false;
 
 		std::function<void(FStrataOperator&, bool)> WalkOperators = [&](FStrataOperator& CurrentOperator, bool bInsideParameterBlendingSubTree) -> void
 		{
@@ -9758,6 +9759,7 @@ bool FHLSLMaterialTranslator::StrataGenerateDerivedMaterialOperatorData()
 			{
 				WalkOperators(StrataMaterialExpressionRegisteredOperators[CurrentOperator.LeftIndex], bUseParameterBlending);
 				WalkOperators(StrataMaterialExpressionRegisteredOperators[CurrentOperator.RightIndex], bUseParameterBlending);
+				bOperatorEncountered = true;
 				break;
 			}
 			case STRATA_OPERATOR_HORIZONTAL:
@@ -9765,12 +9767,14 @@ bool FHLSLMaterialTranslator::StrataGenerateDerivedMaterialOperatorData()
 			{
 				WalkOperators(StrataMaterialExpressionRegisteredOperators[CurrentOperator.LeftIndex], bUseParameterBlending);
 				WalkOperators(StrataMaterialExpressionRegisteredOperators[CurrentOperator.RightIndex], bUseParameterBlending);
+				bOperatorEncountered = true;
 				break;
 			}
 			case STRATA_OPERATOR_WEIGHT:
 			case STRATA_OPERATOR_THINFILM:
 			{
 				WalkOperators(StrataMaterialExpressionRegisteredOperators[CurrentOperator.LeftIndex], bUseParameterBlending);
+				bOperatorEncountered = true;
 				break;
 			}
 			case STRATA_OPERATOR_BSDF:
@@ -9814,6 +9818,12 @@ bool FHLSLMaterialTranslator::StrataGenerateDerivedMaterialOperatorData()
 			// Even though we could support Unlit with slab.
 		}
 
+		if ((bHasUnlit || bHasVFogCloud || bHasHair || bHasSLW) && bOperatorEncountered)
+		{
+			Errorf(TEXT("Unlit, Fog/Cloud, Hair or SingleLayerWater cannot be used with operators. See %s (asset: %s).\r\n"), *Material->GetDebugName(), *Material->GetAssetPath().ToString());
+			// This is because it will results in simpler lighting loops focusin on slab.
+		}
+
 		if (StrataMaterialBSDFCount > STRATA_MAX_BSDF_COUNT)
 		{
 			Errorf(TEXT("Material tries to register more BSDF than can be supproted (%d > %d). See %s (asset: %s).\r\n"), StrataMaterialBSDFCount, STRATA_MAX_BSDF_COUNT, *Material->GetDebugName(), *Material->GetAssetPath().ToString());
