Strata - Removed Opacity form the root node.

diff --git a/Engine/Source/Runtime/Engine/Private/Materials/Material.cpp b/Engine/Source/Runtime/Engine/Private/Materials/Material.cpp
--- a/Engine/Source/Runtime/Engine/Private/Materials/Material.cpp
+++ b/Engine/Source/Runtime/Engine/Private/Materials/Material.cpp
@@ -5718,7 +5718,6 @@ bool UMaterial::IsPropertySupported(EMaterialProperty InProperty) const
 		switch (InProperty)
 		{
 		case MP_Refraction:
-		case MP_Opacity:
 		case MP_OpacityMask:
 		case MP_AmbientOcclusion:
 		case MP_WorldPositionOffset:
