Texture: Temporarily reverted to using bPagedToDerivedData to determine whether a mip is streaming or inline

#rnx

diff --git a/Engine/Source/Runtime/Engine/Private/Texture2D.cpp b/Engine/Source/Runtime/Engine/Private/Texture2D.cpp
--- a/Engine/Source/Runtime/Engine/Private/Texture2D.cpp
+++ b/Engine/Source/Runtime/Engine/Private/Texture2D.cpp
@@ -148,14 +148,6 @@ void FTexture2DMipMap::Serialize(FArchive& Ar, UObject* Owner, int32 MipIdx)
 	BulkData.Serialize(Ar, Owner, MipIdx, false);
 #endif
 
-	// Streaming mips are saved with a size of 0 because they are stored separately.
-	// IsBulkDataLoaded() returns true for this empty bulk data. Remove the empty
-	// bulk data to allow unloaded streaming mips to be detected.
-	if (BulkData.GetBulkDataSize() == 0)
-	{
-		BulkData.RemoveBulkData();
-	}
-
 	Ar << SizeX;
 	Ar << SizeY;
 	Ar << SizeZ;
@@ -163,9 +155,21 @@ void FTexture2DMipMap::Serialize(FArchive& Ar, UObject* Owner, int32 MipIdx)
 #if WITH_EDITORONLY_DATA
 	if (!Ar.IsFilterEditorOnly())
 	{
-		bool bLocalPagedToDerivedData = DerivedData.HasData();
 		Ar << FileRegionType;
-		Ar << bLocalPagedToDerivedData;
+		PRAGMA_DISABLE_DEPRECATION_WARNINGS;
+		Ar << bPagedToDerivedData;
+		PRAGMA_ENABLE_DEPRECATION_WARNINGS;
+	}
+
+	// Streaming mips are saved with a size of 0 because they are stored separately.
+	// IsBulkDataLoaded() returns true for this empty bulk data. Remove the empty
+	// bulk data to allow unloaded streaming mips to be detected.
+	PRAGMA_DISABLE_DEPRECATION_WARNINGS;
+	const bool bLocalPagedToDerivedData = bPagedToDerivedData;
+	PRAGMA_ENABLE_DEPRECATION_WARNINGS;
+	if (BulkData.GetBulkDataSize() == 0 && bLocalPagedToDerivedData)
+	{
+		BulkData.RemoveBulkData();
 	}
 #endif // #if WITH_EDITORONLY_DATA
 }
@@ -189,6 +193,9 @@ int64 FTexture2DMipMap::StoreInDerivedDataCache(const FStringView InKey, const F
 	GetCache().PutValue({{Name, Key, MoveTemp(Value), Policy}}, AsyncOwner);
 	AsyncOwner.KeepAlive();
 
+	PRAGMA_DISABLE_DEPRECATION_WARNINGS;
+	bPagedToDerivedData = true;
+	PRAGMA_ENABLE_DEPRECATION_WARNINGS;
 	DerivedData = FDerivedData(Name, Key);
 	BulkData.RemoveBulkData();
 	return BulkDataSizeInBytes;
diff --git a/Engine/Source/Runtime/Engine/Private/TextureDerivedDataTask.cpp b/Engine/Source/Runtime/Engine/Private/TextureDerivedDataTask.cpp
--- a/Engine/Source/Runtime/Engine/Private/TextureDerivedDataTask.cpp
+++ b/Engine/Source/Runtime/Engine/Private/TextureDerivedDataTask.cpp
@@ -881,7 +881,10 @@ void FTextureCacheDerivedDataWorker::DoWork()
 				const FSharedString Name(TexturePathName);
 				for (FTexture2DMipMap& Mip : DerivedData->Mips)
 				{
-					if (!Mip.BulkData.IsBulkDataLoaded())
+					PRAGMA_DISABLE_DEPRECATION_WARNINGS;
+					const bool bPagedToDerivedData = Mip.bPagedToDerivedData;
+					PRAGMA_ENABLE_DEPRECATION_WARNINGS;
+					if (bPagedToDerivedData)
 					{
 						Mip.DerivedData = FDerivedData(Name, ConvertLegacyCacheKey(DerivedData->GetDerivedDataMipKeyString(MipIndex, Mip)));
 					}
@@ -1580,6 +1583,9 @@ private:
 
 				FSharedString MipName(WriteToString<256>(Output.GetName(), TEXT(" [MIP "), MipIndex, TEXT("]")));
 				NewMip->DerivedData = FDerivedData(MoveTemp(MipName), Params.CacheKey, MipId);
+				PRAGMA_DISABLE_DEPRECATION_WARNINGS;
+				NewMip->bPagedToDerivedData = true;
+				PRAGMA_ENABLE_DEPRECATION_WARNINGS;
 			}
 
 			DerivedData.Mips.Add(NewMip.Release());
diff --git a/Engine/Source/Runtime/Engine/Private/VT/VirtualTextureBuiltData.cpp b/Engine/Source/Runtime/Engine/Private/VT/VirtualTextureBuiltData.cpp
--- a/Engine/Source/Runtime/Engine/Private/VT/VirtualTextureBuiltData.cpp
+++ b/Engine/Source/Runtime/Engine/Private/VT/VirtualTextureBuiltData.cpp
@@ -375,14 +375,6 @@ void FVirtualTextureBuiltData::Serialize(FArchive& Ar, UObject* Owner, int32 Fir
 
 		Chunk.BulkData.Serialize(Ar, Owner, SerialzeChunkId, false);
 
-		// Streaming chunks are saved with a size of 0 because they are stored separately.
-		// IsBulkDataLoaded() returns true for this empty bulk data. Remove the empty
-		// bulk data to allow unloaded streaming chunks to be detected.
-		if (Chunk.BulkData.GetBulkDataSize() == 0)
-		{
-			Chunk.BulkData.RemoveBulkData();
-		}
-
 #if WITH_EDITORONLY_DATA
 		if (!bCooked)
 		{
@@ -392,6 +384,14 @@ void FVirtualTextureBuiltData::Serialize(FArchive& Ar, UObject* Owner, int32 Fir
 				Chunk.ShortenKey(Chunk.DerivedDataKey, Chunk.ShortDerivedDataKey);
 			}
 		}
+
+		// Streaming chunks are saved with a size of 0 because they are stored separately.
+		// IsBulkDataLoaded() returns true for this empty bulk data. Remove the empty
+		// bulk data to allow unloaded streaming chunks to be detected.
+		if (Chunk.BulkData.GetBulkDataSize() == 0 && !Chunk.DerivedDataKey.IsEmpty())
+		{
+			Chunk.BulkData.RemoveBulkData();
+		}
 #endif // WITH_EDITORONLY_DATA
 
 		SerialzeChunkId++;
