Add 'Optional' EditorOnly data for UMaterialInterface and UMaterialFunctionInterface

diff --git a/Engine/Source/Editor/UnrealEd/Private/Fbx/FbxMaterialImport.cpp b/Engine/Source/Editor/UnrealEd/Private/Fbx/FbxMaterialImport.cpp
--- a/Engine/Source/Editor/UnrealEd/Private/Fbx/FbxMaterialImport.cpp
+++ b/Engine/Source/Editor/UnrealEd/Private/Fbx/FbxMaterialImport.cpp
@@ -324,7 +324,7 @@ bool UnFbx::FFbxImporter::CreateAndLinkExpressionForMaterialProperty(
 
 						// and link it to the material 
 						UMaterialExpressionTextureSample* UnrealTextureExpression = NewObject<UMaterialExpressionTextureSample>(UnrealMaterial);
-						UnrealMaterial->Expressions.Add( UnrealTextureExpression );
+						UnrealMaterial->GetExpressionCollection().AddExpression( UnrealTextureExpression );
 						MaterialInput.Expression = UnrealTextureExpression;
 						UnrealTextureExpression->Texture = UnrealTexture;
 						UnrealTextureExpression->SamplerType = bSetupAsNormalMap ? 
@@ -346,7 +346,7 @@ bool UnFbx::FFbxImporter::CreateAndLinkExpressionForMaterialProperty(
 						{
 							// Create a texture coord node for the texture sample
 							UMaterialExpressionTextureCoordinate* MyCoordExpression = NewObject<UMaterialExpressionTextureCoordinate>(UnrealMaterial);
-							UnrealMaterial->Expressions.Add(MyCoordExpression);
+							UnrealMaterial->GetExpressionCollection().AddExpression(MyCoordExpression);
 							MyCoordExpression->CoordinateIndex = (SetIndex >= 0) ? SetIndex : 0;
 							MyCoordExpression->UTiling = ScaleU;
 							MyCoordExpression->VTiling = ScaleV;
@@ -382,14 +382,15 @@ bool UnFbx::FFbxImporter::CreateAndLinkExpressionForMaterialProperty(
 //-------------------------------------------------------------------------
 void UnFbx::FFbxImporter::FixupMaterial( const FbxSurfaceMaterial& FbxMaterial, UMaterial* UnrealMaterial )
 {
+	UMaterialEditorOnlyData* UnrealMaterialEditorOnly = UnrealMaterial->GetEditorOnlyData();
 	// add a basic diffuse color if no texture is linked to diffuse
-	if (UnrealMaterial->BaseColor.Expression == NULL)
+	if (UnrealMaterialEditorOnly->BaseColor.Expression == NULL)
 	{
 		FbxDouble3 DiffuseColor;
 		
 		UMaterialExpressionVectorParameter* MyColorExpression = NewObject<UMaterialExpressionVectorParameter>(UnrealMaterial);
-		UnrealMaterial->Expressions.Add( MyColorExpression );
-		UnrealMaterial->BaseColor.Expression = MyColorExpression;
+		UnrealMaterial->GetExpressionCollection().AddExpression( MyColorExpression );
+		UnrealMaterialEditorOnly->BaseColor.Expression = MyColorExpression;
 
 		bool bFoundDiffuseColor = true;
 		if( FbxMaterial.GetClassId().Is(FbxSurfacePhong::ClassId) )
@@ -419,13 +420,13 @@ void UnFbx::FFbxImporter::FixupMaterial( const FbxSurfaceMaterial& FbxMaterial,
 			MyColorExpression->DefaultValue.B = 0.5f+(0.5f*FMath::Rand()) / static_cast<float>(RAND_MAX);
 		}
 
-		TArray<FExpressionOutput> Outputs = UnrealMaterial->BaseColor.Expression->GetOutputs();
+		TArray<FExpressionOutput> Outputs = UnrealMaterialEditorOnly->BaseColor.Expression->GetOutputs();
 		FExpressionOutput* Output = Outputs.GetData();
-		UnrealMaterial->BaseColor.Mask = Output->Mask;
-		UnrealMaterial->BaseColor.MaskR = Output->MaskR;
-		UnrealMaterial->BaseColor.MaskG = Output->MaskG;
-		UnrealMaterial->BaseColor.MaskB = Output->MaskB;
-		UnrealMaterial->BaseColor.MaskA = Output->MaskA;
+		UnrealMaterialEditorOnly->BaseColor.Mask = Output->Mask;
+		UnrealMaterialEditorOnly->BaseColor.MaskR = Output->MaskR;
+		UnrealMaterialEditorOnly->BaseColor.MaskG = Output->MaskG;
+		UnrealMaterialEditorOnly->BaseColor.MaskB = Output->MaskB;
+		UnrealMaterialEditorOnly->BaseColor.MaskA = Output->MaskA;
 	}
 }
 
@@ -799,19 +800,20 @@ UMaterialInterface* UnFbx::FFbxImporter::CreateUnrealMaterial(const FbxSurfaceMa
 				UE_LOG(LogFbxMaterialImport, Display, TEXT("-------------------------------"));
 			}
 #endif
-			CreateAndLinkExpressionForMaterialProperty(FbxMaterial, UnrealMaterial, FbxSurfaceMaterial::sDiffuse, UnrealMaterial->BaseColor, false, OutUVSets, FVector2D(240, -320));
-			CreateAndLinkExpressionForMaterialProperty(FbxMaterial, UnrealMaterial, FbxSurfaceMaterial::sEmissive, UnrealMaterial->EmissiveColor, false, OutUVSets, FVector2D(240, -64));
-			CreateAndLinkExpressionForMaterialProperty(FbxMaterial, UnrealMaterial, FbxSurfaceMaterial::sSpecular, UnrealMaterial->Specular, false, OutUVSets, FVector2D(240, -128));
-			CreateAndLinkExpressionForMaterialProperty(FbxMaterial, UnrealMaterial, FbxSurfaceMaterial::sSpecularFactor, UnrealMaterial->Roughness, false, OutUVSets, FVector2D(240, -180));
-			CreateAndLinkExpressionForMaterialProperty(FbxMaterial, UnrealMaterial, FbxSurfaceMaterial::sShininess, UnrealMaterial->Metallic, false, OutUVSets, FVector2D(240, -210));
-			if (!CreateAndLinkExpressionForMaterialProperty(FbxMaterial, UnrealMaterial, FbxSurfaceMaterial::sNormalMap, UnrealMaterial->Normal, true, OutUVSets, FVector2D(240, 256)))
+			UMaterialEditorOnlyData* UnrealMaterialEditorOnly = UnrealMaterial->GetEditorOnlyData();
+			CreateAndLinkExpressionForMaterialProperty(FbxMaterial, UnrealMaterial, FbxSurfaceMaterial::sDiffuse, UnrealMaterialEditorOnly->BaseColor, false, OutUVSets, FVector2D(240, -320));
+			CreateAndLinkExpressionForMaterialProperty(FbxMaterial, UnrealMaterial, FbxSurfaceMaterial::sEmissive, UnrealMaterialEditorOnly->EmissiveColor, false, OutUVSets, FVector2D(240, -64));
+			CreateAndLinkExpressionForMaterialProperty(FbxMaterial, UnrealMaterial, FbxSurfaceMaterial::sSpecular, UnrealMaterialEditorOnly->Specular, false, OutUVSets, FVector2D(240, -128));
+			CreateAndLinkExpressionForMaterialProperty(FbxMaterial, UnrealMaterial, FbxSurfaceMaterial::sSpecularFactor, UnrealMaterialEditorOnly->Roughness, false, OutUVSets, FVector2D(240, -180));
+			CreateAndLinkExpressionForMaterialProperty(FbxMaterial, UnrealMaterial, FbxSurfaceMaterial::sShininess, UnrealMaterialEditorOnly->Metallic, false, OutUVSets, FVector2D(240, -210));
+			if (!CreateAndLinkExpressionForMaterialProperty(FbxMaterial, UnrealMaterial, FbxSurfaceMaterial::sNormalMap, UnrealMaterialEditorOnly->Normal, true, OutUVSets, FVector2D(240, 256)))
 			{
-				CreateAndLinkExpressionForMaterialProperty(FbxMaterial, UnrealMaterial, FbxSurfaceMaterial::sBump, UnrealMaterial->Normal, true, OutUVSets, FVector2D(240, 256)); // no bump in unreal, use as normal map
+				CreateAndLinkExpressionForMaterialProperty(FbxMaterial, UnrealMaterial, FbxSurfaceMaterial::sBump, UnrealMaterialEditorOnly->Normal, true, OutUVSets, FVector2D(240, 256)); // no bump in unreal, use as normal map
 			}
-			if (CreateAndLinkExpressionForMaterialProperty(FbxMaterial, UnrealMaterial, FbxSurfaceMaterial::sTransparentColor, UnrealMaterial->Opacity, false, OutUVSets, FVector2D(200, 256)))
+			if (CreateAndLinkExpressionForMaterialProperty(FbxMaterial, UnrealMaterial, FbxSurfaceMaterial::sTransparentColor, UnrealMaterialEditorOnly->Opacity, false, OutUVSets, FVector2D(200, 256)))
 			{
 				UnrealMaterial->BlendMode = BLEND_Translucent;
-				CreateAndLinkExpressionForMaterialProperty(FbxMaterial, UnrealMaterial, FbxSurfaceMaterial::sTransparencyFactor, UnrealMaterial->OpacityMask, false, OutUVSets, FVector2D(150, 256));
+				CreateAndLinkExpressionForMaterialProperty(FbxMaterial, UnrealMaterial, FbxSurfaceMaterial::sTransparencyFactor, UnrealMaterialEditorOnly->OpacityMask, false, OutUVSets, FVector2D(150, 256));
 
 			}
 			FixupMaterial(FbxMaterial, UnrealMaterial); // add random diffuse if none exists
diff --git a/Engine/Source/Editor/UnrealEd/Private/LevelEditorViewport.cpp b/Engine/Source/Editor/UnrealEd/Private/LevelEditorViewport.cpp
--- a/Engine/Source/Editor/UnrealEd/Private/LevelEditorViewport.cpp
+++ b/Engine/Source/Editor/UnrealEd/Private/LevelEditorViewport.cpp
@@ -375,12 +375,14 @@ static bool TryAndCreateMaterialInput( UMaterial* UnrealMaterial, EMaterialKind:
 	}
 
 	const bool bSetupAsNormalMap = UnrealTexture->IsNormalMap();
+
+	UMaterialEditorOnlyData* UnrealMaterialEditorOnly = UnrealMaterial->GetEditorOnlyData();
 	
 	UnrealMaterial->BlendMode = bTextureHasAlpha ? EBlendMode::BLEND_Masked : EBlendMode::BLEND_Opaque;
 
 	// Create a new texture sample expression, this is our texture input node into the material output.
 	UMaterialExpressionTextureSample* UnrealTextureExpression = NewObject<UMaterialExpressionTextureSample>(UnrealMaterial);
-	UnrealMaterial->Expressions.Add( UnrealTextureExpression );
+	UnrealMaterial->GetExpressionCollection().AddExpression( UnrealTextureExpression );
 	MaterialInput.Expression = UnrealTextureExpression;
 	UnrealTextureExpression->Texture = UnrealTexture;
 	UnrealTextureExpression->AutoSetSampleType();
@@ -398,13 +400,13 @@ static bool TryAndCreateMaterialInput( UMaterial* UnrealMaterial, EMaterialKind:
 		UMaterialExpressionTransformPosition* TransformPositionExpression = NewObject<UMaterialExpressionTransformPosition>(UnrealMaterial);
 		UMaterialExpressionWorldPosition* WorldPosExpression = NewObject<UMaterialExpressionWorldPosition>(UnrealMaterial);
 
-		UnrealMaterial->Expressions.Add( DivideExpression );
-		UnrealMaterial->Expressions.Add( BoundsRelativePosExpression );
-		UnrealMaterial->Expressions.Add( BoundsSizeExpression );
-		UnrealMaterial->Expressions.Add( LocalBoundsMinExpression );
-		UnrealMaterial->Expressions.Add( LocalBoundsMaxExpression );
-		UnrealMaterial->Expressions.Add( TransformPositionExpression );
-		UnrealMaterial->Expressions.Add( WorldPosExpression );
+		UnrealMaterial->GetExpressionCollection().AddExpression( DivideExpression );
+		UnrealMaterial->GetExpressionCollection().AddExpression( BoundsRelativePosExpression );
+		UnrealMaterial->GetExpressionCollection().AddExpression( BoundsSizeExpression );
+		UnrealMaterial->GetExpressionCollection().AddExpression( LocalBoundsMinExpression );
+		UnrealMaterial->GetExpressionCollection().AddExpression( LocalBoundsMaxExpression );
+		UnrealMaterial->GetExpressionCollection().AddExpression( TransformPositionExpression );
+		UnrealMaterial->GetExpressionCollection().AddExpression( WorldPosExpression );
 
 		int32 EditorPosX = UnrealTextureExpression->MaterialExpressionEditorX;
 		int32 EditorPosY = UnrealTextureExpression->MaterialExpressionEditorY;
@@ -467,29 +469,29 @@ static bool TryAndCreateMaterialInput( UMaterial* UnrealMaterial, EMaterialKind:
 	{
 		if ( TextureKind == EMaterialKind::Base )
 		{
-			UnrealMaterial->BaseColor.Expression = UnrealTextureExpression;
+			UnrealMaterialEditorOnly->BaseColor.Expression = UnrealTextureExpression;
 			if ( bTextureHasAlpha && UnrealTextureExpression->Outputs.IsValidIndex(4) )
 			{
-				UnrealMaterial->Opacity.Connect(4, UnrealTextureExpression);
-				UnrealMaterial->OpacityMask.Connect(4, UnrealTextureExpression);
+				UnrealMaterialEditorOnly->Opacity.Connect(4, UnrealTextureExpression);
+				UnrealMaterialEditorOnly->OpacityMask.Connect(4, UnrealTextureExpression);
 			}
 		}
 		else if ( TextureKind == EMaterialKind::Specular )
 		{
-			UnrealMaterial->Specular.Expression = UnrealTextureExpression;
+			UnrealMaterialEditorOnly->Specular.Expression = UnrealTextureExpression;
 		}
 		else if ( TextureKind == EMaterialKind::Emissive )
 		{
-			UnrealMaterial->EmissiveColor.Expression = UnrealTextureExpression;
+			UnrealMaterialEditorOnly->EmissiveColor.Expression = UnrealTextureExpression;
 		}
 		else
 		{
-			UnrealMaterial->BaseColor.Expression = UnrealTextureExpression;
+			UnrealMaterialEditorOnly->BaseColor.Expression = UnrealTextureExpression;
 		}
 	}
 	else
 	{
-		UnrealMaterial->Normal.Expression = UnrealTextureExpression;
+		UnrealMaterialEditorOnly->Normal.Expression = UnrealTextureExpression;
 	}
 
 
@@ -532,13 +534,15 @@ UObject* FLevelEditorViewportClient::GetOrCreateMaterialFromTexture( UTexture* U
 		return nullptr;
 	}
 
+	UMaterialEditorOnlyData* UnrealMaterialEditorOnly = UnrealMaterial->GetEditorOnlyData();
+
 	const int HSpace = -300;
 
 	// If we were able to figure out the material kind, we need to try and build a complex material
 	// involving multiple textures.  If not, just try and connect what we found to the base map.
 	if ( MaterialKind == EMaterialKind::Unknown )
 	{
-		TryAndCreateMaterialInput( UnrealMaterial, EMaterialKind::Base, UnrealTexture, UnrealMaterial->BaseColor, HSpace, 0 );
+		TryAndCreateMaterialInput( UnrealMaterial, EMaterialKind::Base, UnrealTexture, UnrealMaterialEditorOnly->BaseColor, HSpace, 0 );
 	}
 	else
 	{
@@ -579,10 +583,10 @@ UObject* FLevelEditorViewportClient::GetOrCreateMaterialFromTexture( UTexture* U
 
 		// Connect and layout any textures we find into their respective inputs in the material.
 		const int VSpace = 170;
-		TryAndCreateMaterialInput( UnrealMaterial, EMaterialKind::Base, BaseTexture, UnrealMaterial->BaseColor, HSpace, VSpace * -1 );
-		TryAndCreateMaterialInput( UnrealMaterial, EMaterialKind::Specular, SpecularTexture, UnrealMaterial->Specular, HSpace, VSpace * 0 );
-		TryAndCreateMaterialInput( UnrealMaterial, EMaterialKind::Emissive, EmissiveTexture, UnrealMaterial->EmissiveColor, HSpace, VSpace * 1 );
-		TryAndCreateMaterialInput( UnrealMaterial, EMaterialKind::Normal, NormalTexture, UnrealMaterial->Normal, HSpace, VSpace * 2 );
+		TryAndCreateMaterialInput( UnrealMaterial, EMaterialKind::Base, BaseTexture, UnrealMaterialEditorOnly->BaseColor, HSpace, VSpace * -1 );
+		TryAndCreateMaterialInput( UnrealMaterial, EMaterialKind::Specular, SpecularTexture, UnrealMaterialEditorOnly->Specular, HSpace, VSpace * 0 );
+		TryAndCreateMaterialInput( UnrealMaterial, EMaterialKind::Emissive, EmissiveTexture, UnrealMaterialEditorOnly->EmissiveColor, HSpace, VSpace * 1 );
+		TryAndCreateMaterialInput( UnrealMaterial, EMaterialKind::Normal, NormalTexture, UnrealMaterialEditorOnly->Normal, HSpace, VSpace * 2 );
 	}
 
 	UnrealMaterial->PreEditChange(nullptr);
diff --git a/Engine/Source/Editor/UnrealEd/Private/Lightmass/LightmassRender.cpp b/Engine/Source/Editor/UnrealEd/Private/Lightmass/LightmassRender.cpp
--- a/Engine/Source/Editor/UnrealEd/Private/Lightmass/LightmassRender.cpp
+++ b/Engine/Source/Editor/UnrealEd/Private/Lightmass/LightmassRender.cpp
@@ -566,25 +566,26 @@ public:
 	{
 		bool bConnected = false;
 
+		UMaterialEditorOnlyData* MaterialEditorOnly = InMaterial->GetEditorOnlyData();
 		switch (MaterialInput)
 		{
 		case MP_EmissiveColor:
-			bConnected = InMaterial->EmissiveColor.Expression != NULL;
+			bConnected = MaterialEditorOnly->EmissiveColor.Expression != nullptr;
 			break;
 		case MP_DiffuseColor:
-			bConnected = InMaterial->BaseColor.Expression != NULL;
+			bConnected = MaterialEditorOnly->BaseColor.Expression != nullptr;
 			break;
 		case MP_SpecularColor:
-			bConnected = InMaterial->Specular.Expression != NULL;
+			bConnected = MaterialEditorOnly->Specular.Expression != nullptr;
 			break;
 		case MP_Normal:
-			bConnected = InMaterial->Normal.Expression != NULL;
+			bConnected = MaterialEditorOnly->Normal.Expression != nullptr;
 			break;
 		case MP_Opacity:
-			bConnected = InMaterial->Opacity.Expression != NULL;
+			bConnected = MaterialEditorOnly->Opacity.Expression != nullptr;
 			break;
 		case MP_OpacityMask:
-			bConnected = InMaterial->OpacityMask.Expression != NULL;
+			bConnected = MaterialEditorOnly->OpacityMask.Expression != nullptr;
 			break;
 		default:
 			break;
@@ -592,7 +593,7 @@ public:
 
 		// Note: only checking to see whether the entire material attributes connection exists.  
 		// This means materials using the material attributes input will export more attributes than is necessary.
-		bConnected = InMaterial->bUseMaterialAttributes ? InMaterial->MaterialAttributes.Expression != NULL : bConnected;
+		bConnected = InMaterial->bUseMaterialAttributes ? MaterialEditorOnly->MaterialAttributes.Expression != NULL : bConnected;
 		return bConnected;
 	}
 
diff --git a/Engine/Source/Editor/UnrealEd/Private/MaterialGraph.cpp b/Engine/Source/Editor/UnrealEd/Private/MaterialGraph.cpp
--- a/Engine/Source/Editor/UnrealEd/Private/MaterialGraph.cpp
+++ b/Engine/Source/Editor/UnrealEd/Private/MaterialGraph.cpp
@@ -53,11 +53,11 @@ void UMaterialGraph::RebuildGraph()
 	// Pre-group expressions & comments per subgraph to avoid unnecessary iteration over all material expressions
 	TMap<UMaterialExpression*, TArray<UMaterialExpression*>> SubgraphExpressionMap;
 	TMap<UMaterialExpression*, TArray<UMaterialExpressionComment*>> SubgraphCommentMap;
-	for (UMaterialExpression* Expression : Material->Expressions)
+	for (UMaterialExpression* Expression : Material->GetExpressions())
 	{
 		SubgraphExpressionMap.FindOrAdd(Expression->SubgraphExpression).Add(Expression);
 	}
-	for (UMaterialExpressionComment* Comment : Material->EditorComments)
+	for (UMaterialExpressionComment* Comment : Material->GetEditorComments())
 	{
 		if (Comment)
 		{
@@ -250,7 +250,7 @@ void UMaterialGraph::RebuildGraphInternal(const TMap<UMaterialExpression*, TArra
 		MaterialInputs.Add(FMaterialInputInfo(FMaterialAttributeDefinitionMap::GetDisplayNameForMaterial(MP_AmbientOcclusion, Material), MP_AmbientOcclusion, LOCTEXT("AmbientOcclusionToolTip", "Simulate the self-shadowing that happens within crevices of a surface, or of a volume for volumetric clouds only")));
 		MaterialInputs.Add(FMaterialInputInfo(FMaterialAttributeDefinitionMap::GetDisplayNameForMaterial(MP_Refraction, Material), MP_Refraction, LOCTEXT("RefractionToolTip", "Takes in a texture or value that simulates the index of refraction of the surface")));
 
-		for (int32 UVIndex = 0; UVIndex < UE_ARRAY_COUNT(Material->CustomizedUVs); UVIndex++)
+		for (int32 UVIndex = 0; UVIndex < UE_ARRAY_COUNT(Material->GetEditorOnlyData()->CustomizedUVs); UVIndex++)
 		{
 			//@todo - localize
 			MaterialInputs.Add(FMaterialInputInfo(FText::FromString(FString::Printf(TEXT("Customized UV%u"), UVIndex)), (EMaterialProperty)(MP_CustomizedUVs0 + UVIndex), FText::FromString(FString::Printf(TEXT("CustomizedUV%uToolTip"), UVIndex))));
@@ -274,18 +274,18 @@ void UMaterialGraph::RebuildGraphInternal(const TMap<UMaterialExpression*, TArra
 
 	if (Material->IsUsingControlFlow())
 	{
-		if (ensure(Material->ExpressionExecBegin))
+		if (ensure(Material->GetExpressionExecBegin()))
 		{
-			InitExpressionNewNode<UMaterialGraphNode>(this, Material->ExpressionExecBegin, false);
+			InitExpressionNewNode<UMaterialGraphNode>(this, Material->GetExpressionExecBegin(), false);
 		}
 		
 		if (MaterialFunction)
 		{
-			check(MaterialFunction->ExpressionExecBegin == Material->ExpressionExecBegin);
-			check(MaterialFunction->ExpressionExecEnd == Material->ExpressionExecEnd);
-			if (ensure(Material->ExpressionExecEnd))
+			check(MaterialFunction->GetExpressionExecBegin() == Material->GetExpressionExecBegin());
+			check(MaterialFunction->GetExpressionExecEnd() == Material->GetExpressionExecEnd());
+			if (ensure(Material->GetExpressionExecEnd()))
 			{
-				InitExpressionNewNode<UMaterialGraphNode>(this, Material->ExpressionExecEnd, false);
+				InitExpressionNewNode<UMaterialGraphNode>(this, Material->GetExpressionExecEnd(), false);
 			}
 		}
 	}
@@ -479,9 +479,8 @@ void UMaterialGraph::LinkGraphNodesFromMaterial()
 		}
 	}
 
-	for (int32 Index = 0; Index < Material->Expressions.Num(); Index++)
+	for (UMaterialExpression* Expression : Material->GetExpressions())
 	{
-		UMaterialExpression* Expression = Material->Expressions[Index];
 		if (!Expression)
 		{
 			continue;
@@ -541,7 +540,7 @@ void UMaterialGraph::LinkGraphNodesFromMaterial()
 				UMaterialExpression* ConnectedExpression = ExecOutput->GetExpression();
 				if (ConnectedExpression)
 				{
-					if (ConnectedExpression == Material->ExpressionExecEnd && RootNode)
+					if (ConnectedExpression == Material->GetExpressionExecEnd() && RootNode)
 					{
 						// Exec end point is the root node (assuming there is a RootNode)
 						// MaterialFunctions currently do not have a RootNode, and instead have the ExecEnd node as a stand-alone
@@ -674,7 +673,7 @@ void UMaterialGraph::LinkMaterialExpressionsFromGraph()
 										bModifiedExpression = true;
 										Expression->Modify();
 									}
-									ExpressionOutput->Connect(Material->ExpressionExecEnd);
+									ExpressionOutput->Connect(Material->GetExpressionExecEnd());
 								}
 								else
 								{
diff --git a/Engine/Source/Editor/UnrealEd/Private/PreviewMaterial.cpp b/Engine/Source/Editor/UnrealEd/Private/PreviewMaterial.cpp
--- a/Engine/Source/Editor/UnrealEd/Private/PreviewMaterial.cpp
+++ b/Engine/Source/Editor/UnrealEd/Private/PreviewMaterial.cpp
@@ -828,7 +828,7 @@ void UMaterialEditorInstanceConstant::ApplySourceFunctionChanges()
 		SourceFunction->RuntimeVirtualTextureParameterValues = SourceInstance->RuntimeVirtualTextureParameterValues;
 		SourceFunction->FontParameterValues = SourceInstance->FontParameterValues;
 
-		const FStaticParameterSet& StaticParameters = SourceInstance->GetStaticParameters();
+		const FStaticParameterSetEditorOnlyData& StaticParameters = SourceInstance->GetEditorOnlyStaticParameters();
 		SourceFunction->StaticSwitchParameterValues = StaticParameters.StaticSwitchParameters;
 		SourceFunction->StaticComponentMaskParameterValues = StaticParameters.StaticComponentMaskParameters;
 
diff --git a/Engine/Source/Editor/VirtualTexturingEditor/Private/RuntimeVirtualTextureAssetTypeActions.cpp b/Engine/Source/Editor/VirtualTexturingEditor/Private/RuntimeVirtualTextureAssetTypeActions.cpp
--- a/Engine/Source/Editor/VirtualTexturingEditor/Private/RuntimeVirtualTextureAssetTypeActions.cpp
+++ b/Engine/Source/Editor/VirtualTexturingEditor/Private/RuntimeVirtualTextureAssetTypeActions.cpp
@@ -114,7 +114,7 @@ namespace
 				Task.EnterProgressFrame();
 
 				bool bMaterialModified = false;
-				for (UMaterialExpression* Expression : Material->Expressions)
+				for (UMaterialExpression* Expression : Material->GetExpressions())
 				{
 					UMaterialExpressionRuntimeVirtualTextureSample* RVTSampleExpression = Cast<UMaterialExpressionRuntimeVirtualTextureSample>(Expression);
 					if (RVTSampleExpression)
@@ -153,8 +153,7 @@ namespace
 				Task.EnterProgressFrame();
 
 				bool bFunctionModified = false;
-				const TArray<TObjectPtr<UMaterialExpression>> *Expressions = Function->GetFunctionExpressions();
-				for (const TObjectPtr<UMaterialExpression>& Expression : *Expressions)
+				for (const TObjectPtr<UMaterialExpression>& Expression : Function->GetExpressions())
 				{
 					UMaterialExpressionRuntimeVirtualTextureSample* RVTSampleExpression = Cast<UMaterialExpressionRuntimeVirtualTextureSample>(Expression);
 					if (RVTSampleExpression)
diff --git a/Engine/Source/Runtime/Engine/Classes/Materials/Material.h b/Engine/Source/Runtime/Engine/Classes/Materials/Material.h
--- a/Engine/Source/Runtime/Engine/Classes/Materials/Material.h
+++ b/Engine/Source/Runtime/Engine/Classes/Materials/Material.h
@@ -150,11 +150,9 @@
 USTRUCT(noexport)
 struct FMaterialInput
 {
-#if WITH_EDITORONLY_DATA
 	/** Material expression that this input is connected to, or NULL if not connected. */
 	UPROPERTY()
 	TObjectPtr<class UMaterialExpression> Expression;
-#endif
 
 	/** Index into Expression's outputs array that this input is connected to. */
 	UPROPERTY()
@@ -167,8 +165,6 @@
 	UPROPERTY()
 	FName InputName;
 
-#if WITH_EDITORONLY_DATA
-
 	UPROPERTY()
 	int32 Mask;
 
@@ -183,12 +179,6 @@
 
 	UPROPERTY()
 	int32 MaskA;
-#endif
-
-	/** Material expression name that this input is connected to, or None if not connected. Used only in cooked builds */
-	UPROPERTY()
-	FName ExpressionName;
-
 };
 #endif
 
@@ -196,7 +186,6 @@
 USTRUCT(noexport)
 struct FColorMaterialInput : public FMaterialInput
 {
-#if WITH_EDITORONLY_DATA
 	UPROPERTY()
 	uint32 UseConstant : 1;
 
@@ -205,7 +194,6 @@
 
 	FColorMaterialInput() :
 		UseConstant(0) {}
-#endif
 };
 #endif
 
@@ -213,7 +201,6 @@
 USTRUCT(noexport)
 struct FScalarMaterialInput : public FMaterialInput
 {
-#if WITH_EDITORONLY_DATA
 	UPROPERTY()
 	uint32 UseConstant : 1;
 
@@ -222,7 +209,6 @@
 
 	FScalarMaterialInput() :
 		UseConstant(0) {}
-#endif
 };
 #endif
 
@@ -246,7 +232,6 @@
 USTRUCT(noexport)
 struct FVectorMaterialInput : public FMaterialInput
 {
-#if WITH_EDITORONLY_DATA
 	UPROPERTY()
 	uint32 UseConstant : 1;
 
@@ -255,7 +240,6 @@
 
 	FVectorMaterialInput:
 		UseConstant(0) {}
-#endif
 };
 #endif
 
@@ -263,7 +247,6 @@
 USTRUCT(noexport)
 struct FVector2MaterialInput : public FMaterialInput
 {
-#if WITH_EDITORONLY_DATA
 	UPROPERTY()
 	uint32 UseConstant : 1;
 
@@ -275,7 +258,6 @@
 
 	FVector2MaterialInput:
 		UseConstant(0) {}
-#endif
 };
 #endif
 
@@ -312,40 +294,14 @@
 	}
 };
 
-/**
- * A Material is an asset which can be applied to a mesh to control the visual look of the scene. 
- * When light from the scene hits the surface, the shading model of the material is used to calculate how that light interacts with the surface. 
- *
- * Warning: Creating new materials directly increases shader compile times!  Consider creating a Material Instance off of an existing material instead.
- */
-UCLASS(hidecategories=Object, MinimalAPI, BlueprintType)
-class UMaterial : public UMaterialInterface
+UCLASS(MinimalAPI, Optional)
+class UMaterialEditorOnlyData : public UMaterialInterfaceEditorOnlyData
 {
-	GENERATED_UCLASS_BODY()
+	GENERATED_BODY()
+public:
+	UMaterialEditorOnlyData();
 
-	// Physics.
-	
-	/** Physical material to use for this graphics material. Used for sounds, effects etc.*/
-	UPROPERTY(EditAnywhere, Category=PhysicalMaterial)
-	TObjectPtr<class UPhysicalMaterial> PhysMaterial;
-
-	/** Physical material mask to use for this graphics material. Used for sounds, effects etc.*/
-	UPROPERTY(EditAnywhere, Category = PhysicalMaterial)
-	TObjectPtr<class UPhysicalMaterialMask> PhysMaterialMask;
-
-	/** Physical material mask map to use for this graphics material. Used for sounds, effects etc.*/
-	UPROPERTY(EditAnywhere, Category = PhysicalMaterialMask)
-	TObjectPtr<class UPhysicalMaterial> PhysicalMaterialMap[EPhysicalMaterialMaskColor::MAX];
-
-	// Reflection.
-#if WITH_EDITORONLY_DATA
+	// Reflection
 	UPROPERTY()
-	FColorMaterialInput DiffuseColor_DEPRECATED;
-
-	UPROPERTY()
-	FColorMaterialInput SpecularColor_DEPRECATED;
-#endif
-
-	UPROPERTY()
 	FColorMaterialInput BaseColor;
 
@@ -370,14 +326,98 @@
 	UPROPERTY()
 	FColorMaterialInput EmissiveColor;
 
-#if WITH_EDITORONLY_DATA
 	// Transmission.
 	UPROPERTY()
 	FScalarMaterialInput Opacity;
 
 	UPROPERTY()
 	FScalarMaterialInput OpacityMask;
-#endif
 
+	/** Adds to world position in the vertex shader. */
+	UPROPERTY()
+	FVectorMaterialInput WorldPositionOffset;
+
+	/** Inner material color, only used for ShadingModel=Subsurface */
+	UPROPERTY()
+	FColorMaterialInput SubsurfaceColor;
+
+	/**  */
+	UPROPERTY()
+	FScalarMaterialInput ClearCoat;
+
+	/**  */
+	UPROPERTY()
+	FScalarMaterialInput ClearCoatRoughness;
+
+	/** output ambient occlusion to the GBuffer */
+	UPROPERTY()
+	FScalarMaterialInput AmbientOcclusion;
+
+	/**
+	 * output refraction index for translucent rendering
+	 * Air:1.0 Water:1.333 Ice:1.3 Glass:~1.6 Diamond:2.42
+	 */
+	UPROPERTY()
+	FScalarMaterialInput Refraction;
+
+	/**
+	 * These inputs are evaluated in the vertex shader and allow artists to do arbitrary vertex shader operations and access them in the pixel shader.
+	 * When unconnected or hidden they default to passing through the vertex UVs.
+	 */
+	UPROPERTY()
+	FVector2MaterialInput CustomizedUVs[8];
+
+	UPROPERTY()
+	FMaterialAttributesInput MaterialAttributes;
+
+	UPROPERTY()
+	FScalarMaterialInput PixelDepthOffset;
+
+	UPROPERTY()
+	FShadingModelMaterialInput ShadingModelFromMaterialExpression;
+
+	UPROPERTY()
+	FStrataMaterialInput FrontMaterial;
+
+	UPROPERTY()
+	FMaterialExpressionCollection ExpressionCollection;
+
+	/** Controls where this parameter group is displayed in a material instance parameter list.  The lower the number the higher up in the parameter list. */
+	UPROPERTY(EditAnywhere, EditFixedSize, Category = "Group Sorting")
+	TArray<FParameterGroupData> ParameterGroupData;
+};
+
+/**
+ * A Material is an asset which can be applied to a mesh to control the visual look of the scene. 
+ * When light from the scene hits the surface, the shading model of the material is used to calculate how that light interacts with the surface. 
+ *
+ * Warning: Creating new materials directly increases shader compile times!  Consider creating a Material Instance off of an existing material instead.
+ */
+UCLASS(hidecategories=Object, MinimalAPI, BlueprintType)
+class UMaterial : public UMaterialInterface
+{
+	GENERATED_UCLASS_BODY()
+
+#if WITH_EDITORONLY_DATA
+	ENGINE_API virtual const UClass* GetEditorOnlyDataClass() const override { return UMaterialEditorOnlyData::StaticClass(); }
+
+	UMaterialEditorOnlyData* GetEditorOnlyData() { return CastChecked<UMaterialEditorOnlyData>(Super::GetEditorOnlyData()); }
+	const UMaterialEditorOnlyData* GetEditorOnlyData() const { return CastChecked<UMaterialEditorOnlyData>(Super::GetEditorOnlyData()); }
+#endif // WITH_EDITORONLY_DATA
+
+	// Physics.
+	
+	/** Physical material to use for this graphics material. Used for sounds, effects etc.*/
+	UPROPERTY(EditAnywhere, Category=PhysicalMaterial)
+	TObjectPtr<class UPhysicalMaterial> PhysMaterial;
+
+	/** Physical material mask to use for this graphics material. Used for sounds, effects etc.*/
+	UPROPERTY(EditAnywhere, Category = PhysicalMaterial)
+	TObjectPtr<class UPhysicalMaterialMask> PhysMaterialMask;
+
+	/** Physical material mask map to use for this graphics material. Used for sounds, effects etc.*/
+	UPROPERTY(EditAnywhere, Category = PhysicalMaterialMask)
+	TObjectPtr<class UPhysicalMaterial> PhysicalMaterialMap[EPhysicalMaterialMaskColor::MAX];
+
 	/** 
 	 * The domain that the material's attributes will be evaluated in. 
@@ -440,56 +480,6 @@
 	UPROPERTY(EditAnywhere, Category = Material, AdvancedDisplay)
 	float OpacityMaskClipValue;
 
-	/** Adds to world position in the vertex shader. */
-	UPROPERTY()
-	FVectorMaterialInput WorldPositionOffset;
-
-#if WITH_EDITORONLY_DATA
-	/** Inner material color, only used for ShadingModel=Subsurface */
-	UPROPERTY()
-	FColorMaterialInput SubsurfaceColor;
-
-	/**  */
-	UPROPERTY()
-	FScalarMaterialInput ClearCoat;
-
-	/**  */
-	UPROPERTY()
-	FScalarMaterialInput ClearCoatRoughness;
-#endif
-
-	/** output ambient occlusion to the GBuffer */
-	UPROPERTY()
-	FScalarMaterialInput AmbientOcclusion;
-
-	/**
-	 * output refraction index for translucent rendering
-	 * Air:1.0 Water:1.333 Ice:1.3 Glass:~1.6 Diamond:2.42
-	 */
-	UPROPERTY()
-	FScalarMaterialInput Refraction;
-
-#if WITH_EDITORONLY_DATA
-	/** 
-	 * These inputs are evaluated in the vertex shader and allow artists to do arbitrary vertex shader operations and access them in the pixel shader.
-	 * When unconnected or hidden they default to passing through the vertex UVs.
-	 */
-	UPROPERTY()
-	FVector2MaterialInput CustomizedUVs[8];
-#endif
-
-	UPROPERTY()
-	FMaterialAttributesInput MaterialAttributes;
-
-	UPROPERTY()
-	FScalarMaterialInput PixelDepthOffset;
-
-	UPROPERTY()
-	FShadingModelMaterialInput ShadingModelFromMaterialExpression;
-
-	UPROPERTY()
-	FStrataMaterialInput FrontMaterial;
-
 private:
 	/** Deprecated. Use TranslucencyPass instead. */
 	UPROPERTY()
@@ -857,25 +847,6 @@
 
 	UPROPERTY()
 	int32 EditorYaw;
-
-	/** Array of material expressions, excluding Comments.  Used by the material editor. */
-	UPROPERTY()
-	TArray<TObjectPtr<class UMaterialExpression>> Expressions;
-
-	/** Array of comments associated with this material; viewed in the material editor. */
-	UPROPERTY()
-	TArray<TObjectPtr<class UMaterialExpressionComment>> EditorComments;
-
-	/** The execution begin expression, if material is using an exec wire */
-	UPROPERTY()
-	TObjectPtr<class UMaterialExpressionExecBegin> ExpressionExecBegin;
-	
-	UPROPERTY()
-	TObjectPtr<class UMaterialExpressionExecEnd> ExpressionExecEnd;
-
-	/** Controls where this parameter group is displayed in a material instance parameter list.  The lower the number the higher up in the parameter list. */
-	UPROPERTY(EditAnywhere, EditFixedSize, Category = "Group Sorting")
-	TArray<FParameterGroupData> ParameterGroupData;
 #endif // WITH_EDITORONLY_DATA
 
 	/** true if this Material can be assumed Opaque when set to masked. */
@@ -1176,6 +1147,14 @@
 	ENGINE_API bool IsCompilingOrHadCompileError(ERHIFeatureLevel::Type InFeatureLevel);
 
 #if WITH_EDITOR
+	ENGINE_API TConstArrayView<TObjectPtr<UMaterialExpression>> GetExpressions() const;
+	ENGINE_API TConstArrayView<TObjectPtr<UMaterialExpressionComment>> GetEditorComments() const;
+	ENGINE_API UMaterialExpressionExecBegin* GetExpressionExecBegin() const;
+	ENGINE_API UMaterialExpressionExecEnd* GetExpressionExecEnd() const;
+	ENGINE_API const FMaterialExpressionCollection& GetExpressionCollection() const;
+	ENGINE_API FMaterialExpressionCollection& GetExpressionCollection();
+	ENGINE_API void AssignExpressionCollection(const FMaterialExpressionCollection& InCollection);
+
 	ENGINE_API bool SetParameterValueEditorOnly(const FName& ParameterName, const FMaterialParameterMetadata& Meta);
 	ENGINE_API bool SetVectorParameterValueEditorOnly(FName ParameterName, FLinearColor InValue);
 	ENGINE_API bool SetScalarParameterValueEditorOnly(FName ParameterName, float InValue);
@@ -1271,7 +1250,7 @@
 	UE_DEPRECATED(5.0, "Use GetAllParameterInfoOfType or GetAllParametersOfType")
 	void GetAllParameterInfo(TArray<FMaterialParameterInfo>& OutParameterInfo, TArray<FGuid>& OutParameterIds) const
 	{
-		for (const TObjectPtr<UMaterialExpression>& Expression : Expressions)
+		for (const TObjectPtr<UMaterialExpression>& Expression : GetExpressions())
 		{
 			FMaterialParameterInfo BaseParameterInfo;
 			BaseParameterInfo.Association = EMaterialParameterAssociation::GlobalParameter;
@@ -1349,7 +1328,7 @@
 	{
 		if (InGUID.IsValid())
 		{
-			return FindExpressionByGUIDRecursive<ExpressionType>(InGUID, Expressions);
+			return FindExpressionByGUIDRecursive<ExpressionType>(InGUID, GetExpressions());
 		}
 		return nullptr;
 	}
@@ -1358,9 +1337,9 @@
 	template<typename ExpressionType>
 	void GetAllExpressionsOfType(TArray<const ExpressionType*>& OutExpressions) const
 	{
-		for (int32 ExpressionIndex = 0; ExpressionIndex < Expressions.Num(); ExpressionIndex++)
+		for (UMaterialExpression* Expression : GetExpressions())
 		{
-			ExpressionType* ExpressionPtr = Cast<ExpressionType>(Expressions[ExpressionIndex]);
+			ExpressionType* ExpressionPtr = Cast<ExpressionType>(Expression);
 			if (ExpressionPtr)
 			{
 				OutExpressions.Add(ExpressionPtr);
@@ -1372,7 +1351,7 @@
 	template<typename ExpressionType>
 	void GetAllExpressionsInMaterialAndFunctionsOfType(TArray<ExpressionType*>& OutExpressions) const
 	{
-		for (UMaterialExpression* Expression : Expressions)
+		for (UMaterialExpression* Expression : GetExpressions())
 		{
 			ExpressionType* ExpressionOfType = Cast<ExpressionType>(Expression);
 			if (ExpressionOfType)
@@ -1415,7 +1394,7 @@
 	template<typename ExpressionType>
 	bool HasAnyExpressionsInMaterialAndFunctionsOfType() const
 	{
-		for (UMaterialExpression* Expression : Expressions)
+		for (UMaterialExpression* Expression : GetExpressions())
 		{
 			ExpressionType* ExpressionOfType = Cast<ExpressionType>(Expression);
 			if (ExpressionOfType)
@@ -1732,6 +1711,7 @@
 	ENGINE_API virtual void GetShaderTypes(EShaderPlatform Platform, const ITargetPlatform* TargetPlatform, TArray<FDebugShaderTypeInfo>& OutShaderInfo) override;
 #endif // WITH_EDITOR
 
+	ENGINE_API bool IsPropertyConnected(EMaterialProperty Property) const;
 	ENGINE_API bool HasBaseColorConnected() const;
 	ENGINE_API bool HasRoughnessConnected() const;
 	ENGINE_API bool HasAmbientOcclusionConnected() const;
@@ -1766,14 +1746,14 @@
 	 * @param	bScaleCoords	Whether to scale the coordinates to space out nodes
 	 * @param	Material		The Material to flip its home coords (optional)
 	 */
-	static void FlipExpressionPositions(const TArray<UMaterialExpression*>& Expressions, const TArray<UMaterialExpressionComment*>& Comments, bool bScaleCoords, UMaterial* Material = NULL);
+	static void FlipExpressionPositions(TConstArrayView<TObjectPtr<UMaterialExpression>> Expressions, TConstArrayView<TObjectPtr<UMaterialExpressionComment>> Comments, bool bScaleCoords, UMaterial* Material = NULL);
 
 	/**
 	 * Shifts the positions of comments so that they are aligned correctly with other expressions
 	 *
 	 * @param	Comments	Array of comments to fix
 	 */
-	static void FixCommentPositions(const TArray<UMaterialExpressionComment*>& Comments);
+	static void FixCommentPositions(TConstArrayView<TObjectPtr<UMaterialExpressionComment>> Comments);
 
 	/**
 	 * Checks whether a Material is arranged in the old style, with inputs flowing from right to left
@@ -1803,7 +1783,7 @@
 #if WITH_EDITORONLY_DATA
 	/* Helper function to help finding expression GUID taking into account UMaterialExpressionMaterialFunctionCall */
 	template<typename ExpressionType>
-	ExpressionType* FindExpressionByGUIDRecursive(const FGuid &InGUID, const TArray<TObjectPtr<UMaterialExpression>>& InMaterialExpression)
+	ExpressionType* FindExpressionByGUIDRecursive(const FGuid& InGUID, TConstArrayView<TObjectPtr<UMaterialExpression>> InMaterialExpression)
 	{
 		for (int32 ExpressionIndex = 0; ExpressionIndex < InMaterialExpression.Num(); ExpressionIndex++)
 		{
@@ -1818,12 +1798,9 @@
 			}
 			else if (MaterialFunctionCall && MaterialFunctionCall->MaterialFunction)
 			{
-				if (const TArray<TObjectPtr<UMaterialExpression>>* FunctionExpressions = MaterialFunctionCall->MaterialFunction->GetFunctionExpressions())
+				if (ExpressionType* Expression = FindExpressionByGUIDRecursive<ExpressionType>(InGUID, MaterialFunctionCall->MaterialFunction->GetExpressions()))
 				{
-					if (ExpressionType* Expression = FindExpressionByGUIDRecursive<ExpressionType>(InGUID, *FunctionExpressions))
-					{
-						return Expression;
-					}
+					return Expression;
 				}
 			}
 			else if (MaterialLayers)
@@ -1835,12 +1812,9 @@
 				{
 					if (Layer)
 					{
-						if (const TArray<TObjectPtr<UMaterialExpression>>* FunctionExpressions = Layer->GetFunctionExpressions())
+						if (ExpressionType* Expression = FindExpressionByGUIDRecursive<ExpressionType>(InGUID, Layer->GetExpressions()))
 						{
-							if (ExpressionType* Expression = FindExpressionByGUIDRecursive<ExpressionType>(InGUID, *FunctionExpressions))
-							{
-								return Expression;
-							}
+							return Expression;
 						}
 					}
 				}
@@ -1849,12 +1823,9 @@
 				{
 					if (Blend)
 					{
-						if (const TArray<TObjectPtr<UMaterialExpression>>* FunctionExpressions = Blend->GetFunctionExpressions())
+						if (ExpressionType* Expression = FindExpressionByGUIDRecursive<ExpressionType>(InGUID, Blend->GetExpressions()))
 						{
-							if (ExpressionType* Expression = FindExpressionByGUIDRecursive<ExpressionType>(InGUID, *FunctionExpressions))
-							{
-								return Expression;
-							}
+							return Expression;
 						}
 					}
 				}
@@ -1864,6 +1835,92 @@
 		return nullptr;
 	}
 #endif // WITH_EDITORONLY_DATA
+
+#if WITH_EDITORONLY_DATA
+	UPROPERTY()
+	TArray<TObjectPtr<class UMaterialExpressionComment>> EditorComments_DEPRECATED;
+
+	UPROPERTY()
+	TObjectPtr<class UMaterialExpressionExecBegin> ExpressionExecBegin_DEPRECATED;
+
+	UPROPERTY()
+	TObjectPtr<class UMaterialExpressionExecEnd> ExpressionExecEnd_DEPRECATED;
+
+	UPROPERTY()
+	TArray<TObjectPtr<class UMaterialExpression>> Expressions_DEPRECATED;
+
+	UPROPERTY(meta = (DeprecatedProperty, DeprecationMessage = "GetEditorOnlyData().ParameterGroupData"))
+	TArray<FParameterGroupData> ParameterGroupData_DEPRECATED;
+
+	UPROPERTY()
+	FColorMaterialInput DiffuseColor_DEPRECATED;
+
+	UPROPERTY()
+	FColorMaterialInput SpecularColor_DEPRECATED;
+
+	UPROPERTY()
+	FColorMaterialInput BaseColor_DEPRECATED;
+
+	UPROPERTY()
+	FScalarMaterialInput Metallic_DEPRECATED;
+
+	UPROPERTY()
+	FScalarMaterialInput Specular_DEPRECATED;
+
+	UPROPERTY()
+	FScalarMaterialInput Roughness_DEPRECATED;
+
+	UPROPERTY()
+	FScalarMaterialInput Anisotropy_DEPRECATED;
+
+	UPROPERTY()
+	FVectorMaterialInput Normal_DEPRECATED;
+
+	UPROPERTY()
+	FVectorMaterialInput Tangent_DEPRECATED;
+
+	UPROPERTY()
+	FColorMaterialInput EmissiveColor_DEPRECATED;
+
+	UPROPERTY()
+	FScalarMaterialInput Opacity_DEPRECATED;
+
+	UPROPERTY()
+	FScalarMaterialInput OpacityMask_DEPRECATED;
+
+	UPROPERTY()
+	FVectorMaterialInput WorldPositionOffset_DEPRECATED;
+
+	UPROPERTY()
+	FColorMaterialInput SubsurfaceColor_DEPRECATED;
+
+	UPROPERTY()
+	FScalarMaterialInput ClearCoat_DEPRECATED;
+
+	UPROPERTY()
+	FScalarMaterialInput ClearCoatRoughness_DEPRECATED;
+
+	UPROPERTY()
+	FScalarMaterialInput AmbientOcclusion_DEPRECATED;
+
+	UPROPERTY()
+	FScalarMaterialInput Refraction_DEPRECATED;
+
+	UPROPERTY()
+	FVector2MaterialInput CustomizedUVs_DEPRECATED[8];
+
+	UPROPERTY()
+	FMaterialAttributesInput MaterialAttributes_DEPRECATED;
+
+	UPROPERTY()
+	FScalarMaterialInput PixelDepthOffset_DEPRECATED;
+
+	UPROPERTY()
+	FShadingModelMaterialInput ShadingModelFromMaterialExpression_DEPRECATED;
+
+	UPROPERTY()
+	FStrataMaterialInput FrontMaterial_DEPRECATED;
+#endif // WITH_EDITORONLY_DATA
 };
 
 
diff --git a/Engine/Source/Runtime/Engine/Classes/Materials/MaterialExpression.h b/Engine/Source/Runtime/Engine/Classes/Materials/MaterialExpression.h
--- a/Engine/Source/Runtime/Engine/Classes/Materials/MaterialExpression.h
+++ b/Engine/Source/Runtime/Engine/Classes/Materials/MaterialExpression.h
@@ -18,6 +18,11 @@ struct FPropertyChangedEvent;
 struct FMaterialShadingModelField;
 struct FStrataOperator;
 
+class UMaterialExpression;
+class UMaterialExpressionComment;
+class UMaterialExpressionExecBegin;
+class UMaterialExpressionExecEnd;
+
 class FMaterialHLSLGenerator;
 enum class EMaterialNewScopeFlag : uint8;
 
@@ -36,11 +41,9 @@ class FExpression;
 USTRUCT(noexport)
 struct FExpressionInput
 {
-#if WITH_EDITORONLY_DATA
 	/** UMaterial expression that this input is connected to, or NULL if not connected. */
 	UPROPERTY()
 	TObjectPtr<class UMaterialExpression> Expression;
-#endif
 
 	/** Index into Expression's outputs array that this input is connected to. */
 	UPROPERTY()
@@ -53,7 +56,6 @@ struct FExpressionInput
 	UPROPERTY()
 	FName InputName;
 
-#if WITH_EDITORONLY_DATA
 
 	UPROPERTY()
 	int32 Mask;
@@ -69,11 +71,6 @@ struct FExpressionInput
 
 	UPROPERTY()
 	int32 MaskA;
-#endif
-
-	/** Material expression name that this input is connected to, or None if not connected. Used only in cooked builds */
-	UPROPERTY()
-	FName ExpressionName;
 };
 
 USTRUCT(noexport)
@@ -93,7 +90,6 @@ struct FExpressionOutput
 	UPROPERTY()
 	FName OutputName;
 
-#if WITH_EDITORONLY_DATA
 	UPROPERTY()
 	int32 Mask;
 
@@ -108,7 +104,6 @@ struct FExpressionOutput
 
 	UPROPERTY()
 	int32 MaskA;
-#endif
 };
 #endif
 
@@ -151,7 +146,34 @@ enum class EMaterialExpressionSetParameterValueFlags : uint32
 };
 ENUM_CLASS_FLAGS(EMaterialExpressionSetParameterValueFlags);
 
-UCLASS(abstract, BlueprintType, hidecategories=Object)
+USTRUCT()
+struct ENGINE_API FMaterialExpressionCollection
+{
+	GENERATED_BODY()
+
+	void AddExpression(UMaterialExpression* InExpression);
+	void RemoveExpression(UMaterialExpression* InExpression);
+	void AddComment(UMaterialExpressionComment* InExpression);
+	void RemoveComment(UMaterialExpressionComment* InExpression);
+	void Empty();
+
+	/** Array of material expressions, excluding Comments.  Used by the material editor. */
+	UPROPERTY()
+	TArray<TObjectPtr<UMaterialExpression>> Expressions;
+
+	/** Array of comments associated with this material; viewed in the material editor. */
+	UPROPERTY()
+	TArray<TObjectPtr<UMaterialExpressionComment>> EditorComments;
+
+	/** The execution begin expression, if material is using an exec wire */
+	UPROPERTY()
+	TObjectPtr<UMaterialExpressionExecBegin> ExpressionExecBegin = nullptr;
+
+	UPROPERTY()
+	TObjectPtr<UMaterialExpressionExecEnd> ExpressionExecEnd = nullptr;
+};
+
+UCLASS(abstract, Optional, BlueprintType, hidecategories=Object)
 class ENGINE_API UMaterialExpression : public UObject
 {
 	GENERATED_UCLASS_BODY()
@@ -162,7 +184,6 @@ class ENGINE_API UMaterialExpression : public UObject
 	static void InitializeNumExecutionInputs(TArrayView<UMaterialExpression*> Expressions);
 #endif
 
-#if WITH_EDITORONLY_DATA
 	UPROPERTY()
 	int32 MaterialExpressionEditorX;
 
@@ -184,7 +205,6 @@ class ENGINE_API UMaterialExpression : public UObject
 	UPROPERTY()
 	FGuid MaterialExpressionGuid;
 
-#endif // WITH_EDITORONLY_DATA
 	/** 
 	 * The material that this expression is currently being compiled in.  
 	 * This is not necessarily the object which owns this expression, for example a preview material compiling a material function's expressions.
@@ -199,7 +219,6 @@ class ENGINE_API UMaterialExpression : public UObject
 	UPROPERTY()
 	TObjectPtr<class UMaterialFunction> Function;
 
-#if WITH_EDITORONLY_DATA
 	/** A description that level designers can add (shows in the material editor UI). */
 	UPROPERTY(EditAnywhere, Category=MaterialExpression, meta=(MultiLine=true, DisplayAfter = "SortPriority"))
 	FString Desc;
@@ -214,13 +233,11 @@ class ENGINE_API UMaterialExpression : public UObject
 	/** If true, we should update the preview next render. This is set when changing bRealtimePreview. */
 	UPROPERTY(transient)
 	uint32 bNeedToUpdatePreview:1;
-#endif // WITH_EDITORONLY_DATA
 
 	/** Indicates that this is a 'parameter' type of expression and should always be loaded (ie not cooked away) because we might want the default parameter. */
 	UPROPERTY()
 	uint8 bIsParameterExpression : 1;
 
-#if WITH_EDITORONLY_DATA
 	/** If true, the comment bubble will be visible in the graph editor */
 	UPROPERTY()
 	uint32 bCommentBubbleVisible:1;
@@ -260,7 +277,6 @@ class ENGINE_API UMaterialExpression : public UObject
 	/** The expression's outputs, which are set in default properties by derived classes. */
 	UPROPERTY()
 	TArray<FExpressionOutput> Outputs;
-#endif // WITH_EDITORONLY_DATA
 
 	//~ Begin UObject Interface.
 	virtual void PostInitProperties() override;
@@ -274,14 +290,6 @@ class ENGINE_API UMaterialExpression : public UObject
 	virtual bool Modify( bool bAlwaysMarkDirty=true ) override;
 #endif // WITH_EDITOR
 	virtual void Serialize( FStructuredArchive::FRecord Record ) override;
-	virtual bool IsEditorOnly() const
-	{
-		return true;
-	}
-	virtual bool HasNonEditorOnlyReferences() const override
-	{
-		return true;
-	}
 	//~ End UObject Interface.
 
 	UObject* GetAssetOwner() const;
diff --git a/Engine/Source/Runtime/Engine/Classes/Materials/MaterialExpressionChannelMaskParameter.h b/Engine/Source/Runtime/Engine/Classes/Materials/MaterialExpressionChannelMaskParameter.h
--- a/Engine/Source/Runtime/Engine/Classes/Materials/MaterialExpressionChannelMaskParameter.h
+++ b/Engine/Source/Runtime/Engine/Classes/Materials/MaterialExpressionChannelMaskParameter.h
@@ -27,10 +27,8 @@ class UMaterialExpressionChannelMaskParameter : public UMaterialExpressionVector
 	UPROPERTY(EditAnywhere, Category=MaterialExpressionChannelMaskParameter)
 	TEnumAsByte<EChannelMaskParameterColor::Type> MaskChannel;
 
-#if WITH_EDITORONLY_DATA
 	UPROPERTY()
 	FExpressionInput Input;
-#endif
 
 #if WITH_EDITOR
 	virtual bool GetParameterValue(FMaterialParameterMetadata& OutMeta) const override
diff --git a/Engine/Source/Runtime/Engine/Classes/Materials/MaterialExpressionFontSampleParameter.h b/Engine/Source/Runtime/Engine/Classes/Materials/MaterialExpressionFontSampleParameter.h
--- a/Engine/Source/Runtime/Engine/Classes/Materials/MaterialExpressionFontSampleParameter.h
+++ b/Engine/Source/Runtime/Engine/Classes/Materials/MaterialExpressionFontSampleParameter.h
@@ -29,11 +29,9 @@ class UMaterialExpressionFontSampleParameter : public UMaterialExpressionFontSam
 	UPROPERTY(EditAnywhere, Category=MaterialExpressionFontSampleParameter)
 	FName Group;
 
-#if WITH_EDITORONLY_DATA
 	/** Controls where the this parameter is displayed in a material instance parameter list. The lower the number the higher up in the parameter list. */
 	UPROPERTY(EditAnywhere, Category = MaterialExpressionFontSampleParameter)
 	int32 SortPriority = 32;
-#endif
 
 	//~ Begin UMaterialExpression Interface
 #if WITH_EDITOR
@@ -55,6 +53,7 @@ class UMaterialExpressionFontSampleParameter : public UMaterialExpressionFontSam
 		OutMeta.ExpressionGuid = ExpressionGUID;
 		OutMeta.Group = Group;
 		OutMeta.SortPriority = SortPriority;
+		OutMeta.AssetPath = GetAssetPathName();
 		return true;
 	}
 	virtual bool SetParameterValue(const FName& Name, const FMaterialParameterMetadata& Meta, EMaterialExpressionSetParameterValueFlags Flags) override
diff --git a/Engine/Source/Runtime/Engine/Classes/Materials/MaterialExpressionMaterialAttributeLayers.h b/Engine/Source/Runtime/Engine/Classes/Materials/MaterialExpressionMaterialAttributeLayers.h
--- a/Engine/Source/Runtime/Engine/Classes/Materials/MaterialExpressionMaterialAttributeLayers.h
+++ b/Engine/Source/Runtime/Engine/Classes/Materials/MaterialExpressionMaterialAttributeLayers.h
@@ -38,27 +38,27 @@ class UMaterialExpressionMaterialAttributeLayers : public UMaterialExpression
 #if WITH_EDITOR
 	const TArray<FText>& GetLayerNames() const
 	{
-		return ParamLayers ? ParamLayers->LayerNames : DefaultLayers.LayerNames;
+		return ParamLayers ? ParamLayers->EditorOnly.LayerNames : DefaultLayers.EditorOnly.LayerNames;
 	}
 
 	const TArray<bool>& GetShouldFilterLayers() const
 	{
-		return ParamLayers ? ParamLayers->RestrictToLayerRelatives : DefaultLayers.RestrictToLayerRelatives;
+		return ParamLayers ? ParamLayers->EditorOnly.RestrictToLayerRelatives : DefaultLayers.EditorOnly.RestrictToLayerRelatives;
 	}
 
 	const TArray<bool>& GetShouldFilterBlends() const
 	{
-		return ParamLayers ? ParamLayers->RestrictToBlendRelatives : DefaultLayers.RestrictToBlendRelatives;
+		return ParamLayers ? ParamLayers->EditorOnly.RestrictToBlendRelatives : DefaultLayers.EditorOnly.RestrictToBlendRelatives;
 	}
 
 	const TArray<FGuid>& GetLayerGuids() const
 	{
-		return ParamLayers ? ParamLayers->LayerGuids : DefaultLayers.LayerGuids;
+		return ParamLayers ? ParamLayers->EditorOnly.LayerGuids : DefaultLayers.EditorOnly.LayerGuids;
 	}
 
 	const TArray<bool>& GetLayerStates() const
 	{
-		return ParamLayers ? ParamLayers->LayerStates : DefaultLayers.LayerStates;
+		return ParamLayers ? ParamLayers->EditorOnly.LayerStates : DefaultLayers.EditorOnly.LayerStates;
 	}
 #endif // WITH_EDITOR
 
diff --git a/Engine/Source/Runtime/Engine/Classes/Materials/MaterialExpressionMaterialFunctionCall.h b/Engine/Source/Runtime/Engine/Classes/Materials/MaterialExpressionMaterialFunctionCall.h
--- a/Engine/Source/Runtime/Engine/Classes/Materials/MaterialExpressionMaterialFunctionCall.h
+++ b/Engine/Source/Runtime/Engine/Classes/Materials/MaterialExpressionMaterialFunctionCall.h
@@ -83,7 +83,6 @@ class UMaterialExpressionMaterialFunctionCall : public UMaterialExpression
 	UPROPERTY(EditAnywhere, Category=MaterialExpressionMaterialFunctionCall, meta=(DisallowedClasses="MaterialFunctionInstance"))
 	TObjectPtr<class UMaterialFunctionInterface> MaterialFunction;
 
-#if WITH_EDITORONLY_DATA
 	/** Array of all the function inputs that this function exposes. */
 	UPROPERTY()
 	TArray<struct FFunctionExpressionInput> FunctionInputs;
@@ -91,7 +90,6 @@ class UMaterialExpressionMaterialFunctionCall : public UMaterialExpression
 	/** Array of all the function outputs that this function exposes. */
 	UPROPERTY()
 	TArray<struct FFunctionExpressionOutput> FunctionOutputs;
-#endif // WITH_EDITOR
 
 	/** Used by material parameters to split references to separate instances. */
 	UPROPERTY(Transient)
@@ -170,7 +168,7 @@ private:
 	void FixupReferencingExpressions(
 		const TArray<FFunctionExpressionOutput>& NewOutputs,
 		const TArray<FFunctionExpressionOutput>& OriginalOutputs,
-		TArray<UMaterialExpression*>& Expressions, 
+		TConstArrayView<TObjectPtr<UMaterialExpression>> Expressions, 
 		TArray<FExpressionInput*>& MaterialInputs,
 		bool bMatchByName);
 
diff --git a/Engine/Source/Runtime/Engine/Classes/Materials/MaterialExpressionNamedReroute.h b/Engine/Source/Runtime/Engine/Classes/Materials/MaterialExpressionNamedReroute.h
--- a/Engine/Source/Runtime/Engine/Classes/Materials/MaterialExpressionNamedReroute.h
+++ b/Engine/Source/Runtime/Engine/Classes/Materials/MaterialExpressionNamedReroute.h
@@ -19,7 +19,9 @@ protected:
 	 * @param	Expressions		The expressions to search in
 	 * @return	null if not found
 	 */
-	 UMaterialExpressionNamedRerouteDeclaration* FindDeclarationInArray(const FGuid& VariableGuid, const TArray<UMaterialExpression*>& Expressions) const;
+	template<typename ExpressionsArrayType>
+	UMaterialExpressionNamedRerouteDeclaration* FindDeclarationInArray(const FGuid& VariableGuid, const ExpressionsArrayType& Expressions) const;
+
 	/**
 	 * Find a variable declaration in the entire graph
 	 * @param	VariableGuid	The GUID of the variable to find
@@ -41,11 +43,9 @@ public:
 	UPROPERTY(EditAnywhere, Category = MaterialExpressionNamedRerouteDeclaration)
 	FName Name;
 
-#if WITH_EDITORONLY_DATA
 	/** The color of the graph node. The same color will apply to all linked usages of this Declaration node */
 	UPROPERTY(EditAnywhere, Category = MaterialExpressionNamedRerouteDeclaration)
 	FLinearColor NodeColor;
-#endif
 
 	// The variable GUID, to support copy across graphs
 	UPROPERTY()
@@ -127,4 +127,18 @@ protected:
 private:
 	// Check that the declaration isn't deleted
 	bool IsDeclarationValid() const;
-};
\ No newline at end of file
+};
+
+template<typename ExpressionsArrayType>
+inline UMaterialExpressionNamedRerouteDeclaration* UMaterialExpressionNamedRerouteBase::FindDeclarationInArray(const FGuid& VariableGuid, const ExpressionsArrayType& Expressions) const
+{
+	for (UMaterialExpression* Expression : Expressions)
+	{
+		auto* Declaration = Cast<UMaterialExpressionNamedRerouteDeclaration>(Expression);
+		if (Declaration && this != Declaration && Declaration->VariableGuid == VariableGuid)
+		{
+			return Declaration;
+		}
+	}
+	return nullptr;
+}
diff --git a/Engine/Source/Runtime/Engine/Classes/Materials/MaterialExpressionParameter.h b/Engine/Source/Runtime/Engine/Classes/Materials/MaterialExpressionParameter.h
--- a/Engine/Source/Runtime/Engine/Classes/Materials/MaterialExpressionParameter.h
+++ b/Engine/Source/Runtime/Engine/Classes/Materials/MaterialExpressionParameter.h
@@ -24,7 +24,6 @@ class UMaterialExpressionParameter : public UMaterialExpression
 	UPROPERTY()
 	FGuid ExpressionGUID;
 
-#if WITH_EDITORONLY_DATA
 	/** The name of the parameter Group to display in MaterialInstance Editor. Default is None group */
 	UPROPERTY(EditAnywhere, Category=MaterialExpressionParameter)
 	FName Group;
@@ -32,7 +31,7 @@ class UMaterialExpressionParameter : public UMaterialExpression
 	/** Controls where the this parameter is displayed in a material instance parameter list.  The lower the number the higher up in the parameter list. */
 	UPROPERTY(EditAnywhere, Category=MaterialExpressionParameter)
 	int32 SortPriority = 32;
-#endif
+
 	static FName ParameterDefaultName;
 
 
@@ -55,6 +54,7 @@ class UMaterialExpressionParameter : public UMaterialExpression
 		OutMeta.ExpressionGuid = ExpressionGUID;
 		OutMeta.Group = Group;
 		OutMeta.SortPriority = SortPriority;
+		OutMeta.AssetPath = GetAssetPathName();
 		return true;
 	}
 
diff --git a/Engine/Source/Runtime/Engine/Classes/Materials/MaterialExpressionRuntimeVirtualTextureSampleParameter.h b/Engine/Source/Runtime/Engine/Classes/Materials/MaterialExpressionRuntimeVirtualTextureSampleParameter.h
--- a/Engine/Source/Runtime/Engine/Classes/Materials/MaterialExpressionRuntimeVirtualTextureSampleParameter.h
+++ b/Engine/Source/Runtime/Engine/Classes/Materials/MaterialExpressionRuntimeVirtualTextureSampleParameter.h
@@ -23,11 +23,9 @@ class UMaterialExpressionRuntimeVirtualTextureSampleParameter : public UMaterial
 	UPROPERTY(EditAnywhere, Category = MaterialParameter)
 	FName Group;
 
-#if WITH_EDITORONLY_DATA
 	/** Controls where the this parameter is displayed in a material instance parameter list. The lower the number the higher up in the parameter list. */
 	UPROPERTY(EditAnywhere, Category = MaterialParameter)
 	int32 SortPriority = 32;
-#endif
 
 #if WITH_EDITOR
 	/** If this is the named parameter from this material expression, then set its value. */
@@ -58,6 +56,7 @@ class UMaterialExpressionRuntimeVirtualTextureSampleParameter : public UMaterial
 		OutMeta.ExpressionGuid = ExpressionGUID;
 		OutMeta.Group = Group;
 		OutMeta.SortPriority = SortPriority;
+		OutMeta.AssetPath = GetAssetPathName();
 		return true;
 	}
 	virtual bool SetParameterValue(const FName& Name, const FMaterialParameterMetadata& Meta, EMaterialExpressionSetParameterValueFlags Flags) override
diff --git a/Engine/Source/Runtime/Engine/Classes/Materials/MaterialExpressionScalarParameter.h b/Engine/Source/Runtime/Engine/Classes/Materials/MaterialExpressionScalarParameter.h
--- a/Engine/Source/Runtime/Engine/Classes/Materials/MaterialExpressionScalarParameter.h
+++ b/Engine/Source/Runtime/Engine/Classes/Materials/MaterialExpressionScalarParameter.h
@@ -23,8 +23,6 @@ class UMaterialExpressionScalarParameter : public UMaterialExpressionParameter
 
 	UPROPERTY(EditAnywhere, Category=CustomPrimitiveData, meta=(ClampMin="0"))
 	uint8 PrimitiveDataIndex = 0;
-
-#if WITH_EDITORONLY_DATA
 	/** 
 	 * Sets the lower bound for the slider on this parameter in the material instance editor. 
 	 */
@@ -37,7 +35,6 @@ class UMaterialExpressionScalarParameter : public UMaterialExpressionParameter
 	 */
 	UPROPERTY(EditAnywhere, Category=MaterialExpressionScalarParameter)
 	float SliderMax;
-#endif
 
 	//~ Begin UMaterialExpression Interface
 #if WITH_EDITOR
diff --git a/Engine/Source/Runtime/Engine/Classes/Materials/MaterialExpressionTextureBase.h b/Engine/Source/Runtime/Engine/Classes/Materials/MaterialExpressionTextureBase.h
--- a/Engine/Source/Runtime/Engine/Classes/Materials/MaterialExpressionTextureBase.h
+++ b/Engine/Source/Runtime/Engine/Classes/Materials/MaterialExpressionTextureBase.h
@@ -23,14 +23,12 @@ class ENGINE_API UMaterialExpressionTextureBase : public UMaterialExpression
 	UPROPERTY(EditAnywhere, BlueprintReadWrite, Category=MaterialExpressionTextureBase)
 	TObjectPtr<class UTexture> Texture;
 
-#if WITH_EDITORONLY_DATA
 	UPROPERTY(EditAnywhere, BlueprintReadWrite, Category=MaterialExpressionTextureBase)
 	TEnumAsByte<enum EMaterialSamplerType> SamplerType;
 	
 	/** Is default selected texture when using mesh paint mode texture painting */
 	UPROPERTY(EditAnywhere, Category=MaterialExpressionTextureBase)
 	uint8 IsDefaultMeshpaintTexture:1;
-#endif // WITH_EDITORONLY_DATA
 	
 	//~ Begin UObject Interface
 #if WITH_EDITOR
diff --git a/Engine/Source/Runtime/Engine/Classes/Materials/MaterialFunction.h b/Engine/Source/Runtime/Engine/Classes/Materials/MaterialFunction.h
--- a/Engine/Source/Runtime/Engine/Classes/Materials/MaterialFunction.h
+++ b/Engine/Source/Runtime/Engine/Classes/Materials/MaterialFunction.h
@@ -16,6 +16,15 @@ class UTexture;
 struct FPropertyChangedEvent;
 class UMaterialExpression;
 
+UCLASS(MinimalAPI, Optional)
+class UMaterialFunctionEditorOnlyData : public UMaterialFunctionInterfaceEditorOnlyData
+{
+	GENERATED_BODY()
+public:
+	UPROPERTY()
+	FMaterialExpressionCollection ExpressionCollection;
+};
+
 /**
  * A Material Function is a collection of material expressions that can be reused in different materials
  */
@@ -24,29 +33,23 @@ class UMaterialFunction : public UMaterialFunctionInterface
 {
 	GENERATED_UCLASS_BODY()
 
+#if WITH_EDITORONLY_DATA
+	virtual const UClass* GetEditorOnlyDataClass() const override { return UMaterialFunctionEditorOnlyData::StaticClass(); }
+
+	UMaterialFunctionEditorOnlyData* GetEditorOnlyData() { return CastChecked<UMaterialFunctionEditorOnlyData>(Super::GetEditorOnlyData()); }
+	const UMaterialFunctionEditorOnlyData* GetEditorOnlyData() const { return CastChecked<UMaterialFunctionEditorOnlyData>(Super::GetEditorOnlyData()); }
+#endif // WITH_EDITORONLY_DATA
+
 #if WITH_EDITORONLY_DATA
 	/** Used in the material editor, points to the function asset being edited, which this function is just a preview for. */
 	UPROPERTY(transient)
 	TObjectPtr<class UMaterialFunction> ParentFunction;
-
 #endif // WITH_EDITORONLY_DATA
+
 	/** Description of the function which will be displayed as a tooltip wherever the function is used. */
 	UPROPERTY(EditAnywhere, Category=MaterialFunction, AssetRegistrySearchable)
 	FString Description;
 
-#if WITH_EDITORONLY_DATA
-	/** Array of material expressions, excluding Comments.  Used by the material editor. */
-	UPROPERTY()
-	TArray<TObjectPtr<UMaterialExpression>> FunctionExpressions;
-
-	/** The execution begin expression, if material is using an exec wire */
-	UPROPERTY()
-	TObjectPtr<class UMaterialExpressionExecBegin> ExpressionExecBegin;
-
-	UPROPERTY()
-	TObjectPtr<class UMaterialExpressionExecEnd> ExpressionExecEnd;
-#endif // WITH_EDITORONLY_DATA
-
 	/** Whether to list this function in the material function library, which is a window in the material editor that lists categorized functions. */
 	UPROPERTY(EditAnywhere, Category=MaterialFunction, AssetRegistrySearchable)
 	uint8 bExposeToLibrary:1;
@@ -78,10 +81,6 @@ class UMaterialFunction : public UMaterialFunctionInterface
 #endif
 
 #if WITH_EDITORONLY_DATA
-	/** Array of comments associated with this material; viewed in the material editor. */
-	UPROPERTY()
-	TArray<TObjectPtr<class UMaterialExpressionComment>> FunctionEditorComments;
-
 	UPROPERTY(transient)
 	TObjectPtr<UMaterial> PreviewMaterial;
 
@@ -167,17 +166,18 @@ public:
 	virtual bool HasFlippedCoordinates() const override;
 #endif
 
-	virtual UMaterialFunctionInterface* GetBaseFunction() override { return this; }
-	virtual const UMaterialFunctionInterface* GetBaseFunction() const override { return this; }
+	virtual UMaterialFunction* GetBaseFunction(FMFRecursionGuard RecursionGuard = FMFRecursionGuard()) override { return this; }
+	virtual const UMaterialFunction* GetBaseFunction(FMFRecursionGuard RecursionGuard = FMFRecursionGuard()) const override { return this; }
 #if WITH_EDITORONLY_DATA
-	virtual const TArray<TObjectPtr<UMaterialExpression>>* GetFunctionExpressions() const override { return &FunctionExpressions; }
-#endif
-	virtual const FString* GetDescription() const override { return &Description; }
-
-#if WITH_EDITOR
-	virtual bool GetReentrantFlag() const override { return bReentrantFlag; }
-	virtual void SetReentrantFlag(const bool bIsReentrant) override { bReentrantFlag = bIsReentrant; }
-#endif
+	ENGINE_API TConstArrayView<TObjectPtr<UMaterialExpression>> GetExpressions() const;
+	ENGINE_API TConstArrayView<TObjectPtr<UMaterialExpressionComment>> GetEditorComments() const;
+	ENGINE_API UMaterialExpressionExecBegin* GetExpressionExecBegin() const;
+	ENGINE_API UMaterialExpressionExecEnd* GetExpressionExecEnd() const;
+
+	ENGINE_API const FMaterialExpressionCollection& GetExpressionCollection() const;
+	ENGINE_API FMaterialExpressionCollection& GetExpressionCollection();
+	ENGINE_API void AssignExpressionCollection(const FMaterialExpressionCollection& InCollection);
+#endif // WITH_EDITORONLY_DATA
 	//~ End UMaterialFunctionInterface interface
 
 
@@ -191,9 +191,27 @@ public:
 	ENGINE_API bool SetStaticComponentMaskParameterValueEditorOnly(FName ParameterName, bool R, bool G, bool B, bool A, FGuid OutExpressionGuid);
 	ENGINE_API bool SetStaticSwitchParameterValueEditorOnly(FName ParameterName, bool OutValue, FGuid OutExpressionGuid);
 
+	inline bool GetReentrantFlag() const { return bReentrantFlag; }
+	inline void SetReentrantFlag(bool bIsReentrant) { bReentrantFlag = bIsReentrant; }
+
 	virtual bool IsUsingControlFlow() const override;
 	virtual bool IsUsingNewHLSLGenerator() const override;
 
 	void CreateExecutionFlowExpressions();
 #endif // WITH_EDITOR
+
+private:
+#if WITH_EDITORONLY_DATA
+	UPROPERTY()
+	TArray<TObjectPtr<UMaterialExpression>> FunctionExpressions_DEPRECATED;
+
+	UPROPERTY()
+	TArray<TObjectPtr<class UMaterialExpressionComment>> FunctionEditorComments_DEPRECATED;
+
+	UPROPERTY()
+	TObjectPtr<class UMaterialExpressionExecBegin> ExpressionExecBegin_DEPRECATED;
+
+	UPROPERTY()
+	TObjectPtr<class UMaterialExpressionExecEnd> ExpressionExecEnd_DEPRECATED;
+#endif // WITH_EDITORONLY_DATA
 };
diff --git a/Engine/Source/Runtime/Engine/Classes/Materials/MaterialFunctionInstance.h b/Engine/Source/Runtime/Engine/Classes/Materials/MaterialFunctionInstance.h
--- a/Engine/Source/Runtime/Engine/Classes/Materials/MaterialFunctionInstance.h
+++ b/Engine/Source/Runtime/Engine/Classes/Materials/MaterialFunctionInstance.h
@@ -101,80 +101,11 @@ class UMaterialFunctionInstance : public UMaterialFunctionInterface
 	virtual bool HasFlippedCoordinates() const override;
 #endif
 
-	virtual UMaterialFunctionInterface* GetBaseFunction() override
-	{
-		UMaterialFunctionInterface* BasePtr = nullptr;
-		UMaterialFunctionInterface* BaseParent = Parent;
-
-		while (true)
-		{
-			BasePtr = BaseParent;
-			if (UMaterialFunctionInstance* BaseInstance = Cast<UMaterialFunctionInstance>(BasePtr))
-			{
-				BaseParent = BaseInstance->Parent;
-			}
-			else
-			{
-				break;
-			}
-		}
-
-		return BasePtr;
-	}
-
-	virtual const UMaterialFunctionInterface* GetBaseFunction() const override
-	{
-		UMaterialFunctionInterface* BasePtr = nullptr;
-		UMaterialFunctionInterface* BaseParent = Parent;
-
-		while (true)
-		{
-			BasePtr = BaseParent;
-			if (UMaterialFunctionInstance* BaseInstance = Cast<UMaterialFunctionInstance>(BasePtr))
-			{
-				BaseParent = BaseInstance->Parent;
-			}
-			else
-			{
-				break;
-			}
-		}
-
-		return BasePtr;
-	}
-
-#if WITH_EDITORONLY_DATA
-	virtual const TArray<TObjectPtr<UMaterialExpression>>* GetFunctionExpressions() const override
-	{
-		const UMaterialFunctionInterface* BaseFunction = GetBaseFunction();
-		return BaseFunction ? BaseFunction->GetFunctionExpressions() : nullptr;
-	}
-#endif // WITH_EDITORONLY_DATA
-
-	virtual const FString* GetDescription() const override
-	{
-		const UMaterialFunctionInterface* BaseFunction = GetBaseFunction();
-		return BaseFunction ? BaseFunction->GetDescription() : nullptr;
-	}
+	virtual UMaterialFunction* GetBaseFunction(FMFRecursionGuard RecursionGuard = FMFRecursionGuard()) override;
+	virtual const UMaterialFunction* GetBaseFunction(FMFRecursionGuard RecursionGuard = FMFRecursionGuard()) const override;
 
 #if WITH_EDITOR
-	virtual bool GetReentrantFlag() const override
-	{
-		const UMaterialFunctionInterface* BaseFunction = GetBaseFunction();
-		return BaseFunction ? BaseFunction->GetReentrantFlag() : false;
-	}
-
-	virtual void SetReentrantFlag(const bool bIsReentrant) override
-	{
-		if (UMaterialFunctionInterface* BaseFunction = GetBaseFunction())
-		{
-			BaseFunction->SetReentrantFlag(bIsReentrant);
-		}
-	}
-
-public:
 	virtual bool GetParameterOverrideValue(EMaterialParameterType Type, const FName& ParameterName, FMaterialParameterMetadata& OutValue, FMFRecursionGuard RecursionGuard = FMFRecursionGuard()) const override;
-	
 #endif // WITH_EDITOR
 	//~ End UMaterialFunctionInterface interface
 
diff --git a/Engine/Source/Runtime/Engine/Classes/Materials/MaterialFunctionInterface.h b/Engine/Source/Runtime/Engine/Classes/Materials/MaterialFunctionInterface.h
--- a/Engine/Source/Runtime/Engine/Classes/Materials/MaterialFunctionInterface.h
+++ b/Engine/Source/Runtime/Engine/Classes/Materials/MaterialFunctionInterface.h
@@ -15,6 +15,7 @@
 #include "MaterialFunctionInterface.generated.h"
 
 class UMaterial;
+class UMaterialFunction;
 class UTexture;
 struct FPropertyChangedEvent;
 class FMaterialHLSLGenerator;
@@ -30,6 +31,14 @@ enum class EMaterialFunctionUsage : uint8
 
 using FMFRecursionGuard = TMaterialRecursionGuard<class UMaterialFunctionInterface>;
 
+UCLASS(Optional)
+class UMaterialFunctionInterfaceEditorOnlyData : public UObject
+{
+	GENERATED_BODY()
+public:
+	
+};
+
 /**
  * A Material Function is a collection of material expressions that can be reused in different materials
  */
@@ -38,6 +47,19 @@ class UMaterialFunctionInterface : public UObject
 {
 	GENERATED_UCLASS_BODY()
 
+#if WITH_EDITORONLY_DATA
+protected:
+	UPROPERTY()
+	TObjectPtr<UMaterialFunctionInterfaceEditorOnlyData> EditorOnlyData;
+
+	ENGINE_API virtual const UClass* GetEditorOnlyDataClass() const;
+	ENGINE_API virtual UMaterialFunctionInterfaceEditorOnlyData* CreateEditorOnlyData();
+
+public:
+	UMaterialFunctionInterfaceEditorOnlyData* GetEditorOnlyData() { return EditorOnlyData; }
+	const UMaterialFunctionInterfaceEditorOnlyData* GetEditorOnlyData() const { return EditorOnlyData; }
+#endif // WITH_EDITORONLY_DATA
+
 	//~ Begin UObject Interface.
 	virtual void PostInitProperties() override;
 	virtual void PostLoad() override;
@@ -128,26 +150,25 @@ public:
 	TObjectPtr<class UThumbnailInfo> ThumbnailInfo;
 #endif
 	
-	virtual UMaterialFunctionInterface* GetBaseFunction()
-		PURE_VIRTUAL(UMaterialFunctionInterface::GetBaseFunction,return nullptr;);
+	virtual UMaterialFunction* GetBaseFunction(FMFRecursionGuard RecursionGuard = FMFRecursionGuard())
+		PURE_VIRTUAL(UMaterialFunction::GetBaseFunction,return nullptr;);
 
-	virtual const UMaterialFunctionInterface* GetBaseFunction() const
-		PURE_VIRTUAL(UMaterialFunctionInterface::GetBaseFunction,return nullptr;);
+	virtual const UMaterialFunction* GetBaseFunction(FMFRecursionGuard RecursionGuard = FMFRecursionGuard()) const
+		PURE_VIRTUAL(UMaterialFunction::GetBaseFunction,return nullptr;);
 
-#if WITH_EDITORONLY_DATA
-	virtual const TArray<TObjectPtr<UMaterialExpression>>* GetFunctionExpressions() const
-		PURE_VIRTUAL(UMaterialFunctionInterface::GetFunctionExpressions,return nullptr;);
-#endif // WITH_EDITORONLY_DATA
-
-	virtual const FString* GetDescription() const
-		PURE_VIRTUAL(UMaterialFunctionInterface::GetDescription,return nullptr;);
+	/** Returns GetBaseFunction() as a UMaterialFunctionInterface, useful if MaterialFunction.h hasn't been included yet, and implicit conversion to UMaterialInterface isn't availiable */
+	ENGINE_API UMaterialFunctionInterface* GetBaseFunctionInterface();
+	ENGINE_API const UMaterialFunctionInterface* GetBaseFunctionInterface() const;
 
 #if WITH_EDITOR
-	virtual bool GetReentrantFlag() const
-		PURE_VIRTUAL(UMaterialFunctionInterface::GetReentrantFlag,return false;);
+	ENGINE_API TConstArrayView<TObjectPtr<UMaterialExpression>> GetExpressions() const;
+
+	UE_DEPRECATED(5.1, "Use GetExpressions()")
+	inline TConstArrayView<TObjectPtr<UMaterialExpression>> GetFunctionExpressions() const { return GetExpressions(); }
 
-	virtual void SetReentrantFlag(const bool bIsReentrant)
-		PURE_VIRTUAL(UMaterialFunctionInterface::SetReentrantFlag,);
+	ENGINE_API const FString& GetDescription() const;
+	ENGINE_API bool GetReentrantFlag() const;
+	ENGINE_API void SetReentrantFlag(bool bIsReentrant);
 #endif // WITH_EDITOR
 
 public:
@@ -157,10 +178,10 @@ public:
 	UE_DEPRECATED(5.0, "Use GetAllParameterInfoOfType or GetAllParametersOfType")
 	void GetAllParameterInfo(TArray<FMaterialParameterInfo>& OutParameterInfo, TArray<FGuid>& OutParameterIds, const FMaterialParameterInfo& InBaseParameterInfo) const
 	{
-		if (const UMaterialFunctionInterface* ParameterFunction = GetBaseFunction())
+		if (const UMaterialFunctionInterface* ParameterFunction = GetBaseFunctionInterface())
 		{
 			const UClass* TargetClass = UMaterialExpressionMaterialFunctionCall::StaticClass();
-			for (const TObjectPtr<UMaterialExpression>& Expression : *ParameterFunction->GetFunctionExpressions())
+			for (const TObjectPtr<UMaterialExpression>& Expression : ParameterFunction->GetExpressions())
 			{
 				if (const UMaterialExpressionMaterialFunctionCall* FunctionExpression = (Expression && Expression.IsA(TargetClass)) ? (const UMaterialExpressionMaterialFunctionCall*)Expression.Get() : nullptr)
 				{
@@ -191,14 +212,14 @@ public:
 	{
 		Parameter = nullptr;
 
-		if (UMaterialFunctionInterface* ParameterFunction = GetBaseFunction())
+		if (UMaterialFunctionInterface* ParameterFunction = GetBaseFunctionInterface())
 		{
 			const UClass* TargetClass = ExpressionType::StaticClass();
 
 			auto GetExpressionParameterByNamePredicate = 
 				[&ParameterInfo, &Parameter, &OwningFunction, TargetClass](UMaterialFunctionInterface* Function) -> bool
 			{
-				for (UMaterialExpression* FunctionExpression : *Function->GetFunctionExpressions())
+				for (UMaterialExpression* FunctionExpression : Function->GetExpressions())
 				{
 					if (ExpressionType* ExpressionParameter = (FunctionExpression && FunctionExpression->IsA(TargetClass)) ? (ExpressionType *)FunctionExpression : nullptr)
 					{
@@ -235,7 +256,7 @@ public:
 	{
 		bool bChanged = false;
 
-		if (UMaterialFunctionInterface* ParameterFunction = GetBaseFunction())
+		if (UMaterialFunctionInterface* ParameterFunction = GetBaseFunctionInterface())
 		{
 			TArray<UMaterialFunctionInterface*> Functions;
 			ParameterFunction->GetDependentFunctions(Functions);
@@ -243,7 +264,7 @@ public:
 
 			for (UMaterialFunctionInterface* Function : Functions)
 			{
-				for (TObjectPtr<UMaterialExpression> FunctionExpression : *Function->GetFunctionExpressions())
+				for (TObjectPtr<UMaterialExpression> FunctionExpression : Function->GetExpressions())
 				{
 					if (ExpressionType* ParameterExpression = Cast<ExpressionType>(FunctionExpression))
 					{
@@ -267,7 +288,7 @@ public:
 	template<typename ExpressionType>
 	bool HasAnyExpressionsOfType()
 	{
-		if (UMaterialFunctionInterface* ParameterFunction = GetBaseFunction())
+		if (UMaterialFunctionInterface* ParameterFunction = GetBaseFunctionInterface())
 		{
 			TArray<UMaterialFunctionInterface*> Functions;
 			ParameterFunction->GetDependentFunctions(Functions);
@@ -275,7 +296,7 @@ public:
 
 			for (UMaterialFunctionInterface* Function : Functions)
 			{
-				for (UMaterialExpression* FunctionExpression : *Function->GetFunctionExpressions())
+				for (UMaterialExpression* FunctionExpression : Function->GetExpressions())
 				{
 					if (ExpressionType* FunctionExpressionOfType = Cast<ExpressionType>(FunctionExpression))
 					{
@@ -292,7 +313,7 @@ public:
 	template<typename ExpressionType>
 	void GetAllExpressionsOfType(TArray<ExpressionType*>& OutExpressions, const bool bRecursive = true)
 	{
-		if (UMaterialFunctionInterface* ParameterFunction = GetBaseFunction())
+		if (UMaterialFunctionInterface* ParameterFunction = GetBaseFunctionInterface())
 		{
 			TArray<UMaterialFunctionInterface*> Functions;
 			if (bRecursive)
@@ -303,7 +324,7 @@ public:
 
 			for (UMaterialFunctionInterface* Function : Functions)
 			{
-				for (UMaterialExpression* FunctionExpression : *Function->GetFunctionExpressions())
+				for (UMaterialExpression* FunctionExpression : Function->GetExpressions())
 				{
 					if (ExpressionType* FunctionExpressionOfType = Cast<ExpressionType>(FunctionExpression))
 					{
diff --git a/Engine/Source/Runtime/Engine/Classes/Materials/MaterialInstance.h b/Engine/Source/Runtime/Engine/Classes/Materials/MaterialInstance.h
--- a/Engine/Source/Runtime/Engine/Classes/Materials/MaterialInstance.h
+++ b/Engine/Source/Runtime/Engine/Classes/Materials/MaterialInstance.h
@@ -454,11 +454,27 @@ private:
 };
 #endif // WITH_EDITORONLY_DATA
 
+UCLASS(MinimalAPI, Optional)
+class UMaterialInstanceEditorOnlyData : public UMaterialInterfaceEditorOnlyData
+{
+	GENERATED_BODY()
+public:
+	UPROPERTY()
+	FStaticParameterSetEditorOnlyData StaticParameters;
+};
+
 UCLASS(abstract, BlueprintType,MinimalAPI)
 class UMaterialInstance : public UMaterialInterface
 {
 	GENERATED_UCLASS_BODY()
 
+#if WITH_EDITORONLY_DATA
+	ENGINE_API virtual const UClass* GetEditorOnlyDataClass() const override { return UMaterialInstanceEditorOnlyData::StaticClass(); }
+
+	UMaterialInstanceEditorOnlyData* GetEditorOnlyData() { return CastChecked<UMaterialInstanceEditorOnlyData>(Super::GetEditorOnlyData()); }
+	const UMaterialInstanceEditorOnlyData* GetEditorOnlyData() const { return CastChecked<UMaterialInstanceEditorOnlyData>(Super::GetEditorOnlyData()); }
+#endif // WITH_EDITORONLY_DATA
+
 	/** Physical material to use for this graphics material. Used for sounds, effects etc.*/
 	UPROPERTY(EditAnywhere, BlueprintReadWrite, Category=MaterialInstance)
 	TObjectPtr<class UPhysicalMaterial> PhysMaterial;
@@ -497,6 +513,8 @@ class UMaterialInstance : public UMaterialInterface
 	// An array of custom parameter set updaters.
 	UE_DEPRECATED(5.0, "Custom static parameter delegates no longer supported.")
 	ENGINE_API static TArray<FCustomParameterSetUpdaterDelegate> CustomParameterSetUpdaters;
+
+	ENGINE_API const FStaticParameterSetEditorOnlyData& GetEditorOnlyStaticParameters() const;
 #endif // WITH_EDITORONLY_DATA
 
 	/**
@@ -504,7 +522,8 @@ class UMaterialInstance : public UMaterialInterface
 	 *
 	 * @returns Static parameter set.
 	 */
-	ENGINE_API const FStaticParameterSet& GetStaticParameters() const;
+	ENGINE_API bool HasStaticParameters() const;
+	ENGINE_API FStaticParameterSet GetStaticParameters() const;
 
 	const FMaterialInstanceCachedData& GetCachedInstanceData() const { return CachedData ? *CachedData : FMaterialInstanceCachedData::EmptyData; }
 
@@ -610,19 +629,20 @@ public:
 
 private:
 
+	/** Static parameter values that are overridden in this instance. */
+	UPROPERTY()
+	FStaticParameterSetRuntimeData StaticParametersRuntime;
+
 #if WITH_EDITORONLY_DATA
 	UPROPERTY()
 	TArray<FGuid> ReferencedTextureGuids;
-#endif // WITH_EDITORONLY_DATA
 
-	/** Static parameter values that are overridden in this instance. */
 	UPROPERTY()
-	FStaticParameterSet StaticParameters;
+	FStaticParameterSet StaticParameters_DEPRECATED;
 
-#if WITH_EDITORONLY_DATA
 	UPROPERTY()
 	bool bSavedCachedData_DEPRECATED;
-#endif
+#endif // WITH_EDITORONLY_DATA
 
 protected:
 	TUniquePtr<FMaterialInstanceCachedData> CachedData;
@@ -964,7 +984,7 @@ namespace MaterialInstance_Private
 {
 	/** Workaround - Similar to base call but evaluates all expressions found, not just the first */
 	template<typename ExpressionType>
-	void FindClosestExpressionByGUIDRecursive(const FName& InName, const FGuid& InGUID, const TArray<TObjectPtr<UMaterialExpression>>& InMaterialExpression, ExpressionType*& OutExpression)
+	void FindClosestExpressionByGUIDRecursive(const FName& InName, const FGuid& InGUID, TConstArrayView<TObjectPtr<UMaterialExpression>> InMaterialExpression, ExpressionType*& OutExpression)
 	{
 		for (int32 ExpressionIndex = 0; ExpressionIndex < InMaterialExpression.Num(); ExpressionIndex++)
 		{
@@ -994,10 +1014,7 @@ namespace MaterialInstance_Private
 			}
 			else if (MaterialFunctionCall && MaterialFunctionCall->MaterialFunction)
 			{
-				if (const TArray<TObjectPtr<UMaterialExpression>>* FunctionExpressions = MaterialFunctionCall->MaterialFunction->GetFunctionExpressions())
-				{
-					FindClosestExpressionByGUIDRecursive<ExpressionType>(InName, InGUID, *FunctionExpressions, OutExpression);
-				}
+				FindClosestExpressionByGUIDRecursive<ExpressionType>(InName, InGUID, MaterialFunctionCall->MaterialFunction->GetExpressions(), OutExpression);
 			}
 			else if (MaterialLayers)
 			{
@@ -1008,10 +1025,7 @@ namespace MaterialInstance_Private
 				{
 					if (Layer)
 					{
-						if (const TArray<TObjectPtr<UMaterialExpression>>* FunctionExpressions = Layer->GetFunctionExpressions())
-						{
-							FindClosestExpressionByGUIDRecursive<ExpressionType>(InName, InGUID, *FunctionExpressions, OutExpression);
-						}
+						FindClosestExpressionByGUIDRecursive<ExpressionType>(InName, InGUID, Layer->GetExpressions(), OutExpression);
 					}
 				}
 
@@ -1019,10 +1033,7 @@ namespace MaterialInstance_Private
 				{
 					if (Blend)
 					{
-						if (const TArray<TObjectPtr<UMaterialExpression>>* FunctionExpressions = Blend->GetFunctionExpressions())
-						{
-							FindClosestExpressionByGUIDRecursive<ExpressionType>(InName, InGUID, *FunctionExpressions, OutExpression);
-						}
+						FindClosestExpressionByGUIDRecursive<ExpressionType>(InName, InGUID, Blend->GetExpressions(), OutExpression);
 					}
 				}
 			}
@@ -1032,7 +1043,7 @@ namespace MaterialInstance_Private
 	template <typename ParameterType, typename ExpressionType>
 	bool UpdateParameter_FullTraversal(ParameterType& Parameter, UMaterial* ParentMaterial)
 	{
-		for (UMaterialExpression* Expression : ParentMaterial->Expressions)
+		for (UMaterialExpression* Expression : ParentMaterial->GetExpressions())
 		{
 			if (Expression->IsA<ExpressionType>())
 			{
@@ -1091,7 +1102,7 @@ namespace MaterialInstance_Private
 			if (Parameter.ExpressionGUID.IsValid())
 			{
 				ExpressionType* Expression = nullptr;
-				FindClosestExpressionByGUIDRecursive<ExpressionType>(Parameter.ParameterInfo.Name, Parameter.ExpressionGUID, ParentMaterial->Expressions, Expression);
+				FindClosestExpressionByGUIDRecursive<ExpressionType>(Parameter.ParameterInfo.Name, Parameter.ExpressionGUID, ParentMaterial->GetExpressions(), Expression);
 
 				// Check to see if the parameter name was changed.
 				if (Expression)
diff --git a/Engine/Source/Runtime/Engine/Classes/Materials/MaterialInterface.h b/Engine/Source/Runtime/Engine/Classes/Materials/MaterialInterface.h
--- a/Engine/Source/Runtime/Engine/Classes/Materials/MaterialInterface.h
+++ b/Engine/Source/Runtime/Engine/Classes/Materials/MaterialInterface.h
@@ -10,6 +10,7 @@
 #include "UObject/Object.h"
 #include "Misc/Guid.h"
 #include "Templates/UniquePtr.h"
+#include "Templates/SharedPointer.h"
 #include "Engine/EngineTypes.h"
 #include "UObject/SoftObjectPath.h"
 #include "UObject/ScriptMacros.h"
@@ -39,6 +40,8 @@ class UMaterialInstance;
 struct FMaterialParameterInfo;
 struct FMaterialResourceLocOnDisk;
 class FMaterialCachedData;
+struct FMaterialCachedExpressionData;
+struct FMaterialCachedExpressionEditorOnlyData;
 #if WITH_EDITOR
 class FMaterialCachedHLSLTree;
 #endif
@@ -192,11 +195,44 @@ struct FMaterialInheritanceChain
 	inline const FMaterialCachedExpressionData& GetCachedExpressionData() const { checkSlow(CachedExpressionData); return *CachedExpressionData; }
 };
 
+UCLASS(Optional)
+class UMaterialInterfaceEditorOnlyData : public UObject
+{
+	GENERATED_BODY()
+public:
+	UMaterialInterfaceEditorOnlyData();
+	UMaterialInterfaceEditorOnlyData(FVTableHelper& Helper);
+	ENGINE_API virtual ~UMaterialInterfaceEditorOnlyData();
+
+	//~ Begin UObject Interface.
+	ENGINE_API virtual void Serialize(FArchive& Ar) override;
+	//~ End UObject Interface.
+
+	TSharedPtr<FMaterialCachedExpressionEditorOnlyData> CachedExpressionData;
+
+	/** Set if CachedExpressionData was loaded from disk, should typically be true when running with cooked data, and false in the editor */
+	bool bLoadedCachedExpressionData = false;
+};
+
 UCLASS(abstract, BlueprintType, MinimalAPI, HideCategories = (Thumbnail))
 class UMaterialInterface : public UObject, public IBlendableInterface, public IInterface_AssetUserData
 {
 	GENERATED_UCLASS_BODY()
 
+#if WITH_EDITORONLY_DATA
+protected:
+	UPROPERTY()
+	TObjectPtr<UMaterialInterfaceEditorOnlyData> EditorOnlyData;
+
+	ENGINE_API virtual const UClass* GetEditorOnlyDataClass() const;
+	ENGINE_API UMaterialInterfaceEditorOnlyData* CreateEditorOnlyData();
+
+public:
+	UMaterialInterfaceEditorOnlyData* GetEditorOnlyData() { return EditorOnlyData; }
+	const UMaterialInterfaceEditorOnlyData* GetEditorOnlyData() const { return EditorOnlyData; }
+#endif // WITH_EDITORONLY_DATA
+
+public:
 	/** SubsurfaceProfile, for Screen Space Subsurface Scattering */
 	UPROPERTY(EditAnywhere, BlueprintReadOnly, Category = Material, meta = (DisplayName = "Subsurface Profile"))
 	TObjectPtr<class USubsurfaceProfile> SubsurfaceProfile;
@@ -288,6 +324,7 @@ public:
 	ENGINE_API virtual void BeginDestroy() override;
 	ENGINE_API virtual void FinishDestroy() override;
 	ENGINE_API virtual bool IsReadyForFinishDestroy() override;
+	ENGINE_API virtual void PostInitProperties() override;
 	ENGINE_API virtual void Serialize(FArchive& Ar) override;
 	ENGINE_API virtual void PostLoad() override;
 	ENGINE_API virtual void PostDuplicate(bool bDuplicateForPIE) override;
diff --git a/Engine/Source/Runtime/Engine/Classes/Materials/MaterialLayersFunctions.h b/Engine/Source/Runtime/Engine/Classes/Materials/MaterialLayersFunctions.h
--- a/Engine/Source/Runtime/Engine/Classes/Materials/MaterialLayersFunctions.h
+++ b/Engine/Source/Runtime/Engine/Classes/Materials/MaterialLayersFunctions.h
@@ -10,6 +10,7 @@
 #define LOCTEXT_NAMESPACE "MaterialLayersFunctions"
 
 class FArchive;
+struct FMaterialLayersFunctions;
 
 UENUM()
 enum class EMaterialLayerLinkState : uint8
@@ -20,62 +21,35 @@
 	NotFromParent, // Layer was created locally in this material, not in parent
 };
 
-USTRUCT()
-struct ENGINE_API FMaterialLayersFunctions
+/** Serializable ID structure for FMaterialLayersFunctions which allows us to deterministically recompile shaders*/
+struct FMaterialLayersFunctionsID
 {
-	GENERATED_USTRUCT_BODY()
+	TArray<FGuid> LayerIDs;
+	TArray<FGuid> BlendIDs;
+	TArray<bool> LayerStates;
 
-#if WITH_EDITOR
-	/** Serializable ID structure for FMaterialLayersFunctions which allows us to deterministically recompile shaders*/
-	struct ID
+	bool operator==(const FMaterialLayersFunctionsID& Reference) const;
+	inline bool operator!=(const FMaterialLayersFunctionsID& Reference) const { return !operator==(Reference); }
+
+	void SerializeForDDC(FArchive& Ar);
+
+	friend FMaterialLayersFunctionsID& operator<<(FArchive& Ar, FMaterialLayersFunctionsID& Ref)
 	{
-		TArray<FGuid> LayerIDs;
-		TArray<FGuid> BlendIDs;
-		TArray<bool> LayerStates;
-
-		bool operator==(const ID& Reference) const;
-		inline bool operator!=(const ID& Reference) const { return !operator==(Reference); }
-
-		void SerializeForDDC(FArchive& Ar);
-
-		friend ID& operator<<(FArchive& Ar, ID& Ref)
-		{
-			Ref.SerializeForDDC(Ar);
-			return Ref;
-		}
-
-		void UpdateHash(FSHA1& HashState) const;
-
-		//TODO: Investigate whether this is really required given it is only used by FMaterialShaderMapId AND that one also uses UpdateHash
-		void AppendKeyString(FString& KeyString) const;
-	};
-#endif // WITH_EDITOR
-
-	static const FGuid BackgroundGuid;
-		
-	FMaterialLayersFunctions() = default;
-
-	void Empty()
-	{
-		Layers.Empty();
-		Blends.Empty();
-#if WITH_EDITOR
-		LayerStates.Empty();
-		LayerNames.Empty();
-		RestrictToLayerRelatives.Empty();
-		RestrictToBlendRelatives.Empty();
-		LayerGuids.Empty();
-		LayerLinkStates.Empty();
-#endif
+		Ref.SerializeForDDC(Ar);
+		return Ref;
 	}
 
-	UPROPERTY(EditAnywhere, Category=MaterialLayers)
-	TArray<TObjectPtr<class UMaterialFunctionInterface>> Layers;
+	void UpdateHash(FSHA1& HashState) const;
 
-	UPROPERTY(EditAnywhere, Category=MaterialLayers)
-	TArray<TObjectPtr<class UMaterialFunctionInterface>> Blends;
+	//TODO: Investigate whether this is really required given it is only used by FMaterialShaderMapId AND that one also uses UpdateHash
+	void AppendKeyString(FString& KeyString) const;
+};
 
-#if WITH_EDITORONLY_DATA
+USTRUCT()
+struct FMaterialLayersFunctionsEditorOnlyData
+{
+	GENERATED_BODY()
+
 	UPROPERTY(EditAnywhere, Category = MaterialLayers)
 	TArray<bool> LayerStates;
 
@@ -104,38 +78,163 @@
 	 */
 	UPROPERTY(EditAnywhere, Category = MaterialLayers)
 	TArray<FGuid> DeletedParentLayerGuids;
-#endif // WITH_EDITORONLY_DATA
 
+#if WITH_EDITOR
+	FORCEINLINE bool operator==(const FMaterialLayersFunctionsEditorOnlyData& Other) const
+	{
+		if (LayerStates != Other.LayerStates ||
+			LayerLinkStates != Other.LayerLinkStates ||
+			DeletedParentLayerGuids != Other.DeletedParentLayerGuids)
+		{
+			return false;
+		}
+		return true;
+	}
+
+	FORCEINLINE bool operator!=(const FMaterialLayersFunctionsEditorOnlyData& Other) const
+	{
+		return !operator==(Other);
+	}
+
+	void Empty()
+	{
+		LayerStates.Empty();
+		LayerNames.Empty();
+		RestrictToLayerRelatives.Empty();
+		RestrictToBlendRelatives.Empty();
+		LayerGuids.Empty();
+		LayerLinkStates.Empty();
+	}
+
+	void LinkAllLayersToParent();
+#endif // WITH_EDITOR
+};
+
+USTRUCT()
+struct FMaterialLayersFunctionsRuntimeData
+{
+	GENERATED_BODY()
+
+	UPROPERTY(EditAnywhere, Category = MaterialLayers)
+	TArray<TObjectPtr<class UMaterialFunctionInterface>> Layers;
+
+	UPROPERTY(EditAnywhere, Category = MaterialLayers)
+	TArray<TObjectPtr<class UMaterialFunctionInterface>> Blends;
+
+	FMaterialLayersFunctionsRuntimeData() = default;
+	FMaterialLayersFunctionsRuntimeData(const FMaterialLayersFunctionsRuntimeData& Rhs) = default;
+	FMaterialLayersFunctionsRuntimeData(const FMaterialLayersFunctions& Rhs) = delete;
+	FMaterialLayersFunctionsRuntimeData& operator=(const FMaterialLayersFunctionsRuntimeData& Rhs) = default;
+	FMaterialLayersFunctionsRuntimeData& operator=(const FMaterialLayersFunctions& Rhs) = delete;
+
+	void Empty()
+	{
+		Layers.Empty();
+		Blends.Empty();
+	}
+
+	FORCEINLINE bool operator==(const FMaterialLayersFunctionsRuntimeData& Other) const
+	{
+		if (Layers != Other.Layers || Blends != Other.Blends)
+		{
+			return false;
+		}
+		return true;
+	}
+
+	FORCEINLINE bool operator!=(const FMaterialLayersFunctionsRuntimeData& Other) const
+	{
+		return !operator==(Other);
+	}
+
+#if WITH_EDITOR
+	const FMaterialLayersFunctionsID GetID(const FMaterialLayersFunctionsEditorOnlyData& EditorOnly) const;
+#endif // WITH_EDITOR
+};
+
+USTRUCT()
+struct ENGINE_API FMaterialLayersFunctions : public FMaterialLayersFunctionsRuntimeData
+{
+	GENERATED_USTRUCT_BODY()
+
+#if WITH_EDITOR
+	using ID = FMaterialLayersFunctionsID;
+#endif // WITH_EDITOR
+
+	static const FGuid BackgroundGuid;
+		
+	FMaterialLayersFunctions() = default;
+	FMaterialLayersFunctions(const FMaterialLayersFunctions& Rhs) = default;
+	FMaterialLayersFunctions(const FMaterialLayersFunctionsRuntimeData&) = delete;
+
+	FMaterialLayersFunctionsRuntimeData& GetRuntime() { return *this; }
+	const FMaterialLayersFunctionsRuntimeData& GetRuntime() const { return *this; }
+
+#if WITH_EDITORONLY_DATA
+	UPROPERTY(EditAnywhere, Category = MaterialLayers)
+	FMaterialLayersFunctionsEditorOnlyData EditorOnly;
+#endif
+
+	void Empty()
+	{
+		FMaterialLayersFunctionsRuntimeData::Empty();
+#if WITH_EDITOR
+		EditorOnly.Empty();
+#endif
+	}
+
 	inline bool IsEmpty() const { return Layers.Num() == 0; }
 
+#if WITH_EDITOR
 	void AddDefaultBackgroundLayer()
 	{
 		// Default to a non-blended "background" layer
 		Layers.AddDefaulted();
-#if WITH_EDITORONLY_DATA
-		LayerStates.Add(true);
+		EditorOnly.LayerStates.Add(true);
 		FText LayerName = FText(LOCTEXT("Background", "Background"));
-		LayerNames.Add(LayerName);
-		RestrictToLayerRelatives.Add(false);
+		EditorOnly.LayerNames.Add(LayerName);
+		EditorOnly.RestrictToLayerRelatives.Add(false);
 		// Use a consistent Guid for the background layer
 		// Default constructor assigning different guids will break FStructUtils::AttemptToFindUninitializedScriptStructMembers
-		LayerGuids.Add(BackgroundGuid);
-		LayerLinkStates.Add(EMaterialLayerLinkState::NotFromParent);
-#endif // WITH_EDITORONLY_DATA
+		EditorOnly.LayerGuids.Add(BackgroundGuid);
+		EditorOnly.LayerLinkStates.Add(EMaterialLayerLinkState::NotFromParent);
 	}
 
 	int32 AppendBlendedLayer();
 
-	int32 AddLayerCopy(const FMaterialLayersFunctions& Source, int32 SourceLayerIndex, bool bVisible, EMaterialLayerLinkState LinkState);
+	int32 AddLayerCopy(const FMaterialLayersFunctionsRuntimeData& Source,
+		const FMaterialLayersFunctionsEditorOnlyData& SourceEditorOnly,
+		int32 SourceLayerIndex,
+		bool bVisible,
+		EMaterialLayerLinkState LinkState);
 
-	void InsertLayerCopy(const FMaterialLayersFunctions& Source, int32 SourceLayerIndex, EMaterialLayerLinkState LinkState, int32 LayerIndex);
+	int32 AddLayerCopy(const FMaterialLayersFunctions& Source,
+		int32 SourceLayerIndex,
+		bool bVisible,
+		EMaterialLayerLinkState LinkState)
+	{
+		return AddLayerCopy(Source, Source.EditorOnly, SourceLayerIndex, bVisible, LinkState);
+	}
 
+	void InsertLayerCopy(const FMaterialLayersFunctionsRuntimeData& Source,
+		const FMaterialLayersFunctionsEditorOnlyData& SourceEditorOnly,
+		int32 SourceLayerIndex,
+		EMaterialLayerLinkState LinkState,
+		int32 LayerIndex);
+
+	void InsertLayerCopy(const FMaterialLayersFunctions& Source,
+		int32 SourceLayerIndex,
+		EMaterialLayerLinkState LinkState,
+		int32 LayerIndex)
+	{
+		return InsertLayerCopy(Source, Source.EditorOnly, SourceLayerIndex, LinkState, LayerIndex);
+	}
+
 	void RemoveBlendedLayerAt(int32 Index);
 
 	void MoveBlendedLayer(int32 SrcLayerIndex, int32 DstLayerIndex);
 
-#if WITH_EDITOR
-	const ID GetID() const;
+	const ID GetID() const { return FMaterialLayersFunctionsRuntimeData::GetID(EditorOnly); }
 
 	/** Gets a string representation of the ID */
 	FString GetStaticPermutationString() const;
@@ -147,37 +246,64 @@
 
 	void ToggleBlendedLayerVisibility(int32 Index)
 	{
-		check(LayerStates.IsValidIndex(Index));
-		LayerStates[Index] = !LayerStates[Index];
+		check(EditorOnly.LayerStates.IsValidIndex(Index));
+		EditorOnly.LayerStates[Index] = !EditorOnly.LayerStates[Index];
 	}
 
 	void SetBlendedLayerVisibility(int32 Index, bool InNewVisibility)
 	{
-		check(LayerStates.IsValidIndex(Index));
-		LayerStates[Index] = InNewVisibility;
+		check(EditorOnly.LayerStates.IsValidIndex(Index));
+		EditorOnly.LayerStates[Index] = InNewVisibility;
 	}
 
 	bool GetLayerVisibility(int32 Index) const
 	{
-		check(LayerStates.IsValidIndex(Index));
-		return LayerStates[Index];
+		check(EditorOnly.LayerStates.IsValidIndex(Index));
+		return EditorOnly.LayerStates[Index];
 	}
 
 	FText GetLayerName(int32 Counter) const
 	{
 		FText LayerName = FText::Format(LOCTEXT("LayerPrefix", "Layer {0}"), Counter);
-		if (LayerNames.IsValidIndex(Counter))
+		if (EditorOnly.LayerNames.IsValidIndex(Counter))
 		{
-			LayerName = LayerNames[Counter];
+			LayerName = EditorOnly.LayerNames[Counter];
 		}
 		return LayerName;
 	}
 
-	bool MatchesParent(const FMaterialLayersFunctions& Parent) const;
+	void LinkAllLayersToParent()
+	{
+		EditorOnly.LinkAllLayersToParent();
+	}
 
-	void LinkAllLayersToParent();
+	static bool MatchesParent(const FMaterialLayersFunctionsRuntimeData& Runtime,
+		const FMaterialLayersFunctionsEditorOnlyData& EditorOnly,
+		const FMaterialLayersFunctionsRuntimeData& ParentRuntime,
+		const FMaterialLayersFunctionsEditorOnlyData& ParentEditorOnly);
 
-	bool ResolveParent(const FMaterialLayersFunctions& Parent, TArray<int32>& OutRemapLayerIndices);
+	bool MatchesParent(const FMaterialLayersFunctions& Parent) const
+	{
+		return MatchesParent(GetRuntime(), EditorOnly, Parent.GetRuntime(), Parent.EditorOnly);
+	}
 
+	static bool ResolveParent(const FMaterialLayersFunctionsRuntimeData& ParentRuntime,
+		const FMaterialLayersFunctionsEditorOnlyData& ParentEditorOnly,
+		FMaterialLayersFunctionsRuntimeData& Runtime,
+		FMaterialLayersFunctionsEditorOnlyData& EditorOnly,
+		TArray<int32>& OutRemapLayerIndices);
+
+	bool ResolveParent(const FMaterialLayersFunctions& Parent, TArray<int32>& OutRemapLayerIndices)
+	{
+		return FMaterialLayersFunctions::ResolveParent(Parent.GetRuntime(), Parent.EditorOnly, GetRuntime(), EditorOnly, OutRemapLayerIndices);
+	}
+
+	static void Validate(const FMaterialLayersFunctionsRuntimeData& Runtime, const FMaterialLayersFunctionsEditorOnlyData& EditorOnly);
+
+	void Validate()
+	{
+		Validate(GetRuntime(), EditorOnly);
+	}
+
 	void SerializeLegacy(FArchive& Ar);
 #endif // WITH_EDITOR
@@ -186,12 +312,12 @@
 
 	FORCEINLINE bool operator==(const FMaterialLayersFunctions& Other) const
 	{
-		if (Layers != Other.Layers || Blends != Other.Blends)
+		if (!FMaterialLayersFunctionsRuntimeData::operator==(Other))
 		{
 			return false;
 		}
 #if WITH_EDITORONLY_DATA
-		if (LayerStates != Other.LayerStates || LayerLinkStates != Other.LayerLinkStates || DeletedParentLayerGuids != Other.DeletedParentLayerGuids)
+		if (EditorOnly != Other.EditorOnly)
 		{
 			return false;
 		}
@@ -203,8 +329,38 @@
 	{
 		return !operator==(Other);
 	}
+
+#if WITH_EDITORONLY_DATA
+private:
+	UPROPERTY()
+	TArray<bool> LayerStates_DEPRECATED;
+
+	UPROPERTY()
+	TArray<FText> LayerNames_DEPRECATED;
+
+	UPROPERTY()
+	TArray<bool> RestrictToLayerRelatives_DEPRECATED;
+
+	UPROPERTY()
+	TArray<bool> RestrictToBlendRelatives_DEPRECATED;
+
+	UPROPERTY()
+	TArray<FGuid> LayerGuids_DEPRECATED;
+
+	UPROPERTY()
+	TArray<EMaterialLayerLinkState> LayerLinkStates_DEPRECATED;
+
+	UPROPERTY()
+	TArray<FGuid> DeletedParentLayerGuids_DEPRECATED;
+#endif // WITH_EDITORONLY_DATA
 };
 
+// Don't allowing comparing a full FMaterialLayersFunctions against partial RuntimeData
+bool operator==(const FMaterialLayersFunctions&, const FMaterialLayersFunctionsRuntimeData&) = delete;
+bool operator==(const FMaterialLayersFunctionsRuntimeData&, const FMaterialLayersFunctions&) = delete;
+bool operator!=(const FMaterialLayersFunctions&, const FMaterialLayersFunctionsRuntimeData&) = delete;
+bool operator!=(const FMaterialLayersFunctionsRuntimeData&, const FMaterialLayersFunctions&) = delete;
+
 template<>
 struct TStructOpsTypeTraits<FMaterialLayersFunctions> : TStructOpsTypeTraitsBase2<FMaterialLayersFunctions>
 {
